Initial drop of Horde Server/Agent for UE5/Main. Still missing some Epic-specific functionality required to build/run correctly.

diff --git a/Engine/Source/Programs/Horde/HordeAgent/AgentSettings.cs b/Engine/Source/Programs/Horde/HordeAgent/AgentSettings.cs
new file mode 100644
--- /dev/null
+++ b/Engine/Source/Programs/Horde/HordeAgent/AgentSettings.cs
@@ -0,0 +1,227 @@
+// Copyright Epic Games, Inc. All Rights Reserved.
+
+using HordeCommon.Rpc;
+using System;
+using System.Collections.Generic;
+using System.ComponentModel.DataAnnotations;
+using System.IO;
+using System.Linq;
+using System.Reflection;
+using System.Text;
+using System.Text.Json;
+using System.Threading.Tasks;
+
+namespace HordeAgent
+{
+	/// <summary>
+	/// The type of executor to use
+	/// </summary>
+	public enum ExecutorType
+	{
+		/// <summary>
+		/// The test executor (predefined job, 
+		/// </summary>
+		Test,
+
+		/// <summary>
+		/// The local executor (run all steps locally)
+		/// </summary>
+		Local,
+
+		/// <summary>
+		/// Manage workspaces and run steps in a synced branch
+		/// </summary>
+		Perforce,
+	}
+
+	/// <summary>
+	/// Describes a network share to mount
+	/// </summary>
+	public class MountNetworkShare
+	{
+		/// <summary>
+		/// Where the share should be mounted on the local machine. Must be a drive letter for Windows.
+		/// </summary>
+		public string? MountPoint { get; set; }
+
+		/// <summary>
+		/// Path to the remote resource
+		/// </summary>
+		public string? RemotePath { get; set; }
+	}
+
+	/// <summary>
+	/// Information about a server to use
+	/// </summary>
+	public class ServerProfile
+	{
+		/// <summary>
+		/// Name of this server profile
+		/// </summary>
+		[Required]
+		public string Name { get; set; } = null!;
+
+		/// <summary>
+		/// Name of the environment (currrently just used for tracing)
+		/// </summary>
+		[Required]
+		public string Environment { get; set; } = "prod";
+
+		/// <summary>
+		/// Url of the server
+		/// </summary>
+		[Required]
+		public string Url { get; set; } = null!;
+
+		/// <summary>
+		/// Bearer token to use to initiate the connection
+		/// </summary>
+		[Required]
+		public string Token { get; set; } = null!;
+
+		/// <summary>
+		/// Thumbprint of a certificate to trust. Allows using self-signed certs for the server.
+		/// </summary>
+		public string? Thumbprint { get; set; }
+
+		/// <summary>
+		/// Thumbprints of certificates to trust. Allows using self-signed certs for the server.
+		/// </summary>
+		public List<string> Thumbprints { get; set; } = new List<string>();
+
+		/// <summary>
+		/// Checks whether the given certificate thumbprint should be trusted
+		/// </summary>
+		/// <param name="CertificateThumbprint">The cert thumbprint</param>
+		/// <returns>True if the cert should be trusted</returns>
+		public bool IsTrustedCertificate(string CertificateThumbprint)
+		{
+			if (Thumbprint != null && Thumbprint.Equals(CertificateThumbprint, StringComparison.OrdinalIgnoreCase))
+			{
+				return true;
+			}
+			if (Thumbprints.Any(x => x.Equals(CertificateThumbprint, StringComparison.OrdinalIgnoreCase)))
+			{
+				return true;
+			}
+			return false;
+		}
+	}
+
+	/// <summary>
+	/// Settings for the local executor
+	/// </summary>
+	public class LocalExecutorSettings
+	{
+		/// <summary>
+		/// Path to the local workspace to use with the local executor
+		/// </summary>
+		public string? WorkspaceDir { get; set; }
+
+		/// <summary>
+		/// Whether to actually execute steps, or just do job setup
+		/// </summary>
+		public bool RunSteps { get; set; } = true;
+	}
+
+	/// <summary>
+	/// Setttings for the perforce executor
+	/// </summary>
+	public class PerforceExecutorSettings
+	{
+		/// <summary>
+		/// Whether to run conform jobs
+		/// </summary>
+		public bool RunConform { get; set; } = true;
+	}
+
+	/// <summary>
+	/// Global settings for the agent
+	/// </summary>
+	public class AgentSettings
+	{
+		/// <summary>
+		/// Name of the section containing these settings
+		/// </summary>
+		public const string SectionName = "Horde";
+
+		/// <summary>
+		/// Known servers to connect to
+		/// </summary>
+		public List<ServerProfile> ServerProfiles { get; set; } = new List<ServerProfile>();
+
+		/// <summary>
+		/// The default server, unless overridden from the command line
+		/// </summary>
+		[Required]
+		public string Server { get; set; } = null!;
+
+		/// <summary>
+		/// The unique version identifier for this client. This is inserted into the appsettings.json file when a client is uploaded to the server.
+		/// </summary>
+		public string? Version { get; set; }
+
+		/// <summary>
+		/// The executor to use for jobs. Defaults to the Perforce executor.
+		/// </summary>
+		public ExecutorType Executor { get; set; } = ExecutorType.Perforce;
+
+		/// <summary>
+		/// Settings for the local executor
+		/// </summary>
+		public LocalExecutorSettings LocalExecutor { get; set; } = new LocalExecutorSettings();
+
+		/// <summary>
+		/// Settings for the perforce executor
+		/// </summary>
+		public PerforceExecutorSettings PerforceExecutor { get; set; } = new PerforceExecutorSettings();
+
+		/// <summary>
+		/// Working directory
+		/// </summary>
+		public string? WorkingDir { get; set; }
+
+		/// <summary>
+		/// List of network shares to mount
+		/// </summary>
+		public List<MountNetworkShare> Shares { get; set; } = new List<MountNetworkShare>();
+
+		/// <summary>
+		/// List of process names to terminate after a job
+		/// </summary>
+		public List<string> ProcessNamesToTerminate { get; set; } = new List<string>();
+
+		/// <summary>
+		/// Whether to write step output to the logging device
+		/// </summary>
+		public bool WriteStepOutputToLogger { get; set; }
+
+		/// <summary>
+		/// Gets the current server settings
+		/// </summary>
+		/// <returns>The current server settings</returns>
+		public ServerProfile GetServerProfile(string Name)
+		{
+			ServerProfile? ServerProfile = ServerProfiles.FirstOrDefault(x => x.Name.Equals(Name, StringComparison.OrdinalIgnoreCase));
+			if (ServerProfile == null)
+			{
+				throw new Exception($"Unknown server profile name '{Server}'");
+			}
+			return ServerProfile;
+		}
+
+		/// <summary>
+		/// Gets the current server settings
+		/// </summary>
+		/// <returns>The current server settings</returns>
+		public ServerProfile GetCurrentServerProfile()
+		{
+			if (Server == null)
+			{
+				throw new Exception("Server is not set");
+			}
+
+			return GetServerProfile(Server);
+		}
+	}
+}
diff --git a/Engine/Source/Programs/Horde/HordeAgent/Commands/Command.cs b/Engine/Source/Programs/Horde/HordeAgent/Commands/Command.cs
new file mode 100644
--- /dev/null
+++ b/Engine/Source/Programs/Horde/HordeAgent/Commands/Command.cs
@@ -0,0 +1,44 @@
+// Copyright Epic Games, Inc. All Rights Reserved.
+
+using Microsoft.Extensions.Logging;
+using System;
+using System.Collections.Generic;
+using System.Text;
+using System.Threading.Tasks;
+using EpicGames.Core;
+
+namespace HordeAgent
+{
+	/// <summary>
+	/// Base class for all commands that can be executed by HordeAgent
+	/// </summary>
+	abstract class Command
+	{
+		/// <summary>
+		/// Configure this object with the given command line arguments
+		/// </summary>
+		/// <param name="Arguments">Command line arguments</param>
+		/// <param name="Logger">Logging output device</param>
+		public virtual void Configure(CommandLineArguments Arguments, ILogger Logger)
+		{
+			Arguments.ApplyTo(this, Logger);
+		}
+
+		/// <summary>
+		/// Gets all command line parameters to show in help for this command
+		/// </summary>
+		/// <param name="Arguments">The command line arguments</param>
+		/// <returns>List of name/description pairs</returns>
+		public virtual List<KeyValuePair<string, string>> GetParameters(CommandLineArguments Arguments)
+		{
+			return CommandLineArguments.GetParameters(GetType());
+		}
+
+		/// <summary>
+		/// Execute this command
+		/// </summary>
+		/// <param name="Logger">The logger to use for this command</param>
+		/// <returns>Exit code</returns>
+		public abstract Task<int> ExecuteAsync(ILogger Logger);
+	}
+}
diff --git a/Engine/Source/Programs/Horde/HordeAgent/Commands/CommandAttribute.cs b/Engine/Source/Programs/Horde/HordeAgent/Commands/CommandAttribute.cs
new file mode 100644
--- /dev/null
+++ b/Engine/Source/Programs/Horde/HordeAgent/Commands/CommandAttribute.cs
@@ -0,0 +1,47 @@
+// Copyright Epic Games, Inc. All Rights Reserved.
+
+using System;
+using System.Collections.Generic;
+using System.Text;
+
+namespace HordeAgent.Commands
+{
+	/// <summary>
+	/// Attribute used to specify names of program modes, and help text
+	/// </summary>
+	class CommandAttribute : Attribute
+	{
+		/// <summary>
+		/// Names for this command
+		/// </summary>
+		public string[] Names;
+
+		/// <summary>
+		/// Short description for the mode. Will be displayed in the help text.
+		/// </summary>
+		public string Description;
+
+		/// <summary>
+		/// Constructor
+		/// </summary>
+		/// <param name="Name">Name of the mode</param>
+		/// <param name="Description">Short description for display in the help text</param>
+		public CommandAttribute(string Name, string Description)
+		{
+			this.Names = new string[] { Name };
+			this.Description = Description;
+		}
+
+		/// <summary>
+		/// Constructor
+		/// </summary>
+		/// <param name="Category">Category for this command</param>
+		/// <param name="Name">Name of the mode</param>
+		/// <param name="Description">Short description for display in the help text</param>
+		public CommandAttribute(string Category, string Name, string Description)
+		{
+			this.Names = new string[] { Category, Name };
+			this.Description = Description;
+		}
+	}
+}
diff --git a/Engine/Source/Programs/Horde/HordeAgent/Commands/Execution/ExecuteCommand.cs b/Engine/Source/Programs/Horde/HordeAgent/Commands/Execution/ExecuteCommand.cs
new file mode 100644
--- /dev/null
+++ b/Engine/Source/Programs/Horde/HordeAgent/Commands/Execution/ExecuteCommand.cs
@@ -0,0 +1,310 @@
+using Build.Bazel.Remote.Execution.V2;
+using Google.Protobuf;
+using HordeAgent.Utility;
+using Microsoft.Extensions.Logging;
+using System;
+using System.Collections.Generic;
+using System.IO;
+using System.Text;
+using System.Threading.Tasks;
+
+using Action = Build.Bazel.Remote.Execution.V2.Action;
+using ActionResult = Build.Bazel.Remote.Execution.V2.ActionResult;
+using RpcCommand = Build.Bazel.Remote.Execution.V2.Command;
+using Directory = Build.Bazel.Remote.Execution.V2.Directory;
+using RpcExecution = Build.Bazel.Remote.Execution.V2.Execution;
+using EpicGames.Core;
+using System.Text.Json;
+using HordeAgent.Services;
+using Grpc.Net.Client;
+using System.Linq;
+using Grpc.Core;
+using Google.LongRunning;
+using Microsoft.Extensions.Hosting;
+using Microsoft.Extensions.DependencyInjection;
+using Microsoft.Extensions.Configuration;
+
+namespace HordeAgent.Commands
+{
+	/// <summary>
+	/// Installs the agent as a service
+	/// </summary>
+	[Command("Execute", "Executes a command through the remote execution API")]
+	class ExecuteCommand : Command
+	{
+		/// <summary>
+		/// List of files that need to be uploaded
+		/// </summary>
+		class UploadList
+		{
+			public Dictionary<string, BatchUpdateBlobsRequest.Types.Request> Blobs = new Dictionary<string, BatchUpdateBlobsRequest.Types.Request>();
+
+			public Digest Add<T>(IMessage<T> Message) where T : IMessage<T>
+			{
+				ByteString ByteString = Message.ToByteString();
+
+				Digest Digest = StorageExtensions.GetDigest(ByteString);
+				if (!Blobs.ContainsKey(Digest.Hash))
+				{
+					BatchUpdateBlobsRequest.Types.Request Request = new BatchUpdateBlobsRequest.Types.Request();
+					Request.Data = ByteString;
+					Request.Digest = Digest;
+					Blobs.Add(Digest.Hash, Request);
+				}
+				return Digest;
+			}
+
+			public async Task<Digest> AddFileAsync(string LocalFile)
+			{
+				BatchUpdateBlobsRequest.Types.Request Content = new BatchUpdateBlobsRequest.Types.Request();
+				using (FileStream Stream = File.Open(LocalFile, FileMode.Open, FileAccess.Read))
+				{
+					Content.Data = await ByteString.FromStreamAsync(Stream);
+					Content.Digest = StorageExtensions.GetDigest(Content.Data);
+				}
+				Blobs[Content.Digest.Hash] = Content;
+				return Content.Digest;
+			}
+		}
+
+		class JsonFileNode
+		{
+			public string Name { get; set; } = String.Empty;
+			public string LocalFile { get; set; } = String.Empty;
+
+			public async Task<FileNode> BuildAsync(UploadList UploadList)
+			{
+				FileNode Node = new FileNode();
+				Node.Name = String.IsNullOrEmpty(Name) ? Path.GetFileName(LocalFile) : Name;
+				Node.Digest = await UploadList.AddFileAsync(LocalFile);
+				return Node;
+			}
+		}
+
+		class JsonDirectory
+		{
+			public List<JsonFileNode> Files { get; set; } = new List<JsonFileNode>();
+			public List<JsonDirectoryNode> Directories { get; set; } = new List<JsonDirectoryNode>();
+
+			public async Task<Digest> BuildAsync(UploadList UploadList)
+			{
+				Directory Directory = new Directory();
+				foreach (JsonFileNode File in Files)
+				{
+					Directory.Files.Add(await File.BuildAsync(UploadList));
+				}
+				foreach (JsonDirectoryNode DirectoryNode in Directories)
+				{
+					Directory.Directories.Add(await DirectoryNode.BuildAsync(UploadList));
+				}
+				return UploadList.Add(Directory);
+			}
+		}
+
+		class JsonDirectoryNode : JsonDirectory
+		{
+			public string Name { get; set; } = String.Empty;
+
+			public new async Task<DirectoryNode> BuildAsync(UploadList UploadList)
+			{
+				DirectoryNode Node = new DirectoryNode();
+				Node.Name = Name;
+				Node.Digest = await base.BuildAsync(UploadList);
+				return Node;
+			}
+		}
+
+		class JsonCommand
+		{
+			public List<string> Arguments { get; set; } = new List<string>();
+			public string WorkingDirectory { get; set; } = String.Empty;
+
+			public Digest Build(UploadList UploadList)
+			{
+				RpcCommand Command = new RpcCommand();
+				Command.Arguments.Add(Arguments);
+				Command.WorkingDirectory = WorkingDirectory;
+				return UploadList.Add(Command);
+			}
+		}
+
+		class JsonAction
+		{
+			public JsonCommand Command { get; set; } = new JsonCommand();
+			public JsonDirectory Workspace { get; set; } = new JsonDirectory();
+
+			public async Task<Digest> BuildAsync(UploadList UploadList, ByteString? Salt)
+			{
+				Action NewAction = new Action();
+				NewAction.CommandDigest = Command.Build(UploadList);
+				NewAction.InputRootDigest = await Workspace.BuildAsync(UploadList);
+				NewAction.Salt = Salt;
+				return UploadList.Add(NewAction);
+			}
+		}
+
+		/// <summary>
+		/// Input file describing the work to execute
+		/// </summary>
+		[CommandLine(Required = true)]
+		public FileReference Input = null!;
+
+		/// <summary>
+		/// Apply a random salt to the cached value
+		/// </summary>
+		[CommandLine("-Salt")]
+		public bool RandomSalt;
+
+		/// <summary>
+		/// Add a known salt value to the cache
+		/// </summary>
+		[CommandLine("-Salt=")]
+		public string? Salt = null;
+
+		/// <inheritdoc/>
+		public override async Task<int> ExecuteAsync(ILogger Logger)
+		{
+			IConfiguration Configuration = new ConfigurationBuilder()
+				.AddEnvironmentVariables()
+				.AddJsonFile("appsettings.json")
+				.Build();
+
+			IConfigurationSection ConfigSection = Configuration.GetSection(AgentSettings.SectionName);
+			IHostBuilder HostBuilder = Host.CreateDefaultBuilder()
+				.ConfigureLogging(Builder => { })
+				.ConfigureServices(Services =>
+				{
+					Services.AddOptions<AgentSettings>().Configure(Options => ConfigSection.Bind(Options)).ValidateDataAnnotations();
+					Services.AddLogging();
+					Services.AddSingleton<GrpcService>();
+				});
+
+			using (IHost Host = HostBuilder.Build())
+			{
+				GrpcService GrpcService = Host.Services.GetRequiredService<GrpcService>();
+
+				byte[] Json = await FileReference.ReadAllBytesAsync(Input);
+				JsonAction Action = JsonSerializer.Deserialize<JsonAction>(Json, new JsonSerializerOptions { PropertyNameCaseInsensitive = true });
+
+				ByteString SaltBytes = ByteString.Empty;
+				if (Salt != null)
+				{
+					SaltBytes = ByteString.CopyFromUtf8(Salt);
+				}
+				else if(RandomSalt)
+				{
+					SaltBytes = ByteString.CopyFrom(Guid.NewGuid().ToByteArray());
+				}
+
+				if (SaltBytes.Length > 0)
+				{
+					Logger.LogInformation("Using salt: {SaltBytes}", StringUtils.FormatHexString(SaltBytes.ToByteArray()));
+				}
+
+				UploadList UploadList = new UploadList();
+				Digest ActionDigest = await Action.BuildAsync(UploadList, SaltBytes);
+
+				using (GrpcChannel Channel = GrpcService.CreateGrpcChannel())
+				{
+					ContentAddressableStorage.ContentAddressableStorageClient StorageClient = new ContentAddressableStorage.ContentAddressableStorageClient(Channel);
+
+					// Find the missing blobs and upload them
+					FindMissingBlobsRequest MissingBlobsRequest = new FindMissingBlobsRequest();
+					MissingBlobsRequest.BlobDigests.Add(UploadList.Blobs.Values.Select(x => x.Digest));
+
+					FindMissingBlobsResponse MissingBlobs = await StorageClient.FindMissingBlobsAsync(MissingBlobsRequest);
+					if (MissingBlobs.MissingBlobDigests.Count > 0)
+					{
+						foreach (Digest MissingBlobDigest in MissingBlobs.MissingBlobDigests)
+						{
+							BatchUpdateBlobsRequest UpdateRequest = new BatchUpdateBlobsRequest();
+							UpdateRequest.Requests.Add(UploadList.Blobs[MissingBlobDigest.Hash]);
+							await StorageClient.BatchUpdateBlobsAsync(UpdateRequest);
+						}
+					}
+
+					// Execute the action
+					RpcExecution.ExecutionClient ExecutionClient = new RpcExecution.ExecutionClient(Channel);
+
+					ExecuteRequest ExecuteRequest = new ExecuteRequest();
+					ExecuteRequest.ActionDigest = ActionDigest;
+
+					using (AsyncServerStreamingCall<Operation> Call = ExecutionClient.Execute(ExecuteRequest))
+					{
+						while (await Call.ResponseStream.MoveNext())
+						{
+							Operation Operation = Call.ResponseStream.Current;
+
+							ExecuteOperationMetadata Metadata = Operation.Metadata.Unpack<ExecuteOperationMetadata>();
+							Logger.LogInformation("Execution state: {State}", Metadata.Stage);
+
+							if (Operation.Done)
+							{
+								// Unpack the response
+								ExecuteResponse ExecuteResponse;
+								if(!Operation.Response.TryUnpack(out ExecuteResponse))
+								{
+									Logger.LogError("Unable to decode response");
+									break;
+								}
+								if((StatusCode)ExecuteResponse.Status.Code != StatusCode.OK)
+								{
+									Logger.LogError("Error {0}: {1}", ExecuteResponse.Status.Code, ExecuteResponse.Status.Message);
+									break;
+								}
+
+								ActionResult? ActionResult = ExecuteResponse.Result;
+								if (ExecuteResponse.Result == null)
+								{
+									Logger.LogError("Empty result");
+									break;
+								}
+
+								Logger.LogInformation("Cached: {CachedResult}", ExecuteResponse.CachedResult);
+
+								// Print the result
+								BatchReadBlobsRequest BatchReadRequest = new BatchReadBlobsRequest();
+								if (ActionResult.StdoutDigest != null && ActionResult.StdoutDigest.SizeBytes > 0)
+								{
+									BatchReadRequest.Digests.Add(ActionResult.StdoutDigest);
+								}
+								if (ActionResult.StderrDigest != null && ActionResult.StderrDigest.SizeBytes > 0)
+								{
+									BatchReadRequest.Digests.Add(ActionResult.StderrDigest);
+								}
+
+								BatchReadBlobsResponse BatchReadResponse = await StorageClient.BatchReadBlobsAsync(BatchReadRequest);
+								if (ActionResult.StdoutDigest != null)
+								{
+									BatchReadBlobsResponse.Types.Response? StdoutResponse = BatchReadResponse.Responses.FirstOrDefault(x => x.Digest.Hash == ActionResult.StdoutDigest.Hash);
+									if (StdoutResponse != null && StdoutResponse.Data.Length > 0)
+									{
+										foreach (string Line in Encoding.UTF8.GetString(StdoutResponse.Data.ToByteArray()).Split('\n'))
+										{
+											Logger.LogInformation("stdout: {Line}", Line);
+										}
+									}
+								}
+								if (ActionResult.StderrDigest != null)
+								{
+									BatchReadBlobsResponse.Types.Response? StderrResponse = BatchReadResponse.Responses.FirstOrDefault(x => x.Digest.Hash == ActionResult.StderrDigest.Hash);
+									if (StderrResponse != null)
+									{
+										foreach (string Line in Encoding.UTF8.GetString(StderrResponse.Data.ToByteArray()).Split('\n'))
+										{
+											Logger.LogError("stderr: {Line}", Line);
+										}
+									}
+								}
+
+								Logger.LogInformation("exit: {ExitCode}", ActionResult.ExitCode);
+								break;
+							}
+						}
+					}
+				}
+			}
+			return 0;
+		}
+	}
+}
diff --git a/Engine/Source/Programs/Horde/HordeAgent/Commands/Parse/ParseCommand.cs b/Engine/Source/Programs/Horde/HordeAgent/Commands/Parse/ParseCommand.cs
new file mode 100644
--- /dev/null
+++ b/Engine/Source/Programs/Horde/HordeAgent/Commands/Parse/ParseCommand.cs
@@ -0,0 +1,86 @@
+using EpicGames.Core;
+using HordeAgent.Parser;
+using Microsoft.Extensions.Logging;
+using System;
+using System.Collections.Generic;
+using System.ComponentModel;
+using System.IO;
+using System.Text;
+using System.Threading.Tasks;
+
+namespace HordeAgent.Commands.Parse
+{
+	/// <summary>
+	/// Installs the agent as a service
+	/// </summary>
+	[Command("Parse", "Parses a file into structured logging output")]
+	class ParseCommand : Command
+	{
+		[CommandLine("-File=", Required = true)]
+		[Description("Log file to parse rather than executing an external program.")]
+		FileReference InputFile = null!;
+
+		[CommandLine("-Ignore=")]
+		[Description("Path to a file containing error patterns to ignore, one regex per line.")]
+		FileReference? IgnorePatternsFile = null;
+
+		[CommandLine("-WorkspaceDir=")]
+		[Description("The root workspace directory")]
+		DirectoryReference? WorkspaceDir = null;
+
+		[CommandLine("-Stream=")]
+		[Description("The stream synced to the workspace")]
+		string? Stream = null;
+
+		[CommandLine("-Change=")]
+		[Description("The changelist number that has been synced")]
+		int? Change = null;
+
+		public override Task<int> ExecuteAsync(ILogger Logger)
+		{
+			// Read all the ignore patterns
+			List<string> IgnorePatterns = new List<string>();
+			if (IgnorePatternsFile != null)
+			{
+				if (!FileReference.Exists(IgnorePatternsFile))
+				{
+					throw new FatalErrorException("Unable to read '{0}", IgnorePatternsFile);
+				}
+
+				// Read all the ignore patterns
+				string[] Lines = FileReference.ReadAllLines(IgnorePatternsFile);
+				foreach (string Line in Lines)
+				{
+					string TrimLine = Line.Trim();
+					if (TrimLine.Length > 0 && !TrimLine.StartsWith("#"))
+					{
+						IgnorePatterns.Add(TrimLine);
+					}
+				}
+			}
+
+			// Read the file and pipe it through the event parser
+			using (StreamReader Reader = new StreamReader(InputFile.FullName))
+			{
+				LogParserContext Context = new LogParserContext();
+				Context.WorkspaceDir = WorkspaceDir;
+				Context.PerforceStream = Stream;
+				Context.PerforceChange = Change;
+
+				using (LogParser Parser = new LogParser(Logger, Context, IgnorePatterns))
+				{
+					for (; ; )
+					{
+						string? Line = Reader.ReadLine();
+						if(Line == null)
+						{
+							break;
+						}
+						Parser.WriteLine(Line);
+					}
+				}
+			}
+			return Task.FromResult(0);
+		}
+	}
+}
diff --git a/Engine/Source/Programs/Horde/HordeAgent/Commands/Service/InstallCommand.cs b/Engine/Source/Programs/Horde/HordeAgent/Commands/Service/InstallCommand.cs
new file mode 100644
--- /dev/null
+++ b/Engine/Source/Programs/Horde/HordeAgent/Commands/Service/InstallCommand.cs
@@ -0,0 +1,98 @@
+using HordeAgent.Utility;
+using Microsoft.Extensions.Logging;
+using System;
+using System.Collections.Generic;
+using System.Reflection;
+using System.Text;
+using System.Threading.Tasks;
+using EpicGames.Core;
+
+namespace HordeAgent.Commands.Service
+{
+	/// <summary>
+	/// Installs the agent as a service
+	/// </summary>
+	[Command("Service", "Install", "Installs the agent as a service")]
+	class InstallCommand : Command
+	{
+		/// <summary>
+		/// Name of the service
+		/// </summary>
+		public const string ServiceName = "HordeAgent";
+
+		/// <summary>
+		/// Specifies the username for the service to run under
+		/// </summary>
+		[CommandLine("-UserName=")]
+		public string? UserName = null;
+
+		/// <summary>
+		/// Password for the username
+		/// </summary>
+		[CommandLine("-Password=")]
+		public string? Password = null;
+
+		/// <summary>
+		/// The server profile to use
+		/// </summary>
+		[CommandLine("-Server=")]
+		public string? Server = null;
+
+		/// <summary>
+		/// Runs the service indefinitely
+		/// </summary>
+		/// <param name="Logger">Logger to use</param>
+		/// <returns>Exit code</returns>
+		public override Task<int> ExecuteAsync(ILogger Logger)
+		{
+			using (WindowsServiceManager ServiceManager = new WindowsServiceManager())
+			{
+				using (WindowsService Service = ServiceManager.Open(ServiceName))
+				{
+					if (Service.IsValid)
+					{
+						Logger.LogInformation("Stopping existing service...");
+						Service.Stop();
+
+						WindowsServiceStatus Status = Service.WaitForStatusChange(WindowsServiceStatus.Stopping, TimeSpan.FromSeconds(30.0));
+						if (Status != WindowsServiceStatus.Stopped)
+						{
+							Logger.LogError("Unable to stop service (status = {Status})", Status);
+							return Task.FromResult(1);
+						}
+
+						Logger.LogInformation("Deleting service");
+						Service.Delete();
+					}
+				}
+
+				Logger.LogInformation("Registering {ServiceName} service", ServiceName);
+
+				StringBuilder CommandLine = new StringBuilder();
+				CommandLine.AppendFormat("dotnet \"{0}\" service run", Assembly.GetEntryAssembly()!.Location);
+				if(Server != null)
+				{
+					CommandLine.Append($" -server={Server}");
+				}
+
+				using (WindowsService Service = ServiceManager.Create(ServiceName, "Horde Agent", CommandLine.ToString(), UserName, Password))
+				{
+					Service.SetDescription("Allows this machine to participate in a Horde farm.");
+
+					Logger.LogInformation("Starting...");
+					Service.Start();
+
+					WindowsServiceStatus Status = Service.WaitForStatusChange(WindowsServiceStatus.Starting, TimeSpan.FromSeconds(30.0));
+					if (Status != WindowsServiceStatus.Running)
+					{
+						Logger.LogError("Unable to start service (status = {Status})", Status);
+						return Task.FromResult(1);
+					}
+
+					Logger.LogInformation("Done.");
+				}
+			}
+			return Task.FromResult(0);
+		}
+	}
+}
diff --git a/Engine/Source/Programs/Horde/HordeAgent/Commands/Service/RunCommand.cs b/Engine/Source/Programs/Horde/HordeAgent/Commands/Service/RunCommand.cs
new file mode 100644
--- /dev/null
+++ b/Engine/Source/Programs/Horde/HordeAgent/Commands/Service/RunCommand.cs
@@ -0,0 +1,154 @@
+// Copyright Epic Games, Inc. All Rights Reserved.
+
+using HordeAgent.Commands;
+using HordeAgent.Services;
+using HordeAgent.Utility;
+using Microsoft.Extensions.Configuration;
+using Microsoft.Extensions.DependencyInjection;
+using Microsoft.Extensions.Hosting;
+using Microsoft.Extensions.Logging;
+using Polly;
+using System;
+using System.Collections.Generic;
+using System.Linq;
+using System.Net.Http;
+using System.Net.Security;
+using System.Security.Cryptography.X509Certificates;
+using System.Text;
+using System.Threading.Tasks;
+using EpicGames.Core;
+using Datadog.Trace.Configuration;
+using Datadog.Trace;
+using System.IO;
+using System.Reflection;
+
+namespace HordeAgent.Modes.Service
+{
+	/// <summary>
+	/// 
+	/// </summary>
+	[Command("Service", "Run", "Runs the service in listen mode")]
+	class RunCommand : Command
+	{
+		/// <summary>
+		/// The logger instance
+		/// </summary>
+		ILogger Logger = null!;
+
+		/// <summary>
+		/// Override for the server to use
+		/// </summary>
+		[CommandLine("-Server=")]
+		string? Server = null;
+
+		/// <summary>
+		/// Override the working directory
+		/// </summary>
+		[CommandLine("-WorkingDir=")]
+		string? WorkingDir = null;
+
+		/// <summary>
+		/// Runs the service indefinitely
+		/// </summary>
+		/// <param name="Logger">Logger to use</param>
+		/// <returns>Exit code</returns>
+		public override async Task<int> ExecuteAsync(ILogger Logger)
+		{
+			this.Logger = Logger;
+
+			IHostBuilder HostBuilder = Host.CreateDefaultBuilder();
+
+			// Attempt to setup this process as a Windows service. A race condition inside Microsoft.Extensions.Hosting.WindowsServices.WindowsServiceHelpers.IsWindowsService
+			// can result in accessing the parent process after it's terminated, so catch any exceptions that it throws.
+			try
+			{
+				HostBuilder = HostBuilder.UseWindowsService();
+			}
+			catch (InvalidOperationException)
+			{
+			}
+
+			HostBuilder = HostBuilder
+				.ConfigureAppConfiguration(Builder =>
+				{
+					Dictionary<string, string> Overrides = new Dictionary<string, string>();
+					if (Server != null)
+					{
+						Overrides.Add($"{AgentSettings.SectionName}:{nameof(AgentSettings.Server)}", Server);
+					}
+					if (WorkingDir != null)
+					{
+						Overrides.Add($"{AgentSettings.SectionName}:{nameof(AgentSettings.WorkingDir)}", WorkingDir);
+					}
+					Builder.AddInMemoryCollection(Overrides);
+				})
+				.ConfigureLogging(Builder =>
+				{
+					Builder.ClearProviders();
+					Builder.AddProvider(new Logging.HordeLoggerProvider());
+					Builder.AddFilter<Logging.HordeLoggerProvider>(null, LogLevel.Trace);
+				})
+				.ConfigureServices((HostContext, Services) =>
+				{
+					Services.Configure<HostOptions>(Options =>
+					{
+						// Allow the worker service to terminate any before shutting down
+						Options.ShutdownTimeout = TimeSpan.FromSeconds(30);
+					});
+
+					IConfigurationSection ConfigSection = HostContext.Configuration.GetSection(AgentSettings.SectionName);
+					Services.AddOptions<AgentSettings>().Configure(Options => ConfigSection.Bind(Options)).ValidateDataAnnotations();
+
+					AgentSettings Settings = new AgentSettings();
+					ConfigSection.Bind(Settings);
+
+					ServerProfile ServerProfile = Settings.GetCurrentServerProfile();
+					ConfigureTracing(ServerProfile.Environment, Settings.Version ?? "unknown");
+					
+					Logging.Version = Settings.Version ?? "";
+
+					Services.AddHttpClient(Program.HordeServerClientName, Config =>
+					{
+						Config.BaseAddress = new Uri(ServerProfile.Url);
+						Config.DefaultRequestHeaders.Add("Accept", "application/json");
+						Config.Timeout = TimeSpan.FromSeconds(300); // Need to make sure this doesn't cancel any long running gRPC streaming calls (eg. session update)
+					})
+					.ConfigurePrimaryHttpMessageHandler(() =>
+					{
+						HttpClientHandler Handler = new HttpClientHandler();
+						Handler.ServerCertificateCustomValidationCallback += (Sender, Cert, Chain, Errors) => CertificateHelper.CertificateValidationCallBack(Logger, Sender, Cert, Chain, Errors, ServerProfile);
+						return Handler;
+					})
+					.AddTransientHttpErrorPolicy(Builder =>
+					{
+						return Builder.WaitAndRetryAsync(new[] { TimeSpan.FromSeconds(1), TimeSpan.FromSeconds(5), TimeSpan.FromSeconds(10) });
+					});
+
+					Services.AddSingleton<GrpcService>();
+					Services.AddHostedService<WorkerService>();
+				});
+
+			try
+			{
+				await HostBuilder.Build().RunAsync();
+			}
+			catch (OperationCanceledException)
+			{
+				// Need to catch this to prevent it propagating to the awaiter when pressing ctrl-c
+			}
+
+			return 0;
+		}
+
+		static void ConfigureTracing(string Environment, string Version)
+		{
+			TracerSettings Settings = TracerSettings.FromDefaultSources();
+			Settings.Environment = Environment;
+			Settings.ServiceName = "hordeagent";
+			Settings.ServiceVersion = Version;
+			Settings.LogsInjectionEnabled = true;
+
+			Tracer.Instance = new Tracer(Settings);
+		}
+	}
+}
diff --git a/Engine/Source/Programs/Horde/HordeAgent/Commands/Service/UninstallCommand.cs b/Engine/Source/Programs/Horde/HordeAgent/Commands/Service/UninstallCommand.cs
new file mode 100644
--- /dev/null
+++ b/Engine/Source/Programs/Horde/HordeAgent/Commands/Service/UninstallCommand.cs
@@ -0,0 +1,57 @@
+using HordeAgent.Utility;
+using Microsoft.Extensions.Logging;
+using System;
+using System.Collections.Generic;
+using System.Reflection;
+using System.Text;
+using System.Threading.Tasks;
+
+namespace HordeAgent.Commands.Service
+{
+	/// <summary>
+	/// Uninstalls the service
+	/// </summary>
+	[Command("Service", "Uninstall", "Uninstalls the service")]
+	class UninstallCommand : Command
+	{
+		/// <summary>
+		/// Name of the service
+		/// </summary>
+		public const string ServiceName = InstallCommand.ServiceName;
+
+		/// <summary>
+		/// Runs the service indefinitely
+		/// </summary>
+		/// <param name="Logger">Logger to use</param>
+		/// <returns>Exit code</returns>
+		public override Task<int> ExecuteAsync(ILogger Logger)
+		{
+			using (WindowsServiceManager ServiceManager = new WindowsServiceManager())
+			{
+				using (WindowsService Service = ServiceManager.Open(ServiceName))
+				{
+					if (Service.IsValid)
+					{
+						Logger.LogInformation("Stopping existing service...");
+						Service.Stop();
+
+						WindowsServiceStatus Status = Service.WaitForStatusChange(WindowsServiceStatus.Stopping, TimeSpan.FromSeconds(30.0));
+						if (Status != WindowsServiceStatus.Stopped)
+						{
+							Logger.LogError("Unable to stop service (status = {Status})", Status);
+							return Task.FromResult(1);
+						}
+
+						Logger.LogInformation("Deleting service");
+						Service.Delete();
+					}
+					else
+					{
+						Logger.LogInformation("Unable to find service {ServiceName}", ServiceName);
+					}
+				}
+			}
+			return Task.FromResult(0);
+		}
+	}
+}
diff --git a/Engine/Source/Programs/Horde/HordeAgent/Commands/Service/UpgradeCommand.cs b/Engine/Source/Programs/Horde/HordeAgent/Commands/Service/UpgradeCommand.cs
new file mode 100644
--- /dev/null
+++ b/Engine/Source/Programs/Horde/HordeAgent/Commands/Service/UpgradeCommand.cs
@@ -0,0 +1,213 @@
+// Copyright Epic Games, Inc. All Rights Reserved.
+
+using HordeAgent.Services;
+using HordeAgent.Utility;
+using Microsoft.Extensions.Logging;
+using System;
+using System.Collections.Generic;
+using System.Diagnostics;
+using System.IO;
+using System.Linq;
+using System.Management;
+using System.Reflection;
+using System.ServiceProcess;
+using System.Text;
+using System.Threading.Tasks;
+using EpicGames.Core;
+using System.Runtime.InteropServices;
+using System.Text.RegularExpressions;
+using System.Threading;
+
+namespace HordeAgent.Commands
+{
+	/// <summary>
+	/// Upgrades a running service to the current application
+	/// </summary>
+	[Command("Service", "Upgrade", "Replaces a running service with the application in the current directory")]
+	class UpgradeCommand : Command
+	{
+		/// <summary>
+		/// The process id to replace
+		/// </summary>
+		[CommandLine("-ProcessId=", Required = true)]
+		int ProcessId = -1;
+
+		/// <summary>
+		/// The target directory to install to
+		/// </summary>
+		[CommandLine("-TargetDir=", Required = true)]
+		DirectoryReference TargetDir = null!;
+
+		/// <summary>
+		/// Arguments to forwar to the target executable
+		/// </summary>
+		[CommandLine("-Arguments=", Required = true)]
+		string Arguments = null!;
+
+		/// <summary>
+		/// Upgrades the application to a new version
+		/// </summary>
+		/// <param name="Logger">The log output device</param>
+		/// <returns>True if the upgrade succeeded</returns>
+		public override Task<int> ExecuteAsync(ILogger Logger)
+		{
+			// Stop the other process
+			Logger.LogInformation("Attempting to perform upgrade on process {ProcessId}", ProcessId);
+			using (Process OtherProcess = Process.GetProcessById(ProcessId))
+			{
+				// Get the directory containing the target application
+				DirectoryInfo TargetDir = new DirectoryInfo(this.TargetDir.FullName);
+				HashSet<string> TargetFiles = new HashSet<string>(TargetDir.EnumerateFiles("*", SearchOption.AllDirectories).Select(x => x.FullName), StringComparer.OrdinalIgnoreCase);
+
+				// Find all the source files
+				DirectoryInfo SourceDir = new DirectoryInfo(Path.GetDirectoryName(Assembly.GetExecutingAssembly().Location));
+				HashSet<string> SourceFiles = new HashSet<string>(SourceDir.EnumerateFiles("*", SearchOption.AllDirectories).Select(x => x.FullName), StringComparer.OrdinalIgnoreCase);
+
+				// Exclude all the source files from the list of target files, since we may be in a subdirectory
+				TargetFiles.ExceptWith(SourceFiles);
+
+				// Ignore any files that are in the saved directory
+				string SourceDataDir = Path.Combine(SourceDir.FullName, Program.SavedDirName) + Path.DirectorySeparatorChar;
+				SourceFiles.RemoveWhere(x => x.StartsWith(SourceDataDir, StringComparison.OrdinalIgnoreCase));
+
+				string TargetDataDir = Path.Combine(TargetDir.FullName, Program.SavedDirName) + Path.DirectorySeparatorChar;
+				TargetFiles.RemoveWhere(x => x.StartsWith(TargetDataDir, StringComparison.OrdinalIgnoreCase));
+
+				// Copy all the files into the target directory
+				List<Tuple<string, string>> RenameFiles = new List<Tuple<string, string>>();
+				foreach (string SourceFile in SourceFiles)
+				{
+					if (!SourceFile.StartsWith(SourceDir.FullName, StringComparison.OrdinalIgnoreCase))
+					{
+						throw new InvalidDataException($"Expected {SourceFile} to be under {SourceDir.FullName}");
+					}
+
+					string TargetFile = TargetDir.FullName + SourceFile.Substring(SourceDir.FullName.Length);
+					Directory.CreateDirectory(Path.GetDirectoryName(TargetFile));
+
+					string TargetFileBeforeRename = TargetFile + ".new";
+					Logger.LogDebug("Copying {SourceFile} to {TargetFileBeforeRename}", SourceFile, TargetFileBeforeRename);
+					File.Copy(SourceFile, TargetFileBeforeRename, true);
+
+					RenameFiles.Add(Tuple.Create(TargetFileBeforeRename, TargetFile));
+					TargetFiles.Remove(TargetFileBeforeRename);
+				}
+
+				if (RuntimeInformation.IsOSPlatform(OSPlatform.Windows))
+				{
+					UpgradeWindowsService(Logger, OtherProcess, TargetFiles, RenameFiles);
+				}
+				else if (RuntimeInformation.IsOSPlatform(OSPlatform.OSX))
+				{
+					UpgradeMacService(Logger, OtherProcess, TargetFiles, RenameFiles);
+				}
+			}
+			Logger.LogInformation("Upgrade complete");
+			return Task.FromResult(0);
+		}
+
+		void UpgradeFilesInPlace(ILogger Logger, HashSet<string> TargetFiles, List<Tuple<string, string>> RenameFiles)
+		{
+			// Remove all the target files
+			foreach (string TargetFile in TargetFiles)
+			{
+				Logger.LogDebug("Deleting {File}", TargetFile);
+				File.SetAttributes(TargetFile, FileAttributes.Normal);
+				File.Delete(TargetFile);
+			}
+
+			// Rename all the new files into place
+			foreach (Tuple<string, string> Pair in RenameFiles)
+			{
+				Logger.LogDebug("Renaming {SourceFile} to {TargetFile}", Pair.Item1, Pair.Item2);
+				File.Move(Pair.Item1, Pair.Item2, true);
+			}
+		}
+
+		void UpgradeMacService(ILogger Logger, Process OtherProcess, HashSet<string> TargetFiles, List<Tuple<string, string>> RenameFiles)
+		{
+			UpgradeFilesInPlace(Logger, TargetFiles, RenameFiles);
+			Logger.LogDebug("Upgrade completed, restarting...");
+			OtherProcess.Kill();
+		}
+
+		void UpgradeWindowsService(ILogger Logger, Process OtherProcess, HashSet<string> TargetFiles, List<Tuple<string, string>> RenameFiles)
+		{
+			// Try to get the service associated with the passed-in process id
+			using (ServiceController? Service = GetServiceForProcess(ProcessId))
+			{
+				// Stop the process
+				if (Service == null)
+				{
+					Logger.LogInformation("Terminating other process");
+					OtherProcess.Kill();
+				}
+				else
+				{
+					Logger.LogInformation("Stopping service");
+					Service.Stop();
+				}
+				OtherProcess.WaitForExit();
+
+				UpgradeFilesInPlace(Logger, TargetFiles, RenameFiles);
+
+				// Run the new application
+				if (Service == null)
+				{
+					string DriverFileName = "dotnet";
+					string AssemblyFileName = Path.Combine(TargetDir.FullName, Path.GetFileName(Assembly.GetExecutingAssembly().Location));
+
+					StringBuilder DriverArguments = new StringBuilder();
+					DriverArguments.AppendArgument(AssemblyFileName);
+					DriverArguments.Append(' ');
+					DriverArguments.Append(Arguments);
+
+					StringBuilder Launch = new StringBuilder();
+					Launch.AppendArgument(DriverFileName);
+					Launch.Append(' ');
+					Launch.Append(DriverArguments);
+					Logger.LogInformation("Launching: {Launch}", Launch.ToString());
+
+					using Process NewProcess = Process.Start(DriverFileName, DriverArguments.ToString());
+				}
+				else
+				{
+					// Start the service again
+					Logger.LogInformation("Restarting service");
+					Service.Start();
+				}
+			}
+		}
+
+		/// <summary>
+		/// Try to find the service controller for the given process id
+		/// </summary>
+		/// <param name="ProcessId">The process id to search for</param>
+		/// <returns>The service controller corresponding to this process</returns>
+		static ServiceController? GetServiceForProcess(int ProcessId)
+		{
+			try
+			{
+				using (ManagementObjectSearcher Searcher = new ManagementObjectSearcher($"SELECT Name FROM Win32_Service WHERE ProcessId={ProcessId}"))
+				{
+					foreach (ManagementObject Service in Searcher.Get())
+					{
+						PropertyData Property = Service.Properties["Name"];
+						if (Property != null)
+						{
+							string? Name = Property.Value as string;
+							if (Name != null)
+							{
+								return new ServiceController(Name);
+							}
+						}
+					}
+				}
+			}
+			catch
+			{
+			}
+			return null;
+		}
+	}
+}
diff --git a/Engine/Source/Programs/Horde/HordeAgent/Commands/Utilities/CreateCertCommand.cs b/Engine/Source/Programs/Horde/HordeAgent/Commands/Utilities/CreateCertCommand.cs
new file mode 100644
--- /dev/null
+++ b/Engine/Source/Programs/Horde/HordeAgent/Commands/Utilities/CreateCertCommand.cs
@@ -0,0 +1,105 @@
+using Microsoft.Extensions.Configuration;
+using Microsoft.Extensions.Logging;
+using System;
+using System.Collections.Generic;
+using System.IO;
+using System.Net;
+using System.Reflection;
+using System.Security.Cryptography;
+using System.Security.Cryptography.X509Certificates;
+using System.Text;
+using System.Text.Json;
+using System.Text.RegularExpressions;
+using System.Threading.Tasks;
+using EpicGames.Core;
+
+namespace HordeAgent.Commands.Certs
+{
+	/// <summary>
+	/// Creates a certificate that can be used for server/agent SSL connections
+	/// </summary>
+	[Command("CreateCert", "Creates a self-signed certificate that can be used for server/agent gRPC connections")]
+	class CreateCertCommand : Command
+	{
+		[CommandLine("-Server=")]
+		public string? Server = null;
+
+		[CommandLine("-DnsName=")]
+		public string? DnsName;
+
+		[CommandLine("-PrivateCert=")]
+		public string? PrivateCertFile;
+
+		[CommandLine("-Environment=")]
+		public string? Environment = "Development";
+
+		/// <summary>
+		/// Main entry point for this command
+		/// </summary>
+		/// <param name="Logger"></param>
+		/// <returns>Async task</returns>
+		public override Task<int> ExecuteAsync(ILogger Logger)
+		{
+			if (DnsName == null)
+			{
+				IConfigurationRoot Config = new ConfigurationBuilder()
+					.SetBasePath(Directory.GetCurrentDirectory())
+					.AddJsonFile("appsettings.json", true)
+					.AddJsonFile($"appsettings.{Environment}.json", true)
+					.Build();
+
+				AgentSettings Settings = new AgentSettings();
+				Config.GetSection("Horde").Bind(Settings);
+
+				if(Server != null)
+				{
+					Settings.Server = Server;
+				}
+
+				DnsName = Settings.GetCurrentServerProfile().Url;
+				DnsName = Regex.Replace(DnsName, @"^[a-zA-Z]+://", "");
+				DnsName = Regex.Replace(DnsName, "/.*$", "");
+			}
+
+			if (PrivateCertFile == null)
+			{
+				FileReference SolutionFile = FileReference.Combine(new FileReference(Assembly.GetExecutingAssembly().Location).Directory, "..", "..", "..", "..", "Horde.sln");
+				if (!FileReference.Exists(SolutionFile))
+				{
+					Logger.LogError("The -PrivateCertFile=... arguments must be specified when running outside the default build directory");
+					return Task.FromResult(1);
+				}
+				PrivateCertFile = FileReference.Combine(SolutionFile.Directory, "HordeServer", "Certs", $"ServerToAgent-{Environment}.pfx").FullName;
+			}
+
+			Logger.LogInformation("Creating certificate for {DnsName}", DnsName);
+			using (RSA Algorithm = RSA.Create(2048))
+			{
+				X500DistinguishedName distinguishedName = new X500DistinguishedName($"CN={DnsName}");
+				CertificateRequest Request = new CertificateRequest(distinguishedName, Algorithm, HashAlgorithmName.SHA256, RSASignaturePadding.Pkcs1);
+
+				Request.CertificateExtensions.Add(new X509BasicConstraintsExtension(false, false, 0, true));
+				Request.CertificateExtensions.Add(new X509KeyUsageExtension(X509KeyUsageFlags.KeyEncipherment | X509KeyUsageFlags.DigitalSignature, true));
+				Request.CertificateExtensions.Add(new X509EnhancedKeyUsageExtension(new OidCollection { new Oid("1.3.6.1.5.5.7.3.1") }, true));
+
+				SubjectAlternativeNameBuilder AlternativeNameBuilder = new SubjectAlternativeNameBuilder();
+				AlternativeNameBuilder.AddDnsName(DnsName);
+				Request.CertificateExtensions.Add(AlternativeNameBuilder.Build(true));
+
+				// NB: MacOS requires 825 days or fewer (https://support.apple.com/en-us/HT210176)
+				using (X509Certificate2 Certificate = Request.CreateSelfSigned(new DateTimeOffset(DateTime.UtcNow.AddDays(-1)), new DateTimeOffset(DateTime.UtcNow.AddDays(800))))
+				{
+					Certificate.FriendlyName = "Horde Server";
+
+					byte[] PrivateCertData = Certificate.Export(X509ContentType.Pkcs12); // Note: Need to reimport this to use immediately, otherwise key is ephemeral
+					Logger.LogInformation("Writing private cert: {PrivateCert}", new FileReference(PrivateCertFile).FullName);
+					File.WriteAllBytes(PrivateCertFile, PrivateCertData);
+
+					Logger.LogInformation("Certificate thumbprint is {Thumbprint}", Certificate.Thumbprint);
+					Logger.LogInformation("Add this thumbprint to list of trusted servers in appsettings.json to trust this server.");
+				}
+			}
+			return Task.FromResult(0);
+		}
+	}
+}
diff --git a/Engine/Source/Programs/Horde/HordeAgent/Commands/Utilities/PublishCommand.cs b/Engine/Source/Programs/Horde/HordeAgent/Commands/Utilities/PublishCommand.cs
new file mode 100644
--- /dev/null
+++ b/Engine/Source/Programs/Horde/HordeAgent/Commands/Utilities/PublishCommand.cs
@@ -0,0 +1,121 @@
+using Grpc.Net.Client;
+using HordeAgent.Utility;
+using HordeCommon.Rpc;
+using Microsoft.Extensions.Configuration;
+using Microsoft.Extensions.Logging;
+using System;
+using System.Collections.Generic;
+using System.IO;
+using System.IO.Compression;
+using System.Net;
+using System.Net.Http;
+using System.Net.Http.Headers;
+using System.Reflection;
+using System.Text;
+using System.Threading.Tasks;
+using EpicGames.Core;
+
+namespace HordeAgent.Commands.Utilities
+{
+	/// <summary>
+	/// Publishes the contents of this application directory to the server
+	/// </summary>
+	[Command("Publish", "Publishes this application to the server")]
+	class PublishCommand : Command
+	{
+		/// <summary>
+		/// The server to upload the software to
+		/// </summary>
+		[CommandLine("-Server=", Required = true)]
+		public string Server = null!;
+
+		/// <summary>
+		/// Access token used to authenticate with the server
+		/// </summary>
+		[CommandLine("-Token=", Required = true)]
+		public string Token = null!;
+
+		/// <summary>
+		/// The input directory
+		/// </summary>
+		[CommandLine("-InputDir=", Required = true)]
+		public DirectoryReference InputDir = null!;
+
+		/// <summary>
+		/// The channel to post to
+		/// </summary>
+		[CommandLine("-Channel=", Required = true)]
+		public string Channel = null!;
+
+		/// <summary>
+		/// Main entry point for this command
+		/// </summary>
+		/// <param name="Logger"></param>
+		/// <returns>Async task</returns>
+		public override async Task<int> ExecuteAsync(ILogger Logger)
+		{
+			// Create the zip archive
+			Logger.LogInformation("Creating archive...");
+			byte[] ArchiveData;
+			using (MemoryStream OutputStream = new MemoryStream())
+			{
+				using (ZipArchive OutputArchive = new ZipArchive(OutputStream, ZipArchiveMode.Create, true))
+				{
+					foreach (FileReference InputFile in DirectoryReference.EnumerateFiles(InputDir, "*", SearchOption.AllDirectories))
+					{
+						string RelativePath = InputFile.MakeRelativeTo(InputDir).Replace(Path.DirectorySeparatorChar, '/');
+						if (!RelativePath.Equals("Log.json", StringComparison.OrdinalIgnoreCase) && !RelativePath.Equals("Log.txt", StringComparison.OrdinalIgnoreCase))
+						{
+							Logger.LogInformation("  Adding {RelativePath}", RelativePath);
+
+							ZipArchiveEntry OutputEntry = OutputArchive.CreateEntry(RelativePath, CompressionLevel.Optimal);
+
+							using Stream InputEntryStream = FileReference.Open(InputFile, FileMode.Open, FileAccess.Read, FileShare.Read);
+							using Stream OutputEntryStream = OutputEntry.Open();
+
+							await InputEntryStream.CopyToAsync(OutputEntryStream);
+						}
+					}
+				}
+				ArchiveData = OutputStream.ToArray();
+			}
+			Logger.LogInformation("Complete ({Size:n0} bytes)", ArchiveData.Length);
+
+			// Read the application settings
+			IConfigurationRoot Config = new ConfigurationBuilder()
+				.SetBasePath(Directory.GetCurrentDirectory())
+				.AddJsonFile("appsettings.json", true)
+				.AddJsonFile("appsettings.Production.json", true)
+				.Build();
+
+			AgentSettings Settings = new AgentSettings();
+			Config.GetSection("Horde").Bind(Settings);
+
+			// Get the server we're using
+			ServerProfile ServerProfile = Settings.GetServerProfile(Server);
+
+			// Create the http message handler
+			HttpClientHandler Handler = new HttpClientHandler();
+			Handler.ServerCertificateCustomValidationCallback += (Sender, Cert, Chain, Errors) => CertificateHelper.CertificateValidationCallBack(Logger, Sender, Cert, Chain, Errors, ServerProfile);
+
+			// Create a new http client for uploading
+			HttpClient HttpClient = new HttpClient(Handler);
+			HttpClient.BaseAddress = new Uri(ServerProfile.Url);
+			HttpClient.DefaultRequestHeaders.Authorization = new AuthenticationHeaderValue("Bearer", Token);
+
+			using (GrpcChannel RpcChannel = GrpcChannel.ForAddress(ServerProfile.Url, new GrpcChannelOptions { HttpClient = HttpClient }))
+			{
+				HordeRpc.HordeRpcClient RpcClient = new HordeRpc.HordeRpcClient(RpcChannel);
+
+				UploadSoftwareRequest UploadRequest = new UploadSoftwareRequest();
+				UploadRequest.Channel = Channel;
+				UploadRequest.Data = Google.Protobuf.ByteString.CopyFrom(ArchiveData);
+
+				UploadSoftwareResponse UploadResponse = await RpcClient.UploadSoftwareAsync(UploadRequest);
+				Logger.LogInformation("Created software (version={SoftwareId} channel={Default})", UploadResponse.Version, UploadRequest.Channel);
+			}
+
+			return 0;
+		}
+	}
+}
diff --git a/Engine/Source/Programs/Horde/HordeAgent/Commands/Workspace/CleanCommand.cs b/Engine/Source/Programs/Horde/HordeAgent/Commands/Workspace/CleanCommand.cs
new file mode 100644
--- /dev/null
+++ b/Engine/Source/Programs/Horde/HordeAgent/Commands/Workspace/CleanCommand.cs
@@ -0,0 +1,28 @@
+// Copyright Epic Games, Inc. All Rights Reserved.
+
+using System;
+using System.Collections.Generic;
+using System.ComponentModel;
+using System.Linq;
+using System.Text;
+using System.Threading;
+using System.Threading.Tasks;
+using Microsoft.Extensions.Logging;
+using EpicGames.Core;
+using EpicGames.Perforce.Managed;
+
+namespace HordeAgent.Commands.Workspace
+{
+	[Command("Workspace", "Clean", "Cleans all modified files from the workspace.")]
+	class CleanCommand : WorkspaceCommand
+	{
+		[CommandLine("-Incremental")]
+		[Description("Performs an incremental sync, without removing intermediates")]
+		public bool bIncrementalSync = false;
+
+		protected override Task ExecuteAsync(ManagedWorkspace Repo, ILogger Logger)
+		{
+			return Repo.CleanAsync(!bIncrementalSync, CancellationToken.None);
+		}
+	}
+}
diff --git a/Engine/Source/Programs/Horde/HordeAgent/Commands/Workspace/ClearCommand.cs b/Engine/Source/Programs/Horde/HordeAgent/Commands/Workspace/ClearCommand.cs
new file mode 100644
--- /dev/null
+++ b/Engine/Source/Programs/Horde/HordeAgent/Commands/Workspace/ClearCommand.cs
@@ -0,0 +1,22 @@
+// Copyright Epic Games, Inc. All Rights Reserved.
+
+using EpicGames.Perforce.Managed;
+using Microsoft.Extensions.Logging;
+using System;
+using System.Collections.Generic;
+using System.Linq;
+using System.Text;
+using System.Threading;
+using System.Threading.Tasks;
+
+namespace HordeAgent.Commands.Workspace
+{
+	[Command("Workspace", "Clear", "Empties the staging directory of any files, returning them to the cache")]
+	class ClearCommand : WorkspaceCommand
+	{
+		protected override Task ExecuteAsync(ManagedWorkspace Repo, ILogger Logger)
+		{
+			return Repo.ClearAsync(CancellationToken.None);
+		}
+	}
+}
diff --git a/Engine/Source/Programs/Horde/HordeAgent/Commands/Workspace/DumpCommand.cs b/Engine/Source/Programs/Horde/HordeAgent/Commands/Workspace/DumpCommand.cs
new file mode 100644
--- /dev/null
+++ b/Engine/Source/Programs/Horde/HordeAgent/Commands/Workspace/DumpCommand.cs
@@ -0,0 +1,22 @@
+// Copyright Epic Games, Inc. All Rights Reserved.
+
+using EpicGames.Perforce.Managed;
+using Microsoft.Extensions.Logging;
+using System;
+using System.Collections.Generic;
+using System.Linq;
+using System.Text;
+using System.Threading.Tasks;
+
+namespace HordeAgent.Commands.Workspace
+{
+	[Command("Workspace", "Dump", "Dumps the contents of the repository to the log for analysis")]
+	class Dump : WorkspaceCommand
+	{
+		protected override Task ExecuteAsync(ManagedWorkspace Repo, ILogger Logger)
+		{
+			Repo.Dump();
+			return Task.CompletedTask;
+		}
+	}
+}
diff --git a/Engine/Source/Programs/Horde/HordeAgent/Commands/Workspace/PopulateCacheCommand.cs b/Engine/Source/Programs/Horde/HordeAgent/Commands/Workspace/PopulateCacheCommand.cs
new file mode 100644
--- /dev/null
+++ b/Engine/Source/Programs/Horde/HordeAgent/Commands/Workspace/PopulateCacheCommand.cs
@@ -0,0 +1,52 @@
+// Copyright Epic Games, Inc. All Rights Reserved.
+
+using EpicGames.Perforce.Managed;
+using Microsoft.Extensions.Logging;
+using System;
+using System.Collections.Generic;
+using System.ComponentModel;
+using System.Linq;
+using System.Text;
+using System.Threading;
+using System.Threading.Tasks;
+using EpicGames.Core;
+using EpicGames.Perforce;
+
+namespace HordeAgent.Commands.Workspace
+{
+	[Command("Workspace", "PopulateWorkspaceCache", "Populates the cache with the head revision of the given streams")]
+	class PopulateCacheCommand : WorkspaceCommand
+	{
+		[CommandLine("-ClientAndStream=")]
+		[Description("Specifies client and stream pairs, in the format Client:Stream")]
+		List<string> ClientAndStreamParams = new List<string>();
+
+		[CommandLine("-Filter=")]
+		[Description("Filters for the files to sync, in P4 syntax (eg. /Engine/...)")]
+		List<string> Filters = new List<string>();
+
+		[CommandLine("-FakeSync")]
+		[Description("Simulates the sync without actually fetching any files")]
+		bool bFakeSync = false;
+
+		protected override Task ExecuteAsync(ManagedWorkspace Repo, ILogger Logger)
+		{
+			List<string> ExpandedFilters = ExpandFilters(Filters);
+
+			List<PopulateRequest> PopulateRequests = new List<PopulateRequest>();
+			foreach (string ClientAndStreamParam in ClientAndStreamParams)
+			{
+				int Idx = ClientAndStreamParam.IndexOf(':');
+				if (Idx == -1)
+				{
+					throw new FatalErrorException("Expected -ClientAndStream=<ClientName>:<StreamName>");
+				}
+
+				PerforceClientConnection PerforceClient = new PerforceClientConnection(Perforce, ClientAndStreamParam.Substring(0, Idx));
+				PopulateRequests.Add(new PopulateRequest(PerforceClient, ClientAndStreamParam.Substring(Idx + 1), ExpandedFilters));
+			}
+
+			return Repo.PopulateAsync(PopulateRequests, bFakeSync, CancellationToken.None);
+		}
+	}
+}
diff --git a/Engine/Source/Programs/Horde/HordeAgent/Commands/Workspace/PurgeCacheCommand.cs b/Engine/Source/Programs/Horde/HordeAgent/Commands/Workspace/PurgeCacheCommand.cs
new file mode 100644
--- /dev/null
+++ b/Engine/Source/Programs/Horde/HordeAgent/Commands/Workspace/PurgeCacheCommand.cs
@@ -0,0 +1,32 @@
+// Copyright Epic Games, Inc. All Rights Reserved.
+
+using EpicGames.Perforce.Managed;
+using Microsoft.Extensions.Logging;
+using System;
+using System.Collections.Generic;
+using System.Linq;
+using System.Text;
+using System.Threading;
+using System.Threading.Tasks;
+using EpicGames.Core;
+
+namespace HordeAgent.Commands.Workspace
+{
+	[Command("Workspace", "PurgeCache", "Shrink the size of the cache to the given size")]
+	class PurgeCacheCommand : WorkspaceCommand
+	{
+		[CommandLine("-Size=")]
+		string? SizeParam = null;
+
+		protected override Task ExecuteAsync(ManagedWorkspace Repo, ILogger Logger)
+		{
+			long Size = 0;
+			if (SizeParam != null)
+			{
+				Size = ParseSize(SizeParam);
+			}
+
+			return Repo.PurgeAsync(Size, CancellationToken.None);
+		}
+	}
+}
diff --git a/Engine/Source/Programs/Horde/HordeAgent/Commands/Workspace/RepairCommand.cs b/Engine/Source/Programs/Horde/HordeAgent/Commands/Workspace/RepairCommand.cs
new file mode 100644
--- /dev/null
+++ b/Engine/Source/Programs/Horde/HordeAgent/Commands/Workspace/RepairCommand.cs
@@ -0,0 +1,22 @@
+// Copyright Epic Games, Inc. All Rights Reserved.
+
+using EpicGames.Perforce.Managed;
+using Microsoft.Extensions.Logging;
+using System;
+using System.Collections.Generic;
+using System.Linq;
+using System.Text;
+using System.Threading;
+using System.Threading.Tasks;
+
+namespace HordeAgent.Commands.Workspace
+{
+	[Command("Workspace", "RepairCache", "Checks the integrity of the cache, and removes any invalid files")]
+	class RepairCommand : WorkspaceCommand
+	{
+		protected override Task ExecuteAsync(ManagedWorkspace Repo, ILogger Logger)
+		{
+			return Repo.RepairAsync(CancellationToken.None);
+		}
+	}
+}
diff --git a/Engine/Source/Programs/Horde/HordeAgent/Commands/Workspace/RevertCommand.cs b/Engine/Source/Programs/Horde/HordeAgent/Commands/Workspace/RevertCommand.cs
new file mode 100644
--- /dev/null
+++ b/Engine/Source/Programs/Horde/HordeAgent/Commands/Workspace/RevertCommand.cs
@@ -0,0 +1,30 @@
+// Copyright Epic Games, Inc. All Rights Reserved.
+
+using EpicGames.Perforce.Managed;
+using Microsoft.Extensions.Logging;
+using System;
+using System.Collections.Generic;
+using System.ComponentModel;
+using System.Linq;
+using System.Text;
+using System.Threading;
+using System.Threading.Tasks;
+using EpicGames.Core;
+using EpicGames.Perforce;
+
+namespace HordeAgent.Commands.Workspace
+{
+	[Command("Workspace", "Revert", "Revert all files that are open in the current workspace. Does not replace them with valid revisions.")]
+	class RevertCommand : WorkspaceCommand
+	{
+		[CommandLine("-Client=", Required = true)]
+		[Description("Client to revert all files for")]
+		string ClientName = null!;
+
+		protected override Task ExecuteAsync(ManagedWorkspace Repo, ILogger Logger)
+		{
+			PerforceClientConnection PerforceClient = new PerforceClientConnection(Perforce, ClientName);
+			return Repo.RevertAsync(PerforceClient, CancellationToken.None);
+		}
+	}
+}
diff --git a/Engine/Source/Programs/Horde/HordeAgent/Commands/Workspace/SetupCommand.cs b/Engine/Source/Programs/Horde/HordeAgent/Commands/Workspace/SetupCommand.cs
new file mode 100644
--- /dev/null
+++ b/Engine/Source/Programs/Horde/HordeAgent/Commands/Workspace/SetupCommand.cs
@@ -0,0 +1,34 @@
+// Copyright Epic Games, Inc. All Rights Reserved.
+
+using System;
+using System.Collections.Generic;
+using System.ComponentModel;
+using System.Linq;
+using System.Text;
+using System.Threading;
+using System.Threading.Tasks;
+using EpicGames.Perforce.Managed;
+using Microsoft.Extensions.Logging;
+using EpicGames.Core;
+using EpicGames.Perforce;
+
+namespace HordeAgent.Commands.Workspace
+{
+	[Command("Workspace", "Setup", "Creates or updates a client to use a given stream")]
+	class SetupCommand : WorkspaceCommand
+	{
+		[CommandLine("-Client=", Required = true)]
+		[Description("Name of the client to create")]
+		string ClientName = null!;
+
+		[CommandLine("-Stream=", Required = true)]
+		[Description("Name of the stream to configure")]
+		string StreamName = null!;
+
+		protected override Task ExecuteAsync(ManagedWorkspace Repo, ILogger Logger)
+		{
+			PerforceClientConnection PerforceClient = new PerforceClientConnection(Perforce, ClientName);
+			return Repo.SetupAsync(PerforceClient, StreamName, CancellationToken.None);
+		}
+	}
+}
diff --git a/Engine/Source/Programs/Horde/HordeAgent/Commands/Workspace/SizesCommand.cs b/Engine/Source/Programs/Horde/HordeAgent/Commands/Workspace/SizesCommand.cs
new file mode 100644
--- /dev/null
+++ b/Engine/Source/Programs/Horde/HordeAgent/Commands/Workspace/SizesCommand.cs
@@ -0,0 +1,38 @@
+// Copyright Epic Games, Inc. All Rights Reserved.
+
+using System;
+using System.Collections.Generic;
+using System.ComponentModel;
+using System.Linq;
+using System.Text;
+using System.Threading;
+using System.Threading.Tasks;
+using Microsoft.Extensions.Logging;
+using EpicGames.Core;
+using EpicGames.Perforce;
+using EpicGames.Perforce.Managed;
+
+namespace HordeAgent.Commands.Workspace
+{
+	[Command("Workspace", "Sizes", "Gathers stats for which streams take the most amount of space in the cache for the given configuration")]
+	class SizesCommand : WorkspaceCommand
+	{
+		[CommandLine("-TempClient=", Required = true)]
+		[Description("Name of a temporary client to switch between streams gathering metadata. Will be created if it does not exist.")]
+		string TempClientName = null!;
+
+		[CommandLine("-Stream=")]
+		[Description("Streams that should be included in the output")]
+		List<string> StreamNames = new List<string>();
+
+		[CommandLine("-Filter=")]
+		[Description("Filters for the files to sync, in P4 syntax (eg. /Engine/...)")]
+		List<string> Filters = new List<string>();
+
+		protected override Task ExecuteAsync(ManagedWorkspace Repo, ILogger Logger)
+		{
+			PerforceClientConnection PerforceClient = new PerforceClientConnection(Perforce, TempClientName);
+			return Repo.StatsAsync(PerforceClient, StreamNames, Filters, CancellationToken.None);
+		}
+	}
+}
diff --git a/Engine/Source/Programs/Horde/HordeAgent/Commands/Workspace/StatusCommand.cs b/Engine/Source/Programs/Horde/HordeAgent/Commands/Workspace/StatusCommand.cs
new file mode 100644
--- /dev/null
+++ b/Engine/Source/Programs/Horde/HordeAgent/Commands/Workspace/StatusCommand.cs
@@ -0,0 +1,22 @@
+// Copyright Epic Games, Inc. All Rights Reserved.
+
+using EpicGames.Perforce.Managed;
+using Microsoft.Extensions.Logging;
+using System;
+using System.Collections.Generic;
+using System.Linq;
+using System.Text;
+using System.Threading.Tasks;
+
+namespace HordeAgent.Commands.Workspace
+{
+	[Command("Workspace", "Status", "Prints information about the state of the cache and workspace")]
+	class StatusCommand : WorkspaceCommand
+	{
+		protected override Task ExecuteAsync(ManagedWorkspace Repo, ILogger Logger)
+		{
+			Repo.Status();
+			return Task.CompletedTask;
+		}
+	}
+}
diff --git a/Engine/Source/Programs/Horde/HordeAgent/Commands/Workspace/SyncCommand.cs b/Engine/Source/Programs/Horde/HordeAgent/Commands/Workspace/SyncCommand.cs
new file mode 100644
--- /dev/null
+++ b/Engine/Source/Programs/Horde/HordeAgent/Commands/Workspace/SyncCommand.cs
@@ -0,0 +1,66 @@
+// Copyright Epic Games, Inc. All Rights Reserved.
+
+using EpicGames.Perforce.Managed;
+using Microsoft.Extensions.Logging;
+using System;
+using System.Collections.Generic;
+using System.ComponentModel;
+using System.Linq;
+using System.Text;
+using System.Threading;
+using System.Threading.Tasks;
+using EpicGames.Core;
+using EpicGames.Perforce;
+
+namespace HordeAgent.Commands.Workspace
+{
+	[Command("Workspace", "Sync", "Syncs the files for a particular stream and changelist")]
+	class SyncCommand : WorkspaceCommand
+	{
+		[CommandLine("-Client=", Required = true)]
+		[Description("Name of the client to sync. Will be created if it does not exist.")]
+		public string ClientName = null!;
+
+		[CommandLine("-Stream=", Required = true)]
+		[Description("The stream to sync")]
+		public string StreamName = null!;
+
+		[CommandLine("-Change=")]
+		[Description("The change to sync. May be a changelist number, or 'Latest'")]
+		public string Change = "Latest";
+
+		[CommandLine("-Preflight=")]
+		[Description("The change to unshelve into the workspace")]
+		public int PreflightChange = -1;
+
+		[CommandLine]
+		[Description("Optional path to a cache file used to store workspace metadata. Using a location on a network share allows multiple machines syncing the same CL to only query Perforce state once.")]
+		FileReference? CacheFile = null;
+
+		[CommandLine("-Filter=")]
+		[Description("Filters for the files to sync, in P4 syntax (eg. /Engine/...)")]
+		public List<string> Filters = new List<string>();
+
+		[CommandLine("-Incremental")]
+		[Description("Performs an incremental sync, without removing intermediates")]
+		public bool bIncrementalSync = false;
+
+		[CommandLine("-FakeSync")]
+		[Description("Simulates the sync without actually fetching any files")]
+		public bool bFakeSync = false;
+
+		protected override async Task ExecuteAsync(ManagedWorkspace Repo, ILogger Logger)
+		{
+			int ChangeNumber = ParseChangeNumberOrLatest(Change);
+			List<string> ExpandedFilters = ExpandFilters(Filters);
+
+			PerforceClientConnection PerforceClient = new PerforceClientConnection(Perforce, ClientName);
+			await Repo.SyncAsync(PerforceClient, StreamName, ChangeNumber, ExpandedFilters, !bIncrementalSync, bFakeSync, CacheFile, CancellationToken.None);
+
+			if (PreflightChange != -1)
+			{
+				await Repo.UnshelveAsync(PerforceClient, PreflightChange, CancellationToken.None);
+			}
+		}
+	}
+}
diff --git a/Engine/Source/Programs/Horde/HordeAgent/Commands/Workspace/WorkspaceCommand.cs b/Engine/Source/Programs/Horde/HordeAgent/Commands/Workspace/WorkspaceCommand.cs
new file mode 100644
--- /dev/null
+++ b/Engine/Source/Programs/Horde/HordeAgent/Commands/Workspace/WorkspaceCommand.cs
@@ -0,0 +1,151 @@
+// Copyright Epic Games, Inc. All Rights Reserved.
+
+using HordeAgent.Utility;
+using Microsoft.Extensions.Logging;
+using System;
+using System.Collections.Generic;
+using System.ComponentModel;
+using System.Diagnostics;
+using System.IO;
+using System.Linq;
+using System.Text;
+using System.Threading;
+using System.Threading.Tasks;
+using EpicGames.Core;
+using EpicGames.Perforce;
+using EpicGames.Perforce.Managed;
+
+namespace HordeAgent.Commands.Workspace
+{
+	abstract class WorkspaceCommand : Command
+	{
+		[CommandLine("-Server", Required = true)]
+		[Description("Specifies the Perforce server and port")]
+		protected string ServerAndPort = null!;
+
+		[CommandLine("-User")]
+		[Description("Specifies the Perforce username")]
+		protected string? UserName = null;
+
+		[CommandLine("-BaseDir", Required = true)]
+		[Description("Base directory to use for syncing workspaces")]
+		protected DirectoryReference BaseDir = null!;
+
+		[CommandLine("-Overwrite")]
+		[Description("")]
+		protected bool bOverwrite = false;
+
+		protected PerforceConnection Perforce = null!;
+
+		public override void Configure(CommandLineArguments Arguments, ILogger Logger)
+		{
+			base.Configure(Arguments, Logger);
+
+			if(BaseDir == null)
+			{
+				for (DirectoryReference? ParentDir = DirectoryReference.GetCurrentDirectory(); ParentDir != null; ParentDir = ParentDir.ParentDirectory)
+				{
+					if (ManagedWorkspace.Exists(ParentDir))
+					{
+						BaseDir = ParentDir;
+						break;
+					}
+				}
+
+				if (BaseDir == null)
+				{
+					throw new FatalErrorException("Unable to find existing repository in current working tree. Specify -BaseDir=... explicitly.");
+				}
+			}
+		}
+
+		public override async Task<int> ExecuteAsync(ILogger Logger)
+		{
+			Perforce = new PerforceConnection(ServerAndPort, UserName, null, Logger);
+
+			InfoRecord Info = await Perforce.GetInfoAsync(InfoOptions.ShortOutput, CancellationToken.None);
+
+			ILogger RepoLogger = new Logging.HordeLoggerProvider().CreateLogger("Repository");
+			ManagedWorkspace Repo = await ManagedWorkspace.LoadOrCreateAsync(Info.ClientHost!, BaseDir, bOverwrite, RepoLogger, CancellationToken.None);
+			await ExecuteAsync(Repo, Logger);
+			return 0;
+		}
+
+		protected abstract Task ExecuteAsync(ManagedWorkspace Repo, ILogger Logger);
+
+		protected static long ParseSize(string Size)
+		{
+			long Value;
+			if (Size.EndsWith("gb", StringComparison.OrdinalIgnoreCase))
+			{
+				string SizeValue = Size.Substring(0, Size.Length - 2).TrimEnd();
+				if (long.TryParse(SizeValue, out Value))
+				{
+					return Value * (1024 * 1024 * 1024);
+				}
+			}
+			else if (Size.EndsWith("mb", StringComparison.OrdinalIgnoreCase))
+			{
+				string SizeValue = Size.Substring(0, Size.Length - 2).TrimEnd();
+				if (long.TryParse(SizeValue, out Value))
+				{
+					return Value * (1024 * 1024);
+				}
+			}
+			else if (Size.EndsWith("kb"))
+			{
+				string SizeValue = Size.Substring(0, Size.Length - 2).TrimEnd();
+				if (long.TryParse(SizeValue, out Value))
+				{
+					return Value * 1024;
+				}
+			}
+			else
+			{
+				if (long.TryParse(Size, out Value))
+				{
+					return Value;
+				}
+			}
+			throw new FatalErrorException("Invalid size '{0}'", Size);
+		}
+
+		protected static int ParseChangeNumber(string Change)
+		{
+			int ChangeNumber;
+			if (int.TryParse(Change, out ChangeNumber) && ChangeNumber > 0)
+			{
+				return ChangeNumber;
+			}
+			throw new FatalErrorException("Unable to parse change number from '{0}'", Change);
+		}
+
+		protected static int ParseChangeNumberOrLatest(string Change)
+		{
+			if (Change.Equals("Latest", StringComparison.OrdinalIgnoreCase))
+			{
+				return -1;
+			}
+			else
+			{
+				return ParseChangeNumber(Change);
+			}
+		}
+
+		protected static List<string> ExpandFilters(List<string> Arguments)
+		{
+			List<string> Filters = new List<string>();
+			foreach (string Argument in Arguments)
+			{
+				foreach (string SingleArgument in Argument.Split(';').Select(x => x.Trim()))
+				{
+					if (SingleArgument.Length > 0)
+					{
+						Filters.Add(SingleArgument);
+					}
+				}
+			}
+			return Filters;
+		}
+	}
+}
diff --git a/Engine/Source/Programs/Horde/HordeAgent/Execution/ActionExecutor.cs b/Engine/Source/Programs/Horde/HordeAgent/Execution/ActionExecutor.cs
new file mode 100644
--- /dev/null
+++ b/Engine/Source/Programs/Horde/HordeAgent/Execution/ActionExecutor.cs
@@ -0,0 +1,214 @@
+// Copyright Epic Games, Inc. All Rights Reserved.
+
+using Build.Bazel.Remote.Execution.V2;
+using EpicGames.Core;
+using Google.LongRunning;
+using Google.Protobuf;
+using Grpc.Core;
+using Microsoft.Extensions.Options;
+using System;
+using System.Collections.Generic;
+using System.IO;
+using System.Linq;
+using System.Threading;
+using System.Threading.Tasks;
+using System.Diagnostics;
+using Action = Build.Bazel.Remote.Execution.V2.Action;
+using Directory = Build.Bazel.Remote.Execution.V2.Directory;
+using RpcCommand = Build.Bazel.Remote.Execution.V2.Command;
+using Status = Google.Rpc.Status;
+using System.Text;
+using Google.Protobuf.WellKnownTypes;
+using Microsoft.Extensions.Hosting;
+using System.Collections.Concurrent;
+using Grpc.Net.Client;
+using HordeAgent.Utility;
+using HordeCommon.Rpc.Tasks;
+using Microsoft.Extensions.Logging;
+using HordeCommon.Rpc;
+
+namespace HordeAgent
+{
+	/// <summary>
+	/// Executes remote actions in a sandbox
+	/// </summary>
+	class ActionExecutor
+	{
+		/// <summary>
+		/// The storage wrapper
+		/// </summary>
+		ContentAddressableStorage.ContentAddressableStorageClient Storage;
+
+		/// <summary>
+		/// Action class instance
+		/// </summary>
+		ActionCache.ActionCacheClient Cache;
+
+		/// <summary>
+		/// Horde RPC wrapper
+		/// </summary>
+		ActionRpc.ActionRpcClient ActionRpc;
+
+		/// <summary>
+		/// Logger for internal executor output
+		/// </summary>
+		ILogger Logger;
+
+		/// <summary>
+		/// Constructor
+		/// </summary>
+		/// <param name="Channel"></param>
+		/// <param name="Logger"></param>
+		public ActionExecutor(GrpcChannel Channel, ILogger Logger)
+		{
+			this.Storage = new ContentAddressableStorage.ContentAddressableStorageClient(Channel);
+			this.Cache = new ActionCache.ActionCacheClient(Channel);
+			this.ActionRpc = new ActionRpc.ActionRpcClient(Channel);
+			this.Logger = Logger;
+		}
+
+		/// <summary>
+		/// Execute an action
+		/// </summary>
+		/// <param name="LeaseId">The lease id</param>
+		/// <param name="ActionTask">Task to execute</param>
+		/// <param name="SandboxDir">Directory to use as a sandbox for execution</param>
+		/// <returns>The action result</returns>
+		public async Task<ActionResult> ExecuteActionAsync(string LeaseId, ActionTask ActionTask, DirectoryReference SandboxDir)
+		{
+			Action? Action = await Storage.GetProtoMessageAsync<Action>(ActionTask.Digest);
+			if (Action == null)
+			{
+				throw new RpcException(new Grpc.Core.Status(StatusCode.NotFound, "Unable to find action message"));
+			}
+
+			DirectoryReference.CreateDirectory(SandboxDir);
+			FileUtils.ForceDeleteDirectoryContents(SandboxDir);
+
+			Directory InputDirectory = await Storage.GetProtoMessageAsync<Directory>(Action.InputRootDigest);
+			await SetupSandboxAsync(InputDirectory, SandboxDir);
+
+			RpcCommand Command = await Storage.GetProtoMessageAsync<RpcCommand>(Action.CommandDigest);
+
+			foreach (string OutputFile in Command.OutputFiles)
+			{
+				FileReference OutputFileLocation = FileReference.Combine(SandboxDir, OutputFile);
+				DirectoryReference.CreateDirectory(OutputFileLocation.Directory);
+			}
+
+			if (Command.OutputPaths.Count > 0 || Command.OutputDirectories.Count > 0)
+			{
+				throw new NotImplementedException();
+			}
+
+			ActionResult Result;
+			using (ManagedProcessGroup ProcessGroup = new ManagedProcessGroup())
+			{
+				string FileName = FileReference.Combine(SandboxDir, Command.WorkingDirectory, Command.Arguments[0]).FullName;
+				string Arguments = CommandLineArguments.Join(Command.Arguments.Skip(1));
+				string WorkingDirectory = DirectoryReference.Combine(SandboxDir, Command.WorkingDirectory).FullName;
+
+				Dictionary<string, string> NewEnvironment = new Dictionary<string, string>();
+				foreach (System.Collections.DictionaryEntry? Entry in Environment.GetEnvironmentVariables())
+				{
+					NewEnvironment[Entry!.Value.Key.ToString()!] = Entry!.Value.Value!.ToString()!;
+				}
+				foreach (RpcCommand.Types.EnvironmentVariable EnvironmentVariable in Command.EnvironmentVariables)
+				{
+					NewEnvironment[EnvironmentVariable.Name] = EnvironmentVariable.Value;
+				}
+
+				Logger.LogInformation("Executing {FileName} with arguments {Arguments}", FileName, Arguments);
+				using (ManagedProcess Process = new ManagedProcess(ProcessGroup, FileName, Arguments, WorkingDirectory, NewEnvironment, ProcessPriorityClass.Normal, ManagedProcessFlags.None))
+				{
+					byte[] StdOutData;
+					byte[] StdErrData;
+					using (MemoryStream StdOutStream = new MemoryStream())
+					using (MemoryStream StdErrStream = new MemoryStream())
+					{
+						await Task.WhenAll(Process.StdOut.CopyToAsync(StdOutStream), Process.StdErr.CopyToAsync(StdErrStream));
+						StdOutData = StdOutStream.ToArray();
+						StdErrData = StdErrStream.ToArray();
+					}
+
+					Result = new ActionResult();
+					Result.StdoutDigest = await Storage.PutBulkDataAsync(StdOutData);
+					Result.StderrDigest = await Storage.PutBulkDataAsync(StdErrData);
+					Result.ExitCode = Process.ExitCode;
+
+					foreach (string Line in Encoding.UTF8.GetString(StdOutData).Split('\n'))
+					{
+						Logger.LogInformation("stdout: {Line}", Line);
+					}
+					foreach (string Line in Encoding.UTF8.GetString(StdErrData).Split('\n'))
+					{
+						Logger.LogInformation("stderr: {Line}", Line);
+					}
+
+					Logger.LogInformation("exit: {ExitCode}", Process.ExitCode);
+
+					foreach (string OutputFile in Command.OutputFiles)
+					{
+						FileReference FileRef = FileReference.Combine(SandboxDir, OutputFile);
+						if (FileReference.Exists(FileRef))
+						{
+							byte[] Bytes = await FileReference.ReadAllBytesAsync(FileRef);
+							Digest Digest = await Storage.PutBulkDataAsync(Bytes);
+
+							OutputFile OutputFileInfo = new OutputFile();
+							OutputFileInfo.Path = OutputFile;
+							OutputFileInfo.Digest = Digest;
+							Result.OutputFiles.Add(OutputFileInfo);
+
+							Logger.LogInformation("Uploaded {File} (digest: {Digest}, size: {Size})", FileRef, Digest.Hash, Digest.SizeBytes);
+						}
+					}
+
+					PostActionResultRequest PostResultRequest = new PostActionResultRequest();
+					PostResultRequest.LeaseId = LeaseId;
+					PostResultRequest.ActionDigest = ActionTask.Digest;
+					PostResultRequest.Result = Result;
+					await ActionRpc.PostActionResultAsync(PostResultRequest);
+
+					if (!Action.DoNotCache)
+					{
+						UpdateActionResultRequest Request = new UpdateActionResultRequest();
+						Request.ActionDigest = ActionTask.Digest;
+						Request.ActionResult = Result;
+						await Cache.UpdateActionResultAsync(Request);
+					}
+
+					// Command.OutputPaths
+					// Command.OutputDirectories
+				}
+			}
+			return Result;
+		}
+
+		/// <summary>
+		/// Downloads files to the sandbox
+		/// </summary>
+		/// <param name="InputDirectory">The directory spec</param>
+		/// <param name="OutputDir">Output directory on disk</param>
+		/// <returns>Async task</returns>
+		async Task SetupSandboxAsync(Directory InputDirectory, DirectoryReference OutputDir)
+		{
+			DirectoryReference.CreateDirectory(OutputDir);
+
+			foreach (FileNode FileNode in InputDirectory.Files)
+			{
+				ReadOnlyMemory<byte> Data = await Storage.GetBulkDataAsync(FileNode.Digest);
+				FileReference File = FileReference.Combine(OutputDir, FileNode.Name);
+				Logger.LogInformation("Writing {File} (digest: {Digest}, size: {Size})", File, FileNode.Digest.Hash, FileNode.Digest.SizeBytes);
+				await FileReference.WriteAllBytesAsync(File, Data.ToArray());
+			}
+
+			foreach (DirectoryNode DirectoryNode in InputDirectory.Directories)
+			{
+				Directory InputSubDirectory = await Storage.GetProtoMessageAsync<Directory>(DirectoryNode.Digest);
+				DirectoryReference OutputSubDirectory = DirectoryReference.Combine(OutputDir, DirectoryNode.Name);
+				await SetupSandboxAsync(InputSubDirectory, OutputSubDirectory);
+			}
+		}
+	}
+}
