Initial drop of Horde Server/Agent for UE5/Main. Still missing some Epic-specific functionality required to build/run correctly.

diff --git a/Engine/Source/Programs/Horde/HordeAgent/Parser/LogParser.cs b/Engine/Source/Programs/Horde/HordeAgent/Parser/LogParser.cs
new file mode 100644
--- /dev/null
+++ b/Engine/Source/Programs/Horde/HordeAgent/Parser/LogParser.cs
@@ -0,0 +1,322 @@
+// Copyright Epic Games, Inc. All Rights Reserved.
+
+using EpicGames.Core;
+using HordeAgent.Parser.Interfaces;
+using HordeAgent.Utility;
+using HordeCommon;
+using Microsoft.Extensions.Logging;
+using System;
+using System.Collections.Concurrent;
+using System.Collections.Generic;
+using System.IO;
+using System.Linq;
+using System.Reflection;
+using System.Text;
+using System.Text.RegularExpressions;
+using System.Threading;
+using System.Threading.Tasks;
+
+namespace HordeAgent.Parser
+{
+	/// <summary>
+	/// Concrete implementation of <see cref="ILogContext"/>
+	/// </summary>
+	public class LogParserContext : ILogContext
+	{
+		/// <inheritdoc/>
+		public DirectoryReference? WorkspaceDir { get; set; }
+
+		/// <inheritdoc/>
+		public string? PerforceStream { get; set; }
+
+		/// <inheritdoc/>
+		public int? PerforceChange { get; set; }
+
+		/// <inheritdoc/>
+		public bool HasLoggedErrors { get; set; }
+	}
+
+	/// <summary>
+	/// Turns raw text output into structured logging events
+	/// </summary>
+	public class LogParser : IDisposable
+	{
+		/// <summary>
+		/// Global list of error matchers, obtained by reflection at startup
+		/// </summary>
+		static List<ILogEventMatcher> Matchers;
+
+		/// <summary>
+		/// List of patterns to ignore
+		/// </summary>
+		List<string> IgnorePatterns;
+
+		/// <summary>
+		/// Buffer of input lines
+		/// </summary>
+		LineBuffer Buffer;
+
+		/// <summary>
+		/// Buffer for holding partial line data
+		/// </summary>
+		MemoryStream PartialLine = new MemoryStream();
+
+		/// <summary>
+		/// The inner logger
+		/// </summary>
+		ILogger Logger;
+
+		/// <summary>
+		/// Context for matchers for this log
+		/// </summary>
+		LogParserContext Context;
+
+		/// <summary>
+		/// Refcount for whether output is enabled
+		/// </summary>
+		int MatchingEnabled;
+
+		/// <summary>
+		/// Constructor
+		/// </summary>
+		/// <param name="Logger">The logger to receive parsed output messages</param>
+		/// <param name="Context">Context for parsing this log</param>
+		/// <param name="IgnorePatterns">List of patterns to ignore</param>
+		public LogParser(ILogger Logger, LogParserContext Context, List<string> IgnorePatterns)
+		{
+			this.IgnorePatterns = IgnorePatterns;
+			this.Logger = Logger;
+			this.Buffer = new LineBuffer(50);
+			this.Context = Context;
+		}
+
+		/// <summary>
+		/// Static constructor
+		/// </summary>
+		static LogParser()
+		{
+			System.Text.RegularExpressions.Regex.CacheSize = 1000;
+
+			Matchers = new List<ILogEventMatcher>();
+			foreach (Type Type in Assembly.GetExecutingAssembly().GetTypes())
+			{
+				if (Type.IsClass && typeof(ILogEventMatcher).IsAssignableFrom(Type))
+				{
+					ILogEventMatcher Matcher = (ILogEventMatcher)Activator.CreateInstance(Type)!;
+					Matchers.Add(Matcher);
+				}
+			}
+		}
+
+		/// <summary>
+		/// Adds a raw utf-8 string to the buffer
+		/// </summary>
+		/// <param name="Data">The string data</param>
+		private void AddLine(ReadOnlySpan<byte> Data)
+		{
+			if(Data.Length > 0 && Data[Data.Length - 1] == '\r')
+			{
+				Data = Data.Slice(0, Data.Length - 1);
+			}
+			Buffer.AddLine(Encoding.UTF8.GetString(Data));
+		}
+
+		/// <summary>
+		/// Writes the current partial line data, with the given data appended to it, then clear the buffer
+		/// </summary>
+		private void FlushPartialLine()
+		{
+			AddLine(PartialLine.ToArray());
+			PartialLine.Position = 0;
+			PartialLine.SetLength(0);
+		}
+
+		/// <summary>
+		/// Writes data to the log parser
+		/// </summary>
+		/// <param name="Data">Data to write</param>
+		/// <param name="bEndOfStream">Whether this is the end of the stream</param>
+		public void WriteData(ReadOnlySpan<byte> Data, bool bEndOfStream)
+		{
+			int BaseIdx = 0;
+			int ScanIdx = 0;
+
+			// Handle a partially existing line
+			if (PartialLine.Length > 0)
+			{
+				for (; ScanIdx < Data.Length; ScanIdx++)
+				{
+					if (Data[ScanIdx] == '\n')
+					{
+						PartialLine.Write(Data.Slice(BaseIdx, ScanIdx - BaseIdx));
+						FlushPartialLine();
+						BaseIdx = ++ScanIdx;
+						break;
+					}
+				}
+			}
+
+			// Handle any complete lines
+			for (; ScanIdx < Data.Length; ScanIdx++)
+			{
+				if(Data[ScanIdx] == '\n')
+				{
+					AddLine(Data.Slice(BaseIdx, ScanIdx - BaseIdx));
+					BaseIdx = ScanIdx + 1;
+				}
+			}
+
+			// Add the rest of the text to the partial line buffer
+			PartialLine.Write(Data.Slice(BaseIdx));
+
+			// If it's the end of the stream, force a flush
+			if (bEndOfStream && PartialLine.Length > 0)
+			{
+				FlushPartialLine();
+			}
+
+			// Process the new data
+			ProcessData(bEndOfStream);
+		}
+
+		/// <summary>
+		/// Writes a line to the event filter
+		/// </summary>
+		/// <param name="Line">The line to output</param>
+		public void WriteLine(string? Line)
+		{
+			if (Line == null)
+			{
+				ProcessData(true);
+			}
+			else
+			{
+				Buffer.AddLine(Line);
+				ProcessData(false);
+			}
+		}
+
+		/// <summary>
+		/// Process any data in the buffer
+		/// </summary>
+		/// <param name="bEndOfStream">Whether we've reached the end of the stream</param>
+		void ProcessData(bool bEndOfStream)
+		{
+			while (Buffer.Length > 0)
+			{
+				// Try to match an event
+				LogEvent? Event = null;
+				if (Regex.IsMatch(Buffer[0], "<-- Suspend Log Parsing -->", RegexOptions.IgnoreCase))
+				{
+					MatchingEnabled--;
+				}
+				else if (Regex.IsMatch(Buffer[0], "<-- Resume Log Parsing -->", RegexOptions.IgnoreCase))
+				{
+					MatchingEnabled++;
+				}
+				else if (MatchingEnabled >= 0)
+				{
+					Event = MatchEvent();
+				}
+
+				// Bail out if we need more data
+				if(Buffer.Length < 1024 && !bEndOfStream && Buffer.NeedMoreData)
+				{
+					break;
+				}
+
+				// If we did match something, check if it's not negated by an ignore pattern. We typically have relatively few errors and many more ignore patterns than matchers, so it's quicker 
+				// to check them in response to an identified error than to treat them as matchers of their own.
+				if (Event != null)
+				{
+					foreach (string IgnorePattern in IgnorePatterns)
+					{
+						if (Regex.IsMatch(Buffer[0], IgnorePattern))
+						{
+							Event = null;
+							break;
+						}
+					}
+				}
+
+				// Report the error to the listeners
+				if (Event != null)
+				{
+					WriteEvent(Event);
+					Buffer.Advance(Event.MaxLineNumber + 1 - Buffer.CurrentLineNumber);
+				}
+				else
+				{
+					Logger.Log(LogLevel.Information, KnownLogEvents.None, Buffer[0]!, null, (State, Exception) => State);
+					Buffer.MoveNext();
+				}
+
+				// Also flag that an error has occurred for future add the error to the log context, so that future errors can examine it
+				if (Event != null && Event.Level >= LogLevel.Error)
+				{
+					Context.HasLoggedErrors = true;
+				}
+			}
+		}
+
+		/// <summary>
+		/// Try to match an event from the current buffer
+		/// </summary>
+		/// <returns>The matched event</returns>
+		private LogEvent? MatchEvent()
+		{
+			LogEvent? Event = null;
+			foreach (ILogEventMatcher Matcher in Matchers)
+			{
+				LogEvent? NewEvent = Matcher.Match(Buffer, Context);
+				if (NewEvent != null && (Event == null || NewEvent.Priority > Event.Priority))
+				{
+					Event = NewEvent;
+				}
+			}
+			return Event;
+		}
+
+		/// <summary>
+		/// Writes an event to the log
+		/// </summary>
+		/// <param name="Event">The event to write</param>
+		private void WriteEvent(LogEvent Event)
+		{
+			if (Event.ChildEvents == null)
+			{
+				Logger.Log(Event.Level, Event.EventId, Event.State, null, (State, Exception) => State.ToString());
+			}
+			else
+			{
+				int ChildIdx = 0;
+				for (int LineNumber = Buffer.CurrentLineNumber; LineNumber <= Event.MaxLineNumber;)
+				{
+					if (ChildIdx < Event.ChildEvents.Count && LineNumber == Event.ChildEvents[ChildIdx].MinLineNumber)
+					{
+						LogEvent ChildEvent = Event.ChildEvents[ChildIdx++];
+						Logger.Log(ChildEvent.Level, ChildEvent.EventId, ChildEvent.State, null, (State, Exception) => State.ToString());
+						LineNumber = ChildEvent.MaxLineNumber + 1;
+					}
+					else
+					{
+						string? Line = Buffer[LineNumber - Buffer.CurrentLineNumber];
+						Logger.Log(LogLevel.Information, KnownLogEvents.None, Line!, null, (State, Exception) => State.ToString());
+						LineNumber++;
+					}
+				}
+			}
+		}
+
+		/// <summary>
+		/// Dispose of this object
+		/// </summary>
+		public void Dispose()
+		{
+			if(Buffer.Length > 0)
+			{
+				WriteLine(null);
+			}
+		}
+	}
+}
diff --git a/Engine/Source/Programs/Horde/HordeAgent/Parser/Matchers/CompileEventMatcher.cs b/Engine/Source/Programs/Horde/HordeAgent/Parser/Matchers/CompileEventMatcher.cs
new file mode 100644
--- /dev/null
+++ b/Engine/Source/Programs/Horde/HordeAgent/Parser/Matchers/CompileEventMatcher.cs
@@ -0,0 +1,277 @@
+// Copyright Epic Games, Inc. All Rights Reserved.
+
+using EpicGames.Core;
+using HordeAgent.Parser;
+using HordeAgent.Parser.Interfaces;
+using HordeCommon;
+using Microsoft.Extensions.Logging;
+using System;
+using System.Collections.Generic;
+using System.Diagnostics.CodeAnalysis;
+using System.IO;
+using System.Linq;
+using System.Text;
+using System.Text.RegularExpressions;
+using System.Threading.Tasks;
+
+namespace HordeAgent.Parser.Matchers
+{
+	/// <summary>
+	/// Matches compile errors and annotates with the source file path and revision
+	/// </summary>
+	class CompileEventMatcher : ILogEventMatcher
+	{
+		const string FilePattern =
+			@"(?<file>" +
+				// optional drive letter
+				@"(?:[a-zA-Z]:)?" +
+				// any non-colon character
+				@"[^:(\s]+" +
+				// any path character
+				@"[^:<>*?""]+" +
+				// valid source file extension
+				@"\.(?:(?i)(?:h|c|cc|cpp|inc|inl|cs|targets))" +
+			@")";
+
+		const string VisualCppLocationPattern =
+			@"\(" +
+				@"(?<line>\d+)" + // line number
+				@"(?:,(?<column>\d+))?" + // optional column number
+			@"\)";
+
+		const string VisualCppSeverity =
+			@"(?<severity>fatal error|error|warning)(?: (?<code>[A-Z]+[0-9]+))?";
+
+		const string ClangLocationPattern =
+			@":" +
+				@"(?<line>\d+)" + // line number
+				@"(?::(?<column>\d+))?" + // optional column number
+			@"";
+
+		const string ClangSeverity =
+			@"(?<severity>error|warning)";
+
+		static string ClangOrVisualCppLinePattern =
+			$"(?:{VisualCppLocationPattern}|{ClangLocationPattern})";
+
+		static readonly string[] InvalidExtensions =
+		{
+			".obj",
+			".dll",
+			".exe"
+		};
+
+		const string DefaultSourceFileBaseDir = "Engine/Source";
+
+		/// <inheritdoc/>
+		public LogEvent? Match(ILogCursor Input, ILogContext Context)
+		{
+			// Match the prelude to any error
+			int MaxOffset = 0;
+			while (Input.IsMatch(MaxOffset, "^\\s*(?:In (member )?function|In file included from)"))
+			{
+				MaxOffset++;
+			}
+
+			// Do the match in two phases so we can early out if the strings "error" or "warning" are not present. The patterns before these strings can
+			// produce many false positives, making them very slow to execute.
+			if (Input.IsMatch(MaxOffset, "error|warning"))
+			{
+				// Try to match a Visual C++ diagnostic
+				LogEvent? NewEvent;
+				if(TryMatchVisualCppEvent(Input, MaxOffset, Context, out NewEvent))
+				{
+					LogEventSpan? Code;
+					if(NewEvent.TryGetSpan("code", out Code) && Code.Text == "C2220")
+					{
+						MaxOffset = NewEvent.MaxLineNumber + 1 - Input.CurrentLineNumber;
+
+						List<LogEvent> ChildEvents = new List<LogEvent>();
+
+						LogEvent? ChildEvent;
+						while (TryMatchVisualCppEvent(Input.Rebase(MaxOffset), 0, Context, out ChildEvent))
+						{
+							ChildEvent.Level = LogLevel.Error;
+							ChildEvents.Add(ChildEvent);
+							MaxOffset = ChildEvent.MaxLineNumber + 1 - Input.CurrentLineNumber;
+						}
+
+						NewEvent.MaxLineNumber = Input.CurrentLineNumber + (MaxOffset - 1);
+						NewEvent.ChildEvents = ChildEvents;
+					}
+
+					return NewEvent;
+				}
+
+				// Try to match a Clang diagnostic
+				Match? Match;
+				if (Input.TryMatch($"^\\s*{FilePattern}\\s*{ClangLocationPattern}:\\s*{ClangSeverity}\\s*:", out Match) && IsSourceFile(Match))
+				{
+					LogLevel Level = GetLogLevelFromSeverity(Match);
+
+					LogEventBuilder Builder = new LogEventBuilder(Input);
+
+					LogEventLine FirstLine = Builder.Lines[0];
+					FirstLine.AddSpan(Match.Groups["file"]).MarkAsSourceFile(Context, "Engine/Source");
+					FirstLine.AddSpan(Match.Groups["severity"]).MarkAsSeverity();
+					FirstLine.TryAddSpan(Match.Groups["line"])?.MarkAsLineNumber();
+					FirstLine.TryAddSpan(Match.Groups["column"])?.MarkAsLineNumber();
+
+					string Indent = ExtractIndent(Input[0]!);
+
+					int NoteIdx = 1;
+					while (Builder.TryMatch($"^(?:{Indent} |{Indent}\\s*{FilePattern}\\s*{ClangLocationPattern}\\s*note:| *$)", out Match))
+					{
+						LogEventLine Line = Builder.AddLine();
+
+						Group FileGroup = Match.Groups["file"];
+						if (FileGroup.Success)
+						{
+							Line.AddSpan(FileGroup, $"note{NoteIdx}").MarkAsSourceFile(Context, DefaultSourceFileBaseDir);
+							Line.TryAddSpan(Match.Groups["line"], $"line{NoteIdx}")?.MarkAsLineNumber();
+							NoteIdx++;
+						}
+					}
+
+					return Builder.ToLogEvent(LogEventPriority.High, Level, KnownLogEvents.Compiler);
+				}
+			}
+			return null;
+		}
+
+		bool TryMatchVisualCppEvent(ILogCursor Input, int MaxOffset, ILogContext Context, [NotNullWhen(true)] out LogEvent? OutEvent)
+		{
+			Match? Match;
+			if (!Input.TryMatch(MaxOffset, $"^\\s*(?:ERROR: |WARNING: )?{FilePattern}(?:{VisualCppLocationPattern})? ?:\\s+{VisualCppSeverity}:", out Match) || !IsSourceFile(Match))
+			{
+				OutEvent = null;
+				return false;
+			}
+
+			LogLevel Level = GetLogLevelFromSeverity(Match);
+
+			LogEventBuilder Builder = new LogEventBuilder(Input);
+			Builder.MaxOffset = MaxOffset;
+
+			LogEventLine FirstLine = Builder.Lines[Builder.MaxOffset];
+
+			FirstLine.AddSpan(Match.Groups["severity"]).MarkAsSeverity();
+
+			string SourceFileBaseDir = DefaultSourceFileBaseDir;
+
+			Group CodeGroup = Match.Groups["code"];
+			if (CodeGroup.Success)
+			{
+				FirstLine.AddSpan(CodeGroup).MarkAsErrorCode();
+
+				if (CodeGroup.Value.StartsWith("CS", StringComparison.Ordinal))
+				{
+					Match? ProjectMatch;
+					if (Input.TryMatch(@"\[(?<project>[^[\]]+)]\s*$", out ProjectMatch))
+					{
+						FirstLine.AddSpan(ProjectMatch.Groups[1]).MarkAsSourceFile(Context, "");
+						SourceFileBaseDir = GetPlatformAgnosticDirectoryName(ProjectMatch.Groups[1].Value) ?? SourceFileBaseDir;
+					}
+				}
+				else if (CodeGroup.Value.StartsWith("MSB", StringComparison.Ordinal))
+				{
+					if (CodeGroup.Value.Equals("MSB3026", StringComparison.Ordinal))
+					{
+						OutEvent = Builder.ToLogEvent(LogEventPriority.High, LogLevel.Information, KnownLogEvents.Systemic_MSBuild);
+						return true;
+					}
+
+					Match? ProjectMatch;
+					if (Input.TryMatch(@"\[(?<file>[^[\]]+)]\s*$", out ProjectMatch))
+					{
+						FirstLine.AddSpan(ProjectMatch.Groups[1]).MarkAsSourceFile(Context, "");
+						OutEvent = Builder.ToLogEvent(LogEventPriority.High, Level, KnownLogEvents.MSBuild);
+						return true;
+					}
+				}
+			}
+
+			FirstLine.AddSpan(Match.Groups["file"]).MarkAsSourceFile(Context, SourceFileBaseDir);
+			FirstLine.TryAddSpan(Match.Groups["line"])?.MarkAsLineNumber();
+			FirstLine.TryAddSpan(Match.Groups["column"])?.MarkAsColumnNumber();
+
+			string Indent = ExtractIndent(Input[0] ?? String.Empty);
+
+			int NoteIdx = 1;
+			while (Builder.TryMatch($"^(?:{Indent} |{Indent}\\s*{FilePattern}(?:{VisualCppLocationPattern})?\\s*: note:| *$)", out Match))
+			{
+				LogEventLine Line = Builder.AddLine();
+
+				Group Group = Match.Groups["file"];
+				if (Group.Success)
+				{
+					Line.AddSpan(Group, $"note{NoteIdx}").MarkAsSourceFile(Context, DefaultSourceFileBaseDir);
+					Line.TryAddSpan(Match.Groups["line"], $"line{NoteIdx}")?.MarkAsLineNumber();
+					NoteIdx++;
+				}
+			}
+
+			string Pattern = $"^{Indent} |: note:";
+			while (Input.IsMatch(Builder.MaxOffset + 1, Pattern))
+			{
+				Builder.AddLine();
+			}
+
+			OutEvent = Builder.ToLogEvent(LogEventPriority.High, Level, KnownLogEvents.Compiler);
+			return true;
+		}
+
+		static string? GetPlatformAgnosticDirectoryName(string FileName)
+		{
+			int Index = FileName.LastIndexOfAny(new[] { '/', '\\' });
+			if (Index == -1)
+			{
+				return null;
+			}
+			else
+			{
+				return FileName.Substring(0, Index);
+			}
+		}
+
+		bool IsSourceFile(Match Match)
+		{
+			Group Group = Match.Groups["file"];
+			if (!Group.Success)
+			{
+				return false;
+			}
+
+			string Text = Group.Value;
+			if (InvalidExtensions.Any(x => Text.EndsWith(x, StringComparison.OrdinalIgnoreCase)))
+			{
+				return false;
+			}
+
+			return true;
+		}
+
+		static LogLevel GetLogLevelFromSeverity(Match Match)
+		{
+			string Severity = Match.Groups["severity"].Value;
+			if(Severity.Equals("warning", StringComparison.Ordinal))
+			{
+				return LogLevel.Warning;
+			}
+			else
+			{
+				return LogLevel.Error;
+			}
+		}
+
+		static string ExtractIndent(string Line)
+		{
+			int Length = 0;
+			while(Length < Line.Length && Line[Length] == ' ')
+			{
+				Length++;
+			}
+			return new string(' ', Length);
+		}
+	}
+}
diff --git a/Engine/Source/Programs/Horde/HordeAgent/Parser/Matchers/ContentEventMatcher.cs b/Engine/Source/Programs/Horde/HordeAgent/Parser/Matchers/ContentEventMatcher.cs
new file mode 100644
--- /dev/null
+++ b/Engine/Source/Programs/Horde/HordeAgent/Parser/Matchers/ContentEventMatcher.cs
@@ -0,0 +1,62 @@
+using HordeAgent.Parser.Interfaces;
+using HordeCommon;
+using Microsoft.Extensions.Logging;
+using System;
+using System.Collections.Generic;
+using System.Text;
+using System.Text.RegularExpressions;
+
+namespace HordeAgent.Parser.Matchers
+{
+	/// <summary>
+	/// Matches compile errors and annotates with the source file path and revision
+	/// </summary>
+	class ContentEventMatcher : ILogEventMatcher
+	{
+		const string Pattern = 
+			@"^\s*" + 
+			@"(?<channel>[a-zA-Z0-9]+):\s+" + 
+			@"(?<severity>Error:|Warning:)\s+" + 
+			@"(?:\[AssetLog\] )?" + 
+			@"(?<asset>(?:[a-zA-Z]:)?[^:]+(?:.uasset|.umap)):\s*" + 
+			@"(?<message>.*)";
+
+		/// <inheritdoc/>
+		public LogEvent? Match(ILogCursor Input, ILogContext Context)
+		{
+			// Do the match in two phases so we can early out if the strings "error" or "warning" are not present. The patterns before these strings can
+			// produce many false positives, making them very slow to execute.
+			if (Input.IsMatch("Error:|Warning:"))
+			{
+				Match? Match;
+				if(Input.TryMatch(Pattern, out Match))
+				{
+					LogEventBuilder Builder = new LogEventBuilder(Input);
+
+					LogEventLine FirstLine = Builder.Lines[0];
+					FirstLine.AddSpan(Match.Groups["channel"]).MarkAsChannel();
+					FirstLine.AddSpan(Match.Groups["severity"]).MarkAsSeverity();
+					FirstLine.AddSpan(Match.Groups["asset"]).MarkAsAsset(Context);
+					FirstLine.AddSpan(Match.Groups["message"]).MarkAsErrorMessage();
+
+					LogLevel Level = GetLogLevelFromSeverity(Match);
+					return Builder.ToLogEvent(LogEventPriority.AboveNormal, Level, KnownLogEvents.Engine_AssetLog);
+				}
+			}
+			return null;
+		}
+
+		static LogLevel GetLogLevelFromSeverity(Match Match)
+		{
+			string Severity = Match.Groups["severity"].Value;
+			if (Severity.Equals("Warning:", StringComparison.Ordinal))
+			{
+				return LogLevel.Warning;
+			}
+			else
+			{
+				return LogLevel.Error;
+			}
+		}
+	}
+}
diff --git a/Engine/Source/Programs/Horde/HordeAgent/Parser/Matchers/CrashEventMatcher.cs b/Engine/Source/Programs/Horde/HordeAgent/Parser/Matchers/CrashEventMatcher.cs
new file mode 100644
--- /dev/null
+++ b/Engine/Source/Programs/Horde/HordeAgent/Parser/Matchers/CrashEventMatcher.cs
@@ -0,0 +1,71 @@
+// Copyright Epic Games, Inc. All Rights Reserved.
+
+using HordeAgent.Parser;
+using HordeAgent.Parser.Interfaces;
+using HordeCommon;
+using Microsoft.Extensions.Logging;
+using System;
+using System.Collections.Generic;
+using System.Linq;
+using System.Text;
+using System.Text.RegularExpressions;
+using System.Threading.Tasks;
+
+namespace HordeAgent.Parser.Matchers
+{
+	/// <summary>
+	/// Matcher for engine crashes
+	/// </summary>
+	class CrashEventMatcher : ILogEventMatcher
+	{
+		public LogEvent? Match(ILogCursor Cursor, ILogContext Context)
+		{
+			if (Cursor.IsMatch("begin: stack for UAT"))
+			{
+				for (int MaxOffset = 1; MaxOffset < 100; MaxOffset++)
+				{
+					if (Cursor.IsMatch(MaxOffset, "end: stack for UAT"))
+					{
+						LogEventBuilder Builder = new LogEventBuilder(Cursor);
+						Builder.MaxOffset = MaxOffset;
+						return Builder.ToLogEvent(LogEventPriority.BelowNormal, GetLogLevel(Cursor), KnownLogEvents.Engine_Crash);
+					}
+				}
+			}
+			if (Cursor.IsMatch("AutomationTool: Stack:"))
+			{
+				LogEventBuilder Builder = new LogEventBuilder(Cursor);
+				while (Cursor.IsMatch(Builder.MaxOffset + 1, "AutomationTool: Stack:"))
+				{
+					Builder.MaxOffset++;
+				}
+				return Builder.ToLogEvent(LogEventPriority.Low, LogLevel.Error, KnownLogEvents.AutomationTool_Crash);
+			}
+
+			Match? Match;
+			if (Cursor.TryMatch("ExitCode=(3|139|255)", out Match))
+			{
+				LogEventBuilder Builder = new LogEventBuilder(Cursor);
+				Builder.Lines[0].AddSpan(Match.Groups[1], "exitCode");
+				return Builder.ToLogEvent(LogEventPriority.Low, LogLevel.Error, KnownLogEvents.AutomationTool_CrashExitCode);
+			}
+			return null;
+		}
+
+		static LogLevel GetLogLevel(ILogCursor Cursor)
+		{
+			if(Cursor.IsMatch(0, "[Ee]rror:"))
+			{
+				return LogLevel.Error;
+			}
+			else if(Cursor.IsMatch(0, "[Ww]arning:"))
+			{
+				return LogLevel.Warning;
+			}
+			else
+			{
+				return LogLevel.Information;
+			}
+		}
+	}
+}
diff --git a/Engine/Source/Programs/Horde/HordeAgent/Parser/Matchers/ExceptionEventMatcher.cs b/Engine/Source/Programs/Horde/HordeAgent/Parser/Matchers/ExceptionEventMatcher.cs
new file mode 100644
--- /dev/null
+++ b/Engine/Source/Programs/Horde/HordeAgent/Parser/Matchers/ExceptionEventMatcher.cs
@@ -0,0 +1,35 @@
+// Copyright Epic Games, Inc. All Rights Reserved.
+
+using HordeAgent.Parser;
+using HordeAgent.Parser.Interfaces;
+using HordeCommon;
+using Microsoft.Extensions.Logging;
+using System;
+using System.Collections.Generic;
+using System.Linq;
+using System.Text;
+using System.Threading.Tasks;
+
+namespace HordeAgent.Parser.Matchers
+{
+	/// <summary>
+	/// Matches a generic C# exception
+	/// </summary>
+	class ExceptionEventMatcher : ILogEventMatcher
+	{
+		/// <inheritdoc/>
+		public LogEvent? Match(ILogCursor Cursor, ILogContext Context)
+		{
+			if (Cursor.IsMatch(@"^\s*Unhandled Exception: "))
+			{
+				LogEventBuilder Builder = new LogEventBuilder(Cursor);
+				while(Cursor.IsMatch(Builder.MaxOffset + 1, @"^\s*at "))
+				{
+					Builder.MaxOffset++;
+				}
+				return Builder.ToLogEvent(LogEventPriority.Low, LogLevel.Error, KnownLogEvents.Exception);
+			}
+			return null;
+		}
+	}
+}
diff --git a/Engine/Source/Programs/Horde/HordeAgent/Parser/Matchers/ExitCodeEventMatcher.cs b/Engine/Source/Programs/Horde/HordeAgent/Parser/Matchers/ExitCodeEventMatcher.cs
new file mode 100644
--- /dev/null
+++ b/Engine/Source/Programs/Horde/HordeAgent/Parser/Matchers/ExitCodeEventMatcher.cs
@@ -0,0 +1,47 @@
+using HordeAgent.Parser.Interfaces;
+using HordeCommon;
+using Microsoft.Extensions.Logging;
+
+namespace HordeAgent.Parser.Matchers
+{
+	/// <summary>
+	/// Matcher for editor/UAT instances exiting with an error
+	/// </summary>
+	class ExitCodeEventMatcher : ILogEventMatcher
+	{
+		public LogEvent? Match(ILogCursor Cursor, ILogContext Context)
+		{
+			if (!Context.HasLoggedErrors)
+			{
+				int NumLines = 0;
+				for (; ; )
+				{
+					if (Cursor.IsMatch(NumLines, "Editor terminated with exit code [1-9]"))
+					{
+						NumLines++;
+					}
+					else if (Cursor.IsMatch(NumLines, "AutomationTool exiting with ExitCode=[1-9]"))
+					{
+						NumLines++;
+					}
+					else if (Cursor.IsMatch(NumLines, "BUILD FAILED"))
+					{
+						NumLines++;
+					}
+					else
+					{
+						break;
+					}
+				}
+
+				if (NumLines > 0)
+				{
+					LogEventBuilder Builder = new LogEventBuilder(Cursor, NumLines);
+					return Builder.ToLogEvent(LogEventPriority.Low, LogLevel.Error, KnownLogEvents.Generic);
+				}
+			}
+			return null;
+		}
+	}
+}
+
diff --git a/Engine/Source/Programs/Horde/HordeAgent/Parser/Matchers/GauntletEventMatcher.cs b/Engine/Source/Programs/Horde/HordeAgent/Parser/Matchers/GauntletEventMatcher.cs
new file mode 100644
--- /dev/null
+++ b/Engine/Source/Programs/Horde/HordeAgent/Parser/Matchers/GauntletEventMatcher.cs
@@ -0,0 +1,94 @@
+using HordeAgent.Parser.Interfaces;
+using HordeCommon;
+using Microsoft.Extensions.Logging;
+using System;
+using System.Collections.Generic;
+using System.Text;
+using System.Text.RegularExpressions;
+
+namespace HordeAgent.Parser.Matchers
+{
+	/// <summary>
+	/// Matcher for Gauntlet unit tests
+	/// </summary>
+	class GauntletEventMatcher : ILogEventMatcher
+	{
+		public LogEvent? Match(ILogCursor Cursor, ILogContext Context)
+		{
+			Match? Match;
+			if(Cursor.TryMatch(@"^(?<indent>\s*)Error: EngineTest.RunTests Group:(?<group>[^\s]+) \(", out Match))
+			{
+				string Indent = Match.Groups["indent"].Value + " ";
+
+				LogEventBuilder Builder = new LogEventBuilder(Cursor);
+				Builder.Lines[0].AddSpan(Match.Groups["group"]);
+
+				string Group = Match.Groups["group"].Value;
+
+				List<string> TestNames = new List<string>();
+
+				bool bInErrorList = false;
+				List<LogEventBuilder> ChildEvents = new List<LogEventBuilder>();
+
+				int LineCount = 1;
+				while (Cursor.IsMatch(LineCount, $"^(?:{Indent}.*|\\s*)$"))
+				{
+					if (bInErrorList)
+					{
+						Match? TestNameMatch;
+						if (Cursor.IsMatch(LineCount, @"^\s*#{1,3} "))
+						{
+							bInErrorList = false;
+						}
+						else if (Cursor.TryMatch(LineCount, @"^\s*#####\s+(?<friendly_name>.*):\s*(?<name>\S+)\s*", out TestNameMatch))
+						{
+							LogEventBuilder ChildEventBuilder = new LogEventBuilder(Cursor.Rebase(LineCount));
+							ChildEventBuilder.AddProperty("group", Group);
+
+							LogEventLine Line = ChildEventBuilder.Lines[0];
+							Line.AddSpan(TestNameMatch.Groups["name"]);
+							Line.AddSpan(TestNameMatch.Groups["friendly_name"]);
+
+							ChildEvents.Add(ChildEventBuilder);
+						}
+						else if(ChildEvents.Count > 0 && ChildEvents[ChildEvents.Count - 1].CurrentLineNumber == Cursor.CurrentLineNumber + LineCount)
+						{
+							ChildEvents[ChildEvents.Count - 1].LineCount++;
+						}
+					}
+					else
+					{
+						if (Builder.Cursor.IsMatch(LineCount, @"^\s*### The following tests failed:"))
+						{
+							bInErrorList = true;
+						}
+					}
+					LineCount++;
+				}
+				Builder.LineCount = LineCount;
+
+				LogEvent Event = Builder.ToLogEvent(LogEventPriority.High, LogLevel.Error, KnownLogEvents.Gauntlet);
+				if(ChildEvents.Count > 0)
+				{
+					Event.ChildEvents = ChildEvents.ConvertAll(x => x.ToLogEvent(LogEventPriority.High, LogLevel.Error, KnownLogEvents.Gauntlet_UnitTest));
+				}
+				return Event;
+			}
+
+			if (Cursor.TryMatch(@"Error: Screenshot '(?<screenshot>[^']+)' test failed", out Match))
+			{
+				LogEventBuilder Builder = new LogEventBuilder(Cursor);
+				Builder.Lines[0].AddSpan(Match.Groups["screenshot"]).MarkAsScreenshotTest();
+				return Builder.ToLogEvent(LogEventPriority.High, LogLevel.Error, KnownLogEvents.Gauntlet_ScreenshotTest);
+			}
+
+			if (Cursor.TryMatch("\\[ERROR\\] (.*)$", out Match))
+			{
+				LogEventBuilder Builder = new LogEventBuilder(Cursor);
+				return Builder.ToLogEvent(LogEventPriority.High, LogLevel.Error, KnownLogEvents.Generic);
+			}
+
+			return null;
+		}
+	}
+}
diff --git a/Engine/Source/Programs/Horde/HordeAgent/Parser/Matchers/GenericEventMatcher.cs b/Engine/Source/Programs/Horde/HordeAgent/Parser/Matchers/GenericEventMatcher.cs
new file mode 100644
--- /dev/null
+++ b/Engine/Source/Programs/Horde/HordeAgent/Parser/Matchers/GenericEventMatcher.cs
@@ -0,0 +1,64 @@
+// Copyright Epic Games, Inc. All Rights Reserved.
+
+using HordeAgent.Parser;
+using HordeAgent.Parser.Interfaces;
+using HordeCommon;
+using Microsoft.Extensions.Logging;
+using System;
+using System.Collections.Generic;
+using System.Linq;
+using System.Text;
+using System.Text.RegularExpressions;
+using System.Threading.Tasks;
+
+namespace HordeAgent.Parser.Matchers
+{
+	/// <summary>
+	/// Low-priority matcher for generic error strings like "warning:" and "error:"
+	/// </summary>
+	class GenericEventMatcher : ILogEventMatcher
+	{
+		/// <inheritdoc/>
+		public LogEvent? Match(ILogCursor Cursor, ILogContext Context)
+		{
+			Match? Match;
+			if (Cursor.TryMatch(@"^\s*(FATAL|fatal error):", out Match))
+			{
+				return new LogEventBuilder(Cursor).ToLogEvent(LogEventPriority.Low, LogLevel.Error, KnownLogEvents.Generic);
+			}
+			if (Cursor.IsMatch(@"(?<!\w)(ERROR|[Ee]rror) ?(\([^)]+\)|\[[^\]]+\])?: "))
+			{
+				LogEventBuilder Builder = new LogEventBuilder(Cursor);
+				while (Cursor.IsMatch(Builder.MaxOffset + 1, String.Format(@"^({0} | *$)", ExtractIndent(Cursor[0]!))))
+				{
+					Builder.MaxOffset++;
+				}
+				return Builder.ToLogEvent(LogEventPriority.Lowest, LogLevel.Error, KnownLogEvents.Generic);
+			}
+			if (Cursor.IsMatch(@"(?<!\w)(WARNING|[Ww]arning) ?(\([^)]+\)|\[[^\]]+\])?: "))
+			{
+				LogEventBuilder Builder = new LogEventBuilder(Cursor);
+				while (Cursor.IsMatch(Builder.MaxOffset + 1, String.Format(@"^({0} | *$)", ExtractIndent(Cursor[0]!))))
+				{
+					Builder.MaxOffset++;
+				}
+				return Builder.ToLogEvent(LogEventPriority.Lowest, LogLevel.Warning, KnownLogEvents.Generic);
+			}
+			if (Cursor.IsMatch(@"[Ee]rror [A-Z]\d+\s:"))
+			{
+				return new LogEventBuilder(Cursor).ToLogEvent(LogEventPriority.Lowest, LogLevel.Error, KnownLogEvents.Generic);
+			}
+			return null;
+		}
+
+		static string ExtractIndent(string Line)
+		{
+			int Length = 0;
+			while (Length < Line.Length && Line[Length] == ' ')
+			{
+				Length++;
+			}
+			return new string(' ', Length);
+		}
+	}
+}
diff --git a/Engine/Source/Programs/Horde/HordeAgent/Parser/Matchers/GradleEventMatcher.cs b/Engine/Source/Programs/Horde/HordeAgent/Parser/Matchers/GradleEventMatcher.cs
new file mode 100644
--- /dev/null
+++ b/Engine/Source/Programs/Horde/HordeAgent/Parser/Matchers/GradleEventMatcher.cs
@@ -0,0 +1,47 @@
+using HordeAgent.Parser.Interfaces;
+using HordeCommon;
+using Microsoft.Extensions.Logging;
+using System.Text.RegularExpressions;
+
+namespace HordeAgent.Parser.Matchers
+{
+	/// <summary>
+	/// Matcher for Gradle errors
+	/// </summary>
+	class GradleEventMatcher : ILogEventMatcher
+	{
+		public LogEvent? Match(ILogCursor Cursor, ILogContext Context)
+		{
+			if (Cursor.IsMatch(@"^\s*FAILURE:"))
+			{
+				int MaxOffset = 1;
+				for (; ; )
+				{
+					int NewMaxOffset = MaxOffset;
+					while (Cursor.IsMatch(NewMaxOffset, @"^\s*$"))
+					{
+						NewMaxOffset++;
+					}
+
+					Match? Match;
+					if (!Cursor.TryMatch(NewMaxOffset, @"^(\s*)\*", out Match))
+					{
+						break;
+					}
+					MaxOffset = NewMaxOffset + 1;
+
+					while (Cursor.IsMatch(MaxOffset, $"^{Match.Groups[1].Value}"))
+					{
+						MaxOffset++;
+					}
+				}
+
+				LogEventBuilder Builder = new LogEventBuilder(Cursor, MaxOffset);
+				return Builder.ToLogEvent(LogEventPriority.Normal, LogLevel.Error, KnownLogEvents.AutomationTool);
+			}
+
+			return null;
+		}
+	}
+}
+
diff --git a/Engine/Source/Programs/Horde/HordeAgent/Parser/Matchers/LinkEventMatcher.cs b/Engine/Source/Programs/Horde/HordeAgent/Parser/Matchers/LinkEventMatcher.cs
new file mode 100644
--- /dev/null
+++ b/Engine/Source/Programs/Horde/HordeAgent/Parser/Matchers/LinkEventMatcher.cs
@@ -0,0 +1,167 @@
+// Copyright Epic Games, Inc. All Rights Reserved.
+
+using HordeAgent.Parser;
+using HordeAgent.Parser.Interfaces;
+using HordeCommon;
+using Microsoft.Extensions.Logging;
+using System;
+using System.Collections.Generic;
+using System.Linq;
+using System.Text;
+using System.Text.RegularExpressions;
+using System.Threading.Tasks;
+
+namespace HordeAgent.Parser.Matchers
+{
+	/// <summary>
+	/// Matcher for linker errors
+	/// </summary>
+	class LinkEventMatcher : ILogEventMatcher
+	{
+		/// <inheritdoc/>
+		public LogEvent? Match(ILogCursor Cursor, ILogContext Context)
+		{
+			Match? Match;
+
+			int LineCount = 0;
+			bool bIsWarning = true;
+			while (Cursor.IsMatch(LineCount, @": In function |: undefined reference to |undefined symbol|^\s*(ld|ld.lld):|[^a-zA-Z]ld: |^\s*>>>"))
+			{
+				bIsWarning &= Cursor.IsMatch("warning:");
+				LineCount++;
+			}
+			if (Cursor.IsMatch(LineCount, @"error: linker command failed with exit code "))
+			{
+				bIsWarning = false;
+				LineCount++;
+			}
+			if (LineCount > 0)
+			{
+				LogEventBuilder Builder = new LogEventBuilder(Cursor);
+				Builder.LineCount = LineCount;
+				EventId EventId = AddSymbolMarkup(Builder) ? KnownLogEvents.Linker_UndefinedSymbol : KnownLogEvents.Linker;
+				return Builder.ToLogEvent(LogEventPriority.Normal, bIsWarning? LogLevel.Warning : LogLevel.Error, EventId);
+			}
+
+			if (Cursor.TryMatch(@"^(\s*)Undefined symbols for architecture", out Match))
+			{
+				LogEventBuilder Builder = new LogEventBuilder(Cursor);
+
+				string Prefix = $"^(?<prefix>{Match.Groups[1].Value}\\s+)";
+				if (Cursor.TryMatch(Builder.MaxOffset + 1, Prefix + @"""(?<symbol>[^""]+)""", out Match))
+				{
+					Prefix = $"^{Match.Groups["prefix"].Value}\\s+";
+
+					Builder.MaxOffset++;
+					Builder.Lines[Builder.MaxOffset].AddSpan(Match.Groups["symbol"]).MarkAsSymbol();
+
+					while(Cursor.TryMatch(Builder.MaxOffset + 1, Prefix + "(?<symbol>[^ ].*) in ", out Match))
+					{
+						Builder.MaxOffset++;
+//						Builder.Lines[Builder.MaxOffset].AddSpan(Match.Groups["symbol"]).MarkAsSymbol();
+					}
+				}
+
+				Builder.MaxOffset = Cursor.MatchForwards(Builder.MaxOffset, @"^\s*(ld|clang):");
+
+				AddSymbolMarkup(Builder);
+				return Builder.ToLogEvent(LogEventPriority.Normal, LogLevel.Error, KnownLogEvents.Linker_UndefinedSymbol);
+			}
+			if (Cursor.TryMatch(@"error (?<code>LNK\d+):", out Match))
+			{
+				LogEventBuilder Builder = new LogEventBuilder(Cursor);
+
+				LogEventLine FirstLine = Builder.Lines[0];
+
+				Group CodeGroup = Match.Groups["code"];
+				FirstLine.AddSpan(CodeGroup).MarkAsErrorCode();
+				AddSymbolMarkup(Builder);
+
+				if (CodeGroup.Value.Equals("LNK2001", StringComparison.Ordinal) || CodeGroup.Value.Equals("LNK2019", StringComparison.Ordinal))
+				{
+					return Builder.ToLogEvent(LogEventPriority.High, LogLevel.Error, KnownLogEvents.Linker_UndefinedSymbol);
+				}
+				else if (CodeGroup.Value.Equals("LNK2005", StringComparison.Ordinal) || CodeGroup.Value.Equals("LNK4022", StringComparison.Ordinal))
+				{
+					return Builder.ToLogEvent(LogEventPriority.High, LogLevel.Error, KnownLogEvents.Linker_DuplicateSymbol);
+				}
+				else
+				{
+					return Builder.ToLogEvent(LogEventPriority.High, LogLevel.Error, KnownLogEvents.Linker);
+				}
+			}
+
+			return null;
+		}
+
+		static bool AddSymbolMarkup(LogEventBuilder Builder)
+		{
+			bool bHasSymbols = false;
+
+			const string SpanName = "symbol";
+			foreach (LogEventLine Line in Builder.Lines)
+			{
+				if(Line.Spans.ContainsKey(SpanName))
+				{
+					continue;
+				}
+
+				string Message = Line.Text;
+
+				// Mac link error:
+				//   Undefined symbols for architecture arm64:
+				//     "Foo::Bar() const", referenced from:
+				Match SymbolMatch = Regex.Match(Message, "^  \"(.+)\"");
+				if (SymbolMatch.Success)
+				{
+					Line.AddSpan(SymbolMatch.Groups[1], SpanName).MarkAsSymbol();
+					bHasSymbols = true;
+				}
+
+				// Android link error:
+				//   Foo.o:(.data.rel.ro + 0x5d88): undefined reference to `Foo::Bar()'
+				Match UndefinedReference = Regex.Match(Message, ": undefined reference to [`']([^`']+)");
+				if (UndefinedReference.Success)
+				{
+					Line.AddSpan(UndefinedReference.Groups[1], SpanName).MarkAsSymbol();
+					bHasSymbols = true;
+				}
+
+				// LLD link error:
+				//   ld.lld.exe: error: undefined symbol: Foo::Bar() const
+				Match LldMatch = Regex.Match(Message, "error: undefined symbol:\\s*(.+)");
+				if (LldMatch.Success)
+				{
+					Line.AddSpan(LldMatch.Groups[1], SpanName).MarkAsSymbol();
+					bHasSymbols = true;
+				}
+
+				// Link error:
+				//   Link: error: L0039: reference to undefined symbol `Foo::Bar() const' in file
+				Match LinkMatch = Regex.Match(Message, ": reference to undefined symbol [`']([^`']+)");
+				if (LinkMatch.Success)
+				{
+					Line.AddSpan(LinkMatch.Groups[1], SpanName).MarkAsSymbol();
+					bHasSymbols = true;
+				}
+				Match LinkMultipleMatch = Regex.Match(Message, @": ([^\s]+) already defined in");
+				if (LinkMultipleMatch.Success)
+				{
+					Line.AddSpan(LinkMultipleMatch.Groups[1], SpanName).MarkAsSymbol();
+					bHasSymbols = true;
+				}
+
+				// Microsoft linker error:
+				//   Foo.cpp.obj : error LNK2001: unresolved external symbol \"private: virtual void __cdecl UAssetManager::InitializeAssetBundlesFromMetadata_Recursive(class UStruct const *,void const *,struct FAssetBundleData &,class FName,class TSet<void const *,struct DefaultKeyFuncs<void const *,0>,class FDefaultSetAllocator> &)const \" (?InitializeAssetBundlesFromMetadata_Recursive@UAssetManager@@EEBAXPEBVUStruct@@PEBXAEAUFAssetBundleData@@VFName@@AEAV?$TSet@PEBXU?$DefaultKeyFuncs@PEBX$0A@@@VFDefaultSetAllocator@@@@@Z)",
+				Match MicrosoftMatch = Regex.Match(Message, " symbol \"([^\"]*)\"");
+				if (MicrosoftMatch.Success)
+				{
+					Line.AddSpan(MicrosoftMatch.Groups[1], SpanName).MarkAsSymbol();
+					bHasSymbols = true;
+				}
+			}
+
+			return bHasSymbols;
+		}
+	}
+}
diff --git a/Engine/Source/Programs/Horde/HordeAgent/Parser/Matchers/LogChannelEventMatcher.cs b/Engine/Source/Programs/Horde/HordeAgent/Parser/Matchers/LogChannelEventMatcher.cs
new file mode 100644
--- /dev/null
+++ b/Engine/Source/Programs/Horde/HordeAgent/Parser/Matchers/LogChannelEventMatcher.cs
@@ -0,0 +1,59 @@
+// Copyright Epic Games, Inc. All Rights Reserved.
+
+using HordeAgent.Parser;
+using HordeAgent.Parser.Interfaces;
+using HordeCommon;
+using Microsoft.Extensions.Logging;
+using System;
+using System.Collections.Generic;
+using System.Linq;
+using System.Text;
+using System.Text.RegularExpressions;
+using System.Threading.Tasks;
+
+namespace HordeAgent.Parser.Matchers
+{
+	/// <summary>
+	/// Matches events formatted as UE log channel output
+	/// </summary>
+	class LogChannelEventMatcher : ILogEventMatcher
+	{
+		/// <inheritdoc/>
+		public LogEvent? Match(ILogCursor Input, ILogContext Context)
+		{
+			Match? Match;
+			if(Input.TryMatch(@"^(\s*)(?<channel>[a-zA-Z_][a-zA-Z0-9_]*):\s*(?<severity>Error|Warning|Display): ", out Match))
+			{
+				string Indent = Match.Groups[1].Value;
+
+				LogEventBuilder Builder = new LogEventBuilder(Input);
+
+				LogEventLine FirstLine = Builder.Lines[0];
+				FirstLine.AddSpan(Match.Groups["channel"]).MarkAsChannel();
+				FirstLine.AddSpan(Match.Groups["severity"]).MarkAsSeverity();
+
+				while (Input.IsMatch(Builder.MaxOffset + 1, $"^({Indent} | *$)"))
+				{
+					Builder.MaxOffset++;
+				}
+
+				LogLevel Level;
+				switch(Match.Groups["severity"].Value)
+				{
+					case "Error":
+						Level = LogLevel.Error;
+						break;
+					case "Warning":
+						Level = LogLevel.Warning;
+						break;
+					default:
+						Level = LogLevel.Information;
+						break;
+				}
+
+				return Builder.ToLogEvent(LogEventPriority.Low, Level, KnownLogEvents.Engine_LogChannel);
+			}
+			return null;
+		}
+	}
+}
diff --git a/Engine/Source/Programs/Horde/HordeAgent/Parser/Matchers/MicrosoftEventMatcher.cs b/Engine/Source/Programs/Horde/HordeAgent/Parser/Matchers/MicrosoftEventMatcher.cs
new file mode 100644
--- /dev/null
+++ b/Engine/Source/Programs/Horde/HordeAgent/Parser/Matchers/MicrosoftEventMatcher.cs
@@ -0,0 +1,53 @@
+// Copyright Epic Games, Inc. All Rights Reserved.
+
+using HordeAgent.Parser.Interfaces;
+using HordeCommon;
+using Microsoft.Extensions.Logging;
+using System;
+using System.Collections.Generic;
+using System.Text;
+using System.Text.RegularExpressions;
+
+namespace HordeAgent.Parser.Matchers
+{
+	/// <summary>
+	/// Matcher for generic Microsoft errors (see https://docs.microsoft.com/en-us/cpp/build/formatting-the-output-of-a-custom-build-step-or-build-event?view=msvc-160)
+	/// </summary>
+	class MicrosoftEventMatcher : ILogEventMatcher
+	{
+		public LogEvent? Match(ILogCursor Cursor, ILogContext Context)
+		{
+			// filename(line# [, column#]) | toolname} : [ any text ] {error | warning} code+number:localizable string [ any text ]
+
+			Match? Match;
+			if (Cursor.TryMatch(@"(?<severity>(?:error|warning)) (?<code>[a-zA-Z]+[0-9]+)\s*:", out Match))
+			{
+				string Prefix = Cursor.CurrentLine!.Substring(0, Match.Index);
+
+				Match FileOrToolMatch = Regex.Match(Prefix, @"^\s*(.*[^\s])\s*:");
+				if (FileOrToolMatch.Success)
+				{
+					LogEventBuilder Builder = new LogEventBuilder(Cursor);
+					LogEventLine Line = Builder.Lines[0];
+
+					Match FileMatch = Regex.Match(FileOrToolMatch.Value, @"^\s*(?<file>.*)\((?<line>\d+)(?:, (?<column>\d+))?\)\s*:$");
+					if (FileMatch.Success)
+					{
+						Line.AddSpan(FileMatch.Groups["file"]).MarkAsSourceFile(Context, "");
+						Line.AddSpan(FileMatch.Groups["line"]).MarkAsLineNumber();
+						Line.TryAddSpan(FileMatch.Groups["column"]);
+					}
+					else
+					{
+						Line.AddSpan(FileOrToolMatch.Groups[1], "tool").MarkAsToolName();
+					}
+
+					LogEventSpan Severity = Line.AddSpan(Match.Groups["severity"]);
+					Line.AddSpan(Match.Groups["code"]);
+					return Builder.ToLogEvent(LogEventPriority.Normal, Severity.Text.Equals("error")? LogLevel.Error : LogLevel.Warning, KnownLogEvents.Microsoft);
+				}
+			}
+			return null;
+		}
+	}
+}
diff --git a/Engine/Source/Programs/Horde/HordeAgent/Parser/Matchers/SanitizerEventMatcher.cs b/Engine/Source/Programs/Horde/HordeAgent/Parser/Matchers/SanitizerEventMatcher.cs
new file mode 100644
--- /dev/null
+++ b/Engine/Source/Programs/Horde/HordeAgent/Parser/Matchers/SanitizerEventMatcher.cs
@@ -0,0 +1,30 @@
+// Copyright Epic Games, Inc. All Rights Reserved.
+
+#if false
+using HordeAgent.Parser;
+using HordeAgent.Parser.Interfaces;
+using HordeCommon;
+using System;
+using System.Collections.Generic;
+using System.Linq;
+using System.Text;
+using System.Text.RegularExpressions;
+using System.Threading.Tasks;
+
+namespace HordeAgent.Parser.Matchers
+{
+	class SanitizerEventMatcher : IEventMatcher
+	{
+		public EventMatch? Match(ReadOnlyLineBuffer Input)
+		{
+			Match? Match;
+			if (Input.TryMatch(@"^(\s*)WARNING: ThreadSanitizer:", out Match))
+			{
+				int EndIdx = Input.MatchForwards(0, String.Format(@"^([ ]*|{0}  .*|{0}SUMMARY:.*)\$", Match!.Groups[1].Value));
+				return new EventMatch(EventSeverity.Warning, EventMatchPriority.Normal, "ThreadSanitizer", Input, 0, EndIdx);
+			}
+			return null;
+		}
+	}
+}
+#endif
diff --git a/Engine/Source/Programs/Horde/HordeAgent/Parser/Matchers/SourceFileLineEventMatcher.cs b/Engine/Source/Programs/Horde/HordeAgent/Parser/Matchers/SourceFileLineEventMatcher.cs
new file mode 100644
--- /dev/null
+++ b/Engine/Source/Programs/Horde/HordeAgent/Parser/Matchers/SourceFileLineEventMatcher.cs
@@ -0,0 +1,76 @@
+using HordeAgent.Parser.Interfaces;
+using HordeCommon;
+using Microsoft.Extensions.Logging;
+using System;
+using System.Collections.Generic;
+using System.Runtime.CompilerServices;
+using System.Text;
+using System.Text.RegularExpressions;
+
+namespace HordeAgent.Parser.Matchers
+{
+	/// <summary>
+	/// Matches compile errors and annotates with the source file path and revision
+	/// </summary>
+	class SourceFileLineEventMatcher : ILogEventMatcher
+	{
+		const string SeverityPattern =
+			"(?<severity>ERROR|WARNING)";
+
+		const string FilePattern =
+			@"(?<file>" +
+				// optional drive letter
+				@"(?:[a-zA-Z]:)?" +
+				// any non-colon character
+				@"[^:(\s]+" +
+				// any filename character (not whitespace or slash)
+				@"[^:(\s\\/]" +
+			@")";
+
+		const string LinePattern =
+			@"(?<line>\d+)";
+
+		/// <inheritdoc/>
+		public LogEvent? Match(ILogCursor Input, ILogContext Context)
+		{
+			Match? Match;
+			if (Input.TryMatch($"^\\s*{SeverityPattern}: {FilePattern}(?:\\({LinePattern}\\))?: ", out Match))
+			{
+				LogLevel Level = GetLogLevelFromSeverity(Match);
+
+				LogEventBuilder Builder = new LogEventBuilder(Input);
+
+				LogEventLine FirstLine = Builder.Lines[0];
+				FirstLine.AddSpan(Match.Groups["file"]).MarkAsSourceFile(Context, "");
+				FirstLine.AddSpan(Match.Groups["severity"]).MarkAsSeverity();
+				FirstLine.TryAddSpan(Match.Groups["line"])?.MarkAsLineNumber();
+
+				EventId EventId;
+				if (Input.IsMatch("copyright"))
+				{
+					EventId = KnownLogEvents.AutomationTool_MissingCopyright;
+				}
+				else
+				{
+					EventId = KnownLogEvents.AutomationTool_SourceFileLine;
+				}
+
+				return Builder.ToLogEvent(LogEventPriority.AboveNormal, Level, EventId);
+			}
+			return null;
+		}
+
+		static LogLevel GetLogLevelFromSeverity(Match Match)
+		{
+			string Severity = Match.Groups["severity"].Value;
+			if (Severity.Equals("WARNING", StringComparison.Ordinal))
+			{
+				return LogLevel.Warning;
+			}
+			else
+			{
+				return LogLevel.Error;
+			}
+		}
+	}
+}
diff --git a/Engine/Source/Programs/Horde/HordeAgent/Parser/Matchers/SystemicEventMatcher.cs b/Engine/Source/Programs/Horde/HordeAgent/Parser/Matchers/SystemicEventMatcher.cs
new file mode 100644
--- /dev/null
+++ b/Engine/Source/Programs/Horde/HordeAgent/Parser/Matchers/SystemicEventMatcher.cs
@@ -0,0 +1,36 @@
+// Copyright Epic Games, Inc. All Rights Reserved.
+
+using HordeAgent.Parser;
+using HordeAgent.Parser.Interfaces;
+using HordeCommon;
+using Microsoft.Extensions.Logging;
+using System;
+using System.Collections.Generic;
+using System.Linq;
+using System.Text;
+using System.Text.RegularExpressions;
+using System.Threading.Tasks;
+
+namespace HordeAgent.Parser.Matchers
+{
+	class SystemicEventMatcher : ILogEventMatcher
+	{
+		public LogEvent? Match(ILogCursor Cursor, ILogContext Context)
+		{
+			if (Cursor.IsMatch(@"LogDerivedDataCache: .*queries/writes will be limited"))
+			{
+				return new LogEventBuilder(Cursor).ToLogEvent(LogEventPriority.High, LogLevel.Information, KnownLogEvents.Systemic_SlowDDC);
+			}
+			if (Cursor.IsMatch(@"^\s*ERROR: Error: EC_OK"))
+			{
+				LogEventBuilder Builder = new LogEventBuilder(Cursor);
+				while (Builder.IsMatch(Builder.MaxOffset + 1, @"^\s*ERROR:\s*$"))
+				{
+					Builder.MaxOffset++;
+				}
+				return new LogEventBuilder(Cursor).ToLogEvent(LogEventPriority.High, LogLevel.Information, KnownLogEvents.Systemic_PdbUtil);
+			}
+			return null;
+		}
+	}
+}
diff --git a/Engine/Source/Programs/Horde/HordeAgent/Parser/Matchers/XoreaxEventMatcher.cs b/Engine/Source/Programs/Horde/HordeAgent/Parser/Matchers/XoreaxEventMatcher.cs
new file mode 100644
--- /dev/null
+++ b/Engine/Source/Programs/Horde/HordeAgent/Parser/Matchers/XoreaxEventMatcher.cs
@@ -0,0 +1,48 @@
+// Copyright Epic Games, Inc. All Rights Reserved.
+
+using HordeAgent.Parser;
+using HordeAgent.Parser.Interfaces;
+using HordeCommon;
+using Microsoft.Extensions.Logging;
+using System;
+using System.Collections.Generic;
+using System.Linq;
+using System.Text;
+using System.Text.RegularExpressions;
+using System.Threading.Tasks;
+
+namespace HordeAgent.Parser.Matchers
+{
+	class XoreaxEventMatcher : ILogEventMatcher
+	{
+		public LogEvent? Match(ILogCursor Cursor, ILogContext Context)
+		{
+			// Return it as information, but add a property indicating it's something that should be 
+			if (Cursor.IsMatch(@"^\s*--------------------Build System Warning[- ]"))
+			{
+				LogEventBuilder Builder = new LogEventBuilder(Cursor);
+				if (Builder.TryMatch(@"^(\s*)([^ ].*):", out Match? Prefix))
+				{
+					Builder.MaxOffset++;
+
+					string Message = Prefix.Groups[2].Value;
+
+					EventId EventId = KnownLogEvents.Systemic_Xge;
+					if (Regex.IsMatch(Message, "Failed to connect to Coordinator"))
+					{
+						EventId = KnownLogEvents.Systemic_Xge_Standalone;
+					}
+
+					string PrefixPattern = $"^{Prefix.Groups[1].Value}\\s";
+					while (Builder.IsMatch(PrefixPattern))
+					{
+						Builder.MaxOffset++;
+					}
+
+					return Builder.ToLogEvent(LogEventPriority.High, LogLevel.Information, EventId);
+				}
+			}
+			return null;
+		}
+	}
+}
diff --git a/Engine/Source/Programs/Horde/HordeAgent/Program.cs b/Engine/Source/Programs/Horde/HordeAgent/Program.cs
new file mode 100644
--- /dev/null
+++ b/Engine/Source/Programs/Horde/HordeAgent/Program.cs
@@ -0,0 +1,238 @@
+// Copyright Epic Games, Inc. All Rights Reserved.
+
+using HordeAgent.Commands;
+using HordeAgent.Services;
+using HordeAgent.Utility;
+using Microsoft.Extensions.Configuration;
+using Microsoft.Extensions.DependencyInjection;
+using Microsoft.Extensions.Hosting;
+using Microsoft.Extensions.Hosting.WindowsServices;
+using Microsoft.Extensions.Logging;
+using Polly;
+using System;
+using System.Collections.Generic;
+using System.Diagnostics;
+using System.IO;
+using System.Linq;
+using System.Management;
+using System.Net;
+using System.Net.Http.Headers;
+using System.Reflection;
+using System.Runtime.InteropServices;
+using System.ServiceProcess;
+using System.Text.Json;
+using System.Threading;
+using System.Threading.Tasks;
+using EpicGames.Core;
+
+namespace HordeAgent
+{
+	/// <summary>
+	/// Entry point
+	/// </summary>
+	public static class Program
+	{
+		/// <summary>
+		/// Name of the http client
+		/// </summary>
+		public const string HordeServerClientName = "HordeServer";
+
+		/// <summary>
+		/// Path to the root application directory
+		/// </summary>
+		public static readonly string AppDir = Path.GetFullPath(Path.GetDirectoryName(Assembly.GetExecutingAssembly().Location)!);
+
+		/// <summary>
+		/// Name of the saved application data directory
+		/// </summary>
+		public const string SavedDirName = "Saved";
+
+		/// <summary>
+		/// Path to the saved application data directory
+		/// </summary>
+		public static readonly string SavedDir = Path.Combine(AppDir, SavedDirName);
+
+		/// <summary>
+		/// The launch arguments
+		/// </summary>
+		public static string[] Args { get; private set; } = null!;
+
+		/// <summary>
+		/// Width to use for printing out help
+		/// </summary>
+		static int HelpWidth
+		{
+			get { return HelpUtils.WindowWidth - 20; }
+		}
+
+		/// <summary>
+		/// Entry point
+		/// </summary>
+		/// <param name="Args">Command-line arguments</param>
+		/// <returns>Exit code</returns>
+		public static async Task<int> Main(string[] Args)
+		{
+			Program.Args = Args;
+
+			ILogger Logger = new Logging.HordeLoggerProvider().CreateLogger("HordeAgent");
+			try
+			{
+				int Result = await GuardedMain(Args, Logger);
+				return Result;
+			}
+			catch (FatalErrorException Ex)
+			{
+				Logger.LogCritical(Ex, "Fatal error.");
+				return Ex.ExitCode;
+			}
+			catch (Exception Ex)
+			{
+				Logger.LogCritical(Ex, "Fatal error.");
+				return 1;
+			}
+		}
+
+		static bool MatchCommand(string[] Args, CommandAttribute Attribute)
+		{
+			if(Args.Length < Attribute.Names.Length)
+			{
+				return false;
+			}
+
+			for (int Idx = 0; Idx < Attribute.Names.Length; Idx++)
+			{
+				if (!Attribute.Names[Idx].Equals(Args[Idx], StringComparison.OrdinalIgnoreCase))
+				{
+					return false;
+				}
+			}
+
+			return true;
+		}
+
+		/// <summary>
+		/// Actual Main function, without exception guards
+		/// </summary>
+		/// <param name="Args">Command-line arguments</param>
+		/// <param name="Logger">The logger interface</param>
+		/// <returns>Exit code</returns>
+		static async Task<int> GuardedMain(string[] Args, ILogger Logger)
+		{
+			// Find all the command types
+			List<(CommandAttribute, Type)> Commands = new List<(CommandAttribute, Type)>();
+			foreach (Type Type in Assembly.GetExecutingAssembly().GetTypes())
+			{
+				CommandAttribute? Attribute = Type.GetCustomAttribute<CommandAttribute>();
+				if (Attribute != null)
+				{
+					Commands.Add((Attribute, Type));
+				}
+			}
+
+			// Check if there's a matching command
+			Type? CommandType = null;
+			CommandAttribute? CommandAttribute = null;
+			foreach ((CommandAttribute Attribute, Type Type) in Commands.OrderBy(x => x.Item1.Names.Length))
+			{
+				if (MatchCommand(Args, Attribute))
+				{
+					CommandType = Type;
+					CommandAttribute = Attribute;
+				}
+			}
+
+			// Check if there are any commands specified on the command line.
+			if (CommandType == null)
+			{
+				if (Args.Length > 0 && !Args[0].StartsWith("-"))
+				{
+					Logger.LogError("Invalid command");
+					Logger.LogInformation("");
+					Logger.LogInformation("Available commands");
+
+					PrintCommands(Commands.Select(x => x.Item1), Logger);
+					return 1;
+				}
+				else
+				{
+					Logger.LogInformation("HordeAgent");
+					Logger.LogInformation("");
+					Logger.LogInformation("Utility for managing automated processes on build machines.");
+					Logger.LogInformation("");
+					Logger.LogInformation("Usage:");
+					Logger.LogInformation("    HordeAgent.exe [Command] [-Option1] [-Option2]...");
+					Logger.LogInformation("");
+					Logger.LogInformation("Commands:");
+
+					PrintCommands(Commands.Select(x => x.Item1), Logger);
+
+					Logger.LogInformation("");
+					Logger.LogInformation("Specify \"Command -Help\" for command-specific help");
+					return 0;
+				}
+			}
+
+			// Extract the 
+			string CommandName = String.Join(" ", CommandAttribute!.Names);
+
+			// Build an argument list for the command, including all the global arguments as well as arguments until the next command
+			CommandLineArguments CommandArguments = new CommandLineArguments(Args.Skip(CommandAttribute.Names.Length).ToArray());
+
+			// Create the command instance
+			Command Command = (Command)Activator.CreateInstance(CommandType)!;
+
+			// If the help flag is specified, print the help info and exit immediately
+			if (CommandArguments.HasOption("-Help"))
+			{
+				HelpUtils.PrintHelp(CommandName, CommandAttribute.Description, Command.GetParameters(CommandArguments), HelpWidth, Logger);
+				return 1;
+			}
+
+			// Configure the command
+			try
+			{
+				Command.Configure(CommandArguments, Logger);
+				CommandArguments.CheckAllArgumentsUsed(Logger);
+			}
+			catch (CommandLineArgumentException Ex)
+			{
+				Logger.LogError("{0}: {1}", CommandName, Ex.Message);
+
+				Logger.LogInformation("");
+				Logger.LogInformation("Arguments for {0}:", CommandName);
+
+				HelpUtils.PrintTable(Command.GetParameters(CommandArguments), 4, 24, HelpWidth, Logger);
+				return 1;
+			}
+
+			if (RuntimeInformation.IsOSPlatform(OSPlatform.Windows))
+			{
+				// Prioritize agent execution time over any job its running.
+				// We've seen file copying starving the agent communication to the Horde server, causing a disconnect.
+				// Increasing the process priority is speculative fix to combat this.
+				using (Process Process = Process.GetCurrentProcess())
+				{
+					Process.PriorityClass = ProcessPriorityClass.High;
+				}
+			}
+
+			// Execute all the commands
+			return await Command.ExecuteAsync(Logger);
+		}
+
+		/// <summary>
+		/// Print a formatted list of all the available commands
+		/// </summary>
+		/// <param name="Attributes">List of command attributes</param>
+		/// <param name="Logger">The logging output device</param>
+		static void PrintCommands(IEnumerable<CommandAttribute> Attributes, ILogger Logger)
+		{
+			List<KeyValuePair<string, string>> Commands = new List<KeyValuePair<string, string>>();
+			foreach(CommandAttribute Attribute in Attributes)
+			{
+				Commands.Add(new KeyValuePair<string, string>(String.Join(" ", Attribute.Names), Attribute.Description));
+			}
+			HelpUtils.PrintTable(Commands.OrderBy(x => x.Key).ToList(), 4, 20, HelpWidth, Logger);
+		}
+	}
+}
diff --git a/Engine/Source/Programs/Horde/HordeAgent/Properties/AssemblyInfo.cs b/Engine/Source/Programs/Horde/HordeAgent/Properties/AssemblyInfo.cs
new file mode 100644
--- /dev/null
+++ b/Engine/Source/Programs/Horde/HordeAgent/Properties/AssemblyInfo.cs
@@ -0,0 +1,3 @@
+using System.Runtime.CompilerServices;
+
+[assembly: InternalsVisibleTo("HordeAgentTests")]
diff --git a/Engine/Source/Programs/Horde/HordeAgent/Services/GrpcService.cs b/Engine/Source/Programs/Horde/HordeAgent/Services/GrpcService.cs
new file mode 100644
--- /dev/null
+++ b/Engine/Source/Programs/Horde/HordeAgent/Services/GrpcService.cs
@@ -0,0 +1,80 @@
+using Grpc.Net.Client;
+using HordeAgent.Utility;
+using Microsoft.Extensions.Http;
+using Microsoft.Extensions.Logging;
+using Microsoft.Extensions.Options;
+using Polly;
+using Polly.Extensions.Http;
+using System;
+using System.Collections.Generic;
+using System.Net.Http;
+using System.Net.Http.Headers;
+using System.Text;
+
+namespace HordeAgent.Services
+{
+	/// <summary>
+	/// Service which creates a configured Grpc channel
+	/// </summary>
+	class GrpcService
+	{
+		/// <summary>
+		/// Settings for the current server
+		/// </summary>
+		ServerProfile ServerProfile;
+
+		/// <summary>
+		/// Options for the agent
+		/// </summary>
+		IOptions<AgentSettings> Settings;
+
+		/// <summary>
+		/// Logger instance
+		/// </summary>
+		ILogger Logger;
+
+		/// <summary>
+		/// Constructor
+		/// </summary>
+		/// <param name="Settings"></param>
+		/// <param name="Logger"></param>
+		public GrpcService(IOptions<AgentSettings> Settings, ILogger<GrpcService> Logger)
+		{
+			this.Settings = Settings;
+			this.ServerProfile = Settings.Value.GetCurrentServerProfile();
+			this.Logger = Logger;
+		}
+
+		/// <summary>
+		/// Create a GRPC channel with a default token
+		/// </summary>
+		/// <returns>New grpc channel</returns>
+		public GrpcChannel CreateGrpcChannel()
+		{
+			return CreateGrpcChannel(ServerProfile.Token);
+		}
+
+		/// <summary>
+		/// Create a GRPC channel with the given bearer token
+		/// </summary>
+		/// <returns>New grpc channel</returns>
+		public GrpcChannel CreateGrpcChannel(string BearerToken)
+		{
+			HttpClientHandler CustomCertHandler = new HttpClientHandler();
+			CustomCertHandler.ServerCertificateCustomValidationCallback += (Sender, Cert, Chain, Errors) => CertificateHelper.CertificateValidationCallBack(Logger, Sender, Cert, Chain, Errors, ServerProfile);
+
+			TimeSpan[] RetryDelay = { TimeSpan.FromSeconds(1), TimeSpan.FromSeconds(5), TimeSpan.FromSeconds(10) };
+			IAsyncPolicy<HttpResponseMessage> Policy = HttpPolicyExtensions.HandleTransientHttpError().WaitAndRetryAsync(RetryDelay);
+			PolicyHttpMessageHandler RetryHandler = new PolicyHttpMessageHandler(Policy);
+			RetryHandler.InnerHandler = CustomCertHandler;
+
+			HttpClient HttpClient = new HttpClient(RetryHandler);
+			HttpClient.DefaultRequestHeaders.Add("Accept", "application/json");
+			HttpClient.DefaultRequestHeaders.Authorization = new AuthenticationHeaderValue("Bearer", BearerToken);
+			HttpClient.Timeout = TimeSpan.FromSeconds(210); // Need to make sure this doesn't cancel any long running gRPC streaming calls (eg. session update)
+
+			Logger.LogInformation("Connecting to rpc server {BaseUrl}", ServerProfile.Url);
+			return GrpcChannel.ForAddress(ServerProfile.Url, new GrpcChannelOptions { HttpClient = HttpClient, DisposeHttpClient = true });
+		}
+	}
+}
