Initial drop of Horde Server/Agent for UE5/Main. Still missing some Epic-specific functionality required to build/run correctly.

diff --git a/Engine/Source/Programs/Horde/HordeAgent/Services/WorkerService.cs b/Engine/Source/Programs/Horde/HordeAgent/Services/WorkerService.cs
new file mode 100644
--- /dev/null
+++ b/Engine/Source/Programs/Horde/HordeAgent/Services/WorkerService.cs
@@ -0,0 +1,1611 @@
+// Copyright Epic Games, Inc. All Rights Reserved.
+
+using Google.Protobuf.WellKnownTypes;
+using Grpc.Core;
+using Grpc.Net.Client;
+using HordeAgent.Execution;
+using HordeAgent.Execution.Interfaces;
+using HordeAgent.Parser;
+using HordeAgent.Parser.Interfaces;
+using HordeAgent.Modes.Service;
+using HordeAgent.Utility;
+using HordeCommon;
+using HordeCommon.Rpc;
+using HordeCommon.Rpc.Messages;
+using HordeCommon.Rpc.Tasks;
+using Microsoft.Extensions.Hosting;
+using Microsoft.Extensions.Logging;
+using Microsoft.Extensions.Options;
+using System;
+using System.Collections.Concurrent;
+using System.Collections.Generic;
+using System.Diagnostics;
+using System.IO;
+using System.IO.Compression;
+using System.Linq;
+using System.Management;
+using System.Net;
+using System.Net.Http;
+using System.Net.Http.Headers;
+using System.Net.Sockets;
+using System.Reflection;
+using System.Runtime.InteropServices;
+using System.Security.Cryptography;
+using System.Security.Cryptography.X509Certificates;
+using System.ServiceProcess;
+using System.Text;
+using System.Text.Json;
+using System.Threading;
+using System.Threading.Tasks;
+using System.Xml;
+using EpicGames.Core;
+using System.ComponentModel;
+using System.Diagnostics.CodeAnalysis;
+using System.Runtime.CompilerServices;
+using Datadog.Trace;
+using Serilog.Context;
+using Polly;
+using Microsoft.Extensions.Http;
+using Polly.Extensions.Http;
+using Amazon.Util;
+using Amazon.EC2;
+using Scope = Datadog.Trace.Scope;
+using Amazon.EC2.Model;
+
+[assembly: InternalsVisibleTo("HordeAgentTests")]
+
+namespace HordeAgent.Services
+{
+	/// <summary>
+	/// Implements the message handling loop for an agent. Runs asynchronously until disposed.
+	/// </summary>
+	class WorkerService : BackgroundService, IDisposable
+	{
+		/// <summary>
+		/// Stores information about an active session
+		/// </summary>
+		class LeaseInfo
+		{
+			/// <summary>
+			/// The worker lease state
+			/// </summary>
+			public Lease Lease;
+
+			/// <summary>
+			/// The task being executed for this lease
+			/// </summary>
+			public Task? Task;
+
+			/// <summary>
+			/// Source for cancellation tokens for this session.
+			/// </summary>
+			public CancellationTokenSource CancellationTokenSource;
+
+			/// <summary>
+			/// Constructor
+			/// </summary>
+			/// <param name="Lease">The worker lease state</param>
+			public LeaseInfo(Lease Lease)
+			{
+				this.Lease = Lease;
+				this.CancellationTokenSource = new CancellationTokenSource();
+			}
+		}
+
+		/// <summary>
+		/// List of processes that should be terminated before running a job
+		/// </summary>
+		private HashSet<string> ProcessNamesToTerminate = new HashSet<string>(StringComparer.OrdinalIgnoreCase);
+
+		/// <summary>
+		/// Object used for controlling access to the access tokens and active sessions list
+		/// </summary>
+		object LockObject = new object();
+
+		/// <summary>
+		/// Sink for log messages
+		/// </summary>
+		ILogger<WorkerService> Logger;
+
+		/// <summary>
+		/// The current settings instance
+		/// </summary>
+		AgentSettings Settings;
+
+		/// <summary>
+		/// Settings for the current server
+		/// </summary>
+		ServerProfile ServerProfile;
+
+		/// <summary>
+		/// The grpc service instance
+		/// </summary>
+		GrpcService GrpcService;
+
+		/// <summary>
+		/// The working directory
+		/// </summary>
+		DirectoryReference WorkingDir;
+
+		/// <summary>
+		/// The list of active leases.
+		/// </summary>
+		List<LeaseInfo> ActiveLeases = new List<LeaseInfo>();
+
+		/// <summary>
+		/// Whether the agent is currently in an unhealthy state
+		/// </summary>
+		bool Unhealthy = false;
+
+		/// <summary>
+		/// Whether the 
+		/// </summary>
+		bool RequestShutdown = false;
+
+		/// <summary>
+		/// Whether to restart after shutting down
+		/// </summary>
+		bool RestartAfterShutdown = false;
+
+		/// <summary>
+		/// Time the the service started
+		/// </summary>
+		static readonly DateTimeOffset StartTime = DateTimeOffset.Now;
+
+		/// <summary>
+		/// Time at which the computer started
+		/// </summary>
+		static readonly DateTimeOffset BootTime = DateTimeOffset.Now - TimeSpan.FromTicks(Environment.TickCount64 * TimeSpan.TicksPerMillisecond);
+
+		/// <summary>
+		/// Task completion source used to trigger the background thread to update the leases. Must take a lock on LockObject before 
+		/// </summary>
+		AsyncEvent UpdateLeasesEvent = new AsyncEvent();
+
+		/// <summary>
+		/// Number of times UpdateSession has failed
+		/// </summary>
+		int UpdateSessionFailures;
+
+		/// <summary>
+		/// Function for creating a new executor. Primarily done this way to aid testing. 
+		/// </summary>
+		private readonly Func<IRpcConnection, ExecuteJobTask, BeginBatchResponse, IExecutor> CreateExecutor;
+
+		/// <summary>
+		/// How often to poll the server checking if a step has been aborted
+		/// Exposed as internal to ease testing. 
+		/// </summary>
+		internal TimeSpan StepAbortPollInterval = TimeSpan.FromSeconds(5);
+
+		/// <summary>
+		/// Constructor. Registers with the server and starts accepting connections.
+		/// </summary>
+		/// <param name="Logger">Log sink</param>
+		/// <param name="OptionsMonitor">The current settings</param>
+		/// <param name="GrpcService">Instance of the Grpc service</param>
+		/// <param name="CreateExecutor"></param>
+		public WorkerService(ILogger<WorkerService> Logger, IOptionsMonitor<AgentSettings> OptionsMonitor, GrpcService GrpcService, Func<IRpcConnection, ExecuteJobTask, BeginBatchResponse, IExecutor>? CreateExecutor = null)
+		{
+			this.Logger = Logger;
+			this.Settings = OptionsMonitor.CurrentValue;
+			this.ServerProfile = Settings.GetCurrentServerProfile();
+			this.GrpcService = GrpcService;
+
+			if (Settings.WorkingDir == null)
+			{
+				throw new Exception("WorkingDir is not set. Unable to run service.");
+			}
+
+			DirectoryReference BaseDir = new FileReference(Assembly.GetExecutingAssembly().Location).Directory;
+
+			WorkingDir = DirectoryReference.Combine(BaseDir, Settings.WorkingDir);
+			Logger.LogInformation("Using working directory {WorkingDir}", WorkingDir);
+			DirectoryReference.CreateDirectory(WorkingDir);
+
+			ProcessNamesToTerminate.UnionWith(Settings.ProcessNamesToTerminate);
+
+			if (CreateExecutor == null)
+			{
+				this.CreateExecutor = Settings.Executor switch
+				{
+					ExecutorType.Test => (RpcClient, ExecuteTask, Batch) =>
+						new TestExecutor(RpcClient, ExecuteTask.JobId, ExecuteTask.BatchId, Batch.AgentType),
+					ExecutorType.Local => (RpcClient, ExecuteTask, Batch) =>
+						new LocalExecutor(RpcClient, ExecuteTask.JobId, ExecuteTask.BatchId, Batch.AgentType, Settings.LocalExecutor),
+					ExecutorType.Perforce => (RpcClient, ExecuteTask, Batch) =>
+						new PerforceExecutor(RpcClient, ExecuteTask.JobId, ExecuteTask.BatchId, Batch.AgentType, ExecuteTask.AutoSdkWorkspace, ExecuteTask.Workspace, WorkingDir),
+					_ => throw new InvalidDataException($"Unknown executor type '{Settings.Executor}'")
+				};
+			}
+			else
+			{
+				this.CreateExecutor = CreateExecutor;
+			}
+		}
+
+		/// <summary>
+		/// Executes the ServerTaskAsync method and swallows the exception for the task being cancelled. This allows waiting for it to terminate.
+		/// </summary>
+		/// <param name="StoppingToken">Indicates that the service is trying to stop</param>
+		protected override async Task ExecuteAsync(CancellationToken StoppingToken)
+		{
+			try
+			{
+				await ExecuteInnerAsync(StoppingToken);
+			}
+			catch (Exception Ex)
+			{
+				Logger.LogCritical(Ex, "Unhandled exception");
+			}
+		}
+
+		/// <summary>
+		/// Background task to cycle access tokens and update the state of the agent with the server.
+		/// </summary>
+		/// <param name="StoppingToken">Indicates that the service is trying to stop</param>
+		async Task ExecuteInnerAsync(CancellationToken StoppingToken)
+		{
+			// Print the server info
+			Logger.LogInformation("Server: {Server}", ServerProfile.Url);
+			Logger.LogInformation("Arguments: {Arguments}", Environment.CommandLine);
+
+			// Show the current client id
+			string Version = Settings.Version ?? "(unknown)";
+			Logger.LogInformation("Version: {Version}", Version);
+
+			// Keep trying to start an agent session with the server
+			while (!StoppingToken.IsCancellationRequested)
+			{
+				Stopwatch SessionTime = Stopwatch.StartNew();
+				try
+				{
+					await HandleSessionAsync(StoppingToken);
+				}
+				catch (Exception Ex)
+				{
+					if (StoppingToken.IsCancellationRequested && IsCancellationException(Ex))
+					{
+						break;
+					}
+					else
+					{
+						Logger.LogError(Ex, "Exception while executing session. Restarting.");
+					}
+				}
+
+				if (RequestShutdown)
+				{
+					Logger.LogInformation("Initiating shutdown (restart={Restart})", RestartAfterShutdown);
+					if (Shutdown.InitiateShutdown(RestartAfterShutdown, Logger))
+					{
+						for (int Idx = 10; Idx > 0; Idx--)
+						{
+							Logger.LogInformation("Waiting for shutdown ({Count})", Idx);
+							await Task.Delay(TimeSpan.FromSeconds(60.0));
+						}
+						Logger.LogInformation("Shutdown aborted.");
+					}
+					RequestShutdown = RestartAfterShutdown = false;
+				}
+				else if (SessionTime.Elapsed < TimeSpan.FromSeconds(2.0))
+				{
+					Logger.LogInformation("Waiting 5 seconds before restarting session...");
+					await Task.Delay(TimeSpan.FromSeconds(5.0));
+				}
+			}
+		}
+
+		async Task WaitForMutexAsync(Mutex Mutex, CancellationToken StoppingToken)
+		{
+			try
+			{
+				if (!Mutex.WaitOne(0))
+				{
+					Logger.LogError("Another instance of HordeAgent is already running. Waiting until it terminates.");
+					while (!Mutex.WaitOne(0))
+					{
+						StoppingToken.ThrowIfCancellationRequested();
+						await Task.Delay(TimeSpan.FromSeconds(1.0), StoppingToken);
+					}
+				}
+			}
+			catch (AbandonedMutexException)
+			{
+			}
+		}
+
+		/// <summary>
+		/// Handles the lifetime of an agent session.
+		/// </summary>
+		/// <param name="StoppingToken">Indicates that the service is trying to stop</param>
+		/// <returns>Async task</returns>
+		async Task HandleSessionAsync(CancellationToken StoppingToken)
+		{
+			// Make sure there's only one instance of the agent running
+			using Mutex SingleInstanceMutex = new Mutex(false, "Global\\HordeAgent-DB828ACB-0AA5-4D32-A62A-21D4429B1014");
+			await WaitForMutexAsync(SingleInstanceMutex, StoppingToken);
+
+			// Terminate any remaining child processes from other instances
+			TerminateProcesses(Logger, StoppingToken);
+
+			// Show the worker capabilities
+			AgentCapabilities Capabilities = await GetAgentCapabilities(WorkingDir, Logger);
+			if (Capabilities.Properties.Count > 0)
+			{
+				Logger.LogInformation("Global:");
+				foreach (string Property in Capabilities.Properties)
+				{
+					Logger.LogInformation("  {AgentProperty}", Property);
+				}
+			}
+			foreach (DeviceCapabilities Device in Capabilities.Devices)
+			{
+				Logger.LogInformation($"{Device.Handle} Device:");
+				foreach (string Property in Device.Properties)
+				{
+					Logger.LogInformation("   {DeviceProperty}", Property);
+				}
+			}
+
+			// Mount all the necessary network shares. Currently only supported on Windows.
+			if (RuntimeInformation.IsOSPlatform(OSPlatform.Windows))
+			{
+				foreach (MountNetworkShare Share in Settings.Shares)
+				{
+					if (Share.MountPoint != null && Share.RemotePath != null)
+					{
+						Logger.LogInformation("Mounting {RemotePath} as {MountPoint}", Share.RemotePath, Share.MountPoint);
+						NetworkShare.Mount(Share.MountPoint, Share.RemotePath);
+					}
+				}
+			}
+
+			// Create the session
+			CreateSessionResponse CreateSessionResponse;
+			using (GrpcChannel Channel = GrpcService.CreateGrpcChannel(ServerProfile.Token))
+			{
+				HordeRpc.HordeRpcClient RpcClient = new HordeRpc.HordeRpcClient(Channel);
+
+				// Create the session information
+				CreateSessionRequest SessionRequest = new CreateSessionRequest();
+				SessionRequest.Name = Environment.MachineName;
+				SessionRequest.Status = AgentStatus.Ok;
+				SessionRequest.Capabilities = Capabilities;
+				SessionRequest.Version = Settings.Version ?? "";
+
+				// Create a session
+				CreateSessionResponse = await RpcClient.CreateSessionAsync(SessionRequest, null, null, StoppingToken);
+				Logger.LogInformation("Session started ({SessionId})", CreateSessionResponse.SessionId);
+			}
+
+			// Open a connection to the server
+			await using (IRpcConnection RpcCon = RpcConnection.Create(() => GrpcService.CreateGrpcChannel(CreateSessionResponse.Token), Logger))
+			{
+				// Loop until we're ready to exit
+				Stopwatch UpdateCapabilitiesTimer = Stopwatch.StartNew();
+				for (; ; )
+				{
+					Task WaitTask = UpdateLeasesEvent.Task;
+
+					// Flag for whether the service is stopping
+					bool bStopping = StoppingToken.IsCancellationRequested || RequestShutdown;
+
+					// Build the next update request
+					UpdateSessionRequest UpdateSessionRequest = new UpdateSessionRequest();
+					UpdateSessionRequest.AgentId = CreateSessionResponse.AgentId;
+					UpdateSessionRequest.SessionId = CreateSessionResponse.SessionId;
+
+					// Get the new the lease states. If a restart is requested and we have no active leases, signal to the server that we're stopping.
+					lock (LockObject)
+					{
+						foreach (LeaseInfo LeaseInfo in ActiveLeases)
+						{
+							UpdateSessionRequest.Leases.Add(new Lease(LeaseInfo.Lease));
+						}
+						if (RequestShutdown && ActiveLeases.Count == 0)
+						{
+							bStopping = true;
+						}
+					}
+
+					// Get the new agent status
+					if (bStopping)
+					{
+						UpdateSessionRequest.Status = AgentStatus.Stopping;
+					}
+					else if (Unhealthy)
+					{
+						UpdateSessionRequest.Status = AgentStatus.Unhealthy;
+					}
+					else
+					{
+						UpdateSessionRequest.Status = AgentStatus.Ok;
+					}
+
+
+					// Update the capabilities every 5m
+					if (UpdateCapabilitiesTimer.Elapsed > TimeSpan.FromMinutes(5.0))
+					{
+						UpdateSessionRequest.Capabilities = await GetAgentCapabilities(WorkingDir, Logger);
+						UpdateCapabilitiesTimer.Restart();
+					}
+
+					// Complete the wait task if we subsequently stop
+					using (bStopping ? (CancellationTokenRegistration?)null : StoppingToken.Register(() => UpdateLeasesEvent.Set()))
+					{
+						// Update the state with the server
+						UpdateSessionResponse? UpdateSessionResponse = null;
+						using (IRpcClientRef? RpcClientRef = RpcCon.TryGetClientRef(new RpcContext()))
+						{
+							if (RpcClientRef == null)
+							{
+								await Task.WhenAny(Task.Delay(TimeSpan.FromSeconds(5.0)), WaitTask);
+							}
+							else
+							{
+								UpdateSessionResponse = await UpdateSessionAsync(RpcClientRef, UpdateSessionRequest, WaitTask);
+							}
+						}
+
+						lock (LockObject)
+						{
+							// Now reconcile the local state to match what the server reports
+							if (UpdateSessionResponse != null)
+							{
+								// Remove any leases which have completed
+								ActiveLeases.RemoveAll(x => (x.Lease.State == LeaseState.Completed || x.Lease.State == LeaseState.Cancelled) && !UpdateSessionResponse.Leases.Any(y => y.Id == x.Lease.Id && y.State != LeaseState.Cancelled));
+
+								// Create any new leases and cancel any running leases
+								foreach (Lease ServerLease in UpdateSessionResponse.Leases)
+								{
+									if (ServerLease.State == LeaseState.Cancelled)
+									{
+										LeaseInfo? Info = ActiveLeases.FirstOrDefault(x => x.Lease.Id == ServerLease.Id);
+										if (Info != null)
+										{
+											Logger.LogInformation("Cancelling lease {LeaseId}", ServerLease.Id);
+											Info.CancellationTokenSource.Cancel();
+										}
+									}
+									if (ServerLease.State == LeaseState.Pending && !ActiveLeases.Any(x => x.Lease.Id == ServerLease.Id))
+									{
+										ServerLease.State = LeaseState.Active;
+
+										Logger.LogInformation("Adding lease {LeaseId}", ServerLease.Id);
+										LeaseInfo Info = new LeaseInfo(ServerLease);
+										Info.Task = Task.Run(() => HandleLeaseAsync(RpcCon, CreateSessionResponse.AgentId, Info));
+										ActiveLeases.Add(Info);
+									}
+								}
+							}
+
+							// If there's nothing still running and cancellation was requested, exit
+							if (ActiveLeases.Count == 0 && UpdateSessionRequest.Status == AgentStatus.Stopping)
+							{
+								Logger.LogInformation("No leases are active. Agent is stopping.");
+								break;
+							}
+						}
+					}
+				}
+
+				Logger.LogInformation("Disposing RpcConnection");
+			}
+		}
+
+		/// <summary>
+		/// Wrapper for <see cref="UpdateSessionInternalAsync"/> which filters/logs exceptions
+		/// </summary>
+		/// <param name="RpcClientRef">The RPC client connection</param>
+		/// <param name="UpdateSessionRequest">The session update request</param>
+		/// <param name="WaitTask">Task which can be used to jump out of the update early</param>
+		/// <returns>Response from the call</returns>
+		async Task<UpdateSessionResponse?> UpdateSessionAsync(IRpcClientRef RpcClientRef, UpdateSessionRequest UpdateSessionRequest, Task WaitTask)
+		{
+			UpdateSessionResponse? UpdateSessionResponse = null;
+			try
+			{
+				UpdateSessionResponse = await UpdateSessionInternalAsync(RpcClientRef, UpdateSessionRequest, WaitTask);
+				UpdateSessionFailures = 0;
+			}
+			catch (RpcException Ex)
+			{
+				if (++UpdateSessionFailures >= 3)
+				{
+					throw;
+				}
+				else if (Ex.StatusCode == StatusCode.Unavailable)
+				{
+					Logger.LogInformation(Ex, "Service unavailable while calling UpdateSessionAsync(), will retry");
+				}
+				else
+				{
+					Logger.LogError(Ex, "Error while executing RPC. Will retry.");
+				}
+			}
+			return UpdateSessionResponse;
+		}
+
+		/// <summary>
+		/// Tries to update the session state on the server.
+		/// 
+		/// This operation is a little gnarly due to the fact that we want to long-poll for the result.
+		/// Since we're doing the update via a gRPC call, the way to do that without using cancellation tokens is to keep the request stream open
+		/// until we want to terminate the call (see [REDACTED_PROFILE_URL]/grpc/issues/8277). In order to do that, we need to make a 
+		/// bidirectional streaming call, even though we only expect one response/response.
+		/// </summary>
+		/// <param name="RpcClientRef">The RPC client</param>
+		/// <param name="Request">The session update request</param>
+		/// <param name="WaitTask">Task to use to terminate the wait</param>
+		/// <returns>The response object</returns>
+		async Task<UpdateSessionResponse?> UpdateSessionInternalAsync(IRpcClientRef RpcClientRef, UpdateSessionRequest Request, Task WaitTask)
+		{
+			DateTime Deadline = DateTime.UtcNow + TimeSpan.FromMinutes(2.0);
+			using (AsyncDuplexStreamingCall<UpdateSessionRequest, UpdateSessionResponse> Call = RpcClientRef.Client.UpdateSession(deadline: Deadline))
+			{
+				Logger.LogDebug("Updating session {SessionId} (Status={Status})", Request.SessionId, Request.Status);
+
+				// Write the request to the server
+				await Call.RequestStream.WriteAsync(Request);
+
+				// Wait until the server responds or we need to trigger a new update
+				Task<bool> MoveNextAsync = Call.ResponseStream.MoveNext();
+
+				Task Task = await Task.WhenAny(MoveNextAsync, WaitTask, RpcClientRef.DisposingTask);
+				if(Task == WaitTask)
+				{
+					Logger.LogDebug("Cancelling long poll from client side (new update)");
+				}
+				else if (Task == RpcClientRef.DisposingTask)
+				{
+					Logger.LogDebug("Cancelling long poll from client side (server migration)");
+				}
+
+				// Close the request stream to indicate that we're finished
+				await Call.RequestStream.CompleteAsync();
+
+				// Wait for a response or a new update to come in, then close the request stream
+				UpdateSessionResponse? Response = null;
+				while (await MoveNextAsync)
+				{
+					Response = Call.ResponseStream.Current;
+					MoveNextAsync = Call.ResponseStream.MoveNext();
+				}
+				return Response;
+			}
+		}
+
+		/// <summary>
+		/// Handle a lease request
+		/// </summary>
+		/// <param name="RpcConnection">The RPC connection to the server</param>
+		/// <param name="AgentId">The agent id</param>
+		/// <param name="LeaseInfo">Information about the lease</param>
+		/// <returns>Async task</returns>
+		async Task HandleLeaseAsync(IRpcConnection RpcConnection, string AgentId, LeaseInfo LeaseInfo)
+		{
+			using Scope Scope = Tracer.Instance.StartActive("HandleLease");
+			Scope.Span.ResourceName = LeaseInfo.Lease.Id;
+			Scope.Span.SetTag("LeaseId", LeaseInfo.Lease.Id);
+			Scope.Span.SetTag("AgentId", AgentId);
+			using IDisposable TraceProperty = LogContext.PushProperty("dd.trace_id", CorrelationIdentifier.TraceId.ToString());
+			using IDisposable SpanProperty = LogContext.PushProperty("dd.span_id", CorrelationIdentifier.SpanId.ToString());
+
+			Logger.LogInformation("Handling lease {LeaseId}", LeaseInfo.Lease.Id);
+
+			// Get the lease outcome
+			LeaseOutcome Outcome = LeaseOutcome.Failed;
+			try
+			{
+				Outcome = await HandleLeasePayloadAsync(RpcConnection, AgentId, LeaseInfo);
+			}
+			catch (Exception Ex)
+			{
+				Logger.LogError(Ex, "Unhandled exception while executing lease {LeaseId}", LeaseInfo.Lease.Id);
+			}
+
+			// Update the state of the lease
+			lock (LockObject)
+			{
+				if (LeaseInfo.CancellationTokenSource.IsCancellationRequested)
+				{
+					LeaseInfo.Lease.State = LeaseState.Cancelled;
+					LeaseInfo.Lease.Outcome = LeaseOutcome.Failed;
+				}
+				else
+				{
+					LeaseInfo.Lease.State = (Outcome == LeaseOutcome.Cancelled) ? LeaseState.Cancelled : LeaseState.Completed;
+					LeaseInfo.Lease.Outcome = Outcome;
+				}
+				Logger.LogInformation("Transitioning lease {LeaseId} to {State}, outcome={Outcome}", LeaseInfo.Lease.Id, LeaseInfo.Lease.State, LeaseInfo.Lease.Outcome);
+			}
+			UpdateLeasesEvent.Set();
+		}
+
+		/// <summary>
+		/// Dispatch a lease payload to the appropriate handler
+		/// </summary>
+		/// <param name="RpcConnection">The RPC connection to the server</param>
+		/// <param name="AgentId">The current agent id</param>
+		/// <param name="LeaseInfo">Information about the lease</param>
+		/// <returns>Outcome from the lease</returns>
+		async Task<LeaseOutcome> HandleLeasePayloadAsync(IRpcConnection RpcConnection, string AgentId, LeaseInfo LeaseInfo)
+		{
+			Any Payload = LeaseInfo.Lease.Payload;
+
+			ActionTask ActionTask;
+			if (LeaseInfo.Lease.Payload.TryUnpack(out ActionTask))
+			{
+				Tracer.Instance.ActiveScope.Span.SetTag("task", "Action");
+				Func<ILogger, Task<LeaseOutcome>> Handler = NewLogger => ActionAsync(RpcConnection, LeaseInfo.Lease.Id, ActionTask, NewLogger, LeaseInfo.CancellationTokenSource.Token);
+				return await HandleLeasePayloadWithLogAsync(RpcConnection, ActionTask.LogId, AgentId, null, null, Handler, LeaseInfo.CancellationTokenSource.Token);
+			}
+
+			ConformTask ConformTask;
+			if (LeaseInfo.Lease.Payload.TryUnpack(out ConformTask))
+			{
+				Tracer.Instance.ActiveScope.Span.SetTag("task", "Conform");
+				Func<ILogger, Task<LeaseOutcome>> Handler = NewLogger => ConformAsync(RpcConnection, AgentId, LeaseInfo.Lease.Id, ConformTask, NewLogger, LeaseInfo.CancellationTokenSource.Token);
+				return await HandleLeasePayloadWithLogAsync(RpcConnection, ConformTask.LogId, AgentId, null, null, Handler, LeaseInfo.CancellationTokenSource.Token);
+			}
+
+			ExecuteJobTask JobTask;
+			if (LeaseInfo.Lease.Payload.TryUnpack(out JobTask))
+			{
+				Tracer.Instance.ActiveScope.Span.SetTag("task", "Job");
+				Func<ILogger, Task<LeaseOutcome>> Handler = NewLogger => ExecuteJobAsync(RpcConnection, AgentId, LeaseInfo.Lease.Id, JobTask, NewLogger, LeaseInfo.CancellationTokenSource.Token);
+				return await HandleLeasePayloadWithLogAsync(RpcConnection, JobTask.LogId, AgentId, JobTask.JobId, JobTask.BatchId, Handler, LeaseInfo.CancellationTokenSource.Token);
+			}
+
+			UpgradeTask UpgradeTask;
+			if (LeaseInfo.Lease.Payload.TryUnpack(out UpgradeTask))
+			{
+				Tracer.Instance.ActiveScope.Span.SetTag("task", "Upgrade");
+				Func<ILogger, Task<LeaseOutcome>> Handler = NewLogger => UpgradeAsync(RpcConnection, AgentId, UpgradeTask, NewLogger, LeaseInfo.CancellationTokenSource.Token);
+				return await HandleLeasePayloadWithLogAsync(RpcConnection, UpgradeTask.LogId, AgentId, null, null, Handler, LeaseInfo.CancellationTokenSource.Token);
+			}
+
+			ShutdownTask ShutdownTask;
+			if (LeaseInfo.Lease.Payload.TryUnpack(out ShutdownTask))
+			{
+				Tracer.Instance.ActiveScope.Span.SetTag("task", "Shutdown");
+				Func<ILogger, Task<LeaseOutcome>> Handler = NewLogger => ShutdownAsync(RpcConnection, AgentId, ShutdownTask, NewLogger, LeaseInfo.CancellationTokenSource.Token);
+				return await HandleLeasePayloadWithLogAsync(RpcConnection, ShutdownTask.LogId, AgentId, null, null, Handler, LeaseInfo.CancellationTokenSource.Token);
+			}
+
+			RestartTask RestartTask;
+			if (LeaseInfo.Lease.Payload.TryUnpack(out RestartTask))
+			{
+				Tracer.Instance.ActiveScope.Span.SetTag("task", "Restart");
+				Func<ILogger, Task<LeaseOutcome>> Handler = NewLogger => RestartAsync(RpcConnection, AgentId, RestartTask, NewLogger, LeaseInfo.CancellationTokenSource.Token);
+				return await HandleLeasePayloadWithLogAsync(RpcConnection, RestartTask.LogId, AgentId, null, null, Handler, LeaseInfo.CancellationTokenSource.Token);
+			}
+
+			Logger.LogError("Invalid lease payload type ({PayloadType})", Payload.TypeUrl);
+			return LeaseOutcome.Failed;
+		}
+
+		/// <summary>
+		/// Terminates any processes which are still running under the given directory
+		/// </summary>
+		/// <param name="Logger">Logger device</param>
+		/// <param name="CancellationToken">Cancellation token</param>
+		void TerminateProcesses(ILogger Logger, CancellationToken CancellationToken)
+		{
+			// Terminate child processes from any previous runs
+			ProcessUtils.TerminateProcesses(ShouldTerminateProcess, Logger, CancellationToken);
+		}
+
+		/// <summary>
+		/// Callback for determining whether a process should be terminated
+		/// </summary>
+		/// <param name="ImageFile">The file to terminate</param>
+		/// <returns>True if the process should be terminated</returns>
+		bool ShouldTerminateProcess(FileReference ImageFile)
+		{
+			if (ImageFile.IsUnderDirectory(WorkingDir))
+			{
+				return true;
+			}
+
+			string FileName = ImageFile.GetFileName();
+			if (ProcessNamesToTerminate.Contains(FileName))
+			{
+				return true;
+			}
+
+			return false;
+		}
+
+		/// <summary>
+		/// Dispatch a lease payload to the appropriate handler
+		/// </summary>
+		/// <param name="RpcConnection">The RPC connection to the server</param>
+		/// <param name="LogId">Unique id for the log to create</param>
+		/// <param name="AgentId">The current agent id</param>
+		/// <param name="JobId">The job being executed</param>
+		/// <param name="JobBatchId">The batch being executed</param>
+		/// <param name="ExecuteLease">Action to perform to execute the lease</param>
+		/// <param name="CancellationToken">The cancellation token for this lease</param>
+		/// <returns>Outcome from the lease</returns>
+		async Task<LeaseOutcome> HandleLeasePayloadWithLogAsync(IRpcConnection RpcConnection, string LogId, string AgentId, string? JobId, string? JobBatchId, Func<ILogger, Task<LeaseOutcome>> ExecuteLease, CancellationToken CancellationToken)
+		{
+			await using (JsonRpcLogger? ConformLogger = String.IsNullOrEmpty(LogId) ? null : new JsonRpcLogger(RpcConnection, LogId, JobId, JobBatchId, null, null, Logger))
+			{
+				ILogger LeaseLogger = (ConformLogger == null) ? (ILogger)Logger : new ForwardingLogger(Logger, new DefaultLoggerIndentHandler(ConformLogger));
+				try
+				{
+					LeaseOutcome Outcome = await ExecuteLease(LeaseLogger);
+					return Outcome;
+				}
+				catch (Exception Ex)
+				{
+					if (CancellationToken.IsCancellationRequested && IsCancellationException(Ex))
+					{
+						LeaseLogger.LogInformation("Lease was cancelled");
+						return LeaseOutcome.Cancelled;
+					}
+					else
+					{
+						LeaseLogger.LogError(Ex, "Caught unhandled exception while attempting to execute lease:\n{Exception}", Ex);
+						return LeaseOutcome.Failed;
+					}
+				}
+			}
+		}
+
+		/// <summary>
+		/// Execute a remote execution aciton
+		/// </summary>
+		/// <param name="RpcConnection">RPC client for communicating with the server</param>
+		/// <param name="LeaseId">The lease id</param>
+		/// <param name="ActionTask">The action task parameters</param>
+		/// <param name="Logger">Logger for the task</param>
+		/// <param name="CancellationToken">Token used to cancel the operation</param>
+		/// <returns></returns>
+		async Task<LeaseOutcome> ActionAsync(IRpcConnection RpcConnection, string LeaseId, ActionTask ActionTask, ILogger Logger, CancellationToken CancellationToken)
+		{
+			using (IRpcClientRef Client = await RpcConnection.GetClientRef(new RpcContext(), CancellationToken))
+			{
+				DirectoryReference LeaseDir = DirectoryReference.Combine(WorkingDir, "Remote", LeaseId);
+				DirectoryReference.CreateDirectory(LeaseDir);
+
+				try
+				{
+					ActionExecutor Executor = new ActionExecutor(Client.Channel, Logger);
+					await Executor.ExecuteActionAsync(LeaseId, ActionTask, LeaseDir);
+				}
+				finally
+				{
+					try
+					{
+						DirectoryReference.Delete(LeaseDir, true);
+					}
+					catch
+					{
+					}
+				}
+
+				return LeaseOutcome.Success;
+			}
+		}
+
+		/// <summary>
+		/// Conform a machine
+		/// </summary>
+		/// <param name="RpcConnection">RPC client for communicating with the server</param>
+		/// <param name="AgentId">The current agent id</param>
+		/// <param name="LeaseId">The current lease id</param>
+		/// <param name="ConformTask">The conform task parameters</param>
+		/// <param name="ConformLogger">Logger for the task</param>
+		/// <param name="CancellationToken">Token used to cancel the operation</param>
+		/// <returns>Async task</returns>
+		async Task<LeaseOutcome> ConformAsync(IRpcConnection RpcConnection, string AgentId, string LeaseId, ConformTask ConformTask, ILogger ConformLogger, CancellationToken CancellationToken)
+		{
+			ConformLogger.LogInformation("Conforming, lease {LeaseId}", LeaseId);
+			TerminateProcesses(ConformLogger, CancellationToken);
+
+			IList<AgentWorkspace> PendingWorkspaces = ConformTask.Workspaces;
+			for (; ;)
+			{
+				// Run the conform task
+				if (Settings.Executor == ExecutorType.Perforce && Settings.PerforceExecutor.RunConform)
+				{
+					await PerforceExecutor.ConformAsync(WorkingDir, PendingWorkspaces, ConformLogger, CancellationToken);
+				}
+				else
+				{
+					ConformLogger.LogInformation("Skipping due to Settings.RunConform flag");
+				}
+
+				// Update the new set of workspaces
+				UpdateAgentWorkspacesRequest Request = new UpdateAgentWorkspacesRequest();
+				Request.AgentId = AgentId;
+				Request.Workspaces.AddRange(PendingWorkspaces);
+
+				UpdateAgentWorkspacesResponse Response = await RpcConnection.InvokeAsync(x => x.UpdateAgentWorkspacesAsync(Request, null, null, CancellationToken), new RpcContext(), CancellationToken);
+				if (!Response.Retry)
+				{
+					ConformLogger.LogInformation("Conform finished");
+					break;
+				}
+
+				ConformLogger.LogInformation("Pending workspaces have changed - running conform again...");
+				PendingWorkspaces = Response.PendingWorkspaces;
+			}
+
+			return LeaseOutcome.Success;
+		}
+
+		/// <summary>
+		/// Execute part of a job
+		/// </summary>
+		/// <param name="RpcClient">RPC client for communicating with the server</param>
+		/// <param name="AgentId">The current agent id</param>
+		/// <param name="LeaseId">The current lease id</param>
+		/// <param name="ExecuteTask">The task to execute</param>
+		/// <param name="Logger">The logger to use for this lease</param>
+		/// <param name="CancellationToken">Token used to cancel the operation</param>
+		/// <returns>Async task</returns>
+		internal async Task<LeaseOutcome> ExecuteJobAsync(IRpcConnection RpcClient, string AgentId, string LeaseId, ExecuteJobTask ExecuteTask, ILogger Logger, CancellationToken CancellationToken)
+		{
+			Logger.LogInformation("Executing job \"{JobName}\", jobId {JobId}, batchId {BatchId}, leaseId {LeaseId}", ExecuteTask.JobName, ExecuteTask.JobId, ExecuteTask.BatchId, LeaseId);
+			Tracer.Instance.ActiveScope?.Span.SetTag("jobId", ExecuteTask.JobId.ToString());
+			Tracer.Instance.ActiveScope?.Span.SetTag("jobName", ExecuteTask.JobName.ToString());
+			Tracer.Instance.ActiveScope?.Span.SetTag("batchId", ExecuteTask.BatchId.ToString());	
+
+			// Start executing the current batch
+			BeginBatchResponse Batch = await RpcClient.InvokeAsync(x => x.BeginBatchAsync(new BeginBatchRequest(ExecuteTask.JobId, ExecuteTask.BatchId, LeaseId), null, null, CancellationToken), new RpcContext(), CancellationToken);
+
+			// Execute the batch
+			try
+			{
+				await ExecuteBatchAsync(RpcClient, AgentId, LeaseId, ExecuteTask, Batch, Logger, CancellationToken);
+			}
+			catch (Exception Ex)
+			{
+				if (CancellationToken.IsCancellationRequested && IsCancellationException(Ex))
+				{
+					throw;
+				}
+				else
+				{
+					Logger.LogError(Ex, $"Exception while executing batch: {Ex}");
+				}
+			}
+
+			// Mark the batch as complete
+			await RpcClient.InvokeAsync(x => x.FinishBatchAsync(new FinishBatchRequest(ExecuteTask.JobId, ExecuteTask.BatchId, LeaseId), null, null, CancellationToken), new RpcContext(), CancellationToken);
+			Logger.LogInformation("Done.");
+
+			return LeaseOutcome.Success;
+		}
+
+		/// <summary>
+		/// Executes a batch
+		/// </summary>
+		/// <param name="RpcClient">RPC client for communicating with the server</param>
+		/// <param name="AgentId">The current agent id</param>
+		/// <param name="LeaseId">The current lease id</param>
+		/// <param name="ExecuteTask">The task to execute</param>
+		/// <param name="Batch">The batch to execute</param>
+		/// <param name="BatchLogger">Output log for the batch</param>
+		/// <param name="CancellationToken">Cancellation token to abort the batch</param>
+		/// <returns>Async task</returns>
+		async Task ExecuteBatchAsync(IRpcConnection RpcClient, string AgentId, string LeaseId, ExecuteJobTask ExecuteTask, BeginBatchResponse Batch, ILogger BatchLogger, CancellationToken CancellationToken)
+		{
+			BatchLogger.LogInformation("Executing batch {BatchId} using {Executor} executor", ExecuteTask.BatchId, Settings.Executor.ToString());
+			TerminateProcesses(BatchLogger, CancellationToken);
+
+			// Get the working directory for this lease
+			DirectoryReference ScratchDir = DirectoryReference.Combine(WorkingDir, "Scratch");
+
+			// Create an executor for this job
+			IExecutor Executor = CreateExecutor(RpcClient, ExecuteTask, Batch);
+
+			// Try to initialize the executor
+			BatchLogger.LogInformation("Initializing...");
+			using (BatchLogger.BeginIndentScope("  "))
+			{
+				using Scope Scope = Tracer.Instance.StartActive("Initialize");
+				await Executor.InitializeAsync(BatchLogger, CancellationToken);
+			}
+
+			// Execute the steps
+			for (; ; )
+			{
+				// Get the next step to execute
+				BeginStepResponse Step = await RpcClient.InvokeAsync(x => x.BeginStepAsync(new BeginStepRequest(ExecuteTask.JobId, ExecuteTask.BatchId, LeaseId), null, null, CancellationToken), new RpcContext(), CancellationToken);
+				if (Step.State == BeginStepResponse.Types.Result.Waiting)
+				{
+					BatchLogger.LogInformation("Waiting for dependency to be ready");
+					await Task.Delay(TimeSpan.FromSeconds(20.0), CancellationToken);
+					continue;
+				}
+				else if (Step.State == BeginStepResponse.Types.Result.Complete)
+				{
+					break;
+				}
+				else if (Step.State != BeginStepResponse.Types.Result.Ready)
+				{
+					BatchLogger.LogError("Unexpected step state: {StepState}", Step.State);
+					break;
+				}
+
+				// Print the new state
+				Stopwatch StepTimer = Stopwatch.StartNew();
+				BatchLogger.LogInformation("Starting job {JobId}, batch {BatchId}, step {StepId}", ExecuteTask.JobId, ExecuteTask.BatchId, Step.StepId);
+
+				// Create a trace span
+				using Scope Scope = Tracer.Instance.StartActive("Execute");
+				Scope.Span.ResourceName = Step.Name;
+				Scope.Span.SetTag("stepId", Step.StepId);
+				Scope.Span.SetTag("logId", Step.LogId);
+				using IDisposable TraceProperty = LogContext.PushProperty("dd.trace_id", CorrelationIdentifier.TraceId.ToString());
+				using IDisposable SpanProperty = LogContext.PushProperty("dd.span_id", CorrelationIdentifier.SpanId.ToString());
+
+				// Update the context to include information about this step
+				JobStepOutcome StepOutcome;
+				JobStepState StepState;
+				using (BatchLogger.BeginIndentScope("  "))
+				{
+					// Start writing to the log file
+					await using (JsonRpcLogger StepLogger = new JsonRpcLogger(RpcClient, Step.LogId, ExecuteTask.JobId, ExecuteTask.BatchId, Step.StepId, Step.Warnings, Logger))
+					{
+						// Execute the task
+						ILogger ForwardingLogger = new DefaultLoggerIndentHandler(StepLogger);
+						if (Settings.WriteStepOutputToLogger)
+						{
+							ForwardingLogger = new ForwardingLogger(Logger, ForwardingLogger);
+						}
+
+						using CancellationTokenSource StepPollCancelSource = new CancellationTokenSource();
+						using CancellationTokenSource StepAbortSource = new CancellationTokenSource();
+						TaskCompletionSource<bool> StepFinishedSource = new TaskCompletionSource<bool>();
+						Task StepPollTask = Task.Run(() => PollForStepAbort(RpcClient, ExecuteTask.JobId, ExecuteTask.BatchId, Step.StepId, StepPollCancelSource.Token, StepAbortSource, StepFinishedSource.Task));
+						
+						try
+						{
+							(StepOutcome, StepState) = await ExecuteStepAsync(Executor, Step, ForwardingLogger, CancellationToken, StepAbortSource.Token);
+						}
+						finally
+						{
+							// Will get called even when cancellation token for the lease/batch fires
+							StepFinishedSource.SetResult(true); // Tell background poll task to stop
+							await StepPollTask;
+						}
+						
+						// Kill any processes spawned by the step
+						TerminateProcesses(StepLogger, CancellationToken);
+
+						// Wait for the logger to finish
+						await StepLogger.StopAsync();
+						
+						// Reflect the warnings/errors in the step outcome
+						if (StepOutcome > StepLogger.Outcome)
+						{
+							StepOutcome = StepLogger.Outcome;
+						}
+					}
+
+					// Update the server with the outcome from the step
+					BatchLogger.LogInformation("Marking step as complete (Outcome={Outcome}, State={StepState})", StepOutcome, StepState);
+					await RpcClient.InvokeAsync(x => x.UpdateStepAsync(new UpdateStepRequest(ExecuteTask.JobId, ExecuteTask.BatchId, Step.StepId, StepState, StepOutcome), null, null, CancellationToken), new RpcContext(), CancellationToken);
+				}
+
+				// Print the finishing state
+				StepTimer.Stop();
+				BatchLogger.LogInformation("Completed in {Time}", StepTimer.Elapsed);
+			}
+
+			// Clean the environment
+			BatchLogger.LogInformation("Finalizing...");
+			using (BatchLogger.BeginIndentScope("  "))
+			{
+				using Scope Scope = Tracer.Instance.StartActive("Finalize");
+				await Executor.FinalizeAsync(BatchLogger, CancellationToken);
+			}
+		}
+
+		/// <summary>
+		/// Executes a step
+		/// </summary>
+		/// <param name="Executor">The executor to run this step</param>
+		/// <param name="Step">Step to execute</param>
+		/// <param name="StepLogger">Logger for the step</param>
+		/// <param name="CancellationToken">Cancellation token to abort the batch</param>
+		/// <param name="StepCancellationToken">Cancellation token to abort only this individual step</param>
+		/// <returns>Async task</returns>
+		internal async Task<(JobStepOutcome, JobStepState)> ExecuteStepAsync(IExecutor Executor, BeginStepResponse Step, ILogger StepLogger, CancellationToken CancellationToken, CancellationToken StepCancellationToken)
+		{
+			using CancellationTokenSource Combined = CancellationTokenSource.CreateLinkedTokenSource(CancellationToken, StepCancellationToken);
+			try
+			{
+				JobStepOutcome StepOutcome = await Executor.RunAsync(Step, StepLogger, Combined.Token);
+				return (StepOutcome, JobStepState.Completed);
+			}
+			catch (Exception Ex)
+			{
+				if (CancellationToken.IsCancellationRequested && IsCancellationException(Ex))
+				{
+					StepLogger.LogError("The step was cancelled by batch/lease");
+					throw;
+				}
+				
+				if (StepCancellationToken.IsCancellationRequested && IsCancellationException(Ex))
+				{
+					StepLogger.LogError("The step was intentionally cancelled");
+					return (JobStepOutcome.Failure, JobStepState.Aborted);
+				}
+				
+				StepLogger.LogError(Ex, $"Exception while executing step: {Ex}");
+				return (JobStepOutcome.Failure, JobStepState.Completed);
+			}
+		}
+
+		internal async Task PollForStepAbort(IRpcConnection RpcClient, string JobId, string BatchId, string StepId, CancellationToken CancellationToken, CancellationTokenSource StepCancelSource, Task FinishedTask)
+		{
+			Stopwatch Timer = Stopwatch.StartNew();
+			while (!FinishedTask.IsCompleted)
+			{
+				try
+				{
+					if (Timer.Elapsed > TimeSpan.FromHours(24.0))
+					{
+						Logger.LogDebug("Step was aborted after running for {NumHours} (JobId={JobId} BatchId={BatchId} StepId={StepId})", Timer.Elapsed, JobId, BatchId, StepId);
+						StepCancelSource.Cancel();
+						break;
+					}
+
+					GetStepResponse Res = await RpcClient.InvokeAsync(x => x.GetStepAsync(new GetStepRequest(JobId, BatchId, StepId), null, null, CancellationToken), new RpcContext(), CancellationToken);
+					if (Res.AbortRequested)
+					{
+						Logger.LogDebug("Step was aborted by server (JobId={JobId} BatchId={BatchId} StepId={StepId})", JobId, BatchId, StepId);
+						StepCancelSource.Cancel();
+						break;
+					}
+				}
+				catch (RpcException Ex)
+				{
+					Logger.LogError(Ex, "Poll for step abort has failed. Aborting (JobId={JobId} BatchId={BatchId} StepId={StepId})", JobId, BatchId, StepId);
+					StepCancelSource.Cancel();
+					break;
+				}
+
+				await Task.WhenAny(Task.Delay(StepAbortPollInterval), FinishedTask);
+			}
+		}
+
+		/// <summary>
+		/// Determine if the given exception was triggered due to a cancellation event
+		/// </summary>
+		/// <param name="Ex">The exception to check</param>
+		/// <returns>True if the exception is a cancellation exception</returns>
+		static bool IsCancellationException(Exception Ex)
+		{
+			if(Ex is OperationCanceledException)
+			{
+				return true;
+			}
+
+			RpcException? RpcException = Ex as RpcException;
+			if(RpcException != null && RpcException.StatusCode == StatusCode.Cancelled)
+			{
+				return true;
+			}
+
+			return false;
+		}
+
+		/// <summary>
+		/// Check for an update of the agent software
+		/// </summary>
+		/// <param name="RpcClient">RPC client for communicating with the server</param>
+		/// <param name="AgentId">The current agent id</param>
+		/// <param name="UpgradeTask">The upgrade task</param>
+		/// <param name="Logger">Logging device</param>
+		/// <param name="CancellationToken">Token used to cancel the operation</param>
+		/// <returns>Outcome of this operation</returns>
+		async Task<LeaseOutcome> UpgradeAsync(IRpcConnection RpcClient, string AgentId, UpgradeTask UpgradeTask, ILogger Logger, CancellationToken CancellationToken)
+		{
+			string RequiredVersion = UpgradeTask.SoftwareId;
+
+			// Check if we're running the right version
+			if (RequiredVersion != null && RequiredVersion != Settings.Version)
+			{
+				Logger.LogInformation("Upgrading from {CurrentVersion} to {TargetVersion}", String.IsNullOrEmpty(Settings.Version)? "(unknown)" : Settings.Version, RequiredVersion);
+
+				// Clear out the working directory
+				DirectoryReference UpgradeDir = DirectoryReference.Combine(WorkingDir, "Upgrade");
+				DirectoryReference.CreateDirectory(UpgradeDir);
+				await DeleteDirectoryContentsAsync(new DirectoryInfo(UpgradeDir.FullName));
+
+				// Download the new software
+				FileInfo OutputFile = new FileInfo(Path.Combine(UpgradeDir.FullName, "Agent.zip"));
+				using (IRpcClientRef RpcClientRef = await RpcClient.GetClientRef(new RpcContext(), CancellationToken))
+				using (AsyncServerStreamingCall<DownloadSoftwareResponse> Cursor = RpcClientRef.Client.DownloadSoftware(new DownloadSoftwareRequest(RequiredVersion), null, null, CancellationToken))
+				{
+					using (Stream OutputStream = OutputFile.Open(FileMode.Create))
+					{
+						while (await Cursor.ResponseStream.MoveNext(CancellationToken))
+						{
+							OutputStream.Write(Cursor.ResponseStream.Current.Data.Span);
+						}
+					}
+				}
+
+				// Extract it to a temporary directory
+				DirectoryReference ExtractedDir = DirectoryReference.Combine(UpgradeDir, "Extracted");
+				DirectoryReference.CreateDirectory(ExtractedDir);
+				ZipFile.ExtractToDirectory(OutputFile.FullName, ExtractedDir.FullName);
+
+				//				// Debug code for updating an agent with the local version
+				//				foreach (FileInfo SourceFile in new FileInfo(Assembly.GetExecutingAssembly().Location).Directory.EnumerateFiles())
+				//				{
+				//					SourceFile.CopyTo(Path.Combine(ExtractedDir.FullName, SourceFile.Name), true);
+				//				}
+
+
+				// Get the current process and assembly. This may be different if running through dotnet.exe rather than a native PE image.
+				FileReference AssemblyFileName = new FileReference(Assembly.GetExecutingAssembly().Location);
+
+				// Spawn the other process
+				using (Process Process = new Process())
+				{
+					StringBuilder Arguments = new StringBuilder();
+
+					DirectoryReference TargetDir = AssemblyFileName.Directory;
+
+					// We were launched via an external application (presumably dotnet.exe). Do the same thing again.
+					FileReference NewAssemblyFileName = FileReference.Combine(ExtractedDir, AssemblyFileName.MakeRelativeTo(TargetDir));
+					if (!FileReference.Exists(NewAssemblyFileName))
+					{
+						Logger.LogError("Unable to find {AgentExe} in extracted archive", NewAssemblyFileName);
+						return LeaseOutcome.Failed;
+					}
+
+					Process.StartInfo.FileName = "dotnet";
+
+					StringBuilder CurrentArguments = new StringBuilder();
+					foreach (string Arg in Program.Args)
+					{
+						CurrentArguments.AppendArgument(Arg);
+					}
+
+					Arguments.AppendArgument(NewAssemblyFileName.FullName);
+					Arguments.AppendArgument("Service");
+					Arguments.AppendArgument("Upgrade");
+					Arguments.AppendArgument("-ProcessId=", Process.GetCurrentProcess().Id.ToString());
+					Arguments.AppendArgument("-TargetDir=", TargetDir.FullName);
+					Arguments.AppendArgument("-Arguments=", CurrentArguments.ToString());
+
+					Process.StartInfo.Arguments = Arguments.ToString();
+					Process.StartInfo.UseShellExecute = false;
+					Process.EnableRaisingEvents = true;
+
+					StringBuilder LaunchCommand = new StringBuilder();
+					LaunchCommand.AppendArgument(Process.StartInfo.FileName);
+					LaunchCommand.Append(' ');
+					LaunchCommand.Append(Arguments);
+					Logger.LogInformation("Launching: {Launch}", LaunchCommand.ToString());
+
+					TaskCompletionSource<int> ExitCodeSource = new TaskCompletionSource<int>();
+					Process.Exited += (Sender, Args) => { ExitCodeSource.SetResult(Process.ExitCode); };
+
+					Process.Start();
+
+					using (CancellationToken.Register(() => { ExitCodeSource.SetResult(0); }))
+					{
+						await ExitCodeSource.Task;
+					}
+				}
+			}
+
+			return LeaseOutcome.Success;
+		}
+
+		/// <summary>
+		/// Delete the contents of a directory without deleting it itself
+		/// </summary>
+		/// <param name="BaseDir">Directory to clean</param>
+		/// <returns>Async task</returns>
+		static async Task DeleteDirectoryContentsAsync(DirectoryInfo BaseDir)
+		{
+			List<Task> ChildTasks = new List<Task>();
+			foreach (DirectoryInfo SubDir in BaseDir.EnumerateDirectories())
+			{
+				ChildTasks.Add(Task.Run(() => DeleteDirectory(SubDir)));
+			}
+			foreach (FileInfo File in BaseDir.EnumerateFiles())
+			{
+				File.Attributes = FileAttributes.Normal;
+				File.Delete();
+			}
+			await Task.WhenAll(ChildTasks);
+		}
+
+		/// <summary>
+		/// Deletes a directory and its contents
+		/// </summary>
+		/// <param name="BaseDir">Directory to delete</param>
+		/// <returns>Async task</returns>
+		static async Task DeleteDirectory(DirectoryInfo BaseDir)
+		{
+			await DeleteDirectoryContentsAsync(BaseDir);
+			BaseDir.Delete();
+		}
+
+		/// <summary>
+		/// Check for an update of the agent software
+		/// </summary>
+		/// <param name="RpcClient">RPC client for communicating with the server</param>
+		/// <param name="AgentId">The current agent id</param>
+		/// <param name="RestartTask">The restart task parameters</param>
+		/// <param name="Logger">Logging device</param>
+		/// <param name="CancellationToken">Token used to cancel the operation</param>
+		/// <returns>Outcome of this operation</returns>
+		Task<LeaseOutcome> RestartAsync(IRpcConnection RpcClient, string AgentId, RestartTask RestartTask, ILogger Logger, CancellationToken CancellationToken)
+		{
+			Logger.LogInformation("Setting restart flag");
+			RequestShutdown = true;
+			RestartAfterShutdown = true;
+			return Task.FromResult(LeaseOutcome.Success);
+		}
+
+		/// <summary>
+		/// Shutdown the agent
+		/// </summary>
+		/// <param name="RpcClient">RPC client for communicating with the server</param>
+		/// <param name="AgentId">The current agent id</param>
+		/// <param name="ShutdownTask">The restart task parameters</param>
+		/// <param name="Logger">Logging device</param>
+		/// <param name="CancellationToken">Token used to cancel the operation</param>
+		/// <returns>Outcome of this operation</returns>
+		Task<LeaseOutcome> ShutdownAsync(IRpcConnection RpcClient, string AgentId, ShutdownTask ShutdownTask, ILogger Logger, CancellationToken CancellationToken)
+		{
+			Logger.LogInformation("Setting shutdown flag");
+			RequestShutdown = true;
+			RestartAfterShutdown = false;
+			return Task.FromResult(LeaseOutcome.Success);
+		}
+
+		/// <summary>
+		/// Gets the hardware capabilities of this worker
+		/// </summary>
+		/// <returns>Worker object for advertising to the server</returns>
+		static async Task<AgentCapabilities> GetAgentCapabilities(DirectoryReference WorkingDir, ILogger Logger)
+		{
+			// Create the primary device
+			DeviceCapabilities PrimaryDevice = new DeviceCapabilities();
+			PrimaryDevice.Handle = "Primary";
+
+			List<DeviceCapabilities> OtherDevices = new List<DeviceCapabilities>();
+			if (RuntimeInformation.IsOSPlatform(OSPlatform.Windows))
+			{
+				PrimaryDevice.Properties.Add("OSFamily=Windows");
+
+				// Add OS info
+				using (ManagementObjectSearcher Searcher = new ManagementObjectSearcher("select * from Win32_OperatingSystem"))
+				{
+					foreach (ManagementObject Row in Searcher.Get())
+					{
+						Dictionary<string, object> Properties = GetWmiProperties(Row);
+
+						object? Name;
+						if(Properties.TryGetValue("Caption", out Name))
+						{
+							PrimaryDevice.Properties.Add($"OSDistribution={Name}");
+						}
+
+						object? Version;
+						if(Properties.TryGetValue("Version", out Version))
+						{
+							PrimaryDevice.Properties.Add($"OSKernelVersion={Version}");
+						}
+					}
+				}
+
+				// Add CPU info
+				using (ManagementObjectSearcher Searcher = new ManagementObjectSearcher("select * from Win32_Processor"))
+				{
+					Dictionary<string, int> NameToCount = new Dictionary<string, int>();
+					int TotalPhysicalCores = 0;
+					int TotalLogicalCores = 0;
+
+					foreach (ManagementObject Row in Searcher.Get())
+					{
+						Dictionary<string, object> Properties = GetWmiProperties(Row);
+
+						object? NameObject;
+						if(Properties.TryGetValue("Name", out NameObject))
+						{
+							string Name = NameObject.ToString() ?? String.Empty;
+							int Count;
+							NameToCount.TryGetValue(Name, out Count);
+							NameToCount[Name] = Count + 1;
+						}
+
+						object? NumPhysicalCores;
+						if((Properties.TryGetValue("NumberOfEnabledCore", out NumPhysicalCores) && NumPhysicalCores is uint) || (Properties.TryGetValue("NumberOfCores", out NumPhysicalCores) && NumPhysicalCores is uint))
+						{
+							TotalPhysicalCores += (int)(uint)NumPhysicalCores;
+						}
+
+						object? NumLogicalCores;
+						if(Properties.TryGetValue("NumberOfLogicalProcessors", out NumLogicalCores) && NumLogicalCores is uint)
+						{
+							TotalLogicalCores += (int)(uint)NumLogicalCores;
+						}
+					}
+
+					if (NameToCount.Count > 0)
+					{
+						PrimaryDevice.Properties.Add("CPU=" + String.Join(", ", NameToCount.Select(x => (x.Value > 1) ? $"{x.Key} x {x.Value}" : x.Key)));
+					}
+
+					if (TotalLogicalCores > 0)
+					{
+						PrimaryDevice.Properties.Add($"LogicalCores={TotalLogicalCores}");
+					}
+
+					if (TotalPhysicalCores > 0)
+					{
+						PrimaryDevice.Properties.Add($"PhysicalCores={TotalPhysicalCores}");
+					}
+				}
+
+				// Add RAM info
+				using (ManagementObjectSearcher Searcher = new ManagementObjectSearcher("select Name, DriverVersion, AdapterRAM from Win32_VideoController"))
+				{
+					int Index = 0;
+					foreach (ManagementObject Row in Searcher.Get())
+					{
+						WmiProperties Properties = new WmiProperties(Row);
+						if (Properties.TryGetValue("Name", out string? Name) && Properties.TryGetValue("DriverVersion", out string? DriverVersion))
+						{
+							DeviceCapabilities Device = new DeviceCapabilities();
+							Device.Handle = $"GPU-{++Index}";
+							Device.Properties.Add($"Name={Name}");
+							Device.Properties.Add($"Type=GPU");
+							Device.Properties.Add($"DriverVersion={DriverVersion}");
+							OtherDevices.Add(Device);
+						}
+					}
+				}
+
+				// Add a flag for EC2 instances
+				using (ManagementObjectSearcher Searcher = new ManagementObjectSearcher("select uuid from Win32_ComputerSystemProduct"))
+				{
+					foreach (ManagementObject Row in Searcher.Get())
+					{
+						string? Uuid = Row.GetPropertyValue("uuid") as string;
+						if (Uuid != null && Uuid.StartsWith("EC2", StringComparison.OrdinalIgnoreCase))
+						{
+							PrimaryDevice.Properties.Add("EC2=1");
+							break;
+						}
+					}
+				}
+
+				// Add other EC2 properties
+				await AddAwsProperties(PrimaryDevice.Properties, Logger);
+
+				// Add session information
+				PrimaryDevice.Properties.Add($"User={Environment.UserName}");
+				PrimaryDevice.Properties.Add($"Domain={Environment.UserDomainName}");
+				PrimaryDevice.Properties.Add($"Interactive={Environment.UserInteractive}");
+				PrimaryDevice.Properties.Add($"Elevated={BuildGraphExecutor.IsUserAdministrator()}");
+			}
+			else if(RuntimeInformation.IsOSPlatform(OSPlatform.OSX))
+			{
+				PrimaryDevice.Properties.Add("OSFamily=MacOS");
+				PrimaryDevice.Properties.Add("OSVersion=MacOS");
+
+				string Output;
+				using(Process Process = new Process())
+				{
+					Process.StartInfo.FileName = "system_profiler";
+					Process.StartInfo.Arguments = "SPHardwareDataType SPSoftwareDataType -xml";
+					Process.StartInfo.CreateNoWindow = true;
+					Process.StartInfo.RedirectStandardOutput = true;
+					Process.StartInfo.RedirectStandardInput = false;
+					Process.StartInfo.UseShellExecute = false;
+					Process.Start();
+
+					Output = Process.StandardOutput.ReadToEnd();
+				}
+
+				XmlDocument Xml = new XmlDocument();
+				Xml.LoadXml(Output);
+
+				XmlNode? HardwareNode = Xml.SelectSingleNode("/plist/array/dict/key[. = '_dataType']/following-sibling::string[. = 'SPHardwareDataType']/../key[. = '_items']/following-sibling::array/dict");
+				if(HardwareNode != null)
+				{
+					XmlNode? Model = HardwareNode.SelectSingleNode("key[. = 'machine_model']/following-sibling::string");
+					if(Model != null)
+					{
+						PrimaryDevice.Properties.Add($"Model={Model.InnerText}");
+					}
+
+					XmlNode? CpuTypeNode = HardwareNode.SelectSingleNode("key[. = 'cpu_type']/following-sibling::string");
+					XmlNode? CpuSpeedNode = HardwareNode.SelectSingleNode("key[. = 'current_processor_speed']/following-sibling::string");
+					XmlNode? CpuPackagesNode = HardwareNode.SelectSingleNode("key[. = 'packages']/following-sibling::integer");
+					if(CpuTypeNode != null && CpuSpeedNode != null && CpuPackagesNode != null)
+					{
+						PrimaryDevice.Properties.Add((CpuPackagesNode.InnerText != "1")? $"CPU={CpuPackagesNode.InnerText} x {CpuTypeNode.InnerText} @ {CpuSpeedNode.InnerText}" : $"CPU={CpuTypeNode.InnerText} @ {CpuSpeedNode.InnerText}");
+					}
+
+					PrimaryDevice.Properties.Add($"LogicalCores={Environment.ProcessorCount}");
+
+					XmlNode? CpuCountNode = HardwareNode.SelectSingleNode("key[. = 'number_processors']/following-sibling::integer");
+					if(CpuCountNode != null)
+					{
+						PrimaryDevice.Properties.Add($"PhysicalCores={CpuCountNode.InnerText}");
+					}
+
+					XmlNode? MemoryNode = HardwareNode.SelectSingleNode("key[. = 'physical_memory']/following-sibling::string");
+					if(MemoryNode != null)
+					{
+						string[] Parts = MemoryNode.InnerText.Split(new char[]{' '}, StringSplitOptions.RemoveEmptyEntries);
+						if(Parts.Length == 2 && Parts[1] == "GB")
+						{
+							PrimaryDevice.Properties.Add($"RAM={Parts[0]}");
+						}
+					}
+				}
+
+				XmlNode? SoftwareNode = Xml.SelectSingleNode("/plist/array/dict/key[. = '_dataType']/following-sibling::string[. = 'SPSoftwareDataType']/../key[. = '_items']/following-sibling::array/dict");
+				if(SoftwareNode != null)
+				{
+					XmlNode? OsVersionNode = SoftwareNode.SelectSingleNode("key[. = 'os_version']/following-sibling::string");
+					if(OsVersionNode != null)
+					{
+						PrimaryDevice.Properties.Add($"OSDistribution={OsVersionNode.InnerText}");
+					}
+
+					XmlNode? KernelVersionNode = SoftwareNode.SelectSingleNode("key[. = 'kernel_version']/following-sibling::string");
+					if(KernelVersionNode != null)
+					{
+						PrimaryDevice.Properties.Add($"OSKernelVersion={KernelVersionNode.InnerText}");
+					}
+				}
+			}
+
+			// Get the time that the machine booted
+			PrimaryDevice.Properties.Add($"BootTime={BootTime}");
+			PrimaryDevice.Properties.Add($"StartTime={StartTime}");
+
+			// Add disk info
+			string? DriveName = Path.GetPathRoot(WorkingDir.FullName);
+			if (DriveName != null)
+			{
+				try
+				{
+					DriveInfo Info = new DriveInfo(DriveName);
+					PrimaryDevice.Properties.Add($"DiskFreeSpace={Info.AvailableFreeSpace}");
+					PrimaryDevice.Properties.Add($"DiskTotalSize={Info.TotalSize}");
+				}
+				catch (Exception Ex)
+				{
+					Logger.LogWarning(Ex, "Unable to query disk info for path '{DriveName}'", DriveName);
+				}
+			}
+			PrimaryDevice.Properties.Add($"WorkingDir={WorkingDir}");
+
+			// Create the worker
+			AgentCapabilities Agent = new AgentCapabilities();
+			Agent.Devices.Add(PrimaryDevice);
+			Agent.Devices.AddRange(OtherDevices);
+			return Agent;
+		}
+
+		class WmiProperties
+		{
+			Dictionary<string, object> Properties = new Dictionary<string, object>(StringComparer.Ordinal);
+
+			public WmiProperties(ManagementObject Row)
+			{
+				foreach (PropertyData Property in Row.Properties)
+				{
+					Properties[Property.Name] = Property.Value;
+				}
+			}
+
+			public bool TryGetValue(string Name, [NotNullWhen(true)] out string? Value)
+			{
+				object? Object;
+				if (Properties.TryGetValue(Name, out Object))
+				{
+					Value = Object.ToString();
+					return Value != null;
+				}
+				else
+				{
+					Value = null!;
+					return false;
+				}
+			}
+
+			public bool TryGetValue(string Name, out long Value)
+			{
+				object? Object;
+				if (Properties.TryGetValue(Name, out Object))
+				{
+					if (Object is int)
+					{
+						Value = (int)Object;
+						return true;
+					}
+					else if (Object is uint)
+					{
+						Value = (uint)Object;
+						return true;
+					}
+					else if (Object is long)
+					{
+						Value = (long)Object;
+						return true;
+					}
+					else if (Object is ulong)
+					{
+						Value = (long)(ulong)Object;
+						return true;
+					}
+				}
+
+				Value = 0;
+				return false;
+			}
+		}
+
+		static Dictionary<string, object> GetWmiProperties(ManagementObject Row)
+		{
+			Dictionary<string, object> Properties = new Dictionary<string, object>(StringComparer.Ordinal);
+			foreach (PropertyData Property in Row.Properties)
+			{
+				Properties[Property.Name] = Property.Value;
+			}
+			return Properties;
+		}
+
+		static async Task AddAwsProperties(IList<string> Properties, ILogger Logger)
+		{
+			if (EC2InstanceMetadata.IdentityDocument != null)
+			{
+				AddAwsProperty("aws-instance-id", "/instance-id", Properties);
+				AddAwsProperty("aws-instance-type", "/instance-type", Properties);
+				AddAwsProperty("aws-region", "/region", Properties);
+
+				try
+				{
+					using (AmazonEC2Client Client = new AmazonEC2Client())
+					{
+						DescribeTagsRequest Request = new DescribeTagsRequest();
+						Request.Filters = new List<Filter>();
+						Request.Filters.Add(new Filter("resource-id", new List<string> { EC2InstanceMetadata.InstanceId }));
+
+						DescribeTagsResponse Response = await Client.DescribeTagsAsync(Request);
+						foreach (TagDescription Tag in Response.Tags)
+						{
+							Properties.Add($"aws-tag={Tag.Key}:{Tag.Value}");
+						}
+					}
+				}
+				catch (Exception Ex)
+				{
+					Logger.LogDebug(Ex, "Unable to query EC2 tags.");
+				}
+			}
+		}			
+
+		static void AddAwsProperty(string Name, string AwsKey, IList<string> Properties)
+		{
+			string? Value = EC2InstanceMetadata.GetData(AwsKey);
+			if (Value != null)
+			{
+				Properties.Add($"{Name}={Value}");
+			}
+		}
+	}
+}
diff --git a/Engine/Source/Programs/Horde/HordeAgent/Utility/ArtifactUploader.cs b/Engine/Source/Programs/Horde/HordeAgent/Utility/ArtifactUploader.cs
new file mode 100644
--- /dev/null
+++ b/Engine/Source/Programs/Horde/HordeAgent/Utility/ArtifactUploader.cs
@@ -0,0 +1,139 @@
+using EpicGames.Core;
+using Grpc.Core;
+using HordeCommon;
+using HordeCommon.Rpc;
+using Microsoft.Extensions.Logging;
+using System;
+using System.Collections.Generic;
+using System.IO;
+using System.Net.Http;
+using System.Net.Http.Headers;
+using System.Runtime.Serialization;
+using System.Text;
+using System.Threading;
+using System.Threading.Tasks;
+
+namespace HordeAgent.Utility
+{
+	/// <summary>
+	/// Helper class for uploading artifacts to the server
+	/// </summary>
+	static class ArtifactUploader
+	{
+        /// <summary>
+        /// Dictionary of horde common mime types
+        /// </summary>
+        private static Dictionary<string, string> HordeMimeTypes = new Dictionary<string, string>()
+        {
+            { ".bin", "application/octet-stream" },
+            { ".json", "application/json" },
+            { ".pm", "text/plain" },
+            { ".pl", "text/plain" },
+            { ".sh", "application/x-sh" },
+            { ".txt", "text/plain" },
+            { ".log", "text/plain" },
+            { ".xml", "text/xml" }
+        };
+
+        /// <summary>
+        /// Gets a mime type from a file extension
+        /// </summary>
+        /// <param name="File">The file object to parse</param>
+        /// <returns>the mimetype if present in the dictionary, binary otherwise</returns>
+        private static string GetMimeType(FileReference File)
+        {
+            string? ContentType;
+            if (!HordeMimeTypes.TryGetValue(File.GetExtension(), out ContentType))
+            {
+                ContentType = "application/octet-stream";
+            }
+            return ContentType;
+        }
+        
+        /// <summary>
+        /// Uploads an artifact (with retries)
+        /// </summary>
+        /// <param name="RpcConnection">The grpc client</param>
+        /// <param name="JobId">Job id</param>
+        /// <param name="BatchId">Job batch id</param>
+        /// <param name="StepId">Job step id</param>
+        /// <param name="ArtifactName">Name of the artifact</param>
+        /// <param name="ArtifactFile">File to upload</param>
+        /// <param name="Logger">Logger interfact</param>
+        /// <param name="CancellationToken">Cancellation token</param>
+        /// <returns></returns>
+        public static async Task<string?> UploadAsync(IRpcConnection RpcConnection, string JobId, string BatchId, string StepId, string ArtifactName, FileReference ArtifactFile, ILogger Logger, CancellationToken CancellationToken)
+        {
+			try
+			{
+				string ArtifactId = await RpcConnection.InvokeAsync(RpcClient => DoUploadAsync(RpcClient, JobId, BatchId, StepId, ArtifactName, ArtifactFile, Logger, CancellationToken), new RpcContext(), CancellationToken);
+				return ArtifactId;
+			}
+			catch (Exception Ex)
+			{
+				Logger.LogInformation(KnownLogEvents.Systemic_Horde_ArtifactUpload, Ex, "Exception while attempting to upload artifact: {Message}", Ex.Message);
+				return null;
+			}
+        }
+
+        /// <summary>
+		/// Uploads an artifact
+		/// </summary>
+		/// <param name="Client">The grpc client</param>
+		/// <param name="JobId">Job id</param>
+		/// <param name="BatchId">Job batch id</param>
+		/// <param name="StepId">Job step id</param>
+		/// <param name="ArtifactName">Name of the artifact</param>
+		/// <param name="ArtifactFile">File to upload</param>
+		/// <param name="Logger">Logger interfact</param>
+		/// <param name="CancellationToken">Cancellation token</param>
+		/// <returns></returns>
+		private static async Task<string> DoUploadAsync(HordeRpc.HordeRpcClient Client, string JobId, string BatchId, string StepId, string ArtifactName, FileReference ArtifactFile, ILogger Logger, CancellationToken CancellationToken)
+        {
+			Logger.LogInformation("Uploading artifact {ArtifactName} from {ArtifactFile}", ArtifactName, ArtifactFile);
+            using (FileStream ArtifactStream = FileReference.Open(ArtifactFile, FileMode.Open, FileAccess.Read, FileShare.Read))
+			{
+				using (AsyncClientStreamingCall<UploadArtifactRequest, UploadArtifactResponse> Cursor = Client.UploadArtifact(null, null, CancellationToken))
+                {
+					// Upload the metadata in the initial request
+					UploadArtifactMetadata Metadata = new UploadArtifactMetadata();
+					Metadata.JobId = JobId;
+					Metadata.BatchId = BatchId;
+					Metadata.StepId = StepId;
+					Metadata.Name = ArtifactName;
+					Metadata.MimeType = GetMimeType(ArtifactFile);
+					Metadata.Length = ArtifactStream.Length;
+
+					UploadArtifactRequest InitialRequest = new UploadArtifactRequest();
+					InitialRequest.Metadata = Metadata;
+
+					await Cursor.RequestStream.WriteAsync(InitialRequest);
+
+					// Upload the data in chunks
+					byte[] Buffer = new byte[4096];
+					for (int Offset = 0; Offset < Metadata.Length; )
+                    {
+                        int BytesRead = await ArtifactStream.ReadAsync(Buffer, 0, Buffer.Length);
+						if(BytesRead == 0)
+						{
+							throw new InvalidDataException($"Unable to read data from {ArtifactFile} beyond offset {Offset}; expected length to be {Metadata.Length}");
+						}
+
+						UploadArtifactRequest Request = new UploadArtifactRequest();
+						Request.Data = Google.Protobuf.ByteString.CopyFrom(Buffer, 0, BytesRead);
+						await Cursor.RequestStream.WriteAsync(Request);
+
+                        Offset += BytesRead;
+                    }
+
+					// Close the stream
+                    await Cursor.RequestStream.CompleteAsync();
+
+					// Read the response
+					UploadArtifactResponse Response = await Cursor.ResponseAsync;
+					return Response.Id;
+                }
+            }
+        }
+	}
+}
diff --git a/Engine/Source/Programs/Horde/HordeAgent/Utility/CertificateHelper.cs b/Engine/Source/Programs/Horde/HordeAgent/Utility/CertificateHelper.cs
new file mode 100644
--- /dev/null
+++ b/Engine/Source/Programs/Horde/HordeAgent/Utility/CertificateHelper.cs
@@ -0,0 +1,111 @@
+using Microsoft.Extensions.Logging;
+using System;
+using System.Collections.Generic;
+using System.Linq;
+using System.Net.Http;
+using System.Net.Security;
+using System.Security.Cryptography.X509Certificates;
+using System.Text;
+using EpicGames.Core;
+
+namespace HordeAgent.Utility
+{
+	/// <summary>
+	/// Helper functions for dealing with certificates
+	/// </summary>
+	static class CertificateHelper
+	{
+		/// <summary>
+		/// Provides additional diagnostic information for SSL certificate validation
+		/// </summary>
+		/// <param name="Logger">The logger instance</param>
+		/// <param name="Sender"></param>
+		/// <param name="Certificate"></param>
+		/// <param name="Chain"></param>
+		/// <param name="SslPolicyErrors"></param>
+		/// <param name="ServerProfile">The server profile</param>
+		/// <returns>True if the certificate is allowed, false otherwise</returns>
+		public static bool CertificateValidationCallBack(ILogger Logger, object Sender, X509Certificate Certificate, X509Chain Chain, SslPolicyErrors SslPolicyErrors, ServerProfile ServerProfile)
+		{
+			// If the certificate is a valid, signed certificate, return true.
+			if (SslPolicyErrors == SslPolicyErrors.None)
+			{
+				return true;
+			}
+
+			// Trust the remote certificate if it has the right thumbprint
+			if (SslPolicyErrors == SslPolicyErrors.RemoteCertificateChainErrors)
+			{
+				if (Chain.ChainElements.Count == 1)
+				{
+					X509ChainElement Element = Chain.ChainElements[0];
+					if (Element.ChainElementStatus.Length == 1 && (Element.ChainElementStatus[0].Status == X509ChainStatusFlags.UntrustedRoot || Element.ChainElementStatus[0].Status == X509ChainStatusFlags.PartialChain))
+					{
+						if (ServerProfile.IsTrustedCertificate(Element.Certificate.Thumbprint))
+						{
+							Logger.LogDebug("Trusting server certificate {Thumbprint}", Element.Certificate.Thumbprint);
+							return true;
+						}
+					}
+
+				}
+			}
+
+			// Generate diagnostic information
+			StringBuilder Builder = new StringBuilder();
+			if (Sender != null)
+			{
+				HttpRequestMessage? Message = Sender as HttpRequestMessage;
+				if (Message != null)
+				{
+					Builder.Append($"\nSender: {Message.Method} {Message.RequestUri}");
+				}
+				else
+				{
+					string SenderInfo = StringUtils.Indent(Sender.ToString() ?? String.Empty, "    ");
+					Builder.Append($"\nSender:\n{SenderInfo}");
+				}
+			}
+			if (Certificate != null)
+			{
+				Builder.Append($"\nCertificate: {Certificate.Subject}");
+			}
+			if (Chain != null)
+			{
+				if (Chain.ChainStatus != null && Chain.ChainStatus.Length > 0)
+				{
+					Builder.Append("\nChain status:");
+					foreach (X509ChainStatus Status in Chain.ChainStatus)
+					{
+						Builder.Append($"\n  {Status.StatusInformation}");
+					}
+				}
+				if (Chain.ChainElements != null)
+				{
+					Builder.Append("\nChain elements:");
+					for (int Idx = 0; Idx < Chain.ChainElements.Count; Idx++)
+					{
+						X509ChainElement Element = Chain.ChainElements[Idx];
+						Builder.Append($"\n  {Idx,4} - Certificate: {Element.Certificate.Subject}");
+						Builder.Append($"\n         Thumbprint: {Element.Certificate.Thumbprint}");
+						if (Element.ChainElementStatus != null && Element.ChainElementStatus.Length > 0)
+						{
+							foreach (X509ChainStatus Status in Element.ChainElementStatus)
+							{
+								Builder.Append($"\n         Status: {Status.StatusInformation} ({Status.Status})");
+							}
+						}
+						if (!String.IsNullOrEmpty(Element.Information))
+						{
+							Builder.Append($"\n         Info: {Element.Information}");
+						}
+					}
+				}
+			}
+
+			// Print out additional diagnostic information
+			Logger.LogError("TLS certificate validation failed ({Errors}).{AdditionalInfo}", SslPolicyErrors, StringUtils.Indent(Builder.ToString(), "    "));
+			return false;
+		}
+	}
+}
diff --git a/Engine/Source/Programs/Horde/HordeAgent/Utility/JsonLogger.cs b/Engine/Source/Programs/Horde/HordeAgent/Utility/JsonLogger.cs
new file mode 100644
--- /dev/null
+++ b/Engine/Source/Programs/Horde/HordeAgent/Utility/JsonLogger.cs
@@ -0,0 +1,352 @@
+using HordeAgent.Parser;
+using HordeCommon;
+using Microsoft.Extensions.Logging;
+using System;
+using System.Buffers;
+using System.Collections.Generic;
+using System.Globalization;
+using System.Linq;
+using System.Text;
+using System.Text.Json;
+
+namespace HordeAgent.Utility
+{
+	/// <summary>
+	/// Class to handle uploading log data to the server in the background
+	/// </summary>
+	public abstract class JsonLogger : ILogger
+	{
+		/// <summary>
+		/// Whether the logger should include warnings, or downgrade them to information
+		/// </summary>
+		protected bool? Warnings { get; }
+
+		/// <summary>
+		/// The inner logger device
+		/// </summary>
+		protected ILogger Inner { get; }
+
+		/// <summary>
+		/// Constructor
+		/// </summary>
+		/// <param name="Warnings">Whether to include warnings in the output</param>
+		/// <param name="Inner">Additional logger to write to</param>
+		public JsonLogger(bool? Warnings, ILogger Inner)
+		{
+			this.Warnings = Warnings;
+			this.Inner = Inner;
+		}
+
+		/// <inheritdoc/>
+		public IDisposable? BeginScope<TState>(TState State)
+		{
+			return null;
+		}
+
+		/// <inheritdoc/>
+		public bool IsEnabled(LogLevel LogLevel)
+		{
+			return true;
+		}
+
+		/// <inheritdoc/>
+		public void Log<TState>(LogLevel LogLevel, EventId EventId, TState State, Exception? Exception, Func<TState, Exception?, string> Formatter)
+		{
+			// Downgrade warnings to information if not required
+			if (LogLevel == LogLevel.Warning && !(Warnings ?? true))
+			{
+				LogLevel = LogLevel.Information;
+			}
+
+			// Render the message
+			string Text = Formatter(State, Exception);
+
+			// Try to log the event
+			IEnumerable<KeyValuePair<string, object>>? Values = State as IEnumerable<KeyValuePair<string, object>>;
+			if (Values != null)
+			{
+				KeyValuePair<string, object>? Format = Values.FirstOrDefault(x => x.Key.Equals(MessageTemplate.FormatPropertyName, StringComparison.Ordinal));
+				if (Format != null)
+				{
+					// Format all the other values
+					QueueFormattedEvent(LogLevel, EventId, Text, Format.Value.Value?.ToString() ?? String.Empty, Values.Where(x => !x.Key.Equals(MessageTemplate.FormatPropertyName, StringComparison.Ordinal)), Exception);
+				}
+				else
+				{
+					// Include all the data, but don't use the format string
+					QueueFormattedEvent(LogLevel, EventId, Text, null, Values, Exception);
+				}
+			}
+			else
+			{
+				// Format as a string
+				QueueFormattedEvent(LogLevel, EventId, Text, null, null, Exception);
+			}
+		}
+
+		/// <summary>
+		/// Queues up a log event for the server
+		/// </summary>
+		/// <param name="LogLevel">Level for this log event</param>
+		/// <param name="EventId">Corresponding event id</param>
+		/// <param name="Text">Rendered text for the message</param>
+		/// <param name="Format">Format string</param>
+		/// <param name="Properties">Properties for the message</param>
+		/// <param name="Exception">Exception information</param>
+		void QueueFormattedEvent(LogLevel LogLevel, EventId EventId, string Text, string? Format, IEnumerable<KeyValuePair<string, object>>? Properties, Exception? Exception)
+		{
+			DateTime Time = DateTime.UtcNow;
+
+			// Split it into individual lines
+			string[] Lines = Text.Split('\n');
+
+			// Split the format string into lines too
+			string[]? FormatLines = null;
+			if (Format != null)
+			{
+				FormatLines = Format.Split('\n');
+			}
+
+			// If it's a systemic issue, log it to the regular log sink so we can catch it in Datadog
+			if (EventId.Id >= KnownLogEvents.Systemic.Id && EventId.Id < KnownLogEvents.Systemic.Id + 100)
+			{
+				WriteSystemicEvent(EventId, Text);
+			}
+
+			// In order to support multi-line log events while seeking through the log buffer, we re-output the same json event for each
+			// line in the multi-line output. The 'line' parameter for the second and successive lines gives the number of this line.
+			byte[][] OutputLines = new byte[Lines.Length][];
+			for (int LineIdx = 0; LineIdx < Lines.Length; LineIdx++)
+			{
+				// Encode the data for this line
+				string? LineFormat = null;
+				if (FormatLines != null)
+				{
+					if (LineIdx < FormatLines.Length)
+					{
+						LineFormat = FormatLines[LineIdx];
+					}
+					else
+					{
+						LineFormat = "Missing format string";
+					}
+				}
+
+				OutputLines[LineIdx] = EncodeMessage(Time, LogLevel, Lines[LineIdx], EventId, LineIdx, Lines.Length, LineFormat, Properties, Exception);
+			}
+			WriteFormattedEvent(LogLevel, OutputLines);
+		}
+
+		/// <summary>
+		/// Writes a formatted event
+		/// </summary>
+		/// <param name="Level">The log level</param>
+		/// <param name="Lines">Utf-8 encoded JSON line data</param>
+		protected abstract void WriteFormattedEvent(LogLevel Level, byte[][] Lines);
+
+		/// <summary>
+		/// Callback to write a systemic event
+		/// </summary>
+		/// <param name="EventId">The event id</param>
+		/// <param name="Text">The event text</param>
+		protected virtual void WriteSystemicEvent(EventId EventId, string Text)
+		{
+			Inner.LogWarning("Systemic event {KnownLogEventId}: {Text}", EventId.Id, Text);
+		}
+
+		/// <summary>
+		/// Queues up a log event for the server
+		/// </summary>
+		/// <param name="Time">Time of the event</param>
+		/// <param name="LogLevel">Level for this log event</param>
+		/// <param name="Text">Rendered text for the message</param>
+		/// <param name="EventId">Corresponding event id</param>
+		/// <param name="LineIndex">Line index within the event</param>
+		/// <param name="LineCount">Number of lines in the event</param>
+		/// <param name="Format">Format string</param>
+		/// <param name="Properties">Properties for the message</param>
+		/// <param name="Exception">Exception information</param>
+		byte[] EncodeMessage(DateTime Time, LogLevel LogLevel, string Text, EventId EventId, int LineIndex, int LineCount, string? Format, IEnumerable<KeyValuePair<string, object>>? Properties, Exception? Exception)
+		{
+			ArrayBufferWriter<byte> Writer = new ArrayBufferWriter<byte>();
+			using (Utf8JsonWriter JsonWriter = new Utf8JsonWriter(Writer))
+			{
+				JsonWriter.WriteStartObject();
+				JsonWriter.WriteString("time", Time.ToString("s", CultureInfo.InvariantCulture));
+				JsonWriter.WriteString("level", LogLevel.ToString());
+				JsonWriter.WriteString("message", Text);
+
+				if (EventId != KnownLogEvents.None)
+				{
+					JsonWriter.WriteNumber("id", EventId.Id);
+				}
+
+				if (LineCount > 1)
+				{
+					JsonWriter.WriteNumber("line", LineIndex);
+					JsonWriter.WriteNumber("lineCount", LineCount);
+				}
+
+				if (Format != null)
+				{
+					JsonWriter.WriteString("format", Format);
+				}
+
+				if (Properties != null && Properties.Any())
+				{
+					JsonWriter.WriteStartObject("properties");
+					foreach (KeyValuePair<string, object> Pair in Properties!)
+					{
+						JsonWriter.WritePropertyName(Pair.Key);
+						WritePropertyValue(Pair.Value, JsonWriter);
+					}
+					JsonWriter.WriteEndObject();
+				}
+
+				if (Exception != null)
+				{
+					JsonWriter.WriteStartObject("exception");
+					WriteException(JsonWriter, Exception);
+					JsonWriter.WriteEndObject();
+				}
+
+				JsonWriter.WriteEndObject();
+			}
+			return Writer.WrittenSpan.ToArray();
+		}
+
+		/// <summary>
+		/// Writes an exception to a json object
+		/// </summary>
+		/// <param name="Writer">Writer to receive the exception data</param>
+		/// <param name="Exception">The exception</param>
+		void WriteException(Utf8JsonWriter Writer, Exception Exception)
+		{
+			Writer.WriteString("message", Exception.Message);
+			Writer.WriteString("trace", Exception.StackTrace);
+
+			if (Exception.InnerException != null)
+			{
+				Writer.WriteStartObject("innerException");
+				WriteException(Writer, Exception.InnerException);
+				Writer.WriteEndObject();
+			}
+
+			AggregateException? AggregateException = Exception as AggregateException;
+			if (AggregateException != null)
+			{
+				Writer.WriteStartArray("innerExceptions");
+				for (int Idx = 0; Idx < 16 && Idx < AggregateException.InnerExceptions.Count; Idx++) // Cap number of exceptions returned to avoid huge messages
+				{
+					Exception InnerException = AggregateException.InnerExceptions[Idx];
+					Writer.WriteStartObject();
+					WriteException(Writer, InnerException);
+					Writer.WriteEndObject();
+				}
+				Writer.WriteEndArray();
+			}
+		}
+
+		/// <summary>
+		/// Converts a property value to log markup
+		/// </summary>
+		/// <param name="Value">The value to write</param>
+		/// <param name="Writer">The json writer</param>
+		static void WritePropertyValue(object Value, Utf8JsonWriter Writer)
+		{
+			if (Value == null)
+			{
+				Writer.WriteNullValue();
+			}
+			else
+			{
+				Type ValueType = Value.GetType();
+				if (ValueType.IsEnum)
+				{
+					Writer.WriteStringValue(Enum.GetName(ValueType, Value));
+				}
+				else
+				{
+					switch (Type.GetTypeCode(ValueType))
+					{
+						case TypeCode.Boolean:
+							Writer.WriteBooleanValue((bool)Value);
+							break;
+						case TypeCode.Byte:
+							Writer.WriteNumberValue((Byte)Value);
+							break;
+						case TypeCode.SByte:
+							Writer.WriteNumberValue((SByte)Value);
+							break;
+						case TypeCode.UInt16:
+							Writer.WriteNumberValue((UInt16)Value);
+							break;
+						case TypeCode.UInt32:
+							Writer.WriteNumberValue((UInt32)Value);
+							break;
+						case TypeCode.UInt64:
+							Writer.WriteNumberValue((UInt64)Value);
+							break;
+						case TypeCode.Int16:
+							Writer.WriteNumberValue((Int16)Value);
+							break;
+						case TypeCode.Int32:
+							Writer.WriteNumberValue((Int32)Value);
+							break;
+						case TypeCode.Int64:
+							Writer.WriteNumberValue((Int64)Value);
+							break;
+						case TypeCode.Decimal:
+							Writer.WriteNumberValue((Decimal)Value);
+							break;
+						case TypeCode.Double:
+							Writer.WriteNumberValue((Double)Value);
+							break;
+						case TypeCode.Single:
+							Writer.WriteNumberValue((Single)Value);
+							break;
+						case TypeCode.String:
+							Writer.WriteStringValue((String)Value);
+							break;
+						default:
+							WriteComplexPropertyValue(Value, ValueType, Writer);
+							break;
+					}
+				}
+			}
+		}
+
+		/// <summary>
+		/// Writes a complex object to the given json writer
+		/// </summary>
+		/// <param name="Value">The value to write</param>
+		/// <param name="ValueType">Type of the value</param>
+		/// <param name="JsonWriter">The writer to output to</param>
+		static void WriteComplexPropertyValue(object Value, Type ValueType, Utf8JsonWriter JsonWriter)
+		{
+			if (ValueType == typeof(LogEventSpan))
+			{
+				LogEventSpan Span = (LogEventSpan)Value;
+				if (Span.Properties.Count == 0)
+				{
+					JsonWriter.WriteStringValue(Span.Text);
+				}
+				else
+				{
+					JsonWriter.WriteStartObject();
+					JsonWriter.WriteString("text", Span.Text);
+					foreach (KeyValuePair<string, object> Pair in Span.Properties)
+					{
+						JsonWriter.WritePropertyName(Pair.Key);
+						WritePropertyValue(Pair.Value, JsonWriter);
+					}
+					JsonWriter.WriteEndObject();
+				}
+			}
+			else
+			{
+				JsonWriter.WriteStringValue(Value.ToString());
+			}
+		}
+	}
+}
diff --git a/Engine/Source/Programs/Horde/HordeAgent/Utility/JsonRpcLogger.cs b/Engine/Source/Programs/Horde/HordeAgent/Utility/JsonRpcLogger.cs
new file mode 100644
--- /dev/null
+++ b/Engine/Source/Programs/Horde/HordeAgent/Utility/JsonRpcLogger.cs
@@ -0,0 +1,305 @@
+// Copyright Epic Games, Inc. All Rights Reserved.
+
+using HordeAgent.Parser.Interfaces;
+using HordeAgent.Utility;
+using HordeCommon;
+using HordeCommon.Rpc;
+using Microsoft.Extensions.Logging;
+using System;
+using System.Buffers;
+using System.Collections.Concurrent;
+using System.Collections.Generic;
+using System.Globalization;
+using System.Linq;
+using System.Net.Http;
+using System.Text;
+using System.Text.Json;
+using System.Threading;
+using System.Threading.Channels;
+using System.Threading.Tasks;
+
+namespace HordeAgent.Parser
+{
+	/// <summary>
+	/// Class to handle uploading log data to the server in the background
+	/// </summary>
+	class JsonRpcLogger : JsonLogger, IAsyncDisposable
+	{
+		class QueueItem
+		{
+			public byte[] Data;
+			public CreateEventRequest? Event;
+
+			public QueueItem(byte[] Data, CreateEventRequest? Event)
+			{
+				this.Data = Data;
+				this.Event = Event;
+			}
+
+			public override string ToString()
+			{
+				return Encoding.UTF8.GetString(Data);
+			}
+		}
+
+		readonly IRpcConnection RpcClient;
+		readonly string LogId;
+		readonly string? JobId;
+		readonly string? JobBatchId;
+		readonly string? JobStepId;
+		Channel<QueueItem> DataChannel;
+		Task? DataWriter;
+
+		/// <summary>
+		/// The current outcome for this step. Updated to reflect any errors and warnings that occurred.
+		/// </summary>
+		public JobStepOutcome Outcome
+		{
+			get;
+			private set;
+		}
+
+		/// <summary>
+		/// Constructor
+		/// </summary>
+		/// <param name="RpcClient">RPC client to use for server requests</param>
+		/// <param name="LogId">The log id to write to</param>
+		/// <param name="JobId">Id of the job being executed</param>
+		/// <param name="JobBatchId">Batch being executed</param>
+		/// <param name="JobStepId">Id of the step being executed</param>
+		/// <param name="Warnings">Whether to include warnings in the output</param>
+		/// <param name="Inner">Additional logger to write to</param>
+		public JsonRpcLogger(IRpcConnection RpcClient, string LogId, string? JobId, string? JobBatchId, string? JobStepId, bool? Warnings, ILogger Inner)
+			: base(Warnings, Inner)
+		{
+			this.RpcClient = RpcClient;
+			this.LogId = LogId;
+			this.JobId = JobId;
+			this.JobBatchId = JobBatchId;
+			this.JobStepId = JobStepId;
+			this.DataChannel = Channel.CreateUnbounded<QueueItem>();
+			this.DataWriter = Task.Run(() => RunDataWriter());
+			this.Outcome = JobStepOutcome.Success;
+		}
+
+		protected override void WriteFormattedEvent(LogLevel Level, byte[][] Lines)
+		{
+			// Update the state of this job if this is an error status
+			if (Level == LogLevel.Error || Level == LogLevel.Critical)
+			{
+				Outcome = JobStepOutcome.Failure;
+			}
+			else if (Level == LogLevel.Warning && Outcome != JobStepOutcome.Failure)
+			{
+				Outcome = JobStepOutcome.Warnings;
+			}
+
+			for (int Idx = 0; Idx < Lines.Length; Idx++)
+			{
+				// If we want an event for this log event, create one now
+				CreateEventRequest? Event = null;
+				if (Idx == 0)
+				{
+					if (Level == LogLevel.Warning || Level == LogLevel.Error || Level == LogLevel.Critical)
+					{
+						Event = CreateEvent(Level, Lines.Length);
+					}
+				}
+
+				// Write the data to the output channel
+				QueueItem QueueItem = new QueueItem(Lines[Idx], Event);
+				if (!DataChannel.Writer.TryWrite(QueueItem))
+				{
+					throw new InvalidOperationException("Expected unbounded writer to complete immediately");
+				}
+			}
+		}
+
+		/// <summary>
+		/// Callback to write a systemic event
+		/// </summary>
+		/// <param name="EventId">The event id</param>
+		/// <param name="Text">The event text</param>
+		protected override void WriteSystemicEvent(EventId EventId, string Text)
+		{
+			if (JobId == null)
+			{
+				Inner.LogWarning("Systemic event {KnownLogEventId} in log {LogId}: {Text}", EventId.Id, Text);
+			}
+			else if (JobBatchId == null)
+			{
+				Inner.LogWarning("Systemic event {KnownLogEventId} in log {LogId} (job {JobId}): {Text})", EventId.Id, JobId, Text);
+			}
+			else if (JobStepId == null)
+			{
+				Inner.LogWarning("Systemic event {KnownLogEventId} in log {LogId} (job batch {JobId}:{BatchId}): {Text})", EventId.Id, JobId, JobStepId, Text);
+			}
+			else
+			{
+				Inner.LogWarning("Systemic event {KnownLogEventId} in log {LogId} (job step {JobId}:{BatchId}:{StepId}): {Text}", EventId.Id, LogId, JobId, JobBatchId, JobStepId, Text);
+			}
+		}
+
+		/// <summary>
+		/// Makes a <see cref="CreateEventRequest"/> for the given parameters
+		/// </summary>
+		/// <param name="LogLevel">Level for this log event</param>
+		/// <param name="LineCount">Number of lines in the event</param>
+		CreateEventRequest CreateEvent(LogLevel LogLevel, int LineCount)
+		{
+			EventSeverity Severity = (LogLevel == LogLevel.Warning) ? EventSeverity.Warning : EventSeverity.Error;
+			return new CreateEventRequest(Severity, LogId, 0, LineCount);
+		}
+
+		/// <summary>
+		/// Stops the log writer's background task
+		/// </summary>
+		/// <returns>Async task</returns>
+		public async Task StopAsync()
+		{
+			if (DataWriter != null)
+			{
+				DataChannel.Writer.TryComplete();
+				await DataWriter;
+				DataWriter = null;
+			}
+		}
+
+		/// <summary>
+		/// Dispose of this object. Call StopAsync() to stop asynchronously.
+		/// </summary>
+		public async ValueTask DisposeAsync()
+		{
+			await StopAsync();
+		}
+
+		/// <summary>
+		/// Upload the log data to the server in the background
+		/// </summary>
+		/// <returns>Async task</returns>
+		async Task RunDataWriter()
+		{
+			// Current position and line number in the log file
+			long Offset = 0;
+			int LineIndex = 0;
+
+			// Total number of errors and warnings
+			const int MaxErrors = 50;
+			int NumErrors = 0;
+			const int MaxWarnings = 50;
+			int NumWarnings = 0;
+
+			// Buffers for chunks and events read in a single iteration
+			ArrayBufferWriter<byte> Writer = new ArrayBufferWriter<byte>();
+			List<CreateEventRequest> Events = new List<CreateEventRequest>();
+
+			// Line separator for JSON events
+			byte[] Newline = { (byte)'\n' };
+			JsonEncodedText Timestamp = JsonEncodedText.Encode("");
+
+			// The current jobstep outcome
+			JobStepOutcome PostedOutcome = JobStepOutcome.Success;
+
+			// Whether we've written the flush command
+			for (; ; )
+			{
+				Writer.Clear();
+				Events.Clear();
+
+				// Save off the current line number for sending to the server
+				int InitialLineIndex = LineIndex;
+
+				// Get the next data
+				Task WaitTask = Task.Delay(TimeSpan.FromSeconds(2.0));
+				while (Writer.WrittenCount < 256 * 1024)
+				{
+					QueueItem? Data;
+					if (DataChannel.Reader.TryRead(out Data))
+					{
+						if (Data.Event != null)
+						{
+							if ((Data.Event.Severity == EventSeverity.Warning && ++NumWarnings < MaxWarnings) || (Data.Event.Severity == EventSeverity.Error && ++NumErrors < MaxErrors))
+							{
+								Data.Event.LineIndex = LineIndex;
+								Events.Add(Data.Event);
+							}
+						}
+
+						Writer.Write(Data.Data);
+						Writer.Write(Newline);
+
+						LineIndex++;
+					}
+					else
+					{
+						Task<bool> ReadTask = DataChannel.Reader.WaitToReadAsync().AsTask();
+						if (await Task.WhenAny(ReadTask, WaitTask) == WaitTask)
+						{
+							break;
+						}
+						if (!await ReadTask)
+						{
+							break;
+						}
+					}
+				}
+
+				// Upload it to the server
+				if (Writer.WrittenCount > 0)
+				{
+					byte[] Data = Writer.WrittenSpan.ToArray();
+					try
+					{
+						await RpcClient.InvokeAsync(x => x.WriteOutputAsync(new WriteOutputRequest(LogId, Offset, InitialLineIndex, Data, false)), new RpcContext(), CancellationToken.None);
+					}
+					catch (Exception Ex)
+					{
+						Inner.LogWarning(Ex, "Unable to write data to server (log {LogId}, offset {Offset})", LogId, Offset);
+					}
+					Offset += Data.Length;
+				}
+
+				// Write all the events
+				if (Events.Count > 0)
+				{
+					try
+					{
+						await RpcClient.InvokeAsync(x => x.CreateEventsAsync(new CreateEventsRequest(Events)), new RpcContext(), CancellationToken.None);
+					}
+					catch (Exception Ex)
+					{
+						Inner.LogWarning(Ex, "Unable to create events");
+					}
+				}
+
+				// Update the outcome of this jobstep
+				if (JobId != null && JobBatchId != null && JobStepId != null && Outcome != PostedOutcome)
+				{
+					try
+					{
+						await RpcClient.InvokeAsync(x => x.UpdateStepAsync(new UpdateStepRequest(JobId, JobBatchId, JobStepId, JobStepState.Unspecified, Outcome)), new RpcContext(), CancellationToken.None);
+					}
+					catch (Exception Ex)
+					{
+						Inner.LogWarning(Ex, "Unable to update step outcome to {NewOutcome}", Outcome);
+					}
+					PostedOutcome = Outcome;
+				}
+
+				// Wait for more data to be available
+				if (!await DataChannel.Reader.WaitToReadAsync())
+				{
+					try
+					{
+						await RpcClient.InvokeAsync(x => x.WriteOutputAsync(new WriteOutputRequest(LogId, Offset, LineIndex, Array.Empty<byte>(), true)), new RpcContext(), CancellationToken.None);
+					}
+					catch (Exception Ex)
+					{
+						Inner.LogWarning(Ex, "Unable to flush data to server (log {LogId}, offset {Offset})", LogId, Offset);
+					}
+					break;
+				}
+			}
+		}
+	}
+}
diff --git a/Engine/Source/Programs/Horde/HordeAgent/Utility/NetworkShare.cs b/Engine/Source/Programs/Horde/HordeAgent/Utility/NetworkShare.cs
new file mode 100644
--- /dev/null
+++ b/Engine/Source/Programs/Horde/HordeAgent/Utility/NetworkShare.cs
@@ -0,0 +1,129 @@
+using System;
+using System.Collections.Generic;
+using System.ComponentModel;
+using System.Diagnostics.CodeAnalysis;
+using System.Runtime.InteropServices;
+using System.Text;
+
+namespace HordeAgent.Utility
+{
+	/// <summary>
+	/// Utility functions for mapping network shares
+	/// </summary>
+	static class NetworkShare
+	{
+		[StructLayout(LayoutKind.Sequential, CharSet = CharSet.Unicode)]
+		private struct NETRESOURCE
+		{
+			public uint dwScope;
+			public uint dwType;
+			public uint dwDisplayType;
+			public uint dwUsage;
+			public string? lpLocalName;
+			public string? lpRemoteName;
+			public string? lpComment;
+			public string? lpProvider;
+		};
+
+		const uint RESOURCETYPE_DISK = 1;
+		const uint CONNECT_TEMPORARY = 4;
+
+
+		const int ERROR_ALREADY_ASSIGNED = 85;
+		const int ERROR_MORE_DATA = 234;
+
+		[DllImport("Mpr.dll", CharSet = CharSet.Unicode)]
+		private static extern int WNetAddConnection2W(ref NETRESOURCE lpNetResource, string? lpPassword, string? lpUsername, System.UInt32 dwFlags );
+
+		[DllImport("Mpr.dll", CharSet = CharSet.Unicode)]
+		private static extern int WNetCancelConnectionW(string lpName, [MarshalAs(UnmanagedType.Bool)] bool fForce);
+
+		[DllImport("Mpr.dll", CharSet = CharSet.Unicode)]
+		private static extern int WNetGetConnectionW(string lpLocalName, IntPtr lpRemoteName, ref int lpnLength);
+
+		/// <summary>
+		/// Mounts a network share at the given location
+		/// </summary>
+		/// <param name="MountPoint">Mount point for the share</param>
+		/// <param name="RemotePath">Path to the remote resource</param>
+		public static void Mount(string MountPoint, string RemotePath)
+		{
+			NETRESOURCE NetResource = new NETRESOURCE();
+			NetResource.dwType = RESOURCETYPE_DISK;
+			NetResource.lpLocalName = MountPoint;
+			NetResource.lpRemoteName = RemotePath;
+
+			int Result = WNetAddConnection2W(ref NetResource, null, null, CONNECT_TEMPORARY);
+			if(Result != 0)
+			{
+				if (Result == ERROR_ALREADY_ASSIGNED)
+				{
+					string? CurRemotePath;
+					if (TryGetRemotePath(MountPoint, out CurRemotePath))
+					{
+						if (CurRemotePath.Equals(RemotePath, StringComparison.OrdinalIgnoreCase))
+						{
+							return;
+						}
+						else
+						{
+							throw new Win32Exception(Result, $"Unable to mount network share {RemotePath} as {MountPoint} ({Result}). Currently connected to {CurRemotePath}.");
+						}
+					}
+				}
+				throw new Win32Exception(Result, $"Unable to mount network share {RemotePath} as {MountPoint} ({Result}: {new Win32Exception(Result).Message})");
+			}
+		}
+
+		/// <summary>
+		/// Unmounts a network share
+		/// </summary>
+		/// <param name="MountPoint">The mount point to remove</param>
+		public static void Unmount(string MountPoint)
+		{
+			int Result = WNetCancelConnectionW(MountPoint, true);
+			if(Result != 0)
+			{
+				throw new Win32Exception(Result, $"Unable to unmount {MountPoint} ({Result})");
+			}
+		}
+
+		/// <summary>
+		/// Gets the currently mounted path for a network share
+		/// </summary>
+		/// <param name="MountPoint">The mount point</param>
+		/// <param name="OutRemotePath">Receives the remote path</param>
+		/// <returns>True if the remote path is returned</returns>
+		public static bool TryGetRemotePath(string MountPoint, [NotNullWhen(true)] out string? OutRemotePath)
+		{
+			int Length = 260 * sizeof(char);
+			for (; ; )
+			{
+				IntPtr RemotePath = Marshal.AllocHGlobal(Length);
+				try
+				{
+					int PrevLength = Length;
+					int Result = WNetGetConnectionW(MountPoint, RemotePath, ref Length);
+					if (Result == 0)
+					{
+						OutRemotePath = Marshal.PtrToStringUni(RemotePath)!;
+						return OutRemotePath != null;
+					}
+					else if(Result == ERROR_MORE_DATA && Length > PrevLength)
+					{
+						continue;
+					}
+					else
+					{
+						OutRemotePath = null!;
+						return false;
+					}
+				}
+				finally
+				{
+					Marshal.FreeHGlobal(RemotePath);
+				}
+			}
+		}
+	}
+}
