Initial drop of Horde Server/Agent for UE5/Main. Still missing some Epic-specific functionality required to build/run correctly.

diff --git a/Engine/Source/Programs/Horde/HordeAgent/Utility/RpcConnection.cs b/Engine/Source/Programs/Horde/HordeAgent/Utility/RpcConnection.cs
new file mode 100644
--- /dev/null
+++ b/Engine/Source/Programs/Horde/HordeAgent/Utility/RpcConnection.cs
@@ -0,0 +1,788 @@
+using Datadog.Trace;
+using EpicGames.Core;
+using Grpc.Core;
+using Grpc.Net.Client;
+using HordeCommon.Rpc;
+using Microsoft.Extensions.Hosting;
+using Microsoft.Extensions.Logging;
+using System;
+using System.Collections.Concurrent;
+using System.Collections.Generic;
+using System.Diagnostics;
+using System.IO;
+using System.Linq;
+using System.Net;
+using System.Net.Http;
+using System.Net.Http.Headers;
+using System.Runtime.CompilerServices;
+using System.Runtime.InteropServices;
+using System.Threading;
+using System.Threading.Tasks;
+
+[assembly: InternalsVisibleTo("HordeAgentTests")]
+
+namespace HordeAgent.Utility
+{
+	/// <summary>
+	/// Ref-counted reference to a HordeRpcClient. Must be disposed after use.
+	/// </summary>
+	interface IRpcClientRef : IDisposable
+	{
+		/// <summary>
+		/// The Grpc channel instance
+		/// </summary>
+		GrpcChannel Channel { get; }
+
+		/// <summary>
+		/// The client instance
+		/// </summary>
+		HordeRpc.HordeRpcClient Client { get; }
+
+		/// <summary>
+		/// Task which completes when the client needs to be disposed of
+		/// </summary>
+		Task DisposingTask { get; }
+	}
+
+	/// <summary>
+	/// Context for an RPC call. Used for debugging reference leaks.
+	/// </summary>
+	class RpcContext
+	{
+		/// <summary>
+		/// Text to display
+		/// </summary>
+		public string Text { get; }
+
+		/// <summary>
+		/// Constructor
+		/// </summary>
+		/// <param name="Text">Text to display</param>
+		public RpcContext(string Text)
+		{
+			this.Text = Text;
+		}
+
+		/// <summary>
+		/// Constructor
+		/// </summary>
+		/// <param name="Line">Line number of the file</param>
+		/// <param name="File">Calling file</param>
+		public RpcContext([CallerLineNumber] int Line = 0, [CallerFilePath] string File = "")
+		{
+			this.Text = $"{File}({Line})";
+		}
+	}
+
+	internal interface IRpcConnection : IAsyncDisposable
+	{
+		/// <summary>
+		/// Attempts to get a client reference, returning immediately if there's not one available
+		/// </summary>
+		/// <returns></returns>
+		IRpcClientRef? TryGetClientRef(RpcContext Context);
+
+		/// <summary>
+		/// Obtains a new client reference object
+		/// </summary>
+		/// <param name="Context">Context for the call, for debugging</param>
+		/// <param name="CancellationToken">Cancellation token for this request</param>
+		/// <returns>New client reference</returns>
+		Task<IRpcClientRef> GetClientRef(RpcContext Context, CancellationToken CancellationToken);
+
+		/// <summary>
+		/// Invokes an asynchronous command with a HordeRpcClient instance
+		/// </summary>
+		/// <param name="Func">The function to execute</param>
+		/// <param name="Context">Context for the call, for debugging</param>
+		/// <param name="CancellationToken">Cancellation token for the call</param>
+		/// <returns>Response from the call</returns>
+		Task<T> InvokeOnceAsync<T>(Func<HordeRpc.HordeRpcClient, Task<T>> Func, RpcContext Context, CancellationToken CancellationToken);
+
+		/// <summary>
+		/// Obtains a client reference and executes a unary RPC
+		/// </summary>
+		/// <typeparam name="T">Return type from the call</typeparam>
+		/// <param name="Func">Method to execute with the RPC instance</param>
+		/// <param name="Context">Context for the call, for debugging</param>
+		/// <param name="CancellationToken">Cancellation token for the call</param>
+		/// <returns>Response from the call</returns>
+		Task<T> InvokeOnceAsync<T>(Func<HordeRpc.HordeRpcClient, AsyncUnaryCall<T>> Func, RpcContext Context, CancellationToken CancellationToken);
+
+		/// <summary>
+		/// Invokes an asynchronous command with a HordeRpcClient instance
+		/// </summary>
+		/// <param name="Func">The function to execute</param>
+		/// <param name="Context">Context for the call, for debugging</param>
+		/// <param name="CancellationToken">Cancellation token for the call</param>
+		/// <returns>Response from the call</returns>
+		Task<T> InvokeAsync<T>(Func<HordeRpc.HordeRpcClient, Task<T>> Func, RpcContext Context, CancellationToken CancellationToken);
+
+		/// <summary>
+		/// Obtains a client reference and executes a unary RPC
+		/// </summary>
+		/// <typeparam name="T">Return type from the call</typeparam>
+		/// <param name="Func">Method to execute with the RPC instance</param>
+		/// <param name="Context">Context for the call, for debugging</param>
+		/// <param name="CancellationToken">Cancellation token for the call</param>
+		/// <returns>Response from the call</returns>
+		Task<T> InvokeAsync<T>(Func<HordeRpc.HordeRpcClient, AsyncUnaryCall<T>> Func, RpcContext Context, CancellationToken CancellationToken);
+	}
+
+	/// <summary>
+	/// Represents a connection to a GRPC server, which responds to server shutting down events and attempts to reconnect to another instance.
+	/// </summary>
+	class RpcConnection : IRpcConnection
+	{
+		/// <summary>
+		/// An instance of a connection to a particular server
+		/// </summary>
+		class RpcSubConnection : IAsyncDisposable
+		{
+			/// <summary>
+			/// The connection id
+			/// </summary>
+			public int ConnectionId { get; }
+
+			/// <summary>
+			/// Name of the server
+			/// </summary>
+			public string Name
+			{
+				get;
+			}
+
+			/// <summary>
+			/// The Grpc channel
+			/// </summary>
+			public GrpcChannel Channel { get; }
+
+			/// <summary>
+			/// The RPC interface
+			/// </summary>
+			public HordeRpc.HordeRpcClient Client
+			{
+				get;
+			}
+
+			/// <summary>
+			/// Task which is set to indicate the connection is being disposed
+			/// </summary>
+			public Task DisposingTask => DisposingTaskSource.Task;
+
+			/// <summary>
+			/// Logger for messages about refcount changes
+			/// </summary>
+			ILogger Logger;
+
+			/// <summary>
+			/// The current refcount
+			/// </summary>
+			int RefCount;
+
+			/// <summary>
+			/// Whether the connection is terminating
+			/// </summary>
+			bool bDisposing;
+
+			/// <summary>
+			/// Cancellation token set 
+			/// </summary>
+			TaskCompletionSource<bool> DisposingTaskSource = new TaskCompletionSource<bool>();
+
+			/// <summary>
+			/// Wraps a task allowing the disposer to wait for clients to finish using this connection
+			/// </summary>
+			TaskCompletionSource<bool> DisposedTaskSource = new TaskCompletionSource<bool>();
+
+			/// <summary>
+			/// Constructor
+			/// </summary>
+			/// <param name="ConnectionId">The connection id</param>
+			/// <param name="Name">Name of the server</param>
+			/// <param name="Channel">The channel instance</param>
+			/// <param name="Client">The client instance</param>
+			/// <param name="Logger">Logger for debug messages</param>
+			public RpcSubConnection(int ConnectionId, string Name, GrpcChannel Channel, HordeRpc.HordeRpcClient Client, ILogger Logger)
+			{
+				this.ConnectionId = ConnectionId;
+				this.Name = Name;
+				this.Channel = Channel;
+				this.Client = Client;
+				this.Logger = Logger;
+				this.RefCount = 1;
+			}
+
+			/// <summary>
+			/// Attempts to increment the reference count for this subconnection. Fails if it's already zero.
+			/// </summary>
+			/// <returns>True if a reference was added</returns>
+			public bool TryAddRef()
+			{
+				for (; ; )
+				{
+					int InitialRefCount = RefCount;
+					if (InitialRefCount == 0)
+					{
+						return false;
+					}
+					if (Interlocked.CompareExchange(ref RefCount, InitialRefCount + 1, InitialRefCount) == InitialRefCount)
+					{
+						return true;
+					}
+				}
+			}
+
+			/// <summary>
+			/// Release a reference to this connection
+			/// </summary>
+			public void Release()
+			{
+				int NewRefCount = Interlocked.Decrement(ref RefCount);
+				if (bDisposing)
+				{
+					Logger.LogInformation("Connection {ConnectionId}: Decrementing refcount to rpc server {ServerName} ({RefCount} rpc(s) in flight)", ConnectionId, Name, NewRefCount);
+				}
+				if (NewRefCount == 0)
+				{
+					DisposedTaskSource.SetResult(true);
+				}
+			}
+
+			/// <summary>
+			/// Dispose of this connection, waiting for all references to be released first
+			/// </summary>
+			/// <returns>Async task</returns>
+			public ValueTask DisposeAsync()
+			{
+				Logger.LogInformation("Connection {ConnectionId}: Disposing sub-connection", ConnectionId);
+				DisposingTaskSource.TrySetResult(true);
+				bDisposing = true;
+				Release();
+				return new ValueTask(DisposedTaskSource.Task);
+			}
+		}
+
+		/// <summary>
+		/// Reference to a connection instance
+		/// </summary>
+		class RpcClientRef : IRpcClientRef
+		{
+			/// <summary>
+			/// The subconnection containing the client
+			/// </summary>
+			public RpcSubConnection SubConnection
+			{
+				get;
+			}
+
+			/// <summary>
+			/// Context for the reference
+			/// </summary>
+			public RpcContext Context
+			{
+				get;
+			}
+
+			/// <summary>
+			/// The channel instance
+			/// </summary>
+			public GrpcChannel Channel
+			{
+				get { return SubConnection.Channel; }
+			}
+
+			/// <summary>
+			/// The client instance
+			/// </summary>
+			public HordeRpc.HordeRpcClient Client
+			{
+				get { return SubConnection.Client; }
+			}
+
+			/// <summary>
+			/// Task which completes when the connection should be disposed
+			/// </summary>
+			public Task DisposingTask
+			{
+				get { return SubConnection.DisposingTask; }
+			}
+
+			/// <summary>
+			/// Tracks all the active refs
+			/// </summary>
+			static List<RpcClientRef> CurrentRefs = new List<RpcClientRef>();
+
+			/// <summary>
+			/// Constructor
+			/// </summary>
+			/// <param name="SharedClient">The subconnection containing the client</param>
+			/// <param name="Context">Description of the reference, for debugging</param>
+			public RpcClientRef(RpcSubConnection SharedClient, RpcContext Context)
+			{
+				this.SubConnection = SharedClient;
+				this.Context = Context;
+
+				lock (CurrentRefs)
+				{
+					CurrentRefs.Add(this);
+				}
+			}
+
+			/// <summary>
+			/// Dispose of this reference
+			/// </summary>
+			public void Dispose()
+			{
+				SubConnection.Release();
+
+				lock(CurrentRefs)
+				{
+					CurrentRefs.Remove(this);
+				}
+			}
+
+			/// <summary>
+			/// Gets a list of all the active refs
+			/// </summary>
+			/// <returns>List of ref descriptions</returns>
+			public static List<string> GetActiveRefs()
+			{
+				lock(CurrentRefs)
+				{
+					return CurrentRefs.Select(x => x.Context.Text).ToList();
+				}
+			}
+		}
+
+		/// <summary>
+		/// Delays before retrying RPC calls
+		/// </summary>
+		static TimeSpan[] RetryTimes =
+		{
+			TimeSpan.FromSeconds(1.0),
+			TimeSpan.FromSeconds(10.0),
+			TimeSpan.FromSeconds(30.0),
+		};
+
+		/// <summary>
+		/// Delegate to create a new GrpcChannel
+		/// </summary>
+		Func<GrpcChannel> CreateGrpcChannel;
+
+		/// <summary>
+		/// Source for tokens asking the background thread to terminate
+		/// </summary>
+		TaskCompletionSource<bool> StoppingTaskSource = new TaskCompletionSource<bool>();
+
+		/// <summary>
+		/// Task source for waiting for a new subconnection to be made
+		/// </summary>
+		TaskCompletionSource<RpcSubConnection> SubConnectionTaskSource = new TaskCompletionSource<RpcSubConnection>();
+
+		/// <summary>
+		/// Background task to maintain the connection
+		/// </summary>
+		Task? BackgroundTask;
+
+		/// <summary>
+		/// Logger instance
+		/// </summary>
+		ILogger Logger;
+
+		/// <summary>
+		/// Constructor
+		/// </summary>
+		/// <param name="CreateGrpcChannel">Factory method for creating new GRPC channels</param>
+		/// <param name="Logger">Logger instance</param>
+		public RpcConnection(Func<GrpcChannel> CreateGrpcChannel, ILogger Logger)
+		{
+			this.CreateGrpcChannel = CreateGrpcChannel;
+			this.Logger = Logger;
+
+			BackgroundTask = Task.Run(() => ExecuteAsync());
+		}
+
+		/// <summary>
+		/// Attempts to get a client reference, returning immediately if there's not one available
+		/// </summary>
+		/// <returns></returns>
+		public IRpcClientRef? TryGetClientRef(RpcContext Context)
+		{
+			Task<RpcSubConnection> SubConnectionTask = SubConnectionTaskSource.Task;
+			if (SubConnectionTask.IsCompleted)
+			{
+				RpcSubConnection DefaultClient = SubConnectionTask.Result;
+				if(DefaultClient.TryAddRef())
+				{
+					return new RpcClientRef(DefaultClient, Context);
+				}
+			}
+			return null;
+		}
+
+		/// <summary>
+		/// Obtains a new client reference object
+		/// </summary>
+		/// <param name="Context">Context for the call, for debugging</param>
+		/// <param name="CancellationToken">Cancellation token for this request</param>
+		/// <returns>New client reference</returns>
+		public async Task<IRpcClientRef> GetClientRef(RpcContext Context, CancellationToken CancellationToken)
+		{
+			for (; ; )
+			{
+				// Print a message while we're waiting for a client to be available
+				if (!SubConnectionTaskSource.Task.IsCompleted)
+				{
+					Stopwatch Timer = Stopwatch.StartNew();
+					while (!SubConnectionTaskSource.Task.IsCompleted)
+					{
+						Task Delay = Task.Delay(TimeSpan.FromSeconds(30.0), CancellationToken);
+						if (await Task.WhenAny(Delay, SubConnectionTaskSource.Task, StoppingTaskSource.Task) == Delay)
+						{
+							await Delay; // Allow the cancellation token to throw
+							Logger.LogInformation("Thread stalled for {Time:0.0}s while waiting for IRpcClientRef. Source:\n{Context}\nActive references:\n{References}\nStack trace:\n{StackTrace}", Timer.Elapsed.TotalSeconds, Context.Text, String.Join("\n", RpcClientRef.GetActiveRefs()), Environment.StackTrace);
+						}
+						if (StoppingTaskSource.Task.IsCompleted)
+						{
+							throw new OperationCanceledException("RpcConnection is being terminated");
+						}
+					}
+				}
+
+				// Get the new task and increase the refcount
+				RpcSubConnection DefaultClient = await SubConnectionTaskSource.Task;
+				if (DefaultClient.TryAddRef())
+				{
+					return new RpcClientRef(DefaultClient, Context);
+				}
+			}
+		}
+
+		/// <summary>
+		/// Invokes an asynchronous command with a HordeRpcClient instance
+		/// </summary>
+		/// <param name="Func">The function to execute</param>
+		/// <param name="Context">Context for the call, for debugging</param>
+		/// <param name="CancellationToken">Cancellation token for the call</param>
+		/// <returns>Response from the call</returns>
+		public async Task<T> InvokeOnceAsync<T>(Func<HordeRpc.HordeRpcClient, Task<T>> Func, RpcContext Context, CancellationToken CancellationToken)
+		{
+			using (IRpcClientRef ClientRef = await GetClientRef(Context, CancellationToken))
+			{
+				return await Func(ClientRef.Client);
+			}
+		}
+
+		/// <summary>
+		/// Obtains a client reference and executes a unary RPC
+		/// </summary>
+		/// <typeparam name="T">Return type from the call</typeparam>
+		/// <param name="Func">Method to execute with the RPC instance</param>
+		/// <param name="Context">Context for the call, for debugging</param>
+		/// <param name="CancellationToken">Cancellation token for the call</param>
+		/// <returns>Response from the call</returns>
+		public Task<T> InvokeOnceAsync<T>(Func<HordeRpc.HordeRpcClient, AsyncUnaryCall<T>> Func, RpcContext Context, CancellationToken CancellationToken)
+		{
+			return InvokeOnceAsync(async (x) => await Func(x), Context, CancellationToken);
+		}
+
+		/// <summary>
+		/// Invokes an asynchronous command with a HordeRpcClient instance
+		/// </summary>
+		/// <param name="Func">The function to execute</param>
+		/// <param name="Context">Context for the call, for debugging</param>
+		/// <param name="CancellationToken">Cancellation token for the call</param>
+		/// <returns>Response from the call</returns>
+		public async Task<T> InvokeAsync<T>(Func<HordeRpc.HordeRpcClient, Task<T>> Func, RpcContext Context, CancellationToken CancellationToken)
+		{
+			for (int Attempt = 0; ;Attempt++)
+			{
+				// Attempt the RPC call
+				using (IRpcClientRef ClientRef = await GetClientRef(Context, CancellationToken))
+				{
+					try
+					{
+						return await Func(ClientRef.Client);
+					}
+					catch (Exception Ex)
+					{
+						// Otherwise check if we can try again
+						if (Attempt >= RetryTimes.Length)
+						{
+							Logger.LogError("Max number of retry attempts reached");
+							throw;
+						}
+
+						// Check if the user requested cancellation
+						if (CancellationToken.IsCancellationRequested)
+						{
+							Logger.LogInformation("Cancellation of gRPC call requested");
+							throw;
+						}
+
+						// Check we can retry it
+						if (!CanRetry(Ex))
+						{
+							Logger.LogError(Ex, "Exception during RPC: {Message}", Ex.Message);
+							throw;
+						}
+
+						RpcException? RpcEx = Ex as RpcException;
+						if (RpcEx != null && RpcEx.StatusCode == StatusCode.Unavailable)
+						{
+							Logger.LogInformation(Ex, "Failure #{Attempt} during gRPC call (service unavailable). Retrying...", Attempt + 1);
+						}
+						else
+						{
+							Logger.LogError(Ex, "Failure #{Attempt} during gRPC call.", Attempt + 1);
+						}
+					}
+				}
+
+				// Wait before retrying
+				await Task.Delay(RetryTimes[Attempt]);
+			}
+		}
+
+		/// <summary>
+		/// Determines if the given exception should be retried
+		/// </summary>
+		/// <param name="Ex"></param>
+		/// <returns></returns>
+		public static bool CanRetry(Exception Ex)
+		{
+			RpcException? RpcEx = Ex as RpcException;
+			if(RpcEx != null)
+			{
+				return RpcEx.StatusCode == StatusCode.Cancelled || RpcEx.StatusCode == StatusCode.Unavailable || RpcEx.StatusCode == StatusCode.DataLoss /* Interrupted streaming call */;
+			}
+			else
+			{
+				return false;
+			}
+		}
+
+		/// <summary>
+		/// Obtains a client reference and executes a unary RPC
+		/// </summary>
+		/// <typeparam name="T">Return type from the call</typeparam>
+		/// <param name="Func">Method to execute with the RPC instance</param>
+		/// <param name="Context">Context for the call, for debugging</param>
+		/// <param name="CancellationToken">Cancellation token for the call</param>
+		/// <returns>Response from the call</returns>
+		public Task<T> InvokeAsync<T>(Func<HordeRpc.HordeRpcClient, AsyncUnaryCall<T>> Func, RpcContext Context, CancellationToken CancellationToken)
+		{
+			return InvokeAsync(async (x) => await Func(x), Context, CancellationToken);
+		}
+
+		/// <summary>
+		/// Background task to monitor the server state
+		/// </summary>
+		/// <returns>Async task</returns>
+		async Task ExecuteAsync()
+		{
+			// Counter for the connection number. Used to track log messages through async threads.
+			int ConnectionId = 0;
+
+			// This task source is marked as completed once a RPC connection receives a notification from the server that it's shutting down.
+			// Another connection will immediately be allocated for other connection requests to use, and the load balancer should route us to
+			// the new server.
+			TaskCompletionSource<bool>? ReconnectTaskSource = null;
+
+			// List of connection handling async tasks we're tracking
+			List<Task> Tasks = new List<Task>();
+			while (Tasks.Count > 0 || !StoppingTaskSource.Task.IsCompleted)
+			{
+				// Get the task indicating we want a new connection (or create a new one)
+				Task? ReconnectTask = null;
+				if (!StoppingTaskSource.Task.IsCompleted)
+				{
+					if (ReconnectTaskSource == null || ReconnectTaskSource.Task.IsCompleted)
+					{
+						// Need to avoid lambda capture of the reconnect task source
+						int NewConnectionId = ++ConnectionId;
+						TaskCompletionSource<bool> NewReconnectTaskSource = new TaskCompletionSource<bool>();
+						Tasks.Add(Task.Run(() => HandleConnectionAsync(NewConnectionId, NewReconnectTaskSource)));
+						ReconnectTaskSource = NewReconnectTaskSource;
+					}
+					ReconnectTask = ReconnectTaskSource.Task;
+				}
+
+				// Build a list of tasks to wait for
+				List<Task> WaitTasks = new List<Task>(Tasks);
+				if (ReconnectTask != null)
+				{
+					WaitTasks.Add(ReconnectTask);
+				}
+				await Task.WhenAny(WaitTasks);
+
+				// Remove any complete tasks, logging any exceptions
+				for (int Idx = 0; Idx < Tasks.Count; Idx++)
+				{
+					Task Task = Tasks[Idx];
+					if(Task.IsCompleted)
+					{
+						await Task;
+						Tasks.RemoveAt(Idx--);
+					}
+				}
+			}
+		}
+
+		/// <summary>
+		/// Background task to monitor the server state
+		/// </summary>
+		/// <param name="ConnectionId">The connection id</param>
+		/// <param name="ReconnectTaskSource">Task source for reconnecting</param>
+		/// <returns>Async task</returns>
+		async Task HandleConnectionAsync(int ConnectionId, TaskCompletionSource<bool> ReconnectTaskSource)
+		{
+			using Scope Scope = Tracer.Instance.StartActive("RpcConnection {ConnectionId}");
+
+			Logger.LogInformation("Connection {ConnectionId}: Creating connection", ConnectionId);
+			try
+			{
+				await HandleConnectionInternalAsync(ConnectionId, ReconnectTaskSource);
+			}
+			catch (RpcException Ex) when (Ex.StatusCode == StatusCode.Unavailable)
+			{
+				Logger.LogInformation("Connection {ConnectionId}: Rpc service is unavailable. Reconnecting.", ConnectionId);
+			}
+			catch (Exception Ex)
+			{
+				Logger.LogError(Ex, "Connection {ConnectionId}: Exception handling connection: {Message}", ConnectionId, Ex.Message);
+			}
+			finally
+			{
+				TriggerReconnect(ReconnectTaskSource);
+			}
+		}
+
+		/// <summary>
+		/// Background task to monitor the server state
+		/// </summary>
+		/// <returns>Async task</returns>
+		async Task HandleConnectionInternalAsync(int ConnectionId, TaskCompletionSource<bool> ReconnectTaskSource)
+		{
+			using (GrpcChannel Channel = CreateGrpcChannel())
+			{
+				HordeRpc.HordeRpcClient Client = new HordeRpc.HordeRpcClient(Channel);
+
+				RpcSubConnection? SubConnection = null;
+				try
+				{
+					while (!ReconnectTaskSource.Task.IsCompleted && !StoppingTaskSource.Task.IsCompleted)
+					{
+						// Client HAS TO BE the one to end the call, by closing the stream. The server must keep receiving requests from us until it does.
+						using (AsyncDuplexStreamingCall<QueryServerStateRequest, QueryServerStateResponse> Call = Client.QueryServerStateV2())
+						{
+							// Send the name of this agent
+							QueryServerStateRequest Request = new QueryServerStateRequest();
+							Request.Name = Dns.GetHostName();
+							await Call.RequestStream.WriteAsync(Request);
+
+							// Read the server info
+							Task<bool> MoveNextTask = Call.ResponseStream.MoveNext();
+							if (!await MoveNextTask)
+							{
+								Logger.LogError("Connection {ConnectionId}: No response from server", ConnectionId);
+								break;
+							}
+
+							// Read the server response
+							QueryServerStateResponse Response = Call.ResponseStream.Current;
+							if (!Response.Stopping)
+							{
+								// The first time we connect, log the server name and create the subconnection
+								if (SubConnection == null)
+								{
+									Logger.LogInformation("Connection {ConnectionId}: Connected to rpc server {ServerName}", ConnectionId, Response.Name);
+									SubConnection = new RpcSubConnection(ConnectionId, Response.Name, Channel, Client, Logger);
+									SubConnectionTaskSource.SetResult(SubConnection);
+								}
+
+								// Wait for the StoppingTask token to be set, or the server to inform that it's shutting down
+								MoveNextTask = Call.ResponseStream.MoveNext();
+								await Task.WhenAny(StoppingTaskSource.Task, Task.Delay(TimeSpan.FromSeconds(45.0)), MoveNextTask);
+
+								// Update the response
+								if (MoveNextTask.IsCompleted && MoveNextTask.Result)
+								{
+									Response = Call.ResponseStream.Current;
+								}
+							}
+
+							// If the server is stopping, start reconnecting
+							if (Response.Stopping)
+							{
+								Logger.LogInformation("Connection {ConnectionId}: Server is stopping. Triggering reconnect.", ConnectionId);
+								TriggerReconnect(ReconnectTaskSource);
+							}
+
+							// Close the request stream
+							await Call.RequestStream.CompleteAsync();
+
+							// Wait for the server to finish posting responses
+							while (await MoveNextTask)
+							{
+								Response = Call.ResponseStream.Current;
+								MoveNextTask = Call.ResponseStream.MoveNext();
+							}
+						}
+					}
+					Logger.LogInformation("Connection {ConnectionId}: Closing connection to rpc server", ConnectionId);
+				}
+				finally
+				{
+					if (SubConnection != null)
+					{
+						await SubConnection.DisposeAsync();
+					}
+				}
+
+				// Wait a few seconds before finishing
+				await Task.Delay(TimeSpan.FromSeconds(5.0));
+			}
+		}
+
+		/// <summary>
+		/// Clears the current subconnection and triggers a reconnect
+		/// </summary>
+		/// <param name="ReconnectTaskSource"></param>
+		void TriggerReconnect(TaskCompletionSource<bool> ReconnectTaskSource)
+		{
+			if (!ReconnectTaskSource.Task.IsCompleted)
+			{
+				// First reset the task source, so that nothing new starts to use the current connection
+				if (SubConnectionTaskSource.Task.IsCompleted)
+				{
+					SubConnectionTaskSource = new TaskCompletionSource<RpcSubConnection>();
+				}
+
+				// Now trigger another connection
+				ReconnectTaskSource.SetResult(true);
+			}
+		}
+
+		/// <summary>
+		/// Dispose of this connection
+		/// </summary>
+		/// <returns>Async task</returns>
+		public async ValueTask DisposeAsync()
+		{
+			if (BackgroundTask != null)
+			{
+				StoppingTaskSource.TrySetResult(true);
+
+				await BackgroundTask;
+				BackgroundTask = null!;
+			}
+		}
+
+		public static IRpcConnection Create(Func<GrpcChannel> CreateGrpcChannel, ILogger Logger)
+		{
+			return new RpcConnection(CreateGrpcChannel, Logger);
+		}
+	}
+}
diff --git a/Engine/Source/Programs/Horde/HordeAgent/Utility/Shutdown.cs b/Engine/Source/Programs/Horde/HordeAgent/Utility/Shutdown.cs
new file mode 100644
--- /dev/null
+++ b/Engine/Source/Programs/Horde/HordeAgent/Utility/Shutdown.cs
@@ -0,0 +1,125 @@
+using Microsoft.Extensions.Logging;
+using Microsoft.Win32.SafeHandles;
+using System;
+using System.Collections.Generic;
+using System.Diagnostics;
+using System.Net.NetworkInformation;
+using System.Runtime.InteropServices;
+using System.Text;
+
+namespace HordeAgent.Utility
+{
+	static class Shutdown
+	{
+		public const UInt32 TOKEN_QUERY = 0x0008;
+		public const UInt32 TOKEN_ADJUST_PRIVILEGES = 0x0020;
+
+		[DllImport("advapi32.dll", SetLastError = true)]
+		[return: MarshalAs(UnmanagedType.Bool)]
+		static extern bool OpenProcessToken(IntPtr ProcessHandle, UInt32 DesiredAccess, out IntPtr TokenHandle);
+
+		[StructLayout(LayoutKind.Sequential)]
+		public struct LUID
+		{
+			public UInt32 LowPart;
+			public Int32 HighPart;
+		}
+
+		[DllImport("advapi32.dll")]
+		static extern bool LookupPrivilegeValue(string? lpSystemName, string lpName, ref LUID lpLuid);
+
+		[StructLayout(LayoutKind.Sequential, Pack = 4)]
+		public struct LUID_AND_ATTRIBUTES
+		{
+			public LUID Luid;
+			public UInt32 Attributes;
+		}
+
+		struct TOKEN_PRIVILEGES
+		{
+			public int PrivilegeCount;
+			[MarshalAs(UnmanagedType.ByValArray, SizeConst = 1)]
+			public LUID_AND_ATTRIBUTES[] Privileges;
+		}
+
+		// Use this signature if you do not want the previous state
+		[DllImport("advapi32.dll", SetLastError = true)]
+		[return: MarshalAs(UnmanagedType.Bool)]
+		static extern bool AdjustTokenPrivileges(IntPtr TokenHandle,
+		   [MarshalAs(UnmanagedType.Bool)]bool DisableAllPrivileges,
+		   ref TOKEN_PRIVILEGES NewState,
+		   UInt32 Zero,
+		   IntPtr Null1,
+		   IntPtr Null2);
+
+		const uint SHTDN_REASON_MAJOR_APPLICATION = 0x00040000;
+		const uint SHTDN_REASON_MINOR_MAINTENANCE = 0x00000001;
+
+		[DllImport("advapi32.dll", CharSet = CharSet.Auto, SetLastError = true)]
+		static extern bool InitiateSystemShutdownEx(string? lpMachineName, string? lpMessage, uint dwTimeout, bool bForceAppsClosed, bool bRebootAfterShutdown, uint dwReason);
+
+		const string SE_SHUTDOWN_NAME = "SeShutdownPrivilege";
+
+		const int SE_PRIVILEGE_ENABLED = 0x00000002;
+
+		/// <summary>
+		/// Initiate a shutdown operation
+		/// </summary>
+		/// <param name="bRestartAfterShutdown">Whether to restart after the shutdown</param>
+		/// <param name="Logger">Logger for the operation</param>
+		/// <returns></returns>
+		public static bool InitiateShutdown(bool bRestartAfterShutdown, ILogger Logger)
+		{
+			if (RuntimeInformation.IsOSPlatform(OSPlatform.Windows))
+			{
+				if (bRestartAfterShutdown)
+				{
+					Logger.LogInformation("Triggering restart");
+				}
+				else
+				{
+					Logger.LogInformation("Triggering shutdown");
+				}
+
+				IntPtr TokenHandle;
+				if(!OpenProcessToken(Process.GetCurrentProcess().Handle, TOKEN_ADJUST_PRIVILEGES | TOKEN_QUERY, out TokenHandle))
+				{
+					Logger.LogError("OpenProcessToken() failed (code 0x{Code:x8})", Marshal.GetLastWin32Error());
+					return false;
+				}
+
+				// Get the LUID for the shutdown privilege. 
+				LUID Luid = new LUID();
+				if (!LookupPrivilegeValue(null, SE_SHUTDOWN_NAME, ref Luid))
+				{
+					Logger.LogError("LookupPrivilegeValue() failed (code 0x{Code:x8})", Marshal.GetLastWin32Error());
+					return false;
+				}
+
+				TOKEN_PRIVILEGES Privileges = new TOKEN_PRIVILEGES();
+				Privileges.PrivilegeCount = 1;
+				Privileges.Privileges = new LUID_AND_ATTRIBUTES[1];
+				Privileges.Privileges[0].Attributes = SE_PRIVILEGE_ENABLED;
+				Privileges.Privileges[0].Luid = Luid;
+
+				if (!AdjustTokenPrivileges(TokenHandle, false, ref Privileges, 0, IntPtr.Zero, IntPtr.Zero))
+				{
+					Logger.LogError("AdjustTokenPrivileges() failed (code 0x{Code:x8})", Marshal.GetLastWin32Error());
+					return false;
+				}
+
+				if (!InitiateSystemShutdownEx(null, "HordeAgent has initiated shutdown", 10, true, bRestartAfterShutdown, SHTDN_REASON_MAJOR_APPLICATION | SHTDN_REASON_MINOR_MAINTENANCE))
+				{
+					Logger.LogError("Shutdown failed (0x{Code:x8})", Marshal.GetLastWin32Error());
+					return false;
+				}
+				return true;
+			}
+			else
+			{
+				Logger.LogError("Shutdown is not implemented on this platform");
+				return false;
+			}
+		}
+	}
+}
diff --git a/Engine/Source/Programs/Horde/HordeAgent/Utility/StorageExtensions.cs b/Engine/Source/Programs/Horde/HordeAgent/Utility/StorageExtensions.cs
new file mode 100644
--- /dev/null
+++ b/Engine/Source/Programs/Horde/HordeAgent/Utility/StorageExtensions.cs
@@ -0,0 +1,78 @@
+using System;
+using System.Collections.Generic;
+using System.Security.Cryptography;
+using System.Text;
+using System.Threading.Tasks;
+using Build.Bazel.Remote.Execution.V2;
+using EpicGames.Core;
+using Google.Protobuf;
+using ContentAddressableStorageClient = Build.Bazel.Remote.Execution.V2.ContentAddressableStorage.ContentAddressableStorageClient;
+
+namespace HordeAgent.Utility
+{
+	static class StorageExtensions
+	{
+		public static async Task<byte[]> GetBulkDataAsync(this ContentAddressableStorageClient Storage, Digest Digest)
+		{
+			BatchReadBlobsRequest Request = new BatchReadBlobsRequest();
+			Request.Digests.Add(Digest);
+
+			BatchReadBlobsResponse Response = await Storage.BatchReadBlobsAsync(Request);
+			return Response.Responses[0].Data.ToByteArray();
+		}
+
+		public static async Task<Digest> PutBulkDataAsync(this ContentAddressableStorageClient Storage, byte[] Data)
+		{
+			BatchUpdateBlobsRequest.Types.Request BlobRequest = new BatchUpdateBlobsRequest.Types.Request();
+			BlobRequest.Data = ByteString.CopyFrom(Data);
+			BlobRequest.Digest = GetDigest(Data);
+
+			BatchUpdateBlobsRequest Request = new BatchUpdateBlobsRequest();
+			Request.Requests.Add(BlobRequest);
+
+			BatchUpdateBlobsResponse Response = await Storage.BatchUpdateBlobsAsync(Request);
+			//?
+			return BlobRequest.Digest;
+		}
+
+		public static async Task<T> GetProtoMessageAsync<T>(this ContentAddressableStorageClient Storage, Digest Digest) where T : class, IMessage<T>, new()
+		{
+			BatchReadBlobsRequest Request = new BatchReadBlobsRequest();
+			Request.Digests.Add(Digest);
+
+			BatchReadBlobsResponse Response = await Storage.BatchReadBlobsAsync(Request);
+
+			T Item = new T();
+			Item.MergeFrom(Response.Responses[0].Data);
+			return Item;
+		}
+
+		public static async Task<Digest> PutProtoMessageAsync<T>(this ContentAddressableStorageClient Storage, IMessage<T> Message) where T : class, IMessage<T>
+		{
+			BatchUpdateBlobsRequest.Types.Request BlobRequest = new BatchUpdateBlobsRequest.Types.Request();
+			BlobRequest.Data = Message.ToByteString();
+			BlobRequest.Digest = GetDigest(BlobRequest.Data.ToByteArray());
+
+			BatchUpdateBlobsRequest Request = new BatchUpdateBlobsRequest();
+			Request.Requests.Add(BlobRequest);
+
+			BatchUpdateBlobsResponse Response = await Storage.BatchUpdateBlobsAsync(Request);
+			//??
+			return BlobRequest.Digest;
+		}
+
+		public static Digest GetDigest(ByteString Data)
+		{
+			return GetDigest(Data.ToByteArray());
+		}
+
+		public static Digest GetDigest(byte[] Data)
+		{
+			using (SHA256 Algorithm = SHA256.Create())
+			{
+				byte[] Hash = Algorithm.ComputeHash(Data);
+				return new Digest { Hash = StringUtils.FormatHexString(Hash), SizeBytes = Data.Length };
+			}
+		}
+	}
+}
diff --git a/Engine/Source/Programs/Horde/HordeAgent/Utility/WindowsServiceManager.cs b/Engine/Source/Programs/Horde/HordeAgent/Utility/WindowsServiceManager.cs
new file mode 100644
--- /dev/null
+++ b/Engine/Source/Programs/Horde/HordeAgent/Utility/WindowsServiceManager.cs
@@ -0,0 +1,361 @@
+using Microsoft.Win32.SafeHandles;
+using Serilog;
+using System;
+using System.Collections.Generic;
+using System.ComponentModel;
+using System.Diagnostics;
+using System.Reflection;
+using System.Runtime.InteropServices;
+using System.Text;
+using System.Threading;
+
+namespace HordeAgent.Utility
+{
+	/// <summary>
+	/// Native methods for manipulating services on Windows
+	/// </summary>
+	static partial class Native
+	{
+#pragma warning disable CS0649
+		public class ServiceHandle : SafeHandleZeroOrMinusOneIsInvalid
+		{
+			public ServiceHandle()
+				: base(true)
+			{
+			}
+
+			public ServiceHandle(IntPtr Handle)
+				: base(true)
+			{
+				SetHandle(Handle);
+			}
+
+			protected override bool ReleaseHandle()
+			{
+				return Native.CloseServiceHandle(handle);
+			}
+		}
+
+		public const uint SC_MANAGER_ALL_ACCESS = 0xF003F;
+
+		[DllImport("advapi32.dll", EntryPoint = "OpenSCManagerW", ExactSpelling = true, CharSet = CharSet.Unicode, SetLastError = true)]
+		public static extern ServiceHandle OpenSCManager(string? machineName, string? databaseName, uint dwAccess);
+
+		public const uint SERVICE_ALL_ACCESS = 0xf01ff;
+
+		public const uint SERVICE_WIN32_OWN_PROCESS = 0x00000010;
+
+		public const uint SERVICE_AUTO_START = 0x00000002;
+		public const uint SERVICE_DEMAND_START = 0x00000003;
+
+		public const uint SERVICE_ERROR_NORMAL = 0x00000001;
+
+		[DllImport("advapi32.dll", SetLastError = true, CharSet = CharSet.Auto)]
+		public static extern ServiceHandle CreateService(ServiceHandle hSCManager, string lpServiceName, string lpDisplayName, uint dwDesiredAccess, uint dwServiceType, uint dwStartType, uint dwErrorControl, string lpBinaryPathName, string? lpLoadOrderGroup, string? lpdwTagId, string? lpDependencies, string? lpServiceStartName, string? lpPassword);
+
+		[DllImport("advapi32.dll", SetLastError = true, CharSet = CharSet.Auto)]
+		public static extern ServiceHandle OpenService(ServiceHandle hSCManager, string lpServiceName, uint dwDesiredAccess);
+
+		[DllImport("advapi32", SetLastError = true)]
+		[return: MarshalAs(UnmanagedType.Bool)]
+		public static extern bool StartService(ServiceHandle hService, int dwNumServiceArgs, string[]? lpServiceArgVectors);
+
+		public const int ERROR_ACCESS_DENIED = 5;
+		public const int ERROR_SERVICE_DOES_NOT_EXIST = 1060;
+		public const int ERROR_SERVICE_NOT_ACTIVE = 1062;
+
+		[DllImport("advapi32.dll", SetLastError = true)]
+		[return: MarshalAs(UnmanagedType.Bool)]
+		public static extern bool CloseServiceHandle(IntPtr hSCObject);
+
+		[DllImport("advapi32.dll", SetLastError = true)]
+		[return: MarshalAs(UnmanagedType.Bool)]
+		public static extern bool DeleteService(ServiceHandle hService);
+
+		public const int SC_STATUS_PROCESS_INFO = 0;
+
+		public struct SERVICE_STATUS_PROCESS
+		{
+			public uint dwServiceType;
+			public uint dwCurrentState;
+			public uint dwControlsAccepted;
+			public uint dwWin32ExitCode;
+			public uint dwServiceSpecificExitCode;
+			public uint dwCheckPoint;
+			public uint dwWaitHint;
+			public uint dwProcessId;
+			public uint dwServiceFlags;
+		}
+
+		[DllImport("advapi32.dll", SetLastError = true)]
+		[return: MarshalAs(UnmanagedType.Bool)]
+		public static extern bool QueryServiceStatusEx(ServiceHandle hService, int InfoLevel, ref SERVICE_STATUS_PROCESS pBuffer, int cbBufSize, out int pcbBytesNeeded);
+
+		public struct SERVICE_STATUS
+		{
+			public uint dwServiceType;
+			public uint dwCurrentState;
+			public uint dwControlsAccepted;
+			public uint dwWin32ExitCode;
+			public uint dwServiceSpecificExitCode;
+			public uint dwCheckPoint;
+			public uint dwWaitHint;
+		}
+
+		public const int SERVICE_CONTROL_STOP = 1;
+
+		[DllImport("advapi32.dll", SetLastError = true)]
+		[return: MarshalAs(UnmanagedType.Bool)]
+		public static extern bool ControlService(ServiceHandle hService, int dwControl, ref SERVICE_STATUS lpServiceStatus);
+
+		public const int SERVICE_CONFIG_DESCRIPTION = 1;
+
+		[StructLayout(LayoutKind.Sequential)]
+		public struct SERVICE_DESCRIPTION
+		{
+			[MarshalAs(UnmanagedType.LPWStr)]
+			public String lpDescription;
+		}
+
+		[DllImport("advapi32.dll", SetLastError = true, CharSet = CharSet.Auto)]
+		[return: MarshalAs(UnmanagedType.Bool)]
+		public static extern bool ChangeServiceConfig2(ServiceHandle hService, int dwInfoLevel, ref SERVICE_DESCRIPTION lpInfo);
+#pragma warning restore CS0649
+	}
+
+	enum WindowsServiceStatus : uint
+	{
+		/// <summary>
+		/// The service has stopped.
+		/// </summary>
+		Stopped = 1,
+
+		/// <summary>
+		/// The service is starting.
+		/// </summary>
+		Starting = 2,
+
+		/// <summary>
+		/// The service is stopping.
+		/// </summary>
+		Stopping = 3,
+
+		/// <summary>
+		/// The service is running
+		/// </summary>
+		Running = 4,
+
+		/// <summary>
+		/// The service is about to continue.
+		/// </summary>
+		Continuing = 5,
+
+		/// <summary>
+		/// The service is pausing
+		/// </summary>
+		Pausing = 6,
+
+		/// <summary>
+		/// The service is paused.
+		/// </summary>
+		Paused = 7,
+	}
+
+	/// <summary>
+	/// Wrapper around 
+	/// </summary>
+	class WindowsService : IDisposable
+	{
+		/// <summary>
+		/// Handle to the service
+		/// </summary>
+		Native.ServiceHandle Handle;
+
+		/// <summary>
+		/// Constructor
+		/// </summary>
+		/// <param name="Handle">Handle to the service</param>
+		public WindowsService(Native.ServiceHandle Handle)
+		{
+			this.Handle = Handle;
+		}
+
+		/// <summary>
+		/// Determines if the service is valid
+		/// </summary>
+		public bool IsValid
+		{
+			get { return !Handle.IsInvalid; }
+		}
+
+		/// <summary>
+		/// Gets the current service status
+		/// </summary>
+		/// <returns>Status code</returns>
+		public WindowsServiceStatus GetStatus()
+		{
+			Native.SERVICE_STATUS_PROCESS Status = new Native.SERVICE_STATUS_PROCESS();
+			if (!Native.QueryServiceStatusEx(Handle, Native.SC_STATUS_PROCESS_INFO, ref Status, Marshal.SizeOf(Status), out int BytesNeeded))
+			{
+				throw new Win32Exception(String.Format("Unable to query process status (0x{0:X8)}", Marshal.GetLastWin32Error()));
+			}
+			return (WindowsServiceStatus)Status.dwCurrentState;
+		}
+
+		/// <summary>
+		/// Waits for the status to change
+		/// </summary>
+		/// <param name="TransitionStatus">Expected status while transitioning</param>
+		/// <param name="MaxWaitTime">Maximum time to wait</param>
+		public WindowsServiceStatus WaitForStatusChange(WindowsServiceStatus TransitionStatus, TimeSpan MaxWaitTime)
+		{
+			Stopwatch Timer = Stopwatch.StartNew();
+			for(; ;)
+			{
+				WindowsServiceStatus Status = GetStatus();
+				if (Status != TransitionStatus || Timer.Elapsed > MaxWaitTime)
+				{
+					return Status;
+				}
+				else
+				{
+					Thread.Sleep(1000);
+				}
+			}
+		}
+
+		/// <summary>
+		/// Sets the service description
+		/// </summary>
+		/// <param name="Description">Description for the service</param>
+		public void SetDescription(string Description)
+		{
+			Native.SERVICE_DESCRIPTION DescriptionData = new Native.SERVICE_DESCRIPTION();
+			DescriptionData.lpDescription = Description;
+			Native.ChangeServiceConfig2(Handle, Native.SERVICE_CONFIG_DESCRIPTION, ref DescriptionData);
+		}
+
+		/// <summary>
+		/// Starts the service
+		/// </summary>
+		public void Start()
+		{
+			if (!Native.StartService(Handle, 0, null))
+			{
+				throw new Win32Exception(String.Format("Unable to start service (error 0x{0:X8})", Marshal.GetLastWin32Error()));
+			}
+		}
+
+		/// <summary>
+		/// Stops the service
+		/// </summary>
+		public void Stop()
+		{
+			Native.SERVICE_STATUS Status = new Native.SERVICE_STATUS();
+			if (!Native.ControlService(Handle, Native.SERVICE_CONTROL_STOP, ref Status))
+			{
+				int Error = Marshal.GetLastWin32Error();
+				if (Error != Native.ERROR_SERVICE_NOT_ACTIVE)
+				{
+					throw new Win32Exception(String.Format("Unable to stop service (error 0x{0:X8})", Error));
+				}
+			}
+		}
+
+		/// <summary>
+		/// Deletes the service
+		/// </summary>
+		public void Delete()
+		{
+			if (!Native.DeleteService(Handle))
+			{
+				throw new Win32Exception(String.Format("Unable to delete service (error 0x{0:X8})", Marshal.GetLastWin32Error()));
+			}
+		}
+
+		/// <summary>
+		/// Dispose of the service handle
+		/// </summary>
+		public void Dispose()
+		{
+			Handle.Close();
+		}
+	}
+
+	/// <summary>
+	/// Helper functionality for manipulating Windows services
+	/// </summary>
+	class WindowsServiceManager : IDisposable
+	{
+		/// <summary>
+		/// Native handle to the service manager
+		/// </summary>
+		Native.ServiceHandle ServiceManagerHandle;
+
+		/// <summary>
+		/// Constructor. Opens a handle to the service manager.
+		/// </summary>
+		public WindowsServiceManager()
+		{
+			ServiceManagerHandle = Native.OpenSCManager(null, null, Native.SC_MANAGER_ALL_ACCESS);
+			if (ServiceManagerHandle.IsInvalid)
+			{
+				int ErrorCode = Marshal.GetLastWin32Error();
+				if (ErrorCode == Native.ERROR_ACCESS_DENIED)
+				{
+					throw new Win32Exception("Unable to open service manager (access denied). Check you're running as administrator.");
+				}
+				else
+				{
+					throw new Win32Exception(String.Format("Unable to open service manager (0x{0:X8)).", ErrorCode));
+				}
+			}
+		}
+
+		/// <summary>
+		/// Dispose of this object
+		/// </summary>
+		public void Dispose()
+		{
+			ServiceManagerHandle.Close();
+		}
+
+		/// <summary>
+		/// Opens a service with the given name
+		/// </summary>
+		/// <param name="ServiceName">Name of the service</param>
+		/// <returns>New service wrapper</returns>
+		public WindowsService Open(string ServiceName)
+		{
+			Native.ServiceHandle ServiceHandle = Native.OpenService(ServiceManagerHandle, ServiceName, Native.SERVICE_ALL_ACCESS);
+			if (ServiceHandle.IsInvalid)
+			{
+				int ErrorCode = Marshal.GetLastWin32Error();
+				if (ErrorCode != Native.ERROR_SERVICE_DOES_NOT_EXIST)
+				{
+					throw new Win32Exception("Unable to open handle to service");
+				}
+			}
+			return new WindowsService(ServiceHandle);
+		}
+
+		/// <summary>
+		/// Creates a service with the given settings
+		/// </summary>
+		/// <param name="Name">Name of the service</param>
+		/// <param name="DisplayName">Display name</param>
+		/// <param name="CommandLine">Command line to use when starting the service</param>
+		/// <param name="UserName">Username to run this service as</param>
+		/// <param name="Password">Password for the account the service is to run under</param>
+		/// <returns>New service instance</returns>
+		public WindowsService Create(string Name, string DisplayName, string CommandLine, string? UserName, string? Password)
+		{
+			Native.ServiceHandle NewServiceHandle = Native.CreateService(ServiceManagerHandle, Name, DisplayName, Native.SERVICE_ALL_ACCESS, Native.SERVICE_WIN32_OWN_PROCESS, Native.SERVICE_AUTO_START, Native.SERVICE_ERROR_NORMAL, CommandLine, null, null, null, UserName, Password);
+			if (NewServiceHandle.IsInvalid)
+			{
+				throw new Win32Exception(String.Format("Unable to create service (0x{0:X8})", Marshal.GetLastWin32Error()));
+			}
+			return new WindowsService(NewServiceHandle);
+		}
+	}
+}
diff --git a/Engine/Source/Programs/Horde/HordeAgent/Utility/WorkspaceInfo.cs b/Engine/Source/Programs/Horde/HordeAgent/Utility/WorkspaceInfo.cs
new file mode 100644
--- /dev/null
+++ b/Engine/Source/Programs/Horde/HordeAgent/Utility/WorkspaceInfo.cs
@@ -0,0 +1,394 @@
+using EpicGames.Perforce.Managed;
+using HordeCommon.Rpc.Messages;
+using Microsoft.Extensions.Logging;
+using System;
+using System.Collections.Generic;
+using System.Diagnostics;
+using System.IO;
+using System.Runtime.InteropServices;
+using System.Text;
+using System.Text.RegularExpressions;
+using System.Threading;
+using System.Threading.Tasks;
+using EpicGames.Core;
+using EpicGames.Perforce;
+using System.Linq;
+
+namespace HordeAgent.Utility
+{
+	/// <summary>
+	/// Stores information about a managed Perforce workspace
+	/// </summary>
+	class WorkspaceInfo
+	{
+		/// <summary>
+		/// The perforce connection
+		/// </summary>
+		public PerforceClientConnection PerforceClient
+		{
+			get;
+		}
+
+		/// <summary>
+		/// The Perforce server and port. This is checked to not be null in the constructor.
+		/// </summary>
+		public string ServerAndPort
+		{
+			get { return PerforceClient.ServerAndPort!; }
+		}
+
+		/// <summary>
+		/// The Perforce client name. This is checked to not be null in the constructor.
+		/// </summary>
+		public string ClientName
+		{
+			get { return PerforceClient.ClientName!; }
+		}
+
+		/// <summary>
+		/// The Perforce user name. This is checked to not be null in the constructor.
+		/// </summary>
+		public string UserName
+		{
+			get { return PerforceClient.UserName!; }
+		}
+
+		/// <summary>
+		/// The hostname
+		/// </summary>
+		public string HostName
+		{
+			get;
+		}
+
+		/// <summary>
+		/// The current stream
+		/// </summary>
+		public string StreamName
+		{
+			get;
+		}
+
+		/// <summary>
+		/// The directory containing metadata for this workspace
+		/// </summary>
+		public DirectoryReference MetadataDir
+		{
+			get;
+		}
+
+		/// <summary>
+		/// The directory containing workspace data
+		/// </summary>
+		public DirectoryReference WorkspaceDir
+		{
+			get;
+		}
+
+		/// <summary>
+		/// The view for files to sync
+		/// </summary>
+		public IReadOnlyList<string> View
+		{
+			get;
+		}			
+
+		/// <summary>
+		/// Whether untracked files should be removed from this workspace
+		/// </summary>
+		public bool bRemoveUntrackedFiles
+		{
+			get;
+		}
+
+		/// <summary>
+		/// The managed repository instance
+		/// </summary>
+		public ManagedWorkspace Repository
+		{
+			get;
+		}
+
+		/// <summary>
+		/// Constructor
+		/// </summary>
+		/// <param name="Perforce">The perforce connection</param>
+		/// <param name="HostName">Name of this host</param>
+		/// <param name="StreamName">Name of the stream to sync</param>
+		/// <param name="MetadataDir">Path to the metadata directory</param>
+		/// <param name="WorkspaceDir">Path to the workspace directory</param>
+		/// <param name="View">View for files to be synced</param>
+		/// <param name="bRemoveUntrackedFiles">Whether to remove untracked files when syncing</param>
+		/// <param name="Repository">The repository instance</param>
+		public WorkspaceInfo(PerforceClientConnection Perforce, string HostName, string StreamName, DirectoryReference MetadataDir, DirectoryReference WorkspaceDir, IList<string>? View, bool bRemoveUntrackedFiles, ManagedWorkspace Repository)
+		{
+			this.PerforceClient = Perforce;
+
+			if (Perforce.ServerAndPort == null)
+			{
+				throw new ArgumentException("Perforce connection does not have valid server and port");
+			}
+			if (Perforce.UserName == null)
+			{
+				throw new ArgumentException("PerforceConnection does not have valid username");
+			}
+			if (Perforce.ClientName == null)
+			{
+				throw new ArgumentException("PerforceConnection does not have valid client name");
+			}
+
+			this.HostName = HostName;
+			this.StreamName = StreamName;
+			this.MetadataDir = MetadataDir;
+			this.WorkspaceDir = WorkspaceDir;
+			this.View = (View == null) ? new List<string>() : new List<string>(View);
+			this.bRemoveUntrackedFiles = bRemoveUntrackedFiles;
+			this.Repository = Repository;
+		}
+
+		/// <summary>
+		/// Creates a new managed workspace
+		/// </summary>
+		/// <param name="Workspace">The workspace definition</param>
+		/// <param name="RootDir">Root directory for storing the workspace</param>
+		/// <param name="Logger">Logger output</param>
+		/// <param name="CancellationToken">Cancellation token</param>
+		/// <returns>New workspace info</returns>
+		public static async Task<WorkspaceInfo> SetupWorkspaceAsync(AgentWorkspace Workspace, DirectoryReference RootDir, ILogger Logger, CancellationToken CancellationToken)
+		{
+			// Fill in the default credentials iff they are not set
+			string? ServerAndPort = String.IsNullOrEmpty(Workspace.ServerAndPort)? null : Workspace.ServerAndPort;
+			string? UserName = String.IsNullOrEmpty(Workspace.UserName)? null : Workspace.UserName;
+			string? Password = String.IsNullOrEmpty(Workspace.Password)? null : Workspace.Password;
+
+			// Create the connection
+			PerforceConnection Perforce = new PerforceConnection(ServerAndPort, UserName, Logger);
+			if (UserName != null)
+			{
+				if (Password != null)
+				{
+					await Perforce.LoginAsync(Password, CancellationToken);
+				}
+				else
+				{
+					Logger.LogInformation($"Using locally logged in session for {UserName}");
+				}
+			}
+			return await SetupWorkspaceAsync(Perforce, Workspace.Stream, Workspace.Identifier, Workspace.View, !Workspace.Incremental, RootDir, Logger, CancellationToken);
+		}
+
+		/// <summary>
+		/// Creates a new managed workspace
+		/// </summary>
+		/// <param name="Perforce">The perforce connection to use</param>
+		/// <param name="StreamName">The stream being synced</param>
+		/// <param name="Identifier">Identifier to use to share </param>
+		/// <param name="View">View for this workspace</param>
+		/// <param name="bRemoveUntrackedFiles">Whether untracked files should be removed when cleaning this workspace</param>
+		/// <param name="RootDir">Root directory for storing the workspace</param>
+		/// <param name="Logger">Logger output</param>
+		/// <param name="CancellationToken">Cancellation token</param>
+		/// <returns>New workspace info</returns>
+		public static async Task<WorkspaceInfo> SetupWorkspaceAsync(PerforceConnection Perforce, string StreamName, string Identifier, IList<string> View, bool bRemoveUntrackedFiles, DirectoryReference RootDir, ILogger Logger, CancellationToken CancellationToken)
+		{
+			// Get the host name, and fill in any missing metadata about the connection
+			InfoRecord Info = await Perforce.GetInfoAsync(InfoOptions.ShortOutput, CancellationToken);
+
+			if (Perforce.ServerAndPort == null)
+			{
+				Perforce.ServerAndPort = await Perforce.TryGetSettingAsync("p4port", CancellationToken) ?? "perforce:1666";
+				Logger.LogInformation($"Using locally configured Perforce server: '{Perforce.ServerAndPort}'");
+			}
+
+			if (Perforce.UserName == null)
+			{
+				Perforce.UserName = Info.UserName;
+				Logger.LogInformation("Using locally configured and logged in Perforce user: '{UserName}'", Perforce.UserName);
+			}
+
+			string? HostName = Info.ClientHost;
+			if(HostName == null)
+			{
+				throw new Exception("Unable to determine Perforce host name");
+			}
+
+			// replace invalid characters in the workspace identifier with a '+' character
+
+			// append the slot index, if it's non-zero
+			//			my $slot_idx = $optional_arguments->{'slot_idx'} || 0;
+			//			$workspace_identifier .= sprintf("+%02d", $slot_idx) if $slot_idx;
+
+			// if running on an edge server, append the server id to the client name
+			string EdgeSuffix = String.Empty;
+			if(Info.Services != null && Info.ServerID != null)
+			{
+				string[] Services = Info.Services.Split(' ', StringSplitOptions.RemoveEmptyEntries);
+				if (Services.Any(x => x.Equals("edge-server", StringComparison.OrdinalIgnoreCase)))
+				{
+					EdgeSuffix = $"+{Info.ServerID}";
+				}
+			}
+
+			// get all the workspace settings
+			string ClientName = $"Horde+{GetNormalizedHostName(HostName)}+{Identifier}{EdgeSuffix}";
+
+			// Create the client Perforce connection
+			PerforceClientConnection PerforceClient = new PerforceClientConnection(Perforce.ServerAndPort, Perforce.UserName, ClientName, Perforce.Logger);
+
+			// get the workspace names
+			DirectoryReference MetadataDir = DirectoryReference.Combine(RootDir, Identifier);
+			DirectoryReference WorkspaceDir = DirectoryReference.Combine(MetadataDir, "Sync");
+
+			// Create the repository
+			ManagedWorkspace NewRepository = await ManagedWorkspace.LoadOrCreateAsync(HostName, MetadataDir, true, Logger, CancellationToken);
+			if (bRemoveUntrackedFiles)
+			{
+				await NewRepository.DeleteClientAsync(PerforceClient, CancellationToken);
+			}
+			await NewRepository.SetupAsync(PerforceClient, StreamName, CancellationToken);
+
+			// Revert any open files
+			await NewRepository.RevertAsync(PerforceClient, CancellationToken);
+
+			// Create the workspace info
+			return new WorkspaceInfo(PerforceClient, HostName, StreamName, MetadataDir, WorkspaceDir, View, bRemoveUntrackedFiles, NewRepository);
+		}
+
+		/// <summary>
+		/// Gets the latest change in the stream
+		/// </summary>
+		/// <param name="CancellationToken">Cancellation token</param>
+		/// <returns>Latest changelist number</returns>
+		public Task<int> GetLatestChangeAsync(CancellationToken CancellationToken)
+		{
+			return Repository.GetLatestChangeAsync(PerforceClient, StreamName, CancellationToken);
+		}
+
+		/// <summary>
+		/// Revert any open files in the workspace and clean it
+		/// </summary>
+		/// <param name="CancellationToken">Cancellation token</param>
+		/// <returns>Async task</returns>
+		public async Task CleanAsync(CancellationToken CancellationToken)
+		{
+			await Repository.RevertAsync(PerforceClient, CancellationToken);
+			await Repository.CleanAsync(bRemoveUntrackedFiles, CancellationToken);
+		}
+
+		/// <summary>
+		/// Removes the given cache file if it's invalid, and updates the metadata to reflect the version about to be synced
+		/// </summary>
+		/// <param name="CacheFile">Path to the cache file</param>
+		/// <param name="Change">The current change being built</param>
+		/// <param name="PreflightChange">The preflight changelist number</param>
+		/// <returns>Async task</returns>
+		public async Task UpdateLocalCacheMarker(FileReference CacheFile, int Change, int PreflightChange)
+		{
+			// Create the new cache file descriptor
+			string NewDescriptor;
+			if(PreflightChange <= 0)
+			{
+				NewDescriptor = $"CL {Change}";
+			}
+			else
+			{
+				NewDescriptor = $"CL {PreflightChange} with base CL {Change}";
+			}
+
+			// Remove the cache file if the current descriptor doesn't match
+			FileReference DescriptorFile = CacheFile.ChangeExtension(".txt");
+			if (FileReference.Exists(CacheFile))
+			{
+				if (FileReference.Exists(DescriptorFile))
+				{
+					string OldDescriptor = await FileReference.ReadAllTextAsync(DescriptorFile);
+					if (OldDescriptor == NewDescriptor)
+					{
+						return;
+					}
+					else
+					{
+						FileReference.Delete(DescriptorFile);
+					}
+				}
+				FileReference.Delete(CacheFile);
+			}
+
+			// Write the new descriptor file
+			FileReference.WriteAllText(DescriptorFile, NewDescriptor);
+		}
+
+		/// <summary>
+		/// Sync the workspace to a given changelist, and capture it so we can quickly clean in the future
+		/// </summary>
+		/// <param name="Change">The changelist to sync to</param>
+		/// <param name="PreflightChange">Change to preflight, or 0</param>
+		/// <param name="CacheFile">Path to the cache file to use</param>
+		/// <param name="Logger">The logger for output messages</param>
+		/// <param name="CancellationToken">Cancellation token</param>
+		/// <returns>Async task</returns>
+		public async Task SyncAsync(int Change, int PreflightChange, FileReference? CacheFile, ILogger Logger, CancellationToken CancellationToken)
+		{
+			await Repository.SyncAsync(PerforceClient, StreamName, Change, View, bRemoveUntrackedFiles, false, CacheFile, CancellationToken);
+
+			// Purge the cache for incremental workspaces
+			if (!bRemoveUntrackedFiles)
+			{
+				await Repository.PurgeAsync(0, CancellationToken);
+			}
+
+			// if we're running a preflight
+			if (PreflightChange > 0)
+			{
+				await Repository.UnshelveAsync(PerforceClient, PreflightChange, CancellationToken);
+			}
+		}
+
+
+		/// <summary>
+		/// Strip the domain suffix to get the host name
+		/// </summary>
+		/// <param name="HostName">Hostname with optional domain</param>
+		/// <returns>Normalized host name</returns>
+		static string GetNormalizedHostName(string HostName)
+		{
+			return Regex.Replace(HostName, @"\..*$", "").ToUpperInvariant();
+		}
+
+		/// <summary>
+		/// Revert all files in a workspace
+		/// </summary>
+		public static async Task RevertAllChangesAsync(PerforceConnection Perforce, ILogger Logger, CancellationToken CancellationToken)
+		{
+			// Make sure the client name is set
+			if(Perforce.ClientName == null)
+			{
+				throw new ArgumentException("RevertAllChangesAsync() requires PerforceConnection with client");
+			}
+
+			// check if there are any open files, and revert them
+			List<FStatRecord> Files = await Perforce.GetOpenFilesAsync(OpenedOptions.None, -1, null, null, 1, new[] { "//..." }, CancellationToken);
+			if (Files.Count > 0)
+			{
+				await Perforce.RevertAsync(-1, null, RevertOptions.KeepWorkspaceFiles, new[] { "//..." }, CancellationToken);
+			}
+
+			// enumerate all the pending changelists
+			List<ChangesRecord> PendingChanges = await Perforce.GetChangesAsync(ChangesOptions.None, Perforce.ClientName, -1, ChangeStatus.Pending, null, null, CancellationToken);
+			foreach (ChangesRecord PendingChange in PendingChanges)
+			{
+				// delete any shelved files if there are any
+				List<DescribeRecord> Records = await Perforce.DescribeAsync(DescribeOptions.Shelved, -1, new int[] { PendingChange.Number }, CancellationToken);
+				if (Records.Count > 0 && Records[0].Files.Count > 0)
+				{
+					Logger.LogInformation($"Deleting shelved files in changelist {PendingChange.Number}");
+					await Perforce.DeleteChangeAsync(DeleteChangeOptions.None, PendingChange.Number, CancellationToken);
+				}
+
+				// delete the changelist
+				Logger.LogInformation($"Deleting changelist {PendingChange.Number}");
+				await Perforce.DeleteChangeAsync(DeleteChangeOptions.None, PendingChange.Number, CancellationToken);
+			}
+		}
+	}
+}
diff --git a/Engine/Source/Programs/Horde/HordeAgentTests/Helper.cs b/Engine/Source/Programs/Horde/HordeAgentTests/Helper.cs
new file mode 100644
--- /dev/null
+++ b/Engine/Source/Programs/Horde/HordeAgentTests/Helper.cs
@@ -0,0 +1,244 @@
+using System;
+using System.Collections.Generic;
+using System.Threading;
+using System.Threading.Tasks;
+using Google.Protobuf.WellKnownTypes;
+using Grpc.Core;
+using Grpc.Net.Client;
+using HordeAgent.Execution.Interfaces;
+using HordeAgent.Utility;
+using HordeCommon;
+using HordeCommon.Rpc;
+using Microsoft.Extensions.Logging;
+using Microsoft.Extensions.Options;
+
+namespace HordeAgentTests
+{
+	// Stub for fulfilling IOptionsMonitor interface during testing
+	// Copied from HordeServerTests until a good way to share code between these is decided.
+	public class TestOptionsMonitor<T> : IOptionsMonitor<T>, IDisposable
+		where T : class, new()
+	{
+		public TestOptionsMonitor(T CurrentValue)
+		{
+			this.CurrentValue = CurrentValue;
+		}
+
+		public T Get(string Name)
+		{
+			return CurrentValue;
+		}
+
+		public IDisposable OnChange(Action<T, string> listener)
+		{
+			return this;
+		}
+
+		public T CurrentValue { get; }
+
+		public void Dispose()
+		{
+			// Dummy stub to satisfy return value of OnChange 
+		}
+	}
+
+	class RpcClientRefStub : IRpcClientRef
+	{
+		public GrpcChannel Channel { get; }
+		public HordeRpc.HordeRpcClient Client { get; }
+		public Task DisposingTask { get; }
+
+		public RpcClientRefStub(GrpcChannel Channel, HordeRpc.HordeRpcClient Client)
+		{
+			this.Channel = Channel;
+			this.Client = Client;
+			DisposingTask = new TaskCompletionSource<bool>().Task;
+		}
+
+		public void Dispose()
+		{
+		}
+	}
+
+	class RpcConnectionStub : IRpcConnection
+	{
+		private readonly GrpcChannel GrpcChannel;
+		private readonly HordeRpc.HordeRpcClient HordeRpcClient;
+
+		public RpcConnectionStub(GrpcChannel GrpcChannel, HordeRpc.HordeRpcClient HordeRpcClient)
+		{
+			this.GrpcChannel = GrpcChannel;
+			this.HordeRpcClient = HordeRpcClient;
+		}
+
+		public IRpcClientRef? TryGetClientRef(RpcContext Context)
+		{
+			return new RpcClientRefStub(GrpcChannel, HordeRpcClient);
+		}
+
+		public Task<IRpcClientRef> GetClientRef(RpcContext Context, CancellationToken CancellationToken)
+		{
+			IRpcClientRef RpcClientRefStub = new RpcClientRefStub(GrpcChannel, HordeRpcClient);
+			return Task.FromResult(RpcClientRefStub);
+		}
+
+		public Task<T> InvokeOnceAsync<T>(Func<HordeRpc.HordeRpcClient, Task<T>> Func, RpcContext Context,
+			CancellationToken CancellationToken)
+		{
+			return Func(HordeRpcClient);
+		}
+
+		public Task<T> InvokeOnceAsync<T>(Func<HordeRpc.HordeRpcClient, AsyncUnaryCall<T>> Func, RpcContext Context,
+			CancellationToken CancellationToken)
+		{
+			return Func(HordeRpcClient).ResponseAsync;
+		}
+
+		public Task<T> InvokeAsync<T>(Func<HordeRpc.HordeRpcClient, Task<T>> Func, RpcContext Context,
+			CancellationToken CancellationToken)
+		{
+			return Func(HordeRpcClient);
+		}
+
+		public Task<T> InvokeAsync<T>(Func<HordeRpc.HordeRpcClient, AsyncUnaryCall<T>> Func, RpcContext Context,
+			CancellationToken CancellationToken)
+		{
+			return Func(HordeRpcClient).ResponseAsync;
+		}
+
+		public ValueTask DisposeAsync()
+		{
+			return new ValueTask();
+		}
+	}
+
+	class HordeRpcClientStub : HordeRpc.HordeRpcClient
+	{
+		public readonly Queue<BeginStepResponse> BeginStepResponses = new Queue<BeginStepResponse>();
+		public readonly List<UpdateStepRequest> UpdateStepRequests = new List<UpdateStepRequest>();
+		public readonly Dictionary<GetStepRequest, GetStepResponse> GetStepResponses = new Dictionary<GetStepRequest, GetStepResponse>();
+		public Func<GetStepRequest, GetStepResponse>? GetStepFunc = null;
+		private readonly ILogger Logger;
+
+		public HordeRpcClientStub(ILogger Logger)
+		{
+			this.Logger = Logger;
+		}
+
+		public override AsyncUnaryCall<BeginBatchResponse> BeginBatchAsync(BeginBatchRequest Request,
+			CallOptions Options)
+		{
+			Logger.LogDebug("HordeRpcClientStub.BeginBatchAsync()");
+			BeginBatchResponse Res = new BeginBatchResponse();
+
+			Res.AgentType = "agentType1";
+			Res.LogId = "logId1";
+
+			return Wrap(Res);
+		}
+
+		public override AsyncUnaryCall<Empty> FinishBatchAsync(FinishBatchRequest Request, CallOptions Options)
+		{
+			Empty Res = new Empty();
+			return Wrap(Res);
+		}
+
+		public override AsyncUnaryCall<GetStreamResponse> GetStreamAsync(GetStreamRequest Request, CallOptions Options)
+		{
+			GetStreamResponse Res = new GetStreamResponse();
+			return Wrap(Res);
+		}
+
+		public override AsyncUnaryCall<GetJobResponse> GetJobAsync(GetJobRequest Request, CallOptions Options)
+		{
+			GetJobResponse Res = new GetJobResponse();
+			return Wrap(Res);
+		}
+
+		public override AsyncUnaryCall<BeginStepResponse> BeginStepAsync(BeginStepRequest Request, CallOptions Options)
+		{
+			if (BeginStepResponses.Count == 0)
+			{
+				BeginStepResponse CompleteRes = new BeginStepResponse();
+				CompleteRes.State = BeginStepResponse.Types.Result.Complete;
+				return Wrap(CompleteRes);
+			}
+
+			BeginStepResponse Res = BeginStepResponses.Dequeue();
+			Res.State = BeginStepResponse.Types.Result.Ready;
+			return Wrap(Res);
+		}
+
+		public override AsyncUnaryCall<Empty> WriteOutputAsync(WriteOutputRequest Request, CallOptions Options)
+		{
+			Logger.LogDebug("WriteOutputAsync: " + Request.Data);
+			Empty Res = new Empty();
+			return Wrap(Res);
+		}
+
+		public override AsyncUnaryCall<Empty> UpdateStepAsync(UpdateStepRequest Request, CallOptions Options)
+		{
+			Logger.LogDebug($"UpdateStepAsync(Request: {Request})");
+			UpdateStepRequests.Add(Request);
+			Empty Res = new Empty();
+			return Wrap(Res);
+		}
+
+		public override AsyncUnaryCall<Empty> CreateEventsAsync(CreateEventsRequest Request, CallOptions Options)
+		{
+			Logger.LogDebug("CreateEventsAsync: " + Request);
+			Empty Res = new Empty();
+			return Wrap(Res);
+		}
+
+		public override AsyncUnaryCall<GetStepResponse> GetStepAsync(GetStepRequest Request, CallOptions Options)
+		{
+			if (GetStepFunc != null)
+			{
+				return Wrap(GetStepFunc(Request));
+			}
+		
+			if (GetStepResponses.TryGetValue(Request, out GetStepResponse? Res))
+			{
+				return Wrap(Res);
+			}
+			
+			return Wrap(new GetStepResponse());
+		}
+
+		private AsyncUnaryCall<T> Wrap<T>(T Res)
+		{
+			return new AsyncUnaryCall<T>(Task.FromResult(Res), Task.FromResult(Metadata.Empty),
+				() => Status.DefaultSuccess, () => Metadata.Empty, null!);
+		}
+	}
+
+	class SimpleTestExecutor : IExecutor
+	{
+		private Func<BeginStepResponse, ILogger, CancellationToken, Task<JobStepOutcome>> Func;
+
+		public SimpleTestExecutor(Func<BeginStepResponse, ILogger, CancellationToken, Task<JobStepOutcome>> Func)
+		{
+			this.Func = Func;
+		}
+
+		public Task InitializeAsync(ILogger Logger, CancellationToken CancellationToken)
+		{
+			Logger.LogDebug("SimpleTestExecutor.InitializeAsync()");
+			return Task.CompletedTask;
+		}
+
+		public Task<JobStepOutcome> RunAsync(BeginStepResponse Step, ILogger Logger,
+			CancellationToken CancellationToken)
+		{
+			Logger.LogDebug($"SimpleTestExecutor.RunAsync(Step: {Step})");
+			return Func(Step, Logger, CancellationToken);
+		}
+
+		public Task FinalizeAsync(ILogger Logger, CancellationToken CancellationToken)
+		{
+			Logger.LogDebug("SimpleTestExecutor.FinalizeAsync()");
+			return Task.CompletedTask;
+		}
+	}
+}
\ No newline at end of file
diff --git a/Engine/Source/Programs/Horde/HordeAgentTests/HordeAgentTests.csproj b/Engine/Source/Programs/Horde/HordeAgentTests/HordeAgentTests.csproj
new file mode 100644
--- /dev/null
+++ b/Engine/Source/Programs/Horde/HordeAgentTests/HordeAgentTests.csproj
@@ -0,0 +1,22 @@
+<Project Sdk="Microsoft.NET.Sdk">
+
+  <PropertyGroup>
+    <TargetFramework>netcoreapp3.1</TargetFramework>
+    <Nullable>enable</Nullable>
+    <IsPackable>false</IsPackable>
+  </PropertyGroup>
+
+  <ItemGroup>
+    <PackageReference Include="Microsoft.Extensions.Logging.Abstractions" Version="5.0.0" />
+    <PackageReference Include="Microsoft.NET.Test.Sdk" Version="16.5.0" />
+    <PackageReference Include="Moq" Version="4.14.5" />
+    <PackageReference Include="MSTest.TestAdapter" Version="2.1.0" />
+    <PackageReference Include="MSTest.TestFramework" Version="2.1.0" />
+    <PackageReference Include="coverlet.collector" Version="1.2.0" />
+  </ItemGroup>
+
+  <ItemGroup>
+    <ProjectReference Include="..\HordeAgent\HordeAgent.csproj" />
+  </ItemGroup>
+
+</Project>
