Initial drop of Horde Server/Agent for UE5/Main. Still missing some Epic-specific functionality required to build/run correctly.

diff --git a/Engine/Source/Programs/Horde/HordeServer/Api/Credentials.cs b/Engine/Source/Programs/Horde/HordeServer/Api/Credentials.cs
new file mode 100644
--- /dev/null
+++ b/Engine/Source/Programs/Horde/HordeServer/Api/Credentials.cs
@@ -0,0 +1,126 @@
+// Copyright Epic Games, Inc. All Rights Reserved.
+
+using HordeServer.Models;
+using System;
+using System.Collections.Generic;
+using System.ComponentModel.DataAnnotations;
+using System.Linq;
+using System.Threading.Tasks;
+
+namespace HordeServer.Api
+{
+	/// <summary>
+	/// Parameters to create a new credential
+	/// </summary>
+	public class CreateCredentialRequest
+	{
+		/// <summary>
+		/// Name for the new credential
+		/// </summary>
+		[Required]
+		public string Name { get; set; } = null!;
+
+		/// <summary>
+		/// Properties for the new credential
+		/// </summary>
+		public Dictionary<string, string>? Properties { get; set; }
+	}
+
+	/// <summary>
+	/// Response from creating a new credential
+	/// </summary>
+	public class CreateCredentialResponse
+	{
+		/// <summary>
+		/// Unique id for the new credential
+		/// </summary>
+		public string Id { get; set; }
+
+		/// <summary>
+		/// Private constructor for serialization
+		/// </summary>
+		private CreateCredentialResponse()
+		{
+			Id = null!;
+		}
+
+		/// <summary>
+		/// Constructor
+		/// </summary>
+		/// <param name="Id">Unique id for the new credential</param>
+		public CreateCredentialResponse(string Id)
+		{
+			this.Id = Id;
+		}
+	}
+
+	/// <summary>
+	/// Parameters to update a credential
+	/// </summary>
+	public class UpdateCredentialRequest
+	{
+		/// <summary>
+		/// Optional new name for the credential
+		/// </summary>
+		public string? Name { get; set; }
+
+		/// <summary>
+		/// Properties to update for the credential. Properties set to null will be removed.
+		/// </summary>
+		public Dictionary<string, string?>? Properties { get; set; }
+
+		/// <summary>
+		/// Custom permissions for this object
+		/// </summary>
+		public UpdateAclRequest? Acl { get; set; }
+	}
+
+	/// <summary>
+	/// Response describing a credential
+	/// </summary>
+	public class GetCredentialResponse
+	{
+		/// <summary>
+		/// Unique id of the credential
+		/// </summary>
+		public string Id { get; set; }
+
+		/// <summary>
+		/// Name of the credential
+		/// </summary>
+		public string Name { get; set; }
+
+		/// <summary>
+		/// Properties for the credential
+		/// </summary>
+		public Dictionary<string, string> Properties { get; set; }
+
+		/// <summary>
+		/// Custom permissions for this object
+		/// </summary>
+		public GetAclResponse? Acl { get; set; }
+
+		/// <summary>
+		/// Parameterless constructor for serialization
+		/// </summary>
+		private GetCredentialResponse()
+		{
+			Id = null!;
+			Name = null!;
+			Properties = null!;
+		}
+
+		/// <summary>
+		/// Constructor
+		/// </summary>
+		/// <param name="Credential">The credential to construct from</param>
+		/// <param name="bIncludeAcl">Whether to include the ACL in the response</param>
+		public GetCredentialResponse(Credential Credential, bool bIncludeAcl)
+		{
+			this.Id = Credential.Id.ToString();
+			this.Name = Credential.Name;
+			this.Properties = Credential.Properties;
+			this.Acl = (bIncludeAcl && Credential.Acl != null)? new GetAclResponse(Credential.Acl) : null;
+		}
+	}
+}
diff --git a/Engine/Source/Programs/Horde/HordeServer/Api/Events.cs b/Engine/Source/Programs/Horde/HordeServer/Api/Events.cs
new file mode 100644
--- /dev/null
+++ b/Engine/Source/Programs/Horde/HordeServer/Api/Events.cs
@@ -0,0 +1,185 @@
+using EpicGames.Core;
+using HordeCommon;
+using Microsoft.AspNetCore.Mvc;
+using System;
+using System.Collections.Generic;
+using System.ComponentModel.DataAnnotations;
+using System.Linq;
+using System.Threading.Tasks;
+
+using StreamId = HordeServer.Utilities.StringId<HordeServer.Models.IStream>;
+using TemplateRefId = HordeServer.Utilities.StringId<HordeServer.Models.TemplateRef>;
+
+namespace HordeServer.Api
+{
+	/// <summary>
+	/// Base class for event records
+	/// </summary>
+	[JsonKnownTypes(typeof(JobCompleteEventRecord), typeof(LabelCompleteEventRecord), typeof(StepCompleteEventRecord))]
+	public class EventRecord
+	{
+	}
+
+	/// <summary>
+	/// Event that occurs when a job completes
+	/// </summary>
+	[JsonDiscriminator("Job")]
+	public class JobCompleteEventRecord : EventRecord
+	{
+		/// <summary>
+		/// The stream id
+		/// </summary>
+		[Required]
+		public string StreamId { get; set; }
+
+		/// <summary>
+		/// The template id
+		/// </summary>
+		[Required]
+		public string TemplateId { get; set; }
+
+		/// <summary>
+		/// Outcome of the job
+		/// </summary>
+		public LabelOutcome Outcome { get; set; } = LabelOutcome.Success;
+
+		/// <summary>
+		/// Default constructor
+		/// </summary>
+		private JobCompleteEventRecord()
+		{
+			StreamId = String.Empty;
+			TemplateId = String.Empty;
+		}
+
+		/// <summary>
+		/// Constructor
+		/// </summary>
+		/// <param name="StreamId">The stream id</param>
+		/// <param name="TemplateId">The template id</param>
+		/// <param name="Outcome">Outcome of the job</param>
+		public JobCompleteEventRecord(StreamId StreamId, TemplateRefId TemplateId, LabelOutcome Outcome)
+		{
+			this.StreamId = StreamId.ToString();
+			this.TemplateId = TemplateId.ToString();
+			this.Outcome = Outcome;
+		}
+	}
+
+	/// <summary>
+	/// Event that occurs when a label completes
+	/// </summary>
+	[JsonDiscriminator("Label")]
+	public class LabelCompleteEventRecord : EventRecord
+	{
+		/// <summary>
+		/// The stream id
+		/// </summary>
+		[Required]
+		public string StreamId { get; set; }
+
+		/// <summary>
+		/// The template id
+		/// </summary>
+		[Required]
+		public string TemplateId { get; set; }
+
+		/// <summary>
+		/// Name of the category for this label
+		/// </summary>
+		public string? CategoryName { get; set; }
+
+		/// <summary>
+		/// Name of the label to monitor
+		/// </summary>
+		[Required]
+		public string LabelName { get; set; }
+
+		/// <summary>
+		/// Outcome of the job
+		/// </summary>
+		public LabelOutcome Outcome { get; set; } = LabelOutcome.Success;
+
+		/// <summary>
+		/// Default constructor
+		/// </summary>
+		private LabelCompleteEventRecord()
+		{
+			this.StreamId = String.Empty;
+			this.TemplateId = String.Empty;
+			this.LabelName = String.Empty;
+		}
+
+		/// <summary>
+		/// Constructor
+		/// </summary>
+		/// <param name="StreamId">The stream id</param>
+		/// <param name="TemplateId">The template id</param>
+		/// <param name="CategoryName">Name of the category</param>
+		/// <param name="LabelName">The label name</param>
+		/// <param name="Outcome">Outcome of the label</param>
+		public LabelCompleteEventRecord(StreamId StreamId, TemplateRefId TemplateId, string? CategoryName, string LabelName, LabelOutcome Outcome)
+		{
+			this.StreamId = StreamId.ToString();
+			this.TemplateId = TemplateId.ToString();
+			this.CategoryName = CategoryName;
+			this.LabelName = LabelName;
+			this.Outcome = Outcome;
+		}
+	}
+
+	/// <summary>
+	/// Event that occurs when a step completes
+	/// </summary>
+	[JsonDiscriminator("Step")]
+	public class StepCompleteEventRecord : EventRecord
+	{
+		/// <summary>
+		/// The stream id
+		/// </summary>
+		[Required]
+		public string StreamId { get; set; }
+
+		/// <summary>
+		/// The template id
+		/// </summary>
+		[Required]
+		public string TemplateId { get; set; }
+
+		/// <summary>
+		/// Name of the step to monitor
+		/// </summary>
+		[Required]
+		public string StepName { get; set; }
+
+		/// <summary>
+		/// Outcome of the job step
+		/// </summary>
+		public JobStepOutcome Outcome { get; set; } = JobStepOutcome.Success;
+
+		/// <summary>
+		/// Default constructor
+		/// </summary>
+		private StepCompleteEventRecord()
+		{
+			StreamId = String.Empty;
+			TemplateId = String.Empty;
+			StepName = String.Empty;
+		}
+
+		/// <summary>
+		/// Constructor
+		/// </summary>
+		/// <param name="StreamId">The stream id</param>
+		/// <param name="TemplateId">The template id</param>
+		/// <param name="StepName">The label name</param>
+		/// <param name="Outcome">Outcome of the step</param>
+		public StepCompleteEventRecord(StreamId StreamId, TemplateRefId TemplateId, string StepName, JobStepOutcome Outcome)
+		{
+			this.StreamId = StreamId.ToString();
+			this.TemplateId = TemplateId.ToString();
+			this.StepName = StepName;
+			this.Outcome = Outcome;
+		}
+	}
+}
diff --git a/Engine/Source/Programs/Horde/HordeServer/Api/HelixSwarm.cs b/Engine/Source/Programs/Horde/HordeServer/Api/HelixSwarm.cs
new file mode 100644
--- /dev/null
+++ b/Engine/Source/Programs/Horde/HordeServer/Api/HelixSwarm.cs
@@ -0,0 +1,92 @@
+﻿using System.ComponentModel.DataAnnotations;
+using System.Diagnostics.CodeAnalysis;
+
+namespace HordeServer.Api
+{
+	/// <summary>
+	/// Request sent by the Helix Swarm project
+	/// The format of this request can be modified through the settings web UI in Swarm.
+	/// Please ensure it matches the request class below.
+	/// </summary>
+	[SuppressMessage("Compiler", "CA1056:URI properties should not be strings")]
+	public class SwarmAutoTestRequest
+	{
+		/// <summary>
+		/// Name of the test
+		/// </summary>
+		public string? Test { get; set; } = null;
+		
+		/// <summary>
+		/// ID of the test run
+		/// </summary>
+		public string? TestRunId { get; set; } = null;
+		
+		/// <summary>
+		/// Perforce changelist number to run the test on
+		/// </summary>
+		[Required]
+		public string? Changelist { get; set; } = null;
+		
+		/// <summary>
+		/// Status of the changelist, for example "submitted"
+		/// </summary>
+		public string? Status { get; set; } = null;
+		
+		/// <summary>
+		/// Version of the review
+		/// </summary>
+		public string? Version { get; set; } = null;
+		
+		/// <summary>
+		/// Description of the review, usually the changelist message itself
+		/// </summary>
+		public string? Description { get; set; } = null;
+		
+		/// <summary>
+		/// Swarm project ID
+		/// </summary>
+		public string? Project { get; set; } = null;
+		
+		/// <summary>
+		/// Swarm project name
+		/// </summary>
+		public string? ProjectName { get; set; } = null;
+		
+		/// <summary>
+		/// Perforce branch ID in Swarm
+		/// </summary>
+		public string? Branch { get; set; } = null;
+		
+		/// <summary>
+		/// Perforce branch name in Swarm
+		/// </summary>
+		public string? BranchName { get; set; } = null;
+		
+		/// <summary>
+		/// Callback URL to Swarm server to be called when "auto test" is finished.
+		/// In our case, that would be when the job is completed.
+		/// </summary>
+		[Required]
+		public string? UpdateUrl { get; set; } = null;
+	}
+
+	/// <summary>
+	/// Response object describing the created job
+	/// </summary>
+	public class SwarmAutoTestResponse
+	{
+		/// <summary>
+		/// The ID of the job that got started
+		/// </summary>
+		public string Id { get; set; }
+
+		/// <summary>
+		/// Constructor
+		/// </summary>
+		/// <param name="Id">Id of the new document</param>
+		public SwarmAutoTestResponse(string Id)
+		{
+			this.Id = Id;
+		}
+	}
+}
diff --git a/Engine/Source/Programs/Horde/HordeServer/Api/Issues.cs b/Engine/Source/Programs/Horde/HordeServer/Api/Issues.cs
new file mode 100644
--- /dev/null
+++ b/Engine/Source/Programs/Horde/HordeServer/Api/Issues.cs
@@ -0,0 +1,546 @@
+using EpicGames.Core;
+using HordeServer.Models;
+using HordeServer.Services;
+using HordeServer.Utilities;
+using System;
+using System.Collections.Generic;
+using System.Diagnostics.CodeAnalysis;
+using System.Linq;
+using System.Text.RegularExpressions;
+using System.Threading.Tasks;
+
+using StreamId = HordeServer.Utilities.StringId<HordeServer.Models.IStream>;
+using TemplateRefId = HordeServer.Utilities.StringId<HordeServer.Models.TemplateRef>;
+
+namespace HordeServer.Api
+{
+	/// <summary>
+	/// Identifies a particular changelist and job
+	/// </summary>
+	public class GetIssueStepResponse
+	{
+		/// <summary>
+		/// The changelist number
+		/// </summary>
+		public int Change { get; set; }
+
+		/// <summary>
+		/// Name of the job containing this step
+		/// </summary>
+		public string JobName { get; set; }
+
+		/// <summary>
+		/// The unique job id
+		/// </summary>
+		public string JobId { get; set; }
+
+		/// <summary>
+		/// The unique batch id
+		/// </summary>
+		public string BatchId { get; set; }
+
+		/// <summary>
+		/// The unique step id
+		/// </summary>
+		public string StepId { get; set; }
+
+		/// <summary>
+		/// Time at which the step ran
+		/// </summary>
+		public DateTime StepTime { get; set; }
+
+		/// <summary>
+		/// The unique log id
+		/// </summary>
+		public string? LogId { get; set; }
+
+		/// <summary>
+		/// Constructor
+		/// </summary>
+		/// <param name="IssueStep">The issue step to construct from</param>
+		public GetIssueStepResponse(IIssueStep IssueStep)
+		{
+			this.Change = IssueStep.Change;
+			this.JobName = IssueStep.JobName;
+			this.JobId = IssueStep.JobId.ToString();
+			this.BatchId = IssueStep.BatchId.ToString();
+			this.StepId = IssueStep.StepId.ToString();
+			this.StepTime = IssueStep.StepTime;
+			this.LogId = IssueStep.LogId?.ToString();
+		}
+	}
+
+	/// <summary>
+	/// Trace of a set of node failures across multiple steps
+	/// </summary>
+	public class GetIssueNodeResponse
+	{
+		/// <summary>
+		/// The template containing this step
+		/// </summary>
+		public string TemplateId { get; set; }
+
+		/// <summary>
+		/// Name of the step
+		/// </summary>
+		public string Name { get; set; }
+
+		/// <summary>
+		/// The previous build 
+		/// </summary>
+		public GetIssueStepResponse? LastSuccess { get; set; }
+
+		/// <summary>
+		/// The failing builds for a particular event
+		/// </summary>
+		public List<GetIssueStepResponse> Steps { get; set; }
+
+		/// <summary>
+		/// The following successful build
+		/// </summary>
+		public GetIssueStepResponse? NextSuccess { get; set; }
+
+		/// <summary>
+		/// Constructor
+		/// </summary>
+		/// <param name="Span">The node to construct from</param>
+		/// <param name="Steps">Failing steps for this span</param>
+		public GetIssueNodeResponse(IIssueSpan Span, List<IIssueStep> Steps)
+		{
+			this.Name = Span.NodeName;
+			this.TemplateId = Span.TemplateRefId.ToString();
+			this.LastSuccess = (Span.LastSuccess != null) ? new GetIssueStepResponse(Span.LastSuccess) : null;
+			this.Steps = Steps.ConvertAll(x => new GetIssueStepResponse(x));
+			this.NextSuccess = (Span.NextSuccess != null) ? new GetIssueStepResponse(Span.NextSuccess) : null;
+		}
+	}
+
+	/// <summary>
+	/// Information about a particular step
+	/// </summary>
+	public class GetIssueStreamResponse
+	{
+		/// <summary>
+		/// Unique id of the stream
+		/// </summary>
+		public string StreamId { get; set; }
+
+		/// <summary>
+		/// Minimum changelist affected by this issue (ie. last successful build)
+		/// </summary>
+		public int? MinChange { get; set; }
+
+		/// <summary>
+		/// Maximum changelist affected by this issue (ie. next successful build)
+		/// </summary>
+		public int? MaxChange { get; set; }
+
+		/// <summary>
+		/// Map of steps to (event signature id -> trace id)
+		/// </summary>
+		public List<GetIssueNodeResponse> Nodes { get; set; }
+
+		/// <summary>
+		/// Constructor
+		/// </summary>
+		/// <param name="StreamId">The stream to construct from</param>
+		/// <param name="Spans">List of spans for the given stream</param>
+		/// <param name="Steps">List of steps for the given stream</param>
+		public GetIssueStreamResponse(StreamId StreamId, List<IIssueSpan> Spans, List<IIssueStep> Steps)
+		{
+			this.StreamId = StreamId.ToString();
+
+			foreach (IIssueSpan Span in Spans)
+			{
+				if (Span.LastSuccess != null && (MinChange == null || Span.LastSuccess.Change < MinChange.Value))
+				{
+					MinChange = Span.LastSuccess.Change;
+				}
+				if (Span.NextSuccess != null && (MaxChange == null || Span.NextSuccess.Change > MaxChange.Value))
+				{
+					MaxChange = Span.NextSuccess.Change;
+				}
+			}
+
+			this.Nodes = Spans.ConvertAll(x => new GetIssueNodeResponse(x, Steps.Where(y => y.SpanId == x.Id).ToList()));
+		}
+	}
+
+	/// <summary>
+	/// Outcome of a particular build
+	/// </summary>
+	public enum IssueBuildOutcome
+	{
+		/// <summary>
+		/// Unknown outcome
+		/// </summary>
+		Unknown,
+
+		/// <summary>
+		/// Build succeeded
+		/// </summary>
+		Success,
+
+		/// <summary>
+		/// Build failed
+		/// </summary>
+		Error,
+
+		/// <summary>
+		/// Build finished with warnings
+		/// </summary>
+		Warning,
+	}
+
+	/// <summary>
+	/// Information about a suspect changelist that may have caused an issue
+	/// </summary>
+	public class GetIssueSuspectResponse
+	{
+		/// <summary>
+		/// Number of the changelist that was submitted
+		/// </summary>
+		public int Change { get; set; }
+
+		/// <summary>
+		/// Author of the changelist
+		/// </summary>
+		public string Author { get; set; }
+
+		/// <summary>
+		/// The originating change
+		/// </summary>
+		public int? OriginatingChange { get; set; }
+
+		/// <summary>
+		/// Time at which the user declined this issue
+		/// </summary>
+		public DateTime? DeclinedAt { get; set; }
+
+		/// <summary>
+		/// Constructor
+		/// </summary>
+		/// <param name="IssueSuspect">The object to construct from</param>
+		public GetIssueSuspectResponse(IIssueSpanSuspect IssueSuspect)
+		{
+			this.Change = IssueSuspect.Change;
+			this.Author = IssueSuspect.Author;
+			this.OriginatingChange = IssueSuspect.OriginatingChange;
+		}
+	}
+
+	/// <summary>
+	/// Information about a diagnostic
+	/// </summary>
+	public class GetIssueDiagnosticResponse
+	{
+		/// <summary>
+		/// The corresponding build id
+		/// </summary>
+		public long? BuildId { get; set; }
+
+		/// <summary>
+		/// Message for the diagnostic
+		/// </summary>
+		public string Message { get; set; }
+
+		/// <summary>
+		/// Link to the error
+		/// </summary>
+		public Uri Url { get; set; }
+
+		/// <summary>
+		/// Constructor
+		/// </summary>
+		/// <param name="BuildId">The corresponding build id</param>
+		/// <param name="Message">Message for the diagnostic</param>
+		/// <param name="Url">Link to the diagnostic</param>
+		public GetIssueDiagnosticResponse(long? BuildId, string Message, Uri Url)
+		{
+			this.BuildId = BuildId;
+			this.Message = Message;
+			this.Url = Url;
+		}
+	}
+
+	/// <summary>
+	/// Information about a template affected by an issue
+	/// </summary>
+	public class GetIssueAffectedTemplateResponse
+	{
+		/// <summary>
+		/// The template id
+		/// </summary>
+		public string TemplateId { get; set; }
+
+		/// <summary>
+		/// The template name
+		/// </summary>
+		public string TemplateName { get; set; }
+
+		/// <summary>
+		/// Whether it has been resolved or not
+		/// </summary>
+		public bool Resolved { get; set; }
+
+		/// <summary>
+		/// Constructor
+		/// </summary>
+		/// <param name="TemplateId"></param>
+		/// <param name="TemplateName"></param>
+		/// <param name="Resolved"></param>
+		public GetIssueAffectedTemplateResponse(string TemplateId, string TemplateName, bool Resolved)
+		{
+			this.TemplateId = TemplateId;
+			this.TemplateName = TemplateName;
+			this.Resolved = Resolved;
+		}
+	}
+
+	/// <summary>
+	/// Summary for the state of a stream in an issue
+	/// </summary>
+	public class GetIssueAffectedStreamResponse
+	{
+		/// <summary>
+		/// Id of the stream
+		/// </summary>
+		public string StreamId { get; set; }
+
+		/// <summary>
+		/// Name of the stream
+		/// </summary>
+		public string StreamName { get; set; }
+
+		/// <summary>
+		/// Whether the issue has been resolved in this stream
+		/// </summary>
+		public bool Resolved { get; set; }
+
+		/// <summary>
+		/// The affected templates
+		/// </summary>
+		public List<GetIssueAffectedTemplateResponse> AffectedTemplates { get; set; }
+
+		/// <summary>
+		/// List of affected template ids
+		/// </summary>
+		public List<string> TemplateIds { get; set; }
+
+		/// <summary>
+		/// List of resolved template ids
+		/// </summary>
+		public List<string> ResolvedTemplateIds { get; set; }
+
+		/// <summary>
+		/// List of unresolved template ids
+		/// </summary>
+		public List<string> UnresolvedTemplateIds { get; set; }
+
+		/// <summary>
+		/// Constructor
+		/// </summary>
+		/// <param name="Stream"></param>
+		/// <param name="Spans">The spans to construct from</param>
+		public GetIssueAffectedStreamResponse(IStream? Stream, IEnumerable<IIssueSpan> Spans)
+		{
+			IIssueSpan FirstSpan = Spans.First();
+			this.StreamId = FirstSpan.StreamId.ToString();
+			this.StreamName = FirstSpan.StreamName;
+			this.Resolved = Spans.All(x => x.NextSuccess != null);
+
+			this.AffectedTemplates = new List<GetIssueAffectedTemplateResponse>();
+			foreach (IGrouping<TemplateRefId, IIssueSpan> Template in Spans.GroupBy(x => x.TemplateRefId))
+			{
+				string TemplateName = Template.Key.ToString();
+				if (Stream != null && Stream.Templates.TryGetValue(Template.Key, out TemplateRef? TemplateRef))
+				{
+					TemplateName = TemplateRef.Name;
+				}
+				this.AffectedTemplates.Add(new GetIssueAffectedTemplateResponse(Template.Key.ToString(), TemplateName, Template.All(x => x.NextSuccess != null)));
+			}
+
+			HashSet<TemplateRefId> TemplateIdsSet = new HashSet<TemplateRefId>(Spans.Select(x => x.TemplateRefId));
+			this.TemplateIds = TemplateIdsSet.Select(x => x.ToString()).ToList();
+
+			HashSet<TemplateRefId> UnresolvedTemplateIdsSet = new HashSet<TemplateRefId>(Spans.Where(x => x.NextSuccess == null).Select(x => x.TemplateRefId));
+			this.UnresolvedTemplateIds = UnresolvedTemplateIdsSet.Select(x => x.ToString()).ToList();
+			this.ResolvedTemplateIds = TemplateIdsSet.Except(UnresolvedTemplateIdsSet).Select(x => x.ToString()).ToList();
+		}
+	}
+
+	/// <summary>
+	/// Stores information about a build health issue
+	/// </summary>
+	public class GetIssueResponse
+	{
+		/// <summary>
+		/// The unique object id
+		/// </summary>
+		public int Id { get; set; }
+
+		/// <summary>
+		/// Time at which the issue was created
+		/// </summary>
+		public DateTime CreatedAt { get; set; }
+
+		/// <summary>
+		/// Time at which the issue was retrieved
+		/// </summary>
+		public DateTime RetrievedAt { get; set; }
+
+		/// <summary>
+		/// The associated project for the issue
+		/// </summary>
+		public string? Project { get; set; }
+
+		/// <summary>
+		/// The summary text for this issue
+		/// </summary>
+		public string Summary { get; set; }
+
+		/// <summary>
+		/// Details text describing the issue
+		/// </summary>
+		public string? Details { get; set; }
+
+		/// <summary>
+		/// Severity of this issue
+		/// </summary>
+		public IssueSeverity Severity { get; set; }
+
+		/// <summary>
+		/// Owner of the issue
+		/// </summary>
+		public string? Owner { get; set; }
+
+		/// <summary>
+		/// User that nominated the current owner
+		/// </summary>
+		public string? NominatedBy { get; set; }
+
+		/// <summary>
+		/// Time that the issue was acknowledged
+		/// </summary>
+		public DateTime? AcknowledgedAt { get; set; }
+
+		/// <summary>
+		/// Changelist that fixed this issue
+		/// </summary>
+		public int? FixChange { get; set; }
+
+		/// <summary>
+		/// Time at which the issue was resolved
+		/// </summary>
+		public DateTime? ResolvedAt { get; set; }
+
+		/// <summary>
+		/// List of stream paths affected by this issue
+		/// </summary>
+		public List<string> Streams { get; set; }
+
+		/// <summary>
+		/// List of affected stream ids
+		/// </summary>
+		public List<string> ResolvedStreams { get; set; }
+
+		/// <summary>
+		/// List of unresolved streams
+		/// </summary>
+		public List<string> UnresolvedStreams { get; set; }
+
+		/// <summary>
+		/// List of affected streams
+		/// </summary>
+		public List<GetIssueAffectedStreamResponse> AffectedStreams { get; set; }
+
+		/// <summary>
+		/// Most likely suspects for causing this issue
+		/// </summary>
+		public List<string> PrimarySuspects { get; set; }
+
+		/// <summary>
+		/// Whether to show alerts for this issue
+		/// </summary>
+		public bool ShowDesktopAlerts { get; set; }
+
+		/// <summary>
+		/// Constructs a new issue
+		/// </summary>
+		/// <param name="Details">Issue to construct from</param>
+		/// <param name="AffectedStreams">The affected streams</param>
+		/// <param name="ShowDesktopAlerts">Whether to show alerts for this issue</param>
+		public GetIssueResponse(IIssueDetails Details, List<GetIssueAffectedStreamResponse> AffectedStreams, bool ShowDesktopAlerts)
+		{
+			IIssue Issue = Details.Issue;
+			this.Id = Issue.Id;
+			this.CreatedAt = Issue.CreatedAt;
+			this.RetrievedAt = DateTime.UtcNow;
+			this.Summary = String.IsNullOrEmpty(Issue.UserSummary)? Issue.Summary : Issue.UserSummary;
+			this.Severity = Issue.Severity;
+			this.Owner = Issue.Owner;
+			this.NominatedBy = Issue.NominatedBy;
+			this.AcknowledgedAt = Issue.AcknowledgedAt;
+			this.FixChange = Issue.FixChange;
+			this.ResolvedAt = Issue.ResolvedAt;
+			this.Streams = Details.Spans.Select(x => x.StreamName).Distinct().ToList()!;
+			this.ResolvedStreams = new List<string>();
+			this.UnresolvedStreams = new List<string>();
+			this.AffectedStreams = AffectedStreams;
+			foreach (IGrouping<StreamId, IIssueSpan> Stream in Details.Spans.GroupBy(x => x.StreamId))
+			{
+				if (Stream.All(x => x.NextSuccess != null))
+				{
+					this.ResolvedStreams.Add(Stream.Key.ToString());
+				}
+				else
+				{
+					this.UnresolvedStreams.Add(Stream.Key.ToString());
+				}
+			}
+			this.PrimarySuspects = Issue.Suspects.Select(x => x.Author).Distinct(StringComparer.OrdinalIgnoreCase).ToList();
+			this.ShowDesktopAlerts = ShowDesktopAlerts;
+		}
+	}
+
+	/// <summary>
+	/// Request an issue to be updated
+	/// </summary>
+	public class UpdateIssueRequest
+	{
+		/// <summary>
+		/// Summary of the issue
+		/// </summary>
+		public string? Summary { get; set; }
+
+		/// <summary>
+		/// New owner of the issue
+		/// </summary>
+		public string? Owner { get; set; }
+
+		/// <summary>
+		/// User than nominates the new owner
+		/// </summary>
+		public string? NominatedBy { get; set; }
+
+		/// <summary>
+		/// Whether the issue has been acknowledged
+		/// </summary>
+		public bool? Acknowledged { get; set; }
+
+		/// <summary>
+		/// Whether the user has declined this issue
+		/// </summary>
+		public bool? Declined { get; set; }
+
+		/// <summary>
+		/// The change at which the issue is claimed fixed. 0 = not fixed, -1 = systemic issue.
+		/// </summary>
+		public int? FixChange { get; set; }
+
+		/// <summary>
+		/// Whether the issue should be marked as resolved
+		/// </summary>
+		public bool? Resolved { get; set; }
+	}
+}
diff --git a/Engine/Source/Programs/Horde/HordeServer/Api/Jobs.cs b/Engine/Source/Programs/Horde/HordeServer/Api/Jobs.cs
new file mode 100644
--- /dev/null
+++ b/Engine/Source/Programs/Horde/HordeServer/Api/Jobs.cs
@@ -0,0 +1,1092 @@
+// Copyright Epic Games, Inc. All Rights Reserved.
+
+using EpicGames.Core;
+using HordeCommon;
+using HordeServer.Models;
+using HordeServer.Utilities;
+using System;
+using System.Collections.Generic;
+using System.ComponentModel.DataAnnotations;
+using System.Linq;
+using System.Security.Claims;
+using System.Text.Json.Serialization;
+using System.Threading.Tasks;
+
+namespace HordeServer.Api
+{
+	/// <summary>
+	/// State of the job
+	/// </summary>
+	public enum JobState
+	{
+		/// <summary>
+		/// Waiting for resources
+		/// </summary>
+		Waiting,
+
+		/// <summary>
+		/// Currently running one or more steps
+		/// </summary>
+		Running,
+
+		/// <summary>
+		/// All steps have completed
+		/// </summary>
+		Complete,
+	}
+
+	/// <summary>
+	/// Query selecting the base changelist to use
+	/// </summary>
+	public class ChangeQueryRequest
+	{
+		/// <summary>
+		/// The template id to query
+		/// </summary>
+		public string TemplateId { get; set; } = String.Empty;
+
+		/// <summary>
+		/// The target to query
+		/// </summary>
+		public string? Target { get; set; }
+	}
+
+	/// <summary>
+	/// Parameters required to create a job
+	/// </summary>
+	public class CreateJobRequest
+	{
+		/// <summary>
+		/// The stream that this job belongs to
+		/// </summary>
+		[Required]
+		public string StreamId { get; set; }
+
+		/// <summary>
+		/// The template for this job
+		/// </summary>
+		[Required]
+		public string TemplateId { get; set; }
+
+		/// <summary>
+		/// Name of the job
+		/// </summary>
+		public string? Name { get; set; }
+
+		/// <summary>
+		/// The changelist number to build. Can be null for latest.
+		/// </summary>
+		public int? Change { get; set; }
+
+		/// <summary>
+		/// Parameters to use when selecting the change to execute at.
+		/// </summary>
+		public ChangeQueryRequest? ChangeQuery { get; set; }
+
+		/// <summary>
+		/// The preflight changelist number
+		/// </summary>
+		public int? PreflightChange { get; set; }
+
+		/// <summary>
+		/// Priority for the job
+		/// </summary>
+		public Priority? Priority { get; set; }
+
+		/// <summary>
+		/// Whether to automatically submit the preflighted change on completion
+		/// </summary>
+		public bool? AutoSubmit { get; set; }
+
+		/// <summary>
+		/// Nodes for the new job
+		/// </summary>
+		public List<CreateGroupRequest>? Groups { get; set; }
+
+		/// <summary>
+		/// Aggregates for the new job
+		/// </summary>
+		public List<CreateAggregateRequest>? Aggregates { get; set; }
+
+		/// <summary>
+		/// Labels for the new job
+		/// </summary>
+		public List<CreateLabelRequest>? Labels { get; set; }
+
+		/// <summary>
+		/// Arguments for the job
+		/// </summary>
+		public List<string>? Arguments { get; set; }
+
+		/// <summary>
+		/// Private constructor for serialization
+		/// </summary>
+		private CreateJobRequest()
+		{
+			StreamId = null!;
+			TemplateId = null!;
+		}
+	}
+
+	/// <summary>
+	/// Response from creating a new job
+	/// </summary>
+	public class CreateJobResponse
+	{
+		/// <summary>
+		/// Unique id for the new job
+		/// </summary>
+		public string Id { get; set; }
+
+		/// <summary>
+		/// Private constructor for serialization
+		/// </summary>
+		private CreateJobResponse()
+		{
+			Id = null!;
+		}
+
+		/// <summary>
+		/// Constructor
+		/// </summary>
+		/// <param name="Id">Unique id for the new job</param>
+		public CreateJobResponse(string Id)
+		{
+			this.Id = Id;
+		}
+	}
+
+	/// <summary>
+	/// Updates an existing job
+	/// </summary>
+	public class UpdateJobRequest
+	{
+		/// <summary>
+		/// New name for the job
+		/// </summary>
+		public string? Name { get; set; }
+
+		/// <summary>
+		/// New priority for the job
+		/// </summary>
+		public Priority? Priority { get; set; }
+
+		/// <summary>
+		/// Set whether the job should be automatically submitted or not
+		/// </summary>
+		public bool? AutoSubmit { get; set; }
+
+		/// <summary>
+		/// Mark this job as aborted
+		/// </summary>
+		public bool? Aborted { get; set; }
+
+		/// <summary>
+		/// New list of arguments for the job. Only -Target= arguments can be modified after the job has started.
+		/// </summary>
+		public List<string>? Arguments { get; set; }
+
+		/// <summary>
+		/// Custom permissions for this object
+		/// </summary>
+		public UpdateAclRequest? Acl { get; set; }
+	}
+
+	/// <summary>
+	/// Information about a report associated with a job
+	/// </summary>
+	public class GetReportResponse
+	{
+		/// <summary>
+		/// Name of the report
+		/// </summary>
+		public string Name { get; set; }
+
+		/// <summary>
+		/// Report placement
+		/// </summary>
+		public HordeCommon.Rpc.ReportPlacement Placement { get; set; }
+
+		/// <summary>
+		/// The artifact id
+		/// </summary>
+		public string ArtifactId { get; set; }
+
+		/// <summary>
+		/// Constructor
+		/// </summary>
+		/// <param name="Report"></param>
+		public GetReportResponse(IReport Report)
+		{
+			Name = Report.Name;
+			Placement = Report.Placement;
+			ArtifactId = Report.ArtifactId.ToString();
+		}
+	}
+
+	/// <summary>
+	/// Information about a job
+	/// </summary>
+	public class GetJobResponse
+	{
+		/// <summary>
+		/// Unique Id for the job
+		/// </summary>
+		public string Id { get; set; }
+
+		/// <summary>
+		/// Unique id of the stream containing this job
+		/// </summary>
+		public string StreamId { get; set; }
+
+		/// <summary>
+		/// Name of the job
+		/// </summary>
+		public string Name { get; set; }
+
+		/// <summary>
+		/// The changelist number to build
+		/// </summary>
+		public int Change { get; set; }
+
+		/// <summary>
+		/// The code changelist
+		/// </summary>
+		public int? CodeChange { get; set; }
+
+		/// <summary>
+		/// The preflight changelist number
+		/// </summary>
+		public int? PreflightChange { get; set; }
+
+		/// <summary>
+		/// The cloned preflight changelist number
+		/// </summary>
+		public int? ClonedPreflightChange { get; set; }
+
+		/// <summary>
+		/// The template type
+		/// </summary>
+		public string TemplateId { get; set; }
+
+		/// <summary>
+		/// Hash of the actual template data
+		/// </summary>
+		public string TemplateHash { get; set; }
+
+		/// <summary>
+		/// Hash of the graph for this job
+		/// </summary>
+		public string? GraphHash { get; set; }
+
+		/// <summary>
+		/// The user that started this job
+		/// </summary>
+		public string? StartedByUserId { get; set; }
+
+		/// <summary>
+		/// The user that started this job
+		/// </summary>
+		public string? StartedByUser { get; set; }
+
+		/// <summary>
+		/// The user that aborted this job
+		/// </summary>
+		public string? AbortedByUser { get; set; }
+
+		/// <summary>
+		/// Priority of the job
+		/// </summary>
+		public Priority Priority { get; set; }
+
+		/// <summary>
+		/// Whether the change will automatically be submitted or not
+		/// </summary>
+		public bool AutoSubmit { get; set; }
+
+		/// <summary>
+		/// The submitted changelist number
+		/// </summary>
+		public int? AutoSubmitChange { get; }
+
+		/// <summary>
+		/// Message produced by trying to auto-submit the change
+		/// </summary>
+		public string? AutoSubmitMessage { get; }
+
+		/// <summary>
+		/// Time that the job was created
+		/// </summary>
+		public DateTimeOffset CreateTime { get; set; }
+
+		/// <summary>
+		/// The global job state
+		/// </summary>
+		public JobState State { get; set; }
+
+		/// <summary>
+		/// Array of jobstep batches
+		/// </summary>
+		public List<GetBatchResponse>? Batches { get; set; }
+
+		/// <summary>
+		/// List of labels
+		/// </summary>
+		public List<GetLabelStateResponse>? Labels { get; set; }
+
+		/// <summary>
+		/// The default label, containing the state of all steps that are otherwise not matched.
+		/// </summary>
+		public GetDefaultLabelStateResponse? DefaultLabel { get; set; } 
+
+		/// <summary>
+		/// List of reports
+		/// </summary>
+		public List<GetReportResponse>? Reports { get; set; }
+
+		/// <summary>
+		/// Parameters for the job
+		/// </summary>
+		public List<string> Arguments { get; set; }
+
+		/// <summary>
+		/// The last update time for this job
+		/// </summary>
+		public DateTimeOffset UpdateTime { get; set; }
+
+		/// <summary>
+		/// Per-object permissions
+		/// </summary>
+		public GetAclResponse? Acl { get; set; }
+
+		/// <summary>
+		/// Constructor
+		/// </summary>
+		/// <param name="Job">Job to create a response for</param>
+		/// <param name="AclResponse">The ACL response object</param>
+		public GetJobResponse(IJob Job, GetAclResponse? AclResponse)
+		{
+			this.Id = Job.Id.ToString();
+			this.StreamId = Job.StreamId.ToString();
+			this.Name = Job.Name;
+			this.Change = Job.Change;
+			this.CodeChange = (Job.CodeChange != 0) ? (int?)Job.CodeChange : null;
+			this.PreflightChange = (Job.PreflightChange != 0) ? (int?)Job.PreflightChange : null;
+			this.ClonedPreflightChange = (Job.ClonedPreflightChange != 0) ? (int?)Job.ClonedPreflightChange : null;
+			this.TemplateId = Job.TemplateId.ToString();
+			this.TemplateHash = Job.TemplateHash?.ToString() ?? String.Empty;
+			this.GraphHash = Job.GraphHash.ToString();
+			this.StartedByUserId = Job.StartedByUserId?.ToString();
+			this.StartedByUser = Job.StartedByUser;
+			this.AbortedByUser = Job.AbortedByUser;
+			this.CreateTime = new DateTimeOffset(Job.CreateTimeUtc);
+			this.State = Job.GetState();
+			this.Priority = Job.Priority;
+			this.AutoSubmit = Job.AutoSubmit;
+			this.AutoSubmitChange = Job.AutoSubmitChange;
+			this.AutoSubmitMessage = Job.AutoSubmitMessage;
+			this.Reports = Job.Reports?.ConvertAll(x => new GetReportResponse(x));
+			this.Arguments = Job.Arguments.ToList();
+			this.UpdateTime = new DateTimeOffset(Job.UpdateTimeUtc);
+			this.Acl = AclResponse;
+		}
+	}
+
+	/// <summary>
+	/// The timing info for 
+	/// </summary>
+	public class GetJobTimingResponse
+	{
+		/// <summary>
+		/// Timing info for each step
+		/// </summary>
+		public Dictionary<string, GetStepTimingInfoResponse> Steps { get; set; }
+
+		/// <summary>
+		/// Timing information for each label
+		/// </summary>
+		public List<GetLabelTimingInfoResponse> Labels { get; set; }
+
+		/// <summary>
+		/// Constructor
+		/// </summary>
+		/// <param name="Steps">Timing info for each steps</param>
+		/// <param name="Labels">Timing info for each label</param>
+		public GetJobTimingResponse(Dictionary<string, GetStepTimingInfoResponse> Steps, List<GetLabelTimingInfoResponse> Labels)
+		{
+			this.Steps = Steps;
+			this.Labels = Labels;
+		}
+	}
+
+	/// <summary>
+	/// Request used to update a jobstep
+	/// </summary>
+	public class UpdateStepRequest
+	{
+		/// <summary>
+		/// The new jobstep state
+		/// </summary>
+		public JobStepState State { get; set; } = JobStepState.Unspecified;
+
+		/// <summary>
+		/// Outcome from the jobstep
+		/// </summary>
+		public JobStepOutcome Outcome { get; set; } = JobStepOutcome.Unspecified;
+		
+		/// <summary>
+		/// If the step has been requested to abort
+		/// </summary>
+		public bool? AbortRequested { get; set; }
+
+		/// <summary>
+		/// Specifies the log file id for this step
+		/// </summary>
+		public string? LogId { get; set; }
+
+		/// <summary>
+		/// Whether the step should be re-run
+		/// </summary>
+		public bool? Retry { get; set; }
+
+		/// <summary>
+		/// New priority for this step
+		/// </summary>
+		public Priority? Priority { get; set; }
+
+		/// <summary>
+		/// Properties to set. Any entries with a null value will be removed.
+		/// </summary>
+		public Dictionary<string, string?>? Properties { get; set; }
+	}
+
+	/// <summary>
+	/// Response object when updating a jobstep
+	/// </summary>
+	public class UpdateStepResponse
+	{
+		/// <summary>
+		/// If a new step is created (due to specifying the retry flag), specifies the batch id
+		/// </summary>
+		public string? BatchId { get; set; }
+
+		/// <summary>
+		/// If a step is retried, includes the new step id
+		/// </summary>
+		public string? StepId { get; set; }
+	}
+
+	/// <summary>
+	/// Returns information about a jobstep
+	/// </summary>
+	public class GetStepResponse
+	{
+		/// <summary>
+		/// The unique id of the step
+		/// </summary>
+		public string Id { get; set; }
+
+		/// <summary>
+		/// Index of the node which this jobstep is to execute
+		/// </summary>
+		public int NodeIdx { get; set; }
+
+		/// <summary>
+		/// Current state of the job step. This is updated automatically when runs complete.
+		/// </summary>
+		public JobStepState State { get; set; }
+
+		/// <summary>
+		/// Current outcome of the jobstep
+		/// </summary>
+		public JobStepOutcome Outcome { get; set; }
+		
+		/// <summary>
+		/// If the step has been requested to abort
+		/// </summary>
+		public bool AbortRequested { get; set; }
+		
+		/// <summary>
+		/// Name of the user that requested the abort of this step
+		/// </summary>
+		public string? AbortByUser { get; set; }
+
+		/// <summary>
+		/// Name of the user that requested this step be run again
+		/// </summary>
+		public string? RetryByUser { get; set; }
+
+		/// <summary>
+		/// The log id for this step
+		/// </summary>
+		public string? LogId { get; set; }
+
+		/// <summary>
+		/// Time at which the batch was ready (UTC).
+		/// </summary>
+		public DateTimeOffset? ReadyTime { get; set; }
+
+		/// <summary>
+		/// Time at which the batch started (UTC).
+		/// </summary>
+		public DateTimeOffset? StartTime { get; set; }
+
+		/// <summary>
+		/// Time at which the batch finished (UTC)
+		/// </summary>
+		public DateTimeOffset? FinishTime { get; set; }
+
+		/// <summary>
+		/// List of reports
+		/// </summary>
+		public List<GetReportResponse>? Reports { get; set; }
+
+		/// <summary>
+		/// User-defined properties for this jobstep.
+		/// </summary>
+		public Dictionary<string, string>? Properties { get; set; }
+
+		/// <summary>
+		/// Parameterless constructor for serialization
+		/// </summary>
+		private GetStepResponse()
+		{
+			Id = null!;
+		}
+
+		/// <summary>
+		/// Constructor
+		/// </summary>
+		/// <param name="Step">The step to construct from</param>
+		public GetStepResponse(IJobStep Step)
+		{
+			this.Id = Step.Id.ToString();
+			this.NodeIdx = Step.NodeIdx;
+			this.State = Step.State;
+			this.Outcome = Step.Outcome;
+			this.AbortRequested = Step.AbortRequested;
+			this.AbortByUser = Step.AbortByUser;
+			this.RetryByUser = Step.RetryByUser;
+			this.LogId = Step.LogId?.ToString();
+			this.ReadyTime = Step.ReadyTimeUtc;
+			this.StartTime = Step.StartTimeUtc;
+			this.FinishTime = Step.FinishTimeUtc;
+			this.Reports = Step.Reports?.ConvertAll(x => new GetReportResponse(x));
+
+			if (Step.Properties != null && Step.Properties.Count > 0)
+			{
+				this.Properties = Step.Properties;
+			}
+		}
+	}
+
+	/// <summary>
+	/// The state of a particular run
+	/// </summary>
+	[JsonConverter(typeof(JsonStringEnumConverter))]
+	public enum JobStepBatchState
+	{
+		/// <summary>
+		/// Waiting for dependencies of at least one jobstep to complete
+		/// </summary>
+		Waiting = 0,
+
+		/// <summary>
+		/// Ready to execute
+		/// </summary>
+		Ready = 1,
+
+		/// <summary>
+		/// Currently running
+		/// </summary>
+		Running = 2,
+
+		/// <summary>
+		/// All steps have finished executing
+		/// </summary>
+		Complete = 3
+	}
+
+	/// <summary>
+	/// Error code for a batch not being executed
+	/// </summary>
+	[JsonConverter(typeof(JsonStringEnumConverter))]
+	public enum JobStepBatchError
+	{
+		/// <summary>
+		/// No error
+		/// </summary>
+		None = 0,
+
+		/// <summary>
+		/// The stream for this job is unknown
+		/// </summary>
+		UnknownStream = 1,
+
+		/// <summary>
+		/// The given agent type for this batch was not valid for this stream
+		/// </summary>
+		UnknownAgentType = 2,
+
+		/// <summary>
+		/// The pool id referenced by the agent type was not found
+		/// </summary>
+		UnknownPool = 3,
+
+		/// <summary>
+		/// There are no agents in the given pool currently online
+		/// </summary>
+		NoAgentsInPool = 4,
+
+		/// <summary>
+		/// There are no agents in this pool that are onlinbe
+		/// </summary>
+		NoAgentsOnline = 5,
+
+		/// <summary>
+		/// Unknown workspace referenced by the agent type
+		/// </summary>
+		UnknownWorkspace = 6,
+
+		/// <summary>
+		/// Cancelled
+		/// </summary>
+		Cancelled = 7,
+
+		/// <summary>
+		/// Lost connection with the agent machine
+		/// </summary>
+		LostConnection = 8,
+
+		/// <summary>
+		/// Lease terminated prematurely
+		/// </summary>
+		Incomplete = 9,
+	}
+
+	/// <summary>
+	/// Request to update a jobstep batch
+	/// </summary>
+	public class UpdateBatchRequest
+	{
+		/// <summary>
+		/// The new log file id
+		/// </summary>
+		public string? LogId { get; set; }
+
+		/// <summary>
+		/// The state of this batch
+		/// </summary>
+		public JobStepBatchState? State { get; set; }
+	}
+
+	/// <summary>
+	/// Information about a jobstep batch
+	/// </summary>
+	public class GetBatchResponse
+	{
+		/// <summary>
+		/// Unique id for this batch
+		/// </summary>
+		public string Id { get; set; }
+
+		/// <summary>
+		/// The unique log file id
+		/// </summary>
+		public string? LogId { get; set; }
+
+		/// <summary>
+		/// Index of the group being executed
+		/// </summary>
+		public int GroupIdx { get; set; }
+
+		/// <summary>
+		/// The state of this batch
+		/// </summary>
+		public JobStepBatchState State { get; set; }
+
+		/// <summary>
+		/// Error code for this batch
+		/// </summary>
+		public JobStepBatchError Error { get; set; }
+
+		/// <summary>
+		/// Steps within this run
+		/// </summary>
+		public List<GetStepResponse> Steps { get; set; }
+
+		/// <summary>
+		/// The agent assigned to execute this group
+		/// </summary>
+		public string? AgentId { get; set; }
+
+		/// <summary>
+		/// The agent session holding this lease
+		/// </summary>
+		public string? SessionId { get; set; }
+
+		/// <summary>
+		/// The lease that's executing this group
+		/// </summary>
+		public string? LeaseId { get; set; }
+
+		/// <summary>
+		/// The priority of this batch
+		/// </summary>
+		public int WeightedPriority { get; set; }
+
+		/// <summary>
+		/// Time at which the group started (UTC).
+		/// </summary>
+		public DateTimeOffset? StartTime { get; set; }
+
+		/// <summary>
+		/// Time at which the group finished (UTC)
+		/// </summary>
+		public DateTimeOffset? FinishTime { get; set; }
+
+		/// <summary>
+		/// Parameterless constructor for serialization
+		/// </summary>
+		private GetBatchResponse()
+		{
+			Id = null!;
+			Steps = new List<GetStepResponse>();
+		}
+
+		/// <summary>
+		/// Converts this batch into a public response object
+		/// </summary>
+		/// <param name="Batch">The batch to construct from</param>
+		/// <returns>Response instance</returns>
+		public GetBatchResponse(IJobStepBatch Batch)
+		{
+			this.Id = Batch.Id.ToString();
+			this.LogId = Batch.LogId?.ToString();
+			this.GroupIdx = Batch.GroupIdx;
+			this.State = Batch.State;
+			this.Error = Batch.Error;
+			this.Steps = Batch.Steps.ConvertAll(x => new GetStepResponse(x));
+			this.AgentId = Batch.AgentId?.ToString();
+			this.SessionId = Batch.SessionId?.ToString();
+			this.LeaseId = Batch.LeaseId?.ToString();
+			this.WeightedPriority = Batch.SchedulePriority;
+			this.StartTime = Batch.StartTimeUtc;
+			this.FinishTime = Batch.FinishTimeUtc;
+		}
+	}
+
+	/// <summary>
+	/// State of an aggregate
+	/// </summary>
+	public enum LabelState
+	{
+		/// <summary>
+		/// Aggregate is not currently being built (no required nodes are present)
+		/// </summary>
+		Unspecified,
+
+		/// <summary>
+		/// Steps are still running
+		/// </summary>
+		Running,
+
+		/// <summary>
+		/// All steps are complete
+		/// </summary>
+		Complete
+	}
+
+	/// <summary>
+	/// Outcome of an aggregate
+	/// </summary>
+	public enum LabelOutcome
+	{
+		/// <summary>
+		/// Aggregate is not currently being built
+		/// </summary>
+		Unspecified,
+
+		/// <summary>
+		/// A step dependency failed
+		/// </summary>
+		Failure,
+
+		/// <summary>
+		/// A dependency finished with warnings
+		/// </summary>
+		Warnings,
+
+		/// <summary>
+		/// Successful
+		/// </summary>
+		Success,
+	}
+
+	/// <summary>
+	/// State of a label within a job
+	/// </summary>
+	public class GetLabelStateResponse
+	{
+		/// <summary>
+		/// State of the label
+		/// </summary>
+		public LabelState? State { get; set; }
+
+		/// <summary>
+		/// Outcome of the label
+		/// </summary>
+		public LabelOutcome? Outcome { get; set; }
+
+		/// <summary>
+		/// Constructor
+		/// </summary>
+		/// <param name="State">State of the label</param>
+		/// <param name="Outcome">Outcome of the label</param>
+		public GetLabelStateResponse(LabelState State, LabelOutcome Outcome)
+		{
+			this.State = State;
+			this.Outcome = Outcome;
+		}
+	}
+
+	/// <summary>
+	/// Information about the default label (ie. with inlined list of nodes)
+	/// </summary>
+	public class GetDefaultLabelStateResponse : GetLabelStateResponse
+	{
+		/// <summary>
+		/// List of nodes covered by default label
+		/// </summary>
+		public List<string> Nodes { get; set; }
+
+		/// <summary>
+		/// Constructor
+		/// </summary>
+		/// <param name="State">State of the label</param>
+		/// <param name="Outcome">Outcome of the label</param>
+		/// <param name="Nodes">List of nodes that are covered by the default label</param>
+		public GetDefaultLabelStateResponse(LabelState State, LabelOutcome Outcome, List<string> Nodes)
+			: base(State, Outcome)
+		{
+			this.Nodes = Nodes;
+		}
+	}
+
+	/// <summary>
+	/// Information about the timing info for a particular target
+	/// </summary>
+	public class GetTimingInfoResponse
+	{
+		/// <summary>
+		/// Wait time on the critical path
+		/// </summary>
+		public float? TotalWaitTime { get; set; }
+
+		/// <summary>
+		/// Sync time on the critical path
+		/// </summary>
+		public float? TotalInitTime { get; set; }
+
+		/// <summary>
+		/// Duration to this point
+		/// </summary>
+		public float? TotalTimeToComplete { get; set; }
+
+		/// <summary>
+		/// Average wait time by the time the job reaches this point
+		/// </summary>
+		public float? AverageTotalWaitTime { get; set; }
+
+		/// <summary>
+		/// Average sync time to this point
+		/// </summary>
+		public float? AverageTotalInitTime { get; set; }
+
+		/// <summary>
+		/// Average duration to this point
+		/// </summary>
+		public float? AverageTotalTimeToComplete { get; set; }
+
+		/// <summary>
+		/// Constructor
+		/// </summary>
+		/// <param name="TimingInfo">Timing info to construct from</param>
+		public GetTimingInfoResponse(TimingInfo TimingInfo)
+		{
+			this.TotalWaitTime = (float?)TimingInfo.TotalWaitTime?.TotalSeconds;
+			this.TotalInitTime = (float?)TimingInfo.TotalInitTime?.TotalSeconds;
+			this.TotalTimeToComplete = (float?)TimingInfo.TotalTimeToComplete?.TotalSeconds;
+
+			this.AverageTotalWaitTime = (float?)TimingInfo.AverageTotalWaitTime?.TotalSeconds;
+			this.AverageTotalInitTime = (float?)TimingInfo.AverageTotalInitTime?.TotalSeconds;
+			this.AverageTotalTimeToComplete = (float?)TimingInfo.AverageTotalTimeToComplete?.TotalSeconds;
+		}
+	}
+
+	/// <summary>
+	/// Information about the timing info for a particular target
+	/// </summary>
+	public class GetStepTimingInfoResponse : GetTimingInfoResponse
+	{
+		/// <summary>
+		/// Name of this node
+		/// </summary>
+		public string? Name { get; set; }
+
+		/// <summary>
+		/// Average wait time for this step
+		/// </summary>
+		public float? AverageStepWaitTime { get; set; }
+
+		/// <summary>
+		/// Average init time for this step
+		/// </summary>
+		public float? AverageStepInitTime { get; set; }
+
+		/// <summary>
+		/// Average duration for this step
+		/// </summary>
+		public float? AverageStepDuration { get; set; }
+
+		/// <summary>
+		/// Constructor
+		/// </summary>
+		/// <param name="Name">Name of the node</param>
+		/// <param name="JobTimingInfo">Timing info to construct from</param>
+		public GetStepTimingInfoResponse(string? Name, TimingInfo JobTimingInfo)
+			: base(JobTimingInfo)
+		{
+			this.Name = Name;
+			this.AverageStepWaitTime = JobTimingInfo.StepTiming?.AverageWaitTime;
+			this.AverageStepInitTime = JobTimingInfo.StepTiming?.AverageInitTime;
+			this.AverageStepDuration = JobTimingInfo.StepTiming?.AverageDuration;
+		}
+	}
+
+	/// <summary>
+	/// Information about the timing info for a label
+	/// </summary>
+	public class GetLabelTimingInfoResponse : GetTimingInfoResponse
+	{
+		/// <summary>
+		/// Name of the label
+		/// </summary>
+		[Obsolete("Use DashboardName instead")]
+		public string? Name => DashboardName;
+
+		/// <summary>
+		/// Category for the label
+		/// </summary>
+		[Obsolete("Use DashboardCategory instead")]
+		public string? Category => DashboardCategory;
+
+		/// <summary>
+		/// Name of the label
+		/// </summary>
+		public string? DashboardName { get; set; }
+
+		/// <summary>
+		/// Category for the label
+		/// </summary>
+		public string? DashboardCategory { get; set; }
+
+		/// <summary>
+		/// Name of the label
+		/// </summary>
+		public string? UgsName { get; set; }
+
+		/// <summary>
+		/// Category for the label
+		/// </summary>
+		public string? UgsProject { get; set; }
+
+		/// <summary>
+		/// Constructor
+		/// </summary>
+		/// <param name="Label">The label to construct from</param>
+		/// <param name="TimingInfo">Timing info to construct from</param>
+		public GetLabelTimingInfoResponse(ILabel Label, TimingInfo TimingInfo)
+			: base(TimingInfo)
+		{
+			this.DashboardName = Label.DashboardName;
+			this.DashboardCategory = Label.DashboardCategory;
+			this.UgsName = Label.UgsName;
+			this.UgsProject = Label.UgsProject;
+		}
+	}
+
+	/// <summary>
+	/// Describes the history of a step
+	/// </summary>
+	public class GetJobStepRefResponse
+	{
+		/// <summary>
+		/// The job id
+		/// </summary>
+		public string JobId { get; set; }
+
+		/// <summary>
+		/// The batch containing the step
+		/// </summary>
+		public string BatchId { get; set; }
+
+		/// <summary>
+		/// The step identifier
+		/// </summary>
+		public string StepId { get; set; }
+
+		/// <summary>
+		/// The change number being built
+		/// </summary>
+		public int Change { get; set; }
+
+		/// <summary>
+		/// The step log id
+		/// </summary>
+		public string? LogId { get; set; }
+
+		/// <summary>
+		/// The pool id
+		/// </summary>
+		public string? PoolId { get; set; }
+
+		/// <summary>
+		/// The agent id
+		/// </summary>
+		public string? AgentId { get; set; }
+
+		/// <summary>
+		/// Outcome of the step, once complete.
+		/// </summary>
+		public JobStepOutcome? Outcome { get; set; }
+
+		/// <summary>
+		/// Time at which the step started.
+		/// </summary>
+		public DateTimeOffset StartTime { get; }
+
+		/// <summary>
+		/// Time at which the step finished.
+		/// </summary>
+		public DateTimeOffset? FinishTime { get; }
+
+		/// <summary>
+		/// Constructor
+		/// </summary>
+		/// <param name="JobStepRef">The jobstep ref to construct from</param>
+		public GetJobStepRefResponse(IJobStepRef JobStepRef)
+		{
+			this.JobId = JobStepRef.Id.JobId.ToString();
+			this.BatchId = JobStepRef.Id.BatchId.ToString();
+			this.StepId = JobStepRef.Id.StepId.ToString();
+			this.Change = JobStepRef.Change;
+			this.LogId = JobStepRef.LogId.ToString();
+			this.PoolId = JobStepRef.PoolId?.ToString();
+			this.AgentId = JobStepRef.AgentId?.ToString();
+			this.Outcome = JobStepRef.Outcome;
+			this.StartTime = JobStepRef.StartTimeUtc;
+			this.FinishTime = JobStepRef.FinishTimeUtc;
+		}
+	}
+}
diff --git a/Engine/Source/Programs/Horde/HordeServer/Api/LogEvents.cs b/Engine/Source/Programs/Horde/HordeServer/Api/LogEvents.cs
new file mode 100644
--- /dev/null
+++ b/Engine/Source/Programs/Horde/HordeServer/Api/LogEvents.cs
@@ -0,0 +1,73 @@
+using EpicGames.Core;
+using HordeCommon;
+using HordeServer.Models;
+using Microsoft.Extensions.Logging;
+using System;
+using System.Collections.Generic;
+using System.Linq;
+using System.Text.Json;
+using System.Threading.Tasks;
+
+namespace HordeServer.Api
+{
+	/// <summary>
+	/// Information about an uploaded event
+	/// </summary>
+	public class GetLogEventResponse
+	{
+		/// <summary>
+		/// Unique id of the log containing this event
+		/// </summary>
+		public string LogId { get; set; }
+
+		/// <summary>
+		/// Severity of this event
+		/// </summary>
+		public EventSeverity Severity { get; set; }
+
+		/// <summary>
+		/// Index of the first line for this event
+		/// </summary>
+		public int LineIndex { get; set; }
+
+		/// <summary>
+		/// Number of lines in the event
+		/// </summary>
+		public int LineCount { get; set; }
+
+		/// <summary>
+		/// The issue id associated with this event
+		/// </summary>
+		public int? IssueId { get; set; }
+
+		/// <summary>
+		/// The structured message data for this event
+		/// </summary>
+		public List<JsonElement> Lines { get; set; }
+
+		/// <summary>
+		/// Private constructor for serialization
+		/// </summary>
+		private GetLogEventResponse()
+		{
+			LogId = null!;
+			Lines = new List<JsonElement>();
+		}
+
+		/// <summary>
+		/// Constructor
+		/// </summary>
+		/// <param name="Event">The event to construct from</param>
+		/// <param name="EventData">The event data</param>
+		/// <param name="IssueId">The issue for this event</param>
+		public GetLogEventResponse(ILogEvent Event, ILogEventData EventData, int? IssueId)
+		{
+			this.Severity = Event.Severity;
+			this.LogId = Event.LogId.ToString();
+			this.LineIndex = Event.LineIndex;
+			this.LineCount = Event.LineCount;
+			this.IssueId = IssueId;
+			this.Lines = EventData.Lines.ConvertAll(x => x.Data);
+		}
+	}
+}
diff --git a/Engine/Source/Programs/Horde/HordeServer/Api/LogFiles.cs b/Engine/Source/Programs/Horde/HordeServer/Api/LogFiles.cs
new file mode 100644
--- /dev/null
+++ b/Engine/Source/Programs/Horde/HordeServer/Api/LogFiles.cs
@@ -0,0 +1,156 @@
+// Copyright Epic Games, Inc. All Rights Reserved.
+
+using HordeServer.Logs;
+using HordeServer.Models;
+using HordeServer.Services;
+using System;
+using System.Collections.Generic;
+using System.ComponentModel.DataAnnotations;
+using System.Linq;
+using System.Text.Json.Serialization;
+using System.Threading.Tasks;
+
+namespace HordeServer.Api
+{
+	/// <summary>
+	/// The type of data stored in this log file
+	/// </summary>
+	public enum LogType
+	{
+		/// <summary>
+		/// Plain text data
+		/// </summary>
+		Text,
+
+		/// <summary>
+		/// Structured json objects, output as one object per line (without trailing commas)
+		/// </summary>
+		Json
+	}
+
+	/// <summary>
+	/// Parameters required to create log file
+	/// </summary>
+	public class CreateLogFileRequest
+	{
+		/// <summary>
+		/// Job Id this log file belongs to
+		/// </summary>
+		[Required]
+		public string JobId { get; set; }
+
+		/// <summary>
+		/// The type of events stored in this log
+		/// </summary>
+		public LogType Type { get; set; }
+
+		/// <summary>
+		/// Private constructor for serialization
+		/// </summary>
+		private CreateLogFileRequest()
+		{
+			JobId = null!;
+		}
+
+		/// <summary>
+		/// Constructor
+		/// </summary>
+		/// <param name="JobId">Unique job id that this log file belongs to</param>
+		/// <param name="Type">Type of events to be stored in this file</param>
+		public CreateLogFileRequest(string JobId, LogType Type)
+		{
+			this.JobId = JobId;
+			this.Type = Type;
+		}
+	}
+
+	/// <summary>
+	/// Response from creating a log file
+	/// </summary>
+	public class CreateLogFileResponse
+	{
+		/// <summary>
+		/// Unique id for this log file
+		/// </summary>
+		public string Id { get; set; }
+
+		/// <summary>
+		/// Private constructor for serialization
+		/// </summary>
+		private CreateLogFileResponse()
+		{
+			Id = null!;
+		}
+
+		/// <summary>
+		/// Constructor
+		/// </summary>
+		/// <param name="Id">The log file id</param>
+		public CreateLogFileResponse(string Id)
+		{
+			this.Id = Id;
+		}
+	}
+
+	/// <summary>
+	/// Response describing a log file
+	/// </summary>
+	public class GetLogFileResponse
+	{
+        /// <summary>
+		/// Unique id of the log file
+		/// </summary>
+		public string Id { get; set; }
+
+        /// <summary>
+        /// Unique id of the job for this log file
+		/// </summary>
+        public string JobId { get; set; }
+
+		/// <summary>
+		/// Type of events stored in this log
+		/// </summary>
+		public LogType Type { get; set; }
+
+		/// <summary>
+        /// Length of the log, in bytes
+        /// </summary>
+        public long Length { get; set; }
+
+		/// <summary>
+        /// Number of lines in the file
+        /// </summary>
+        public int LineCount { get; set; }
+
+		/// <summary>
+		/// Constructor
+		/// </summary>
+		/// <param name="LogFile">The logfile to construct from</param>
+		/// <param name="Metadata">Metadata about the log file</param>
+		public GetLogFileResponse(ILogFile LogFile, LogMetadata Metadata)
+		{
+            this.Id = LogFile.Id.ToString();
+            this.JobId = LogFile.JobId.ToString();
+			this.Type = LogFile.Type;
+            this.Length = Metadata.Length;
+            this.LineCount = Metadata.MaxLineIndex;
+		}
+	}
+
+	/// <summary>
+	/// Response describing a log file
+	/// </summary>
+	public class SearchLogFileResponse
+	{
+		/// <summary>
+		/// List of line numbers containing the search text
+		/// </summary>
+		public List<int> Lines { get; set; } = new List<int>();
+
+		/// <summary>
+		/// Stats for the search
+		/// </summary>
+		public LogSearchStats? Stats { get; set; }
+	}
+}
+
diff --git a/Engine/Source/Programs/Horde/HordeServer/Api/Nodes.cs b/Engine/Source/Programs/Horde/HordeServer/Api/Nodes.cs
new file mode 100644
--- /dev/null
+++ b/Engine/Source/Programs/Horde/HordeServer/Api/Nodes.cs
@@ -0,0 +1,466 @@
+// Copyright Epic Games, Inc. All Rights Reserved.
+
+using EpicGames.Core;
+using HordeCommon;
+using HordeServer.Models;
+using HordeServer.Utilities;
+using System;
+using System.Collections.Generic;
+using System.ComponentModel.DataAnnotations;
+using System.Linq;
+using System.Threading.Tasks;
+
+namespace HordeServer.Api
+{
+	/// <summary>
+	/// Information required to create a node
+	/// </summary>
+	public class CreateNodeRequest
+	{
+		/// <summary>
+		/// The name of this node 
+		/// </summary>
+		[Required]
+		public string Name { get; set; } = null!;
+
+		/// <summary>
+		/// List of nodes which must succeed for this node to run
+		/// </summary>
+		public List<string>? InputDependencies { get; set; }
+
+		/// <summary>
+		/// List of nodes which must have completed for this node to run
+		/// </summary>
+		public List<string>? OrderDependencies { get; set; }
+
+		/// <summary>
+		/// The priority of this node
+		/// </summary>
+		public Priority? Priority { get; set; }
+
+		/// <summary>
+		/// This node can be run multiple times
+		/// </summary>
+		public bool? AllowRetry { get; set; }
+
+		/// <summary>
+		/// This node can start running early, before dependencies of other nodes in the same group are complete
+		/// </summary>
+		public bool? RunEarly { get; set; }
+
+		/// <summary>
+		/// Whether to include warnings in the diagnostic output
+		/// </summary>
+		public bool? Warnings { get; set; }
+
+		/// <summary>
+		/// Credentials required for this node to run. This dictionary maps from environment variable names to a credential property in the format 'CredentialName.PropertyName'.
+		/// </summary>
+		public Dictionary<string, string>? Credentials { get; set; }
+
+		/// <summary>
+		/// Properties for this node
+		/// </summary>
+		public Dictionary<string, string>? Properties { get; set; }
+
+		/// <summary>
+		/// Private constructor for serialization
+		/// </summary>
+		private CreateNodeRequest()
+		{
+			this.Name = null!;
+		}
+
+		/// <summary>
+		/// Constructor
+		/// </summary>
+		/// <param name="Name">Name of the node</param>
+		/// <param name="InputDependencies">List of nodes which must have completed succesfully for this node to run</param>
+		/// <param name="OrderDependencies">List of nodes which must have completed for this node to run</param>
+		/// <param name="Priority">Priority of this node</param>
+		/// <param name="AllowRetry">Whether the node can be run multiple times</param>
+		/// <param name="RunEarly">Whether the node can run early, before dependencies of other nodes in the same group complete</param>
+		/// <param name="Warnings">Whether to include warnings in the diagnostic output (defaults to true)</param>
+		/// <param name="Credentials">Credentials required for this node to run</param>
+		/// <param name="Properties">Properties for the node</param>
+		public CreateNodeRequest(string Name, List<string>? InputDependencies = null, List<string>? OrderDependencies = null, Priority? Priority = null, bool? AllowRetry = null, bool? RunEarly = null, bool? Warnings = null, Dictionary<string, string>? Credentials = null, Dictionary<string, string>? Properties = null)
+		{
+			this.Name = Name;
+			this.InputDependencies = InputDependencies;
+			this.OrderDependencies = OrderDependencies;
+			this.Priority = Priority;
+			this.AllowRetry = AllowRetry;
+			this.RunEarly = RunEarly;
+			this.Warnings = Warnings;
+			this.Credentials = Credentials;
+			this.Properties = Properties;
+		}
+	}
+
+	/// <summary>
+	/// Information about a group of nodes
+	/// </summary>
+	public class CreateGroupRequest
+	{
+		/// <summary>
+		/// The type of agent to execute this group
+		/// </summary>
+		[Required]
+		public string AgentType { get; set; }
+
+		/// <summary>
+		/// Nodes in the group
+		/// </summary>
+		[Required]
+		public List<CreateNodeRequest> Nodes { get; set; }
+
+		/// <summary>
+		/// Private constructor for serialization
+		/// </summary>
+		private CreateGroupRequest()
+		{
+			AgentType = null!;
+			Nodes = new List<CreateNodeRequest>();
+		}
+
+		/// <summary>
+		/// Constructor
+		/// </summary>
+		/// <param name="AgentType">The type of agent to execute this group</param>
+		/// <param name="Nodes">Nodes in this group</param>
+		public CreateGroupRequest(string AgentType, List<CreateNodeRequest> Nodes)
+		{
+			this.AgentType = AgentType;
+			this.Nodes = Nodes;
+		}
+	}
+
+	/// <summary>
+	/// Information about a group of nodes
+	/// </summary>
+	public class CreateLabelRequest
+	{
+		/// <summary>
+		/// Category for this label
+		/// </summary>
+		[Obsolete("Use DashboardCategory instead")]
+		public string? Category => DashboardCategory;
+
+		/// <summary>
+		/// Name of the aggregate
+		/// </summary>
+		[Obsolete("Use DashboardName instead")]
+		public string? Name => DashboardName;
+
+		/// <summary>
+		/// Name of the aggregate
+		/// </summary>
+		public string? DashboardName { get; set; }
+
+		/// <summary>
+		/// Category for this label
+		/// </summary>
+		public string? DashboardCategory { get; set; }
+
+		/// <summary>
+		/// Name of the badge in UGS
+		/// </summary>
+		public string? UgsName { get; set; }
+
+		/// <summary>
+		/// Project to show this label for in UGS
+		/// </summary>
+		public string? UgsProject { get; set; }
+
+		/// <summary>
+		/// Which change the label applies to
+		/// </summary>
+		public LabelChange Change { get; set; }
+
+		/// <summary>
+		/// Nodes which must be part of the job for the aggregate to be valid
+		/// </summary>
+		public List<string> RequiredNodes { get; set; } = new List<string>();
+
+		/// <summary>
+		/// Nodes which must be part of the job for the aggregate to be valid
+		/// </summary>
+		public List<string> IncludedNodes { get; set; } = new List<string>();
+
+		/// <summary>
+		/// Private constructor for serialization
+		/// </summary>
+		private CreateLabelRequest()
+		{
+		}
+
+		/// <summary>
+		/// Constructor
+		/// </summary>
+		/// <param name="DashboardName">Name of this label</param>
+		/// <param name="DashboardCategory">Category of this label</param>
+		/// <param name="UgsName">Name of this label for UGS</param>
+		/// <param name="UgsProject">Name of the project to show this label for in UGS</param>
+		/// <param name="Change">Which change the label applies to</param>
+		/// <param name="RequiredNodes">Nodes which must be part of the job for the aggregate to be shown</param>
+		/// <param name="IncludedNodes">Nodes to include in the status of this aggregate, if present in the job</param>
+		public CreateLabelRequest(string? DashboardName, string? DashboardCategory, string? UgsName, string? UgsProject, LabelChange Change, List<string> RequiredNodes, List<string> IncludedNodes)
+		{
+			this.DashboardName = String.IsNullOrEmpty(DashboardName)? null : DashboardName;
+			this.DashboardCategory = String.IsNullOrEmpty(DashboardCategory)? null : DashboardCategory;
+			this.UgsName = String.IsNullOrEmpty(UgsName)? null : UgsName;
+			this.UgsProject = String.IsNullOrEmpty(UgsProject)? null : UgsProject;
+			this.Change = Change;
+			this.RequiredNodes = RequiredNodes;
+			this.IncludedNodes = IncludedNodes;
+		}
+	}
+
+	/// <summary>
+	/// Information required to create a node
+	/// </summary>
+	public class GetNodeResponse
+	{
+		/// <summary>
+		/// The name of this node 
+		/// </summary>
+		public string Name { get; set; }
+
+		/// <summary>
+		/// Indices of nodes which must have succeeded for this node to run
+		/// </summary>
+		public List<string> InputDependencies { get; set; }
+
+		/// <summary>
+		/// Indices of nodes which must have completed for this node to run
+		/// </summary>
+		public List<string> OrderDependencies { get; set; }
+
+		/// <summary>
+		/// The priority of this node
+		/// </summary>
+		public Priority Priority { get; set; }
+
+		/// <summary>
+		/// Whether this node can be retried
+		/// </summary>
+		public bool AllowRetry { get; set; }
+
+		/// <summary>
+		/// This node can start running early, before dependencies of other nodes in the same group are complete
+		/// </summary>
+		public bool RunEarly { get; set; }
+
+		/// <summary>
+		/// Whether to include warnings in diagnostic output
+		/// </summary>
+		public bool Warnings { get; set; }
+
+		/// <summary>
+		/// Average time to execute this node based on historical trends
+		/// </summary>
+		public float? AverageDuration { get; set; }
+
+		/// <summary>
+		/// Credentials required for this node to run. This dictionary maps from environment variable names to a credential property in the format 'CredentialName.PropertyName'.
+		/// </summary>
+		public IReadOnlyDictionary<string, string>? Credentials { get; set; }
+
+		/// <summary>
+		/// Properties for this node
+		/// </summary>
+		public IReadOnlyDictionary<string, string>? Properties { get; set; }
+
+		/// <summary>
+		/// Parameterless constructor for serialization
+		/// </summary>
+		private GetNodeResponse()
+		{
+			this.Name = null!;
+			this.InputDependencies = new List<string>();
+			this.OrderDependencies = new List<string>();
+			this.Credentials = new Dictionary<string, string>();
+			this.Properties = new Dictionary<string, string>();
+		}
+
+		/// <summary>
+		/// Constructor
+		/// </summary>
+		/// <param name="Node">The node to construct from</param>
+		/// <param name="Groups">The groups in the graph</param>
+		public GetNodeResponse(INode Node, IReadOnlyList<INodeGroup> Groups)
+		{
+			this.Name = Node.Name;
+			this.InputDependencies = new List<string>(Node.InputDependencies.Select(x => Groups[x.GroupIdx].Nodes[x.NodeIdx].Name)); ;
+			this.OrderDependencies = new List<string>(Node.OrderDependencies.Select(x => Groups[x.GroupIdx].Nodes[x.NodeIdx].Name)); ;
+			this.Priority = Node.Priority;
+			this.AllowRetry = Node.AllowRetry;
+			this.RunEarly = Node.RunEarly;
+			this.Warnings = Node.Warnings;
+			this.Credentials = Node.Credentials;
+			this.Properties = Node.Properties;
+		}
+	}
+
+	/// <summary>
+	/// Information about a group of nodes
+	/// </summary>
+	public class GetGroupResponse
+	{
+		/// <summary>
+		/// The type of agent to execute this group
+		/// </summary>
+		public string AgentType { get; set; }
+
+		/// <summary>
+		/// Nodes in the group
+		/// </summary>
+		public List<GetNodeResponse> Nodes { get; set; }
+
+		/// <summary>
+		/// Constructor
+		/// </summary>
+		/// <param name="Group">The group to construct from</param>
+		/// <param name="Groups">Other groups in this graph</param>
+		public GetGroupResponse(INodeGroup Group, IReadOnlyList<INodeGroup> Groups)
+		{
+			this.AgentType = Group.AgentType;
+			this.Nodes = Group.Nodes.ConvertAll(x => new GetNodeResponse(x, Groups));
+		}
+	}
+
+	/// <summary>
+	/// Information about an aggregate
+	/// </summary>
+	public class GetAggregateResponse
+	{
+		/// <summary>
+		/// Name of the aggregate
+		/// </summary>
+		public string Name { get; set; }
+
+		/// <summary>
+		/// Nodes which must be part of the job for the aggregate to be shown
+		/// </summary>
+		public List<string> Nodes { get; set; }
+
+		/// <summary>
+		/// Constructor
+		/// </summary>
+		/// <param name="Aggregate">The aggregate to construct from</param>
+		/// <param name="Groups">List of groups in this graph</param>
+		public GetAggregateResponse(IAggregate Aggregate, IReadOnlyList<INodeGroup> Groups)
+		{
+			Name = Aggregate.Name;
+			Nodes = Aggregate.Nodes.ConvertAll(x => x.ToNode(Groups).Name);
+		}
+	}
+
+	/// <summary>
+	/// Information about a label
+	/// </summary>
+	public class GetLabelResponse
+	{
+		/// <summary>
+		/// Category of the aggregate
+		/// </summary>
+		public string? Category => DashboardCategory;
+
+		/// <summary>
+		/// Label for this aggregate
+		/// </summary>
+		public string? Name => DashboardName;
+
+		/// <summary>
+		/// Name to show for this label on the dashboard
+		/// </summary>
+		public string? DashboardName { get; set; }
+
+		/// <summary>
+		/// Category to show this label in on the dashboard
+		/// </summary>
+		public string? DashboardCategory { get; set; }
+
+		/// <summary>
+		/// Name to show for this label in UGS
+		/// </summary>
+		public string? UgsName { get; set; }
+
+		/// <summary>
+		/// Project to display this label for in UGS
+		/// </summary>
+		public string? UgsProject { get; set; }
+
+		/// <summary>
+		/// Nodes which must be part of the job for the aggregate to be shown
+		/// </summary>
+		public List<string> RequiredNodes { get; set; }
+
+		/// <summary>
+		/// Nodes to include in the status of this aggregate, if present in the job
+		/// </summary>
+		public List<string> IncludedNodes { get; set; }
+
+		/// <summary>
+		/// Constructor
+		/// </summary>
+		/// <param name="Label">The label to construct from</param>
+		/// <param name="Groups">List of groups in this graph</param>
+		public GetLabelResponse(ILabel Label, IReadOnlyList<INodeGroup> Groups)
+		{
+			DashboardName = Label.DashboardName;
+			DashboardCategory = Label.DashboardCategory;
+			UgsName = Label.UgsName;
+			UgsProject = Label.UgsProject;
+			RequiredNodes = Label.RequiredNodes.ConvertAll(x => x.ToNode(Groups).Name);
+			IncludedNodes = Label.IncludedNodes.ConvertAll(x => x.ToNode(Groups).Name);
+		}
+	}
+
+	/// <summary>
+	/// Information about a graph
+	/// </summary>
+	public class GetGraphResponse
+	{
+		/// <summary>
+		/// The hash of the graph
+		/// </summary>
+		public string Hash { get; set; }
+
+		/// <summary>
+		/// Array of nodes for this job
+		/// </summary>
+		public List<GetGroupResponse>? Groups { get; set; }
+
+		/// <summary>
+		/// List of aggregates
+		/// </summary>
+		public List<GetAggregateResponse>? Aggregates { get; set; }
+
+		/// <summary>
+		/// List of labels for the graph
+		/// </summary>
+		public List<GetLabelResponse>? Labels { get; set; }
+
+		/// <summary>
+		/// Constructor
+		/// </summary>
+		/// <param name="Graph">The graph to construct from</param>
+		public GetGraphResponse(IGraph Graph)
+		{
+			Hash = Graph.Id.ToString();
+			if (Graph.Groups.Count > 0)
+			{
+				Groups = Graph.Groups.ConvertAll(x => new GetGroupResponse(x, Graph.Groups));
+			}
+			if (Graph.Aggregates.Count > 0)
+			{
+				Aggregates = Graph.Aggregates.ConvertAll(x => new GetAggregateResponse(x, Graph.Groups));
+			}
+			if (Graph.Labels.Count > 0)
+			{
+				Labels = Graph.Labels.ConvertAll(x => new GetLabelResponse(x, Graph.Groups));
+			}
+		}
+	}
+}
diff --git a/Engine/Source/Programs/Horde/HordeServer/Api/Notifications.cs b/Engine/Source/Programs/Horde/HordeServer/Api/Notifications.cs
new file mode 100644
--- /dev/null
+++ b/Engine/Source/Programs/Horde/HordeServer/Api/Notifications.cs
@@ -0,0 +1,57 @@
+// Copyright Epic Games, Inc. All Rights Reserved.
+
+using System.ComponentModel.DataAnnotations;
+using HordeServer.Models;
+
+namespace HordeServer.Api
+{
+	/// <summary>
+	/// A request from a user to be notified when something happens.
+	/// </summary>
+	public class UpdateNotificationsRequest
+	{
+		/// <summary>
+		/// Request notifications by email
+		/// </summary>
+		public bool? Email { get; set; }
+
+		/// <summary>
+		/// Request notifications on Slack
+		/// </summary>
+		public bool? Slack { get; set; }
+	}
+
+	/// <summary>
+	/// A request from a user to be notified when something happens.
+	/// </summary>
+	public class GetNotificationResponse
+	{
+		/// <summary>
+		/// Request notifications by email
+		/// </summary>
+		public bool Email { get; set; }
+
+		/// <summary>
+		/// Request notifications on Slack
+		/// </summary>
+		public bool Slack { get; set; }
+
+		/// <summary>
+		/// Constructor
+		/// </summary>
+		/// <param name="Subscription">The subscription interface</param>
+		internal GetNotificationResponse(INotificationSubscription? Subscription)
+		{
+			if (Subscription == null)
+			{
+				Email = false;
+				Slack = false;
+			}
+			else
+			{
+				Email = Subscription.Email;
+				Slack = Subscription.Slack;
+			}
+		}
+	}
+}
diff --git a/Engine/Source/Programs/Horde/HordeServer/Api/Pools.cs b/Engine/Source/Programs/Horde/HordeServer/Api/Pools.cs
new file mode 100644
--- /dev/null
+++ b/Engine/Source/Programs/Horde/HordeServer/Api/Pools.cs
@@ -0,0 +1,207 @@
+// Copyright Epic Games, Inc. All Rights Reserved.
+
+using HordeServer.Models;
+using System;
+using System.Collections.Generic;
+using System.ComponentModel.DataAnnotations;
+using System.Linq;
+using System.Threading.Tasks;
+
+namespace HordeServer.Api
+{
+	/// <summary>
+	/// Parameters to create a new pool
+	/// </summary>
+	public class CreatePoolRequest
+	{
+		/// <summary>
+		/// Name for the new pool
+		/// </summary>
+		[Required]
+		public string Name { get; set; } = null!;
+
+		/// <summary>
+		/// Agent requirements to automatically include in this pool
+		/// </summary>
+		public AgentRequirements? Requirements { get; set; }
+
+		/// <summary>
+		/// Whether to enable autoscaling for this pool
+		/// </summary>
+		public bool? EnableAutoscaling { get; set; }
+
+		/// <summary>
+		/// The minimum nunmber of agents to retain in this pool
+		/// </summary>
+		public int? MinAgents { get; set; }
+
+		/// <summary>
+		/// The minimum number of idle agents in this pool, if autoscaling is enabled
+		/// </summary>
+		public int? NumReserveAgents { get; set; }
+
+		/// <summary>
+		/// Properties for the new pool
+		/// </summary>
+		public Dictionary<string, string>? Properties { get; set; }
+	}
+
+	/// <summary>
+	/// Response from creating a new pool
+	/// </summary>
+	public class CreatePoolResponse
+	{
+		/// <summary>
+		/// Unique id for the new pool
+		/// </summary>
+		public string Id { get; set; }
+
+		/// <summary>
+		/// Private constructor for serialization
+		/// </summary>
+		private CreatePoolResponse()
+		{
+			Id = null!;
+		}
+
+		/// <summary>
+		/// Constructor
+		/// </summary>
+		/// <param name="Id">Unique id for the new pool</param>
+		public CreatePoolResponse(string Id)
+		{
+			this.Id = Id;
+		}
+	}
+
+	/// <summary>
+	/// Parameters to update a pool
+	/// </summary>
+	public class UpdatePoolRequest
+	{
+		/// <summary>
+		/// Optional new name for the pool
+		/// </summary>
+		public string? Name { get; set; }
+
+		/// <summary>
+		/// Requirements for this pool
+		/// </summary>
+		public AgentRequirements? Requirements { get; set; }
+
+		/// <summary>
+		/// Whether to enable autoscaling for this pool
+		/// </summary>
+		public bool? EnableAutoscaling { get; set; }
+
+		/// <summary>
+		/// The minimum nunmber of agents to retain in this pool
+		/// </summary>
+		public int? MinAgents { get; set; }
+
+		/// <summary>
+		/// The minimum number of idle agents in this pool, if autoscaling is enabled
+		/// </summary>
+		public int? NumReserveAgents { get; set; }
+
+		/// <summary>
+		/// Properties to update for the pool. Properties set to null will be removed.
+		/// </summary>
+		public Dictionary<string, string?>? Properties { get; set; }
+	}
+
+	/// <summary>
+	/// Parameters to update a pool
+	/// </summary>
+	public class BatchUpdatePoolRequest
+	{
+		/// <summary>
+		///  ID of the pool to update
+		/// </summary>
+		[Required]
+		public string Id { get; set; } = null!;
+
+		/// <summary>
+		/// Optional new name for the pool
+		/// </summary>
+		public string? Name { get; set; }
+
+		/// <summary>
+		/// Properties to update for the pool. Properties set to null will be removed.
+		/// </summary>
+		public Dictionary<string, string?>? Properties { get; set; }
+	}
+
+	/// <summary>
+	/// Response describing a pool
+	/// </summary>
+	public class GetPoolResponse
+	{
+		/// <summary>
+		/// Unique id of the pool
+		/// </summary>
+		public string Id { get; set; }
+
+		/// <summary>
+		/// Name of the pool
+		/// </summary>
+		public string Name { get; set; }
+
+		/// <summary>
+		/// Requirements for agents to be auto-added to the pool
+		/// </summary>
+		public AgentRequirements? Requirements { get; set; }
+
+		/// <summary>
+		/// Whether to enable autoscaling for this pool
+		/// </summary>
+		public bool EnableAutoscaling { get; set; }
+
+		/// <summary>
+		/// The minimum nunmber of agents to retain in this pool
+		/// </summary>
+		public int? MinAgents { get; set; }
+
+		/// <summary>
+		/// The minimum number of idle agents in this pool, if autoscaling is enabled
+		/// </summary>
+		public int? NumReserveAgents { get; set; }
+
+		/// <summary>
+		/// List of workspaces that this agent contains
+		/// </summary>
+		public List<GetAgentWorkspaceResponse> Workspaces { get; set; }
+
+		/// <summary>
+		/// Arbitrary properties for this pool.
+		/// </summary>
+		public IReadOnlyDictionary<string, string> Properties { get; set; }
+
+		/// <summary>
+		/// Parameterless constructor for serialization
+		/// </summary>
+		private GetPoolResponse()
+		{
+			Id = null!;
+			Name = null!;
+			Workspaces = null!;
+			Properties = null!;
+		}
+
+		/// <summary>
+		/// Constructor
+		/// </summary>
+		/// <param name="Pool">The pool to construct from</param>
+		public GetPoolResponse(IPool Pool)
+		{
+			this.Id = Pool.Id.ToString();
+			this.Name = Pool.Name;
+			this.Requirements = Pool.Requirements;
+			this.EnableAutoscaling = Pool.EnableAutoscaling;
+			this.MinAgents = Pool.MinAgents;
+			this.NumReserveAgents = Pool.NumReserveAgents;
+			this.Workspaces = Pool.Workspaces.Select(x => new GetAgentWorkspaceResponse(x)).ToList();
+			this.Properties = Pool.Properties;
+		}
+	}
+}
diff --git a/Engine/Source/Programs/Horde/HordeServer/Api/Projects.cs b/Engine/Source/Programs/Horde/HordeServer/Api/Projects.cs
new file mode 100644
--- /dev/null
+++ b/Engine/Source/Programs/Horde/HordeServer/Api/Projects.cs
@@ -0,0 +1,287 @@
+// Copyright Epic Games, Inc. All Rights Reserved.
+
+using HordeServer.Models;
+using System;
+using System.Collections.Generic;
+using System.ComponentModel.DataAnnotations;
+using System.Linq;
+using System.Threading.Tasks;
+
+namespace HordeServer.Api
+{
+	/// <summary>
+	/// Information about a category to display for a stream
+	/// </summary>
+	public class CreateProjectCategoryRequest
+	{
+		/// <summary>
+		/// Name of this category
+		/// </summary>
+		[Required]
+		public string Name { get; set; }
+
+		/// <summary>
+		/// Index of the row to display this category on
+		/// </summary>
+		public int Row { get; set; }
+
+		/// <summary>
+		/// Whether to show this category on the nav menu
+		/// </summary>
+		public bool ShowOnNavMenu { get; set; }
+
+		/// <summary>
+		/// Patterns for stream names to include
+		/// </summary>
+		public List<string> IncludePatterns { get; set; } = new List<string>();
+
+		/// <summary>
+		/// Patterns for stream names to exclude
+		/// </summary>
+		public List<string> ExcludePatterns { get; set; } = new List<string>();
+
+		/// <summary>
+		/// Private constructor for serialization
+		/// </summary>
+		private CreateProjectCategoryRequest()
+		{
+			Name = null!;
+		}
+	}
+
+	/// <summary>
+	/// Parameters to create a new project
+	/// </summary>
+	public class CreateProjectRequest
+	{
+		/// <summary>
+		/// Name for the new project
+		/// </summary>
+		[Required]
+		public string Name { get; set; } = null!;
+
+		/// <summary>
+		/// Order to display this project on the dashboard
+		/// </summary>
+		public int? Order { get; set; }
+
+		/// <summary>
+		/// Categories to include in this project
+		/// </summary>
+		public List<CreateProjectCategoryRequest>? Categories { get; set; }
+
+		/// <summary>
+		/// Properties for the new project
+		/// </summary>
+		public Dictionary<string, string>? Properties { get; set; }
+	}
+
+	/// <summary>
+	/// Response from creating a new project
+	/// </summary>
+	public class CreateProjectResponse
+	{
+		/// <summary>
+		/// Unique id for the new project
+		/// </summary>
+		public string Id { get; set; }
+
+		/// <summary>
+		/// Private constructor for serialization
+		/// </summary>
+		private CreateProjectResponse()
+		{
+			Id = null!;
+		}
+
+		/// <summary>
+		/// Constructor
+		/// </summary>
+		/// <param name="Id">Unique id for the new project</param>
+		public CreateProjectResponse(string Id)
+		{
+			this.Id = Id;
+		}
+	}
+
+	/// <summary>
+	/// Parameters to update a project
+	/// </summary>
+	public class UpdateProjectRequest
+	{
+		/// <summary>
+		/// Optional new name for the project
+		/// </summary>
+		public string? Name { get; set; }
+
+		/// <summary>
+		/// Order to display this project on the dashboard
+		/// </summary>
+		public int? Order { get; set; }
+
+		/// <summary>
+		/// New set of categories for the request
+		/// </summary>
+		public List<CreateProjectCategoryRequest>? Categories { get; set; }
+
+		/// <summary>
+		/// Properties to update for the project. Properties set to null will be removed.
+		/// </summary>
+		public Dictionary<string, string>? Properties { get; set; }
+
+		/// <summary>
+		/// Custom permissions for this object
+		/// </summary>
+		public UpdateAclRequest? Acl { get; set; }
+	}
+
+	/// <summary>
+	/// Information about a stream within a project
+	/// </summary>
+	public class GetProjectStreamResponse
+	{
+		/// <summary>
+		/// The stream id
+		/// </summary>
+		public string Id { get; set; }
+
+		/// <summary>
+		/// The stream name
+		/// </summary>
+		public string Name { get; set; }
+
+		/// <summary>
+		/// Constructor
+		/// </summary>
+		/// <param name="Id">The unique stream id</param>
+		/// <param name="Name">The stream name</param>
+		public GetProjectStreamResponse(string Id, string Name)
+		{
+			this.Id = Id;
+			this.Name = Name;
+		}
+	}
+
+	/// <summary>
+	/// Information about a category to display for a stream
+	/// </summary>
+	public class GetProjectCategoryResponse
+	{
+		/// <summary>
+		/// Heading for this column
+		/// </summary>
+		public string Name { get; set; }
+
+		/// <summary>
+		/// Index of the row to display this category on
+		/// </summary>
+		public int Row { get; set; }
+
+		/// <summary>
+		/// Whether to show this category on the nav menu
+		/// </summary>
+		public bool ShowOnNavMenu { get; set; }
+
+		/// <summary>
+		/// Patterns for stream names to include
+		/// </summary>
+		public List<string> IncludePatterns { get; set; } = new List<string>();
+
+		/// <summary>
+		/// Patterns for stream names to exclude
+		/// </summary>
+		public List<string> ExcludePatterns { get; set; } = new List<string>();
+
+		/// <summary>
+		/// Streams to include in this category
+		/// </summary>
+		public List<string> Streams { get; set; } = new List<string>();
+
+		/// <summary>
+		/// Constructor
+		/// </summary>
+		/// <param name="StreamCategory">The category to construct from</param>
+		public GetProjectCategoryResponse(StreamCategory StreamCategory)
+		{
+			this.Name = StreamCategory.Name;
+			this.Row = StreamCategory.Row;
+			this.ShowOnNavMenu = StreamCategory.ShowOnNavMenu;
+			this.IncludePatterns = StreamCategory.IncludePatterns;
+			this.ExcludePatterns = StreamCategory.ExcludePatterns;
+		}
+	}
+
+	/// <summary>
+	/// Response describing a project
+	/// </summary>
+	public class GetProjectResponse
+	{
+		/// <summary>
+		/// Unique id of the project
+		/// </summary>
+		public string Id { get; set; }
+
+		/// <summary>
+		/// Name of the project
+		/// </summary>
+		public string Name { get; set; }
+
+		/// <summary>
+		/// Order to display this project on the dashboard
+		/// </summary>
+		public int Order { get; set; }
+
+		/// <summary>
+		/// List of streams that are in this project
+		/// </summary>
+		public List<GetProjectStreamResponse>? Streams { get; set; }
+
+		/// <summary>
+		/// List of stream categories to display
+		/// </summary>
+		public List<GetProjectCategoryResponse>? Categories { get; set; }
+
+		/// <summary>
+		/// Properties for this project.
+		/// </summary>
+		public Dictionary<string, string> Properties { get; set; }
+
+		/// <summary>
+		/// Custom permissions for this object
+		/// </summary>
+		public GetAclResponse? Acl { get; set; }
+
+		/// <summary>
+		/// Parameterless constructor for serialization
+		/// </summary>
+		private GetProjectResponse()
+		{
+			Id = null!;
+			Name = null!;
+			Properties = null!;
+			Streams = null!;
+			Categories = null!;
+		}
+
+		/// <summary>
+		/// Constructor
+		/// </summary>
+		/// <param name="Id">Unique id of the project</param>
+		/// <param name="Name">Name of the project</param>
+		/// <param name="Order">Order to show this project on the dashboard</param>
+		/// <param name="Streams">List of streams to display</param>
+		/// <param name="Categories">List of stream categories to display</param>
+		/// <param name="Properties">Properties for this project</param>
+		/// <param name="Acl">Custom permissions for this object</param>
+		public GetProjectResponse(string Id, string Name, int Order, List<GetProjectStreamResponse>? Streams, List<GetProjectCategoryResponse>? Categories, Dictionary<string, string> Properties, GetAclResponse? Acl)
+		{
+			this.Id = Id;
+			this.Name = Name;
+			this.Order = Order;
+			this.Streams = Streams;
+			this.Categories = Categories;
+			this.Properties = Properties;
+			this.Acl = Acl;
+		}
+	}
+}
diff --git a/Engine/Source/Programs/Horde/HordeServer/Api/Reports.cs b/Engine/Source/Programs/Horde/HordeServer/Api/Reports.cs
new file mode 100644
--- /dev/null
+++ b/Engine/Source/Programs/Horde/HordeServer/Api/Reports.cs
@@ -0,0 +1,127 @@
+using EpicGames.Core;
+using HordeServer.Models;
+using System;
+using System.Collections.Generic;
+using System.Globalization;
+using System.Linq;
+
+namespace HordeServer.Api
+{
+
+	/// <summary>
+	/// Represents one stream in one pool in one hour of telemetry
+	/// </summary>
+	public class UtilizationTelemetryStream
+	{
+		/// <summary>
+		/// Stream Id
+		/// </summary>
+		public string StreamId { get; set; }
+
+		/// <summary>
+		/// Total time
+		/// </summary>
+		public double Time { get; set; }
+
+		/// <summary>
+		/// Constructor
+		/// </summary>
+		/// <param name="Stream">The stream</param>
+		public UtilizationTelemetryStream(IStreamUtilizationTelemetry Stream)
+		{
+			this.StreamId = Stream.StreamId.ToString();
+			this.Time = Stream.Time;
+		}
+	}
+
+	/// <summary>
+	/// Representation of an hour of time
+	/// </summary>
+	public class UtilizationTelemetryPool
+	{
+		/// <summary>
+		/// Pool id
+		/// </summary>
+		public string PoolId { get; set; }
+
+		/// <summary>
+		/// Number of agents in this pool
+		/// </summary>
+		public int NumAgents { get; set; }
+
+		/// <summary>
+		/// Total time spent doing admin work
+		/// </summary>
+		public double AdminTime { get; set; }
+
+		/// <summary>
+		/// Total time agents in this pool were doing work for other pools
+		/// </summary>
+		public double OtherTime { get; set; }
+
+		/// <summary>
+		/// List of streams
+		/// </summary>
+		public List<UtilizationTelemetryStream> Streams { get; set; }
+
+		/// <summary>
+		/// Constructor
+		/// </summary>
+		/// <param name="Pool"></param>
+		public UtilizationTelemetryPool(IPoolUtilizationTelemetry Pool)
+		{
+			this.PoolId = Pool.PoolId.ToString();
+			this.NumAgents = Pool.NumAgents;
+			this.AdminTime = Pool.AdminTime;
+			this.OtherTime = Pool.OtherTime;
+
+			this.Streams = Pool.Streams.ConvertAll(Stream => new UtilizationTelemetryStream(Stream));
+		}
+	}
+
+	/// <summary>
+	/// Response data for a utilization request
+	/// </summary>
+	public class UtilizationTelemetryResponse
+	{
+		/// <summary>
+		/// Start hour
+		/// </summary>
+		public DateTimeOffset StartTime { get; set; }
+
+		/// <summary>
+		/// End hour
+		/// </summary>
+		public DateTimeOffset FinishTime { get; set; }
+
+		/// <summary>
+		/// List of pools
+		/// </summary>
+		public List<UtilizationTelemetryPool> Pools { get; set; }
+
+		/// <summary>
+		/// Total admin time
+		/// </summary>
+		public double AdminTime { get; set; }
+
+		/// <summary>
+		/// Total agents
+		/// </summary>
+		public int NumAgents { get; set; }
+
+		/// <summary>
+		/// Constructor
+		/// </summary>
+		public UtilizationTelemetryResponse(IUtilizationTelemetry Telemetry)
+		{
+			this.StartTime = Telemetry.StartTime;
+			this.FinishTime = Telemetry.FinishTime;
+			this.AdminTime = Telemetry.AdminTime;
+			this.NumAgents = Telemetry.NumAgents;
+
+			this.Pools = Telemetry.Pools.ConvertAll(Pool => new UtilizationTelemetryPool(Pool));
+		}
+	}
+
+	
+}
