Initial drop of Horde Server/Agent for UE5/Main. Still missing some Epic-specific functionality required to build/run correctly.

diff --git a/Engine/Source/Programs/Horde/HordeServer/Api/Schedules.cs b/Engine/Source/Programs/Horde/HordeServer/Api/Schedules.cs
new file mode 100644
--- /dev/null
+++ b/Engine/Source/Programs/Horde/HordeServer/Api/Schedules.cs
@@ -0,0 +1,331 @@
+// Copyright Epic Games, Inc. All Rights Reserved.
+
+using HordeServer.Models;
+using System;
+using System.Collections.Generic;
+using System.ComponentModel.DataAnnotations;
+using System.Linq;
+using System.Net;
+using System.Security.Claims;
+using System.Text.Json;
+using System.Text.Json.Serialization;
+using System.Threading.Tasks;
+
+using TemplateRefId = HordeServer.Utilities.StringId<HordeServer.Models.TemplateRef>;
+
+namespace HordeServer.Api
+{
+	/// <summary>
+	/// Parameters to create a new schedule
+	/// </summary>
+	public class CreateSchedulePatternRequest
+	{
+		/// <summary>
+		/// Days of the week to run this schedule on. If null, the schedule will run every day.
+		/// </summary>
+		public List<string>? DaysOfWeek { get; set; }
+
+		/// <summary>
+		/// Time during the day for the first schedule to trigger. Measured in minutes from midnight.
+		/// </summary>
+		public int MinTime { get; set; }
+
+		/// <summary>
+		/// Time during the day for the last schedule to trigger. Measured in minutes from midnight.
+		/// </summary>
+		public int? MaxTime { get; set; }
+
+		/// <summary>
+		/// Interval between each schedule triggering
+		/// </summary>
+		public int? Interval { get; set; }
+
+		/// <summary>
+		/// Constructs a model object from this request
+		/// </summary>
+		/// <returns>Model object</returns>
+		public SchedulePattern ToModel()
+		{
+			return new SchedulePattern(DaysOfWeek?.ConvertAll(x => Enum.Parse<DayOfWeek>(x)), MinTime, MaxTime, Interval);
+		}
+	}
+
+	/// <summary>
+	/// Gate allowing a schedule to trigger.
+	/// </summary>
+	public class CreateScheduleGateRequest
+	{
+		/// <summary>
+		/// The template containing the dependency
+		/// </summary>
+		[Required]
+		public string TemplateId { get; set; } = String.Empty;
+
+		/// <summary>
+		/// Target to wait for
+		/// </summary>
+		[Required]
+		public string Target { get; set; } = String.Empty;
+
+		/// <summary>
+		/// Constructs a model object
+		/// </summary>
+		/// <returns>New model object.</returns>
+		public ScheduleGate ToModel()
+		{
+			return new ScheduleGate(new TemplateRefId(TemplateId), Target);
+		}
+	}
+
+	/// <summary>
+	/// Parameters to create a new schedule
+	/// </summary>
+	public class CreateScheduleRequest
+	{
+		/// <summary>
+		/// Roles to impersonate for this schedule
+		/// </summary>
+		public List<CreateAclClaimRequest>? Claims { get; set; }
+
+		/// <summary>
+		/// Whether the schedule should be enabled
+		/// </summary>
+		public bool Enabled { get; set; }
+
+		/// <summary>
+		/// Maximum number of builds that can be active at once
+		/// </summary>
+		public int MaxActive { get; set; }
+
+		/// <summary>
+		/// Maximum number of changes the schedule can fall behind head revision. If greater than zero, builds will be triggered for every submitted changelist until the backlog is this size.
+		/// </summary>
+		public int MaxChanges { get; set; }
+
+		/// <summary>
+		/// Whether the build requires a change to be submitted
+		/// </summary>
+		public bool RequireSubmittedChange { get; set; } = true;
+
+		/// <summary>
+		/// Gate allowing the schedule to trigger
+		/// </summary>
+		public CreateScheduleGateRequest? Gate { get; set; }
+
+		/// <summary>
+		/// The types of changes to run for
+		/// </summary>
+		public List<ChangeContentFlags>? Filter { get; set; }
+
+		/// <summary>
+		/// Parameters for the template
+		/// </summary>
+		public Dictionary<string, string> TemplateParameters { get; set; } = new Dictionary<string, string>();
+
+		/// <summary>
+		/// New patterns for the schedule
+		/// </summary>
+		public List<CreateSchedulePatternRequest> Patterns { get; set; } = new List<CreateSchedulePatternRequest>();
+
+		/// <summary>
+		/// Constructs a model object
+		/// </summary>
+		/// <returns>New model object</returns>
+		public Schedule ToModel()
+		{
+			return new Schedule(Enabled, MaxActive, MaxChanges, RequireSubmittedChange, Gate?.ToModel(), Filter, TemplateParameters, Patterns.ConvertAll(x => x.ToModel()));
+		}
+	}
+
+	/// <summary>
+	/// Information about a schedule pattern
+	/// </summary>
+	public class GetSchedulePatternResponse
+	{
+		/// <summary>
+		/// Days of the week to run this schedule on. If null, the schedule will run every day.
+		/// </summary>
+		public List<string>? DaysOfWeek { get; set; }
+
+		/// <summary>
+		/// Time during the day for the first schedule to trigger. Measured in minutes from midnight.
+		/// </summary>
+		public int MinTime { get; set; }
+
+		/// <summary>
+		/// Time during the day for the last schedule to trigger. Measured in minutes from midnight.
+		/// </summary>
+		public int? MaxTime { get; set; }
+
+		/// <summary>
+		/// Interval between each schedule triggering
+		/// </summary>
+		public int? Interval { get; set; }
+
+		/// <summary>
+		/// Private constructor for serialization
+		/// </summary>
+		private GetSchedulePatternResponse()
+		{
+		}
+
+		/// <summary>
+		/// Constructor
+		/// </summary>
+		/// <param name="Pattern">The pattern to construct from</param>
+		public GetSchedulePatternResponse(SchedulePattern Pattern)
+		{
+			this.DaysOfWeek = Pattern.DaysOfWeek?.ConvertAll(x => x.ToString());
+			this.MinTime = Pattern.MinTime;
+			this.MaxTime = Pattern.MaxTime;
+			this.Interval = Pattern.Interval;
+		}
+	}
+
+	/// <summary>
+	/// Gate allowing a schedule to trigger.
+	/// </summary>
+	public class GetScheduleGateResponse
+	{
+		/// <summary>
+		/// The template containing the dependency
+		/// </summary>
+		public string TemplateId { get; set; }
+
+		/// <summary>
+		/// Target to wait for
+		/// </summary>
+		public string Target { get; set; }
+
+		/// <summary>
+		/// Constructor
+		/// </summary>
+		/// <param name="Gate">Gate to construct from</param>
+		public GetScheduleGateResponse(ScheduleGate Gate)
+		{
+			TemplateId = Gate.TemplateRefId.ToString();
+			Target = Gate.Target;
+		}
+	}
+
+	/// <summary>
+	/// Response describing a schedule
+	/// </summary>
+	public class GetScheduleResponse
+	{
+		/// <summary>
+		/// Whether the schedule is currently enabled
+		/// </summary>
+		public bool Enabled { get; set; }
+
+		/// <summary>
+		/// Maximum number of scheduled jobs at once
+		/// </summary>
+		public int MaxActive { get; set; }
+
+		/// <summary>
+		/// Maximum number of changes the schedule can fall behind head revision. If greater than zero, builds will be triggered for every submitted changelist until the backlog is this size.
+		/// </summary>
+		public int MaxChanges { get; set; }
+
+		/// <summary>
+		/// Whether the build requires a change to be submitted
+		/// </summary>
+		public bool RequireSubmittedChange { get; set; }
+
+		/// <summary>
+		/// Gate for this schedule to trigger
+		/// </summary>
+		public GetScheduleGateResponse? Gate { get; set; }
+
+		/// <summary>
+		/// The types of changes to run for
+		/// </summary>
+		public List<ChangeContentFlags>? Filter { get; set; }
+
+		/// <summary>
+		/// Parameters for the template
+		/// </summary>
+		public Dictionary<string, string> TemplateParameters { get; set; }
+
+		/// <summary>
+		/// New patterns for the schedule
+		/// </summary>
+		public List<GetSchedulePatternResponse> Patterns { get; set; }
+
+		/// <summary>
+		/// Last changelist number that this was triggered for
+		/// </summary>
+		public int LastTriggerChange { get; set; }
+
+		/// <summary>
+		/// Last time that the schedule was triggered
+		/// </summary>
+		public DateTimeOffset LastTriggerTime { get; set; }
+
+		/// <summary>
+		/// List of active jobs
+		/// </summary>
+		public List<string> ActiveJobs { get; set; }
+
+		/// <summary>
+		/// Parameterless constructor for serialization
+		/// </summary>
+		private GetScheduleResponse()
+		{
+			TemplateParameters = null!;
+			Patterns = null!;
+			ActiveJobs = new List<string>();
+		}
+
+		/// <summary>
+		/// Constructor
+		/// </summary>
+		/// <param name="Schedule">The schedule to construct from</param>
+		public GetScheduleResponse(Schedule Schedule)
+		{
+			this.Enabled = Schedule.Enabled;
+			this.MaxActive = Schedule.MaxActive;
+			this.MaxChanges = Schedule.MaxChanges;
+			this.RequireSubmittedChange = Schedule.RequireSubmittedChange;
+			if (Schedule.Gate != null)
+			{
+				this.Gate = new GetScheduleGateResponse(Schedule.Gate);
+			}
+			this.Filter = Schedule.Filter;
+			this.TemplateParameters = Schedule.TemplateParameters;
+			this.Patterns = Schedule.Patterns.ConvertAll(x => new GetSchedulePatternResponse(x));
+			this.LastTriggerChange = Schedule.LastTriggerChange;
+			this.LastTriggerTime = Schedule.LastTriggerTime;
+			this.ActiveJobs = Schedule.ActiveJobs.ConvertAll(x => x.ToString());
+		}
+	}
+
+	/// <summary>
+	/// Response describing when a schedule is expected to trigger
+	/// </summary>
+	public class GetScheduleForecastResponse
+	{
+		/// <summary>
+		/// Next trigger times
+		/// </summary>
+		public List<DateTime> Times { get; set; }
+
+		/// <summary>
+		/// Private constructor for serialization
+		/// </summary>
+		private GetScheduleForecastResponse()
+		{
+			this.Times = null!;
+		}
+
+		/// <summary>
+		/// Constructor
+		/// </summary>
+		/// <param name="Times">List of trigger times</param>
+		public GetScheduleForecastResponse(List<DateTime> Times)
+		{
+			this.Times = Times;
+		}
+	}
+}
diff --git a/Engine/Source/Programs/Horde/HordeServer/Api/Software.cs b/Engine/Source/Programs/Horde/HordeServer/Api/Software.cs
new file mode 100644
--- /dev/null
+++ b/Engine/Source/Programs/Horde/HordeServer/Api/Software.cs
@@ -0,0 +1,104 @@
+// Copyright Epic Games, Inc. All Rights Reserved.
+
+using HordeServer.Models;
+using System;
+using System.Collections.Generic;
+using System.Linq;
+using System.Threading.Tasks;
+
+namespace HordeServer.Api
+{
+	/// <summary>
+	/// Parameters for creating a new software archive
+	/// </summary>
+	public class CreateSoftwareRequest
+	{
+		/// <summary>
+		/// Whether this software should be the default
+		/// </summary>
+		public bool Default { get; set; }
+	}
+
+	/// <summary>
+	/// Information about a client version
+	/// </summary>
+	public class CreateSoftwareResponse
+	{
+		/// <summary>
+		/// The software id
+		/// </summary>
+		public string Id { get; set; }
+
+		/// <summary>
+		/// Constructor
+		/// </summary>
+		/// <param name="Id">Identifier for the software</param>
+		public CreateSoftwareResponse(string Id)
+		{
+			this.Id = Id;
+		}
+	}
+
+	/// <summary>
+	/// Parameters for updating a software archive
+	/// </summary>
+	public class UpdateSoftwareRequest
+	{
+		/// <summary>
+		/// Whether this software should be the default
+		/// </summary>
+		public bool Default { get; set; }
+	}
+
+	/// <summary>
+	/// Information about an uploaded software archive
+	/// </summary>
+	public class GetSoftwareResponse
+	{
+		/// <summary>
+		/// Unique id for this enty
+		/// </summary>
+		public string Id { get; set; }
+
+		/// <summary>
+		/// Name of the user that created this software
+		/// </summary>
+		public string? UploadedByUser { get; set; }
+
+		/// <summary>
+		/// Time at which the client was created
+		/// </summary>
+		public DateTime UploadedAtTime { get; set; }
+
+		/// <summary>
+		/// Name of the user that created this software
+		/// </summary>
+		public string? MadeDefaultByUser { get; set; }
+
+		/// <summary>
+		/// Time at which the client was made default.
+		/// </summary>
+		public DateTime? MadeDefaultAtTime { get; set; }
+
+		/// <summary>
+		/// Default constructor for serialization
+		/// </summary>
+		private GetSoftwareResponse()
+		{
+			Id = null!;
+		}
+
+		/// <summary>
+		/// Constructor
+		/// </summary>
+		/// <param name="Software">The object to construct from</param>
+		public GetSoftwareResponse(ISoftware Software)
+		{
+			this.Id = Software.Id.ToString();
+			this.UploadedByUser = Software.UploadedByUser;
+			this.UploadedAtTime = Software.UploadedAtTime;
+			this.MadeDefaultByUser = Software.MadeDefaultByUser;
+			this.MadeDefaultAtTime = Software.MadeDefaultAtTime;
+		}
+	}
+}
diff --git a/Engine/Source/Programs/Horde/HordeServer/Api/StreamTabs.cs b/Engine/Source/Programs/Horde/HordeServer/Api/StreamTabs.cs
new file mode 100644
--- /dev/null
+++ b/Engine/Source/Programs/Horde/HordeServer/Api/StreamTabs.cs
@@ -0,0 +1,259 @@
+using Amazon.S3.Model.Internal.MarshallTransformations;
+using EpicGames.Core;
+using HordeServer.Models;
+using HordeServer.Utilities;
+using System;
+using System.Collections.Generic;
+using System.ComponentModel.DataAnnotations;
+using System.Linq;
+using System.Text.Json.Serialization;
+using System.Threading.Tasks;
+
+namespace HordeServer.Api
+{
+	/// <summary>
+	/// Information about a page to display in the dashboard for a stream
+	/// </summary>
+	[JsonKnownTypes(typeof(CreateJobsTabRequest))]
+	public abstract class CreateStreamTabRequest
+	{
+		/// <summary>
+		/// Title of this page
+		/// </summary>
+		[Required]
+		public string Title { get; set; } = null!;
+	}
+
+	/// <summary>
+	/// Type of a column in a jobs tab
+	/// </summary>
+	public enum JobsTabColumnType
+	{
+		/// <summary>
+		/// Contains labels
+		/// </summary>
+		Labels,
+
+		/// <summary>
+		/// Contains parameters
+		/// </summary>
+		Parameter
+	}
+
+	/// <summary>
+	/// Describes a column to display on the jobs page
+	/// </summary>
+	public class CreateJobsTabColumnRequest
+	{
+		/// <summary>
+		/// The type of column
+		/// </summary>
+		public JobsTabColumnType Type { get; set; } = JobsTabColumnType.Labels;
+
+		/// <summary>
+		/// Heading for this column
+		/// </summary>
+		[Required]
+		public string Heading { get; set; } = null!;
+
+		/// <summary>
+		/// Category of aggregates to display in this column. If null, includes any aggregate not matched by another column.
+		/// </summary>
+		public string? Category { get; set; }
+
+		/// <summary>
+		/// Parameter to show in this column
+		/// </summary>
+		public string? Parameter { get; set; }
+
+		/// <summary>
+		/// Relative width of this column.
+		/// </summary>
+		public int? RelativeWidth { get; set; }
+
+		/// <summary>
+		/// Construct a JobsTabColumn object from this request
+		/// </summary>
+		/// <returns>Column object</returns>
+		public JobsTabColumn ToModel()
+		{
+			switch(Type)
+			{
+				case JobsTabColumnType.Labels:
+					return new JobsTabLabelColumn(Heading, Category, RelativeWidth);
+				case JobsTabColumnType.Parameter:
+					return new JobsTabParameterColumn(Heading, Parameter ?? "Undefined", RelativeWidth);
+				default:
+					return new JobsTabLabelColumn(Heading, "Undefined", RelativeWidth);
+			}
+		}
+	}
+
+	/// <summary>
+	/// Describes a job page
+	/// </summary>
+	[JsonDiscriminator("Jobs")]
+	public class CreateJobsTabRequest : CreateStreamTabRequest
+	{
+		/// <summary>
+		/// Whether to show job names on this page
+		/// </summary>
+		public bool ShowNames { get; set; }
+
+		/// <summary>
+		/// Names of jobs to include on this page. If there is only one name specified, the name column does not need to be displayed.
+		/// </summary>
+		public List<string>? JobNames { get; set; }
+
+		/// <summary>
+		/// List of job template names to show on this page.
+		/// </summary>
+		public List<string>? Templates { get; set; }
+
+		/// <summary>
+		/// Columns to display for different types of aggregates
+		/// </summary>
+		public List<CreateJobsTabColumnRequest>? Columns { get; set; }
+	}
+
+	/// <summary>
+	/// Information about a page to display in the dashboard for a stream
+	/// </summary>
+	[JsonKnownTypes(typeof(GetJobsTabResponse))]
+	public abstract class GetStreamTabResponse
+	{
+		/// <summary>
+		/// Title of this page
+		/// </summary>
+		public string Title { get; set; }
+
+		/// <summary>
+		/// Constructor
+		/// </summary>
+		/// <param name="Title">Title of this page</param>
+		public GetStreamTabResponse(string Title)//, StreamPageType Type)
+		{
+			this.Title = Title;
+		}
+	}
+
+	/// <summary>
+	/// Describes a column to display on the jobs page
+	/// </summary>
+	[JsonKnownTypes(typeof(GetJobsTabLabelColumnResponse), typeof(GetJobsTabParameterColumnResponse))]
+	public abstract class GetJobsTabColumnResponse
+	{
+		/// <summary>
+		/// Heading for this column
+		/// </summary>
+		public string Heading { get; set; } = null!;
+
+		/// <summary>
+		/// Relative width of this column.
+		/// </summary>
+		public int? RelativeWidth { get; set; }
+
+		/// <summary>
+		/// Constructor
+		/// </summary>
+		/// <param name="Heading">Heading for this column</param>
+		/// <param name="RelativeWidth">Relative width of this column</param>
+		public GetJobsTabColumnResponse(string Heading, int? RelativeWidth)
+		{
+			this.Heading = Heading;
+			this.RelativeWidth = RelativeWidth;
+		}
+	}
+
+	/// <summary>
+	/// Describes a label column to display on the jobs page
+	/// </summary>
+	[JsonDiscriminator(nameof(JobsTabColumnType.Labels))]
+	public class GetJobsTabLabelColumnResponse : GetJobsTabColumnResponse
+	{
+		/// <summary>
+		/// Category of aggregates to display in this column. If null, includes any aggregate not matched by another column.
+		/// </summary>
+		public string? Category { get; set; }
+
+		/// <summary>
+		/// Constructor
+		/// </summary>
+		/// <param name="Heading">Heading for this column</param>
+		/// <param name="Category">Category of aggregates to display in this column. If null, includes any aggregate not matched by another column.</param>
+		/// <param name="RelativeWidth">Relative width of this column</param>
+		public GetJobsTabLabelColumnResponse(string Heading, string? Category, int? RelativeWidth)
+			: base(Heading, RelativeWidth)
+		{
+			this.Category = Category;
+		}
+	}
+
+	/// <summary>
+	/// Describes a parameter column to display on the jobs page
+	/// </summary>
+	[JsonDiscriminator(nameof(JobsTabColumnType.Parameter))]
+	public class GetJobsTabParameterColumnResponse : GetJobsTabColumnResponse
+	{
+		/// <summary>
+		/// Parameter to show in this column
+		/// </summary>
+		public string? Parameter { get; set; }
+
+		/// <summary>
+		/// Constructor
+		/// </summary>
+		/// <param name="Heading">Heading for this column</param>
+		/// <param name="Parameter">Name of the parameter to display</param>
+		/// <param name="RelativeWidth">Relative width of this column</param>
+		public GetJobsTabParameterColumnResponse(string Heading, string? Parameter, int? RelativeWidth)
+			: base(Heading, RelativeWidth)
+		{
+			this.Parameter = Parameter;
+		}
+	}
+
+	/// <summary>
+	/// Describes a job page
+	/// </summary>
+	[JsonDiscriminator("Jobs")]
+	public class GetJobsTabResponse : GetStreamTabResponse
+	{
+		/// <summary>
+		/// Whether to show names on the page
+		/// </summary>
+		public bool ShowNames { get; set; }
+
+		/// <summary>
+		/// List of templates to show on the page
+		/// </summary>
+		public List<string>? Templates { get; set; }
+
+		/// <summary>
+		/// Names of jobs to include on this page. If there is only one name specified, the name column does not need to be displayed.
+		/// </summary>
+		public List<string>? JobNames { get; set; }
+
+		/// <summary>
+		/// Columns to display for different types of aggregates
+		/// </summary>
+		public List<GetJobsTabColumnResponse>? Columns { get; set; }
+
+		/// <summary>
+		/// Constructor
+		/// </summary>
+		/// <param name="Title">Title for this page</param>
+		/// <param name="bShowNames">Whether to show names on the page</param>
+		/// <param name="Templates">Templates to include on this page</param>
+		/// <param name="JobNames">List of job names to include on this page</param>
+		/// <param name="Columns">List of columns to display</param>
+		public GetJobsTabResponse(string Title, bool bShowNames, List<string>? Templates, List<string>? JobNames, List<GetJobsTabColumnResponse>? Columns)
+			: base(Title)
+		{
+			this.ShowNames = bShowNames;
+			this.Templates = Templates;
+			this.JobNames = JobNames;
+			this.Columns = Columns;
+		}
+	}
+}
diff --git a/Engine/Source/Programs/Horde/HordeServer/Api/Streams.cs b/Engine/Source/Programs/Horde/HordeServer/Api/Streams.cs
new file mode 100644
--- /dev/null
+++ b/Engine/Source/Programs/Horde/HordeServer/Api/Streams.cs
@@ -0,0 +1,670 @@
+// Copyright Epic Games, Inc. All Rights Reserved.
+
+using HordeCommon;
+using HordeServer.Models;
+using HordeServer.Utilities;
+using System;
+using System.Collections.Generic;
+using System.ComponentModel.DataAnnotations;
+using System.Linq;
+using System.Threading.Tasks;
+
+using ProjectId = HordeServer.Utilities.StringId<HordeServer.Models.IProject>;
+using TemplateRefId = HordeServer.Utilities.StringId<HordeServer.Models.TemplateRef>;
+
+namespace HordeServer.Api
+{
+	/// <summary>
+	/// Specifies defaults for running a preflight
+	/// </summary>
+	public class DefaultPreflightRequest
+	{
+		/// <summary>
+		/// The template id to query
+		/// </summary>
+		public string? TemplateId { get; set; }
+
+		/// <summary>
+		/// The last successful job type to use for the base changelist
+		/// </summary>
+		public string? ChangeTemplateId { get; set; }
+
+		/// <summary>
+		/// Convert to a model object
+		/// </summary>
+		/// <returns></returns>
+		public DefaultPreflight ToModel()
+		{
+			return new DefaultPreflight((TemplateId != null) ? (TemplateRefId?)new TemplateRefId(TemplateId) : null, (ChangeTemplateId != null) ? (TemplateRefId?)new TemplateRefId(ChangeTemplateId) : null);
+		}
+	}
+
+	/// <summary>
+	/// Mapping from a BuildGraph agent type to a set of machines on the farm
+	/// </summary>
+	public class CreateAgentTypeRequest
+	{
+		/// <summary>
+		/// Pool of agents to use for this agent type
+		/// </summary>
+		[Required]
+		public string Pool { get; set; } = null!;
+
+		/// <summary>
+		/// Name of the workspace to sync
+		/// </summary>
+		public string? Workspace { get; set; }
+
+		/// <summary>
+		/// Path to the temporary storage dir
+		/// </summary>
+		public string? TempStorageDir { get; set; }
+
+		/// <summary>
+		/// Environment variables to be set when executing the job
+		/// </summary>
+		public Dictionary<string, string>? Environment { get; set; }
+	}
+
+	/// <summary>
+	/// Information about a workspace type
+	/// </summary>
+	public class CreateWorkspaceTypeRequest
+	{
+		/// <summary>
+		/// Name of the Perforce server cluster to use
+		/// </summary>
+		public string? Cluster { get; set; }
+
+		/// <summary>
+		/// The Perforce server and port (eg. perforce:1666)
+		/// </summary>
+		public string? ServerAndPort { get; set; }
+
+		/// <summary>
+		/// User to log into Perforce with (defaults to buildmachine)
+		/// </summary>
+		public string? UserName { get; set; }
+
+		/// <summary>
+		/// Password to use to log into the workspace
+		/// </summary>
+		public string? Password { get; set; }
+
+		/// <summary>
+		/// Identifier to distinguish this workspace from other workspaces. Defaults to the workspace type name.
+		/// </summary>
+		public string? Identifier { get; set; }
+
+		/// <summary>
+		/// Override for the stream to sync
+		/// </summary>
+		public string? Stream { get; set; }
+
+		/// <summary>
+		/// Custom view for the workspace
+		/// </summary>
+		public List<string>? View { get; set; }
+
+		/// <summary>
+		/// Whether to use an incrementally synced workspace
+		/// </summary>
+		public bool Incremental { get; set; }
+	}
+
+	/// <summary>
+	/// Trigger for another template
+	/// </summary>
+	public class CreateChainedJobTemplateRequest
+	{
+		/// <summary>
+		/// Name of the target that needs to complete before starting the other template
+		/// </summary>
+		[Required]
+		public string Trigger { get; set; } = String.Empty;
+
+		/// <summary>
+		/// Id of the template to trigger
+		/// </summary>
+		[Required]
+		public string TemplateId { get; set; } = String.Empty;
+	}
+
+	/// <summary>
+	/// Parameters to create a template within a stream
+	/// </summary>
+	public class CreateTemplateRefRequest : CreateTemplateRequest
+	{
+		/// <summary>
+		/// Optional identifier for this ref. If not specified, an id will be generated from the name.
+		/// </summary>
+		public string? Id { get; set; }
+
+		/// <summary>
+		/// Whether to show badges in UGS for these jobs
+		/// </summary>
+		public bool ShowUgsBadges { get; set; }
+
+		/// <summary>
+		/// Whether to show alerts in UGS for these jobs
+		/// </summary>
+		public bool ShowUgsAlerts { get; set; }
+
+		/// <summary>
+		/// Notification channel for this template. Overrides the stream channel if set.
+		/// </summary>
+		public string? NotificationChannel { get; set; }
+
+		/// <summary>
+		/// Notification channel filter for this template. Can be Success|Failure|Warnings
+		/// </summary>
+		public string? NotificationChannelFilter { get; set; }
+
+		/// <summary>
+		/// Triage channel for this template. Overrides the stream channel if set.
+		/// </summary>
+		public string? TriageChannel { get; set; }
+
+		/// <summary>
+		/// Schedule to execute this template
+		/// </summary>
+		public CreateScheduleRequest? Schedule { get; set; }
+
+		/// <summary>
+		/// List of chained job triggers
+		/// </summary>
+		public List<CreateChainedJobTemplateRequest>? ChainedJobs { get; set; }
+
+		/// <summary>
+		/// The ACL for this template
+		/// </summary>
+		public UpdateAclRequest? Acl { get; set; }
+	}
+
+	/// <summary>
+	/// Parameters to create a new stream
+	/// </summary>
+	public class SetStreamRequest
+	{
+		/// <summary>
+		/// Unique id for the project
+		/// </summary>
+		[Required]
+		public string ProjectId { get; set; }
+
+		/// <summary>
+		/// Name for the new stream
+		/// </summary>
+		[Required]
+		public string Name { get; set; }
+
+		/// <summary>
+		/// Path to the config file in Perforce
+		/// </summary>
+		public string? ConfigPath { get; set; }
+
+		/// <summary>
+		/// Order for this stream
+		/// </summary>
+		public int? Order { get; set; }
+
+		/// <summary>
+		/// Notification channel for all jobs in this stream
+		/// </summary>
+		public string? NotificationChannel { get; set; }
+
+		/// <summary>
+		/// Notification channel filter for this template. Can be Success|Failure|Warnings
+		/// </summary>
+		public string? NotificationChannelFilter { get; set; }
+
+		/// <summary>
+		/// Channel to post issue triage notifications
+		/// </summary>
+		public string? TriageChannel { get; set; }
+
+		/// <summary>
+		/// Default template for running preflights
+		/// </summary>
+		public string? DefaultPreflightTemplate { get; set; }
+
+		/// <summary>
+		/// Default template for running preflights
+		/// </summary>
+		public DefaultPreflightRequest? DefaultPreflight { get; set; }
+
+		/// <summary>
+		/// List of tabs to show for the new stream
+		/// </summary>
+		public List<CreateStreamTabRequest>? Tabs { get; set; }
+
+		/// <summary>
+		/// Map of agent name to type
+		/// </summary>
+		public Dictionary<string, CreateAgentTypeRequest>? AgentTypes { get; set; }
+
+		/// <summary>
+		/// Map of workspace name to type
+		/// </summary>
+		public Dictionary<string, CreateWorkspaceTypeRequest>? WorkspaceTypes { get; set; }
+
+		/// <summary>
+		/// List of templates to create
+		/// </summary>
+		public List<CreateTemplateRefRequest>? Templates { get; set; }
+
+		/// <summary>
+		/// Properties for the new stream
+		/// </summary>
+		public Dictionary<string, string>? Properties { get; set; }
+
+		/// <summary>
+		/// Custom permissions for this object
+		/// </summary>
+		public UpdateAclRequest? Acl { get; set; }
+		
+		/// <summary>
+		/// Pause stream builds until specified date
+		/// </summary>
+		public DateTime? PausedUntil { get; set; }
+		
+		/// <summary>
+		/// Reason for pausing builds of the stream
+		/// </summary>
+		public string? PauseComment { get; set; }
+
+		/// <summary>
+		/// Private constructor for serialization
+		/// </summary>
+		private SetStreamRequest()
+		{
+			this.ProjectId = null!;
+			this.Name = null!;
+		}
+	}
+
+	/// <summary>
+	/// Mapping from a BuildGraph agent type to a set of machines on the farm
+	/// </summary>
+	public class GetAgentTypeResponse
+	{
+		/// <summary>
+		/// Pool of agents to use for this agent type
+		/// </summary>
+		public string Pool { get; set; }
+
+		/// <summary>
+		/// Name of the workspace to sync
+		/// </summary>
+		public string? Workspace { get; set; }
+
+		/// <summary>
+		/// Pool of agents to use to execute this work
+		/// </summary>
+		public GetAgentRequirementsResponse? Requirements { get; set; }
+
+		/// <summary>
+		/// Path to the temporary storage dir
+		/// </summary>
+		public string? TempStorageDir { get; set; }
+
+		/// <summary>
+		/// Environment variables to be set when executing the job
+		/// </summary>
+		public Dictionary<string, string>? Environment { get; set; }
+
+		/// <summary>
+		/// Default constructor
+		/// </summary>
+		private GetAgentTypeResponse()
+		{
+			Pool = null!;
+		}
+
+		/// <summary>
+		/// Constructor
+		/// </summary>
+		/// <param name="Pool">Pool of agents to use for this agent type</param>
+		/// <param name="Workspace">Name of the workspace to sync</param>
+		/// <param name="TempStorageDir">Path to the temp storage directory</param>
+		/// <param name="Environment">Environment variables to be set when executing this job</param>
+		public GetAgentTypeResponse(string Pool, string? Workspace, string? TempStorageDir, Dictionary<string, string>? Environment)
+		{
+			this.Pool = Pool;
+			this.Workspace = Workspace;
+			this.TempStorageDir = TempStorageDir;
+			this.Environment = Environment;
+		}
+	}
+
+	/// <summary>
+	/// Information about a workspace type
+	/// </summary>
+	public class GetWorkspaceTypeResponse
+	{
+		/// <summary>
+		/// The Perforce server cluster
+		/// </summary>
+		public string? Cluster { get; set; }
+
+		/// <summary>
+		/// The Perforce server and port (eg. perforce:1666)
+		/// </summary>
+		public string? ServerAndPort { get; set; }
+
+		/// <summary>
+		/// User to log into Perforce with (defaults to buildmachine)
+		/// </summary>
+		public string? UserName { get; set; }
+
+		/// <summary>
+		/// Identifier to distinguish this workspace from other workspaces. Defaults to the workspace type name.
+		/// </summary>
+		public string? Identifier { get; set; }
+
+		/// <summary>
+		/// Override for the stream to sync
+		/// </summary>
+		public string? Stream { get; set; }
+
+		/// <summary>
+		/// Custom view for the workspace
+		/// </summary>
+		public List<string>? View { get; set; }
+
+		/// <summary>
+		/// Whether to use an incrementally synced workspace
+		/// </summary>
+		public bool Incremental { get; set; }
+
+		/// <summary>
+		/// Default constructor
+		/// </summary>
+		public GetWorkspaceTypeResponse()
+		{
+		}
+
+		/// <summary>
+		/// Constructor
+		/// </summary>
+		/// <param name="Cluster">The server cluster</param>
+		/// <param name="ServerAndPort">The perforce server</param>
+		/// <param name="UserName">The perforce user name</param>
+		/// <param name="Identifier">Identifier to distinguish this workspace from other workspaces. Defaults to the workspace type name.</param>
+		/// <param name="Stream">Override for the stream to sync</param>
+		/// <param name="View">Custom view for the workspace</param>
+		/// <param name="bIncremental">Whether to use an incrementally synced workspace</param>
+		public GetWorkspaceTypeResponse(string? Cluster, string? ServerAndPort, string? UserName, string? Identifier, string? Stream, List<string>? View, bool bIncremental)
+		{
+			this.Cluster = Cluster;
+			this.ServerAndPort = ServerAndPort;
+			this.UserName = UserName;
+			this.Identifier = Identifier;
+			this.Stream = Stream;
+			this.View = View;
+			this.Incremental = bIncremental;
+		}
+	}
+
+	/// <summary>
+	/// Trigger for another template
+	/// </summary>
+	public class GetChainedJobTemplateResponse
+	{
+		/// <summary>
+		/// Name of the target that needs to complete before starting the other template
+		/// </summary>
+		public string Trigger { get; set; }
+
+		/// <summary>
+		/// Id of the template to trigger
+		/// </summary>
+		public string TemplateId { get; set; }
+
+		/// <summary>
+		/// Constructor
+		/// </summary>
+		/// <param name="Trigger">The trigger definition</param>
+		public GetChainedJobTemplateResponse(ChainedJobTemplate Trigger)
+		{
+			this.Trigger = Trigger.Trigger;
+			this.TemplateId = Trigger.TemplateRefId.ToString();
+		}
+	}
+
+	/// <summary>
+	/// Information about a template in this stream
+	/// </summary>
+	public class GetTemplateRefResponse : GetTemplateResponseBase
+	{
+		/// <summary>
+		/// Id of the template ref
+		/// </summary>
+		public string Id { get; set; }
+
+		/// <summary>
+		/// Hash of the template definition
+		/// </summary>
+		public string Hash { get; set; }
+
+		/// <summary>
+		/// Whether to show badges in UGS for these jobs
+		/// </summary>
+		public bool ShowUgsBadges { get; set; }
+
+		/// <summary>
+		/// Whether to show alerts in UGS for these jobs
+		/// </summary>
+		public bool ShowUgsAlerts { get; set; }
+
+		/// <summary>
+		/// Notification channel for this template. Overrides the stream channel if set.
+		/// </summary>
+		public string? NotificationChannel { get; set; }
+
+		/// <summary>
+		/// Notification channel filter for this template. Can be Success|Failure|Warnings
+		/// </summary>
+		public string? NotificationChannelFilter { get; set; }
+
+		/// <summary>
+		/// Triage channel for this template. Overrides the stream channel if set.
+		/// </summary>
+		public string? TriageChannel { get; set; }
+
+		/// <summary>
+		/// The schedule for this ref
+		/// </summary>
+		public GetScheduleResponse? Schedule { get; set; }
+
+		/// <summary>
+		/// List of templates to trigger
+		/// </summary>
+		public List<GetChainedJobTemplateResponse>? ChainedJobs { get; set; }
+
+		/// <summary>
+		/// ACL for this template
+		/// </summary>
+		public GetAclResponse? Acl { get; set; }
+
+		/// <summary>
+		/// Constructor
+		/// </summary>
+		/// <param name="Id">The template ref id</param>
+		/// <param name="TemplateRef">The template ref</param>
+		/// <param name="Template">The template instance</param>
+		/// <param name="bIncludeAcl">Whether to include the ACL in the response</param>
+		public GetTemplateRefResponse(TemplateRefId Id, TemplateRef TemplateRef, ITemplate Template, bool bIncludeAcl)
+			: base(Template)
+		{
+			this.Id = Id.ToString();
+			Hash = TemplateRef.Hash.ToString();
+			ShowUgsBadges = TemplateRef.ShowUgsBadges;
+			ShowUgsAlerts = TemplateRef.ShowUgsAlerts;
+			NotificationChannel = TemplateRef.NotificationChannel;
+			NotificationChannelFilter = TemplateRef.NotificationChannelFilter;
+			Schedule = (TemplateRef.Schedule != null) ? new GetScheduleResponse(TemplateRef.Schedule) : null;
+			ChainedJobs = (TemplateRef.ChainedJobs != null && TemplateRef.ChainedJobs.Count > 0) ? TemplateRef.ChainedJobs.ConvertAll(x => new GetChainedJobTemplateResponse(x)) : null;
+			Acl = (bIncludeAcl && TemplateRef.Acl != null)? new GetAclResponse(TemplateRef.Acl) : null;
+		}
+	}
+
+	/// <summary>
+	/// Response describing a stream
+	/// </summary>
+	public class GetStreamResponse
+	{
+		/// <summary>
+		/// Unique id of the stream
+		/// </summary>
+		public string Id { get; set; }
+
+		/// <summary>
+		/// Unique id of the project containing this stream
+		/// </summary>
+		public string ProjectId { get; set; }
+
+		/// <summary>
+		/// Name of the stream
+		/// </summary>
+		public string Name { get; set; }
+
+		/// <summary>
+		/// The config file path on the server
+		/// </summary>
+		public string? ConfigPath { get; set; }
+
+		/// <summary>
+		/// Last version of the config file that was read
+		/// </summary>
+		public int? ConfigChange { get; set; }
+
+		/// <summary>
+		/// Order to display in the list
+		/// </summary>
+		public int Order { get; set; }
+
+		/// <summary>
+		/// Notification channel for all jobs in this stream
+		/// </summary>
+		public string? NotificationChannel { get; set; }
+
+		/// <summary>
+		/// Notification channel filter for this template. Can be Success|Failure|Warnings
+		/// </summary>
+		public string? NotificationChannelFilter { get; set; }
+
+		/// <summary>
+		/// Channel to post issue triage notifications
+		/// </summary>
+		public string? TriageChannel { get; set; }
+
+		/// <summary>
+		/// Default template for running preflights
+		/// </summary>
+		public string? DefaultPreflightTemplate { get; set; }
+
+		/// <summary>
+		/// Default template to use for preflights
+		/// </summary>
+		public DefaultPreflightRequest? DefaultPreflight { get; set; }
+
+		/// <summary>
+		/// List of tabs to display for this stream
+		/// </summary>
+		public List<GetStreamTabResponse> Tabs { get; set; }
+
+		/// <summary>
+		/// Map of agent name to type
+		/// </summary>
+		public Dictionary<string, GetAgentTypeResponse> AgentTypes { get; set; }
+
+		/// <summary>
+		/// Map of workspace name to type
+		/// </summary>
+		public Dictionary<string, GetWorkspaceTypeResponse>? WorkspaceTypes { get; set; }
+
+		/// <summary>
+		/// Templates for jobs in this stream
+		/// </summary>
+		public List<GetTemplateRefResponse> Templates { get; set; }
+
+		/// <summary>
+		/// Properties for this stream
+		/// </summary>
+		public Dictionary<string, string> Properties { get; set; }
+
+		/// <summary>
+		/// Custom permissions for this object
+		/// </summary>
+		public GetAclResponse? Acl { get; set; }
+		
+		/// <summary>
+		/// Stream paused for new builds until this date
+		/// </summary>
+		public DateTime? PausedUntil { get; set; }
+		
+		/// <summary>
+		/// Reason for stream being paused
+		/// </summary>
+		public string? PauseComment { get; set; }
+
+		/// <summary>
+		/// Parameterless constructor for serialization
+		/// </summary>
+		private GetStreamResponse()
+		{
+			Id = null!;
+			ProjectId = null!;
+			Name = null!;
+			Tabs = null!;
+			AgentTypes = null!;
+			Templates = null!;
+			Properties = null!;
+		}
+
+		/// <summary>
+		/// Constructor
+		/// </summary>
+		/// <param name="Id">Unique id of the stream</param>
+		/// <param name="ProjectId">Unique id of the project containing the stream</param>
+		/// <param name="Name">Name of the stream</param>
+		/// <param name="ConfigPath">The config file path on the server</param>
+		/// <param name="ConfigChange">Last config change that was read</param>
+		/// <param name="Order">Order to display this stream</param>
+		/// <param name="NotificationChannel"></param>
+		/// <param name="NotificationChannelFilter"></param>
+		/// <param name="TriageChannel"></param>
+		/// <param name="DefaultPreflight">The default template to use for preflights</param>
+		/// <param name="Tabs">List of tabs to display for this stream</param>
+		/// <param name="AgentTypes">Map of agent type name to description</param>
+		/// <param name="WorkspaceTypes">Map of workspace name to description</param>
+		/// <param name="Templates">Templates for this stream</param>
+		/// <param name="Properties">Properties for the stream</param>
+		/// <param name="Acl">Permissions for this object</param>
+		/// <param name="PausedUntil">Stream paused for new builds until this date</param>
+		/// <param name="PauseComment">Reason for stream being paused</param>
+		public GetStreamResponse(string Id, string ProjectId, string Name, string? ConfigPath, int? ConfigChange, int Order, string? NotificationChannel, string? NotificationChannelFilter, string? TriageChannel, DefaultPreflightRequest? DefaultPreflight, List<GetStreamTabResponse> Tabs, Dictionary<string, GetAgentTypeResponse> AgentTypes, Dictionary<string, GetWorkspaceTypeResponse>? WorkspaceTypes, List<GetTemplateRefResponse> Templates, Dictionary<string, string> Properties, GetAclResponse? Acl, DateTime? PausedUntil, string? PauseComment)
+		{
+			this.Id = Id;
+			this.ProjectId = ProjectId;
+			this.Name = Name;
+			this.ConfigPath = ConfigPath;
+			this.ConfigChange = ConfigChange;
+			this.Order = Order;
+			this.NotificationChannel = NotificationChannel;
+			this.NotificationChannelFilter = NotificationChannelFilter;
+			this.TriageChannel = TriageChannel;
+			this.DefaultPreflightTemplate = DefaultPreflight?.TemplateId;
+			this.DefaultPreflight = DefaultPreflight;
+			this.Tabs = Tabs;
+			this.AgentTypes = AgentTypes;
+			this.WorkspaceTypes = WorkspaceTypes;
+			this.Templates = Templates;
+			this.Properties = Properties;
+			this.Acl = Acl;
+			this.PausedUntil = PausedUntil;
+			this.PauseComment = PauseComment;
+		}
+	}
+}
diff --git a/Engine/Source/Programs/Horde/HordeServer/Api/Subscriptions.cs b/Engine/Source/Programs/Horde/HordeServer/Api/Subscriptions.cs
new file mode 100644
--- /dev/null
+++ b/Engine/Source/Programs/Horde/HordeServer/Api/Subscriptions.cs
@@ -0,0 +1,90 @@
+using HordeServer.Models;
+using System;
+using System.Collections.Generic;
+using System.ComponentModel.DataAnnotations;
+using System.Linq;
+using System.Threading.Tasks;
+
+namespace HordeServer.Api
+{
+	/// <summary>
+	/// Request object for creating a new subscription
+	/// </summary>
+	public class CreateSubscriptionRequest
+	{
+		/// <summary>
+		/// The event to create a subscription for
+		/// </summary>
+		[Required]
+		public EventRecord Event { get; set; } = null!;
+
+		/// <summary>
+		/// The user to subscribe. Defaults to the current user if unspecified.
+		/// </summary>
+		[Required]
+		public string UserId { get; set; } = null!;
+
+		/// <summary>
+		/// Type of notification to send
+		/// </summary>
+		public NotificationType NotificationType { get; set; }
+	}
+
+	/// <summary>
+	/// Response from creating a subscription
+	/// </summary>
+	public class CreateSubscriptionResponse
+	{
+		/// <summary>
+		/// The subscription id
+		/// </summary>
+		public string Id { get; set; }
+
+		/// <summary>
+		/// Constructor
+		/// </summary>
+		/// <param name="Subscription">The subscription to construct from</param>
+		public CreateSubscriptionResponse(ISubscription Subscription)
+		{
+			this.Id = Subscription.Id;
+		}
+	}
+
+	/// <summary>
+	/// Response describing a subscription
+	/// </summary>
+	public class GetSubscriptionResponse
+	{
+		/// <summary>
+		/// The subscription id
+		/// </summary>
+		public string Id { get; set; }
+
+		/// <summary>
+		/// Information about the specific event
+		/// </summary>
+		public EventRecord Event { get; set; }
+
+		/// <summary>
+		/// Unique id of the user subscribed to the event
+		/// </summary>
+		public string UserId { get; set; }
+
+		/// <summary>
+		/// The type of notification to receive
+		/// </summary>
+		public NotificationType NotificationType { get; set; }
+
+		/// <summary>
+		/// Constructor
+		/// </summary>
+		/// <param name="Subscription">Subscription to construct from</param>
+		public GetSubscriptionResponse(ISubscription Subscription)
+		{
+			this.Id = Subscription.Id;
+			this.Event = Subscription.Event.ToRecord();
+			this.UserId = Subscription.UserId.ToString();
+			this.NotificationType = Subscription.NotificationType;
+		}
+	}
+}
diff --git a/Engine/Source/Programs/Horde/HordeServer/Api/Templates.cs b/Engine/Source/Programs/Horde/HordeServer/Api/Templates.cs
new file mode 100644
--- /dev/null
+++ b/Engine/Source/Programs/Horde/HordeServer/Api/Templates.cs
@@ -0,0 +1,579 @@
+// Copyright Epic Games, Inc. All Rights Reserved.
+
+using EpicGames.Core;
+using HordeCommon;
+using HordeServer.Models;
+using HordeServer.Utilities;
+using System;
+using System.Collections.Generic;
+using System.ComponentModel.DataAnnotations;
+using System.Linq;
+using System.Text.Json.Serialization;
+using System.Threading.Tasks;
+
+namespace HordeServer.Api
+{
+	/// <summary>
+	/// Base class for template parameters
+	/// </summary>
+	[JsonKnownTypes(typeof(GroupParameterData), typeof(TextParameterData), typeof(ListParameterData), typeof(BoolParameterData))]
+	public abstract class ParameterData
+	{
+		/// <summary>
+		/// Convert to a parameter object
+		/// </summary>
+		/// <returns><see cref="Parameter"/> object</returns>
+		public abstract Parameter ToModel();
+	}
+
+	/// <summary>
+	/// Describes how to render a group parameter
+	/// </summary>
+	public enum GroupParameterStyle
+	{
+		/// <summary>
+		/// Separate tab on the form
+		/// </summary>
+		Tab,
+
+		/// <summary>
+		/// Section with heading
+		/// </summary>
+		Section,
+	}
+
+	/// <summary>
+	/// Used to group a number of other parameters
+	/// </summary>
+	[JsonDiscriminator("Group")]
+	public class GroupParameterData : ParameterData
+	{
+		/// <summary>
+		/// Label to display next to this parameter
+		/// </summary>
+		public string Label { get; set; }
+
+		/// <summary>
+		/// How to display this group
+		/// </summary>
+		public GroupParameterStyle Style { get; set; }
+
+		/// <summary>
+		/// List of child parameters
+		/// </summary>
+		public List<ParameterData> Children { get; set; }
+
+		/// <summary>
+		/// Private constructor for serialization
+		/// </summary>
+		private GroupParameterData()
+		{
+			Label = null!;
+			Children = null!;
+		}
+
+		/// <summary>
+		/// Constructor
+		/// </summary>
+		/// <param name="Label">Name of the group</param>
+		/// <param name="Style">How to display this group</param>
+		/// <param name="Children">List of child parameters</param>
+		public GroupParameterData(string Label, GroupParameterStyle Style, List<ParameterData> Children)
+		{
+			this.Label = Label;
+			this.Style = Style;
+			this.Children = Children;
+		}
+
+		/// <summary>
+		/// Converts this data to a model object
+		/// </summary>
+		/// <returns>New <see cref="GroupParameter"/> object</returns>
+		public override Parameter ToModel()
+		{
+			return new GroupParameter(Label, Style, Children.ConvertAll(x => x.ToModel()));
+		}
+	}
+
+	/// <summary>
+	/// Free-form text entry parameter
+	/// </summary>
+	[JsonDiscriminator("Text")]
+	public class TextParameterData : ParameterData
+	{
+		/// <summary>
+		/// Name of the parameter associated with this parameter.
+		/// </summary>
+		public string Label { get; set; }
+
+		/// <summary>
+		/// Argument to pass to the executor
+		/// </summary>
+		public string Argument { get; set; }
+
+		/// <summary>
+		/// Default value for this argument
+		/// </summary>
+		public string Default { get; set; }
+
+		/// <summary>
+		/// Hint text for this parameter
+		/// </summary>
+		public string? Hint { get; set; }
+
+		/// <summary>
+		/// Regex used to validate this parameter
+		/// </summary>
+		public string? Validation { get; set; }
+
+		/// <summary>
+		/// Message displayed if validation fails, informing user of valid values.
+		/// </summary>
+		public string? ValidationError { get; set; }
+
+		/// <summary>
+		/// Tool-tip text to display
+		/// </summary>
+		public string? ToolTip { get; set; }
+
+		/// <summary>
+		/// Private constructor for serialization
+		/// </summary>
+		private TextParameterData()
+		{
+			Label = null!;
+			Argument = null!;
+			Default = String.Empty;
+		}
+
+		/// <summary>
+		/// Constructor
+		/// </summary>
+		/// <param name="Label">Label to show next to the parameter</param>
+		/// <param name="Argument">Argument to pass this value with</param>
+		/// <param name="Default">Default value for this parameter</param>
+		/// <param name="Hint">Hint text to display for this parameter</param>
+		/// <param name="Validation">Regex used to validate entries</param>
+		/// <param name="ValidationError">Message displayed to explain validation issues</param>
+		/// <param name="ToolTip">Tool tip text to display</param>
+		public TextParameterData(string Label, string Argument, string Default, string? Hint, string? Validation, string? ValidationError, string? ToolTip)
+		{
+			this.Label = Label;
+			this.Argument = Argument;
+			this.Default = Default;
+			this.Hint = Hint;
+			this.Validation = Validation;
+			this.ValidationError = ValidationError;
+			this.ToolTip = ToolTip;
+		}
+
+		/// <summary>
+		/// Converts this data to a model object
+		/// </summary>
+		/// <returns>New <see cref="TextParameter"/> object</returns>
+		public override Parameter ToModel()
+		{
+			return new TextParameter(Label, Argument, Default, Hint, Validation, ValidationError, ToolTip);
+		}
+	}
+
+	/// <summary>
+	/// Style of list parameter
+	/// </summary>
+	public enum ListParameterStyle
+	{
+		/// <summary>
+		/// Regular drop-down list. One item is always selected.
+		/// </summary>
+		List,
+
+		/// <summary>
+		/// Drop-down list with checkboxes
+		/// </summary>
+		MultiList,
+
+		/// <summary>
+		/// Tag picker from list of options
+		/// </summary>
+		TagPicker,
+	}
+
+	/// <summary>
+	/// Possible option for a list parameter
+	/// </summary>
+	public class ListParameterItemData
+	{
+		/// <summary>
+		/// Optional group heading to display this entry under, if the picker style supports it.
+		/// </summary>
+		public string? Group { get; set; }
+
+		/// <summary>
+		/// Name of the parameter associated with this list.
+		/// </summary>
+		public string Text { get; set; }
+
+		/// <summary>
+		/// Argument to pass with this parameter.
+		/// </summary>
+		public string? ArgumentIfEnabled { get; set; }
+
+		/// <summary>
+		/// Argument to pass with this parameter.
+		/// </summary>
+		public string? ArgumentIfDisabled { get; set; }
+
+		/// <summary>
+		/// Whether this item is selected by default
+		/// </summary>
+		public bool Default { get; set; }
+
+		/// <summary>
+		/// Private constructor for serialization
+		/// </summary>
+		private ListParameterItemData()
+		{
+			Text = null!;
+		}
+
+		/// <summary>
+		/// Constructor
+		/// </summary>
+		/// <param name="Group">The group to put this parameter in</param>
+		/// <param name="Text">Text to display for this option</param>
+		/// <param name="ArgumentIfEnabled">Argument to pass for this item if it's enabled</param>
+		/// <param name="ArgumentIfDisabled">Argument to pass for this item if it's enabled</param>
+		/// <param name="Default">Whether this item is selected by default</param>
+		public ListParameterItemData(string? Group, string Text, string? ArgumentIfEnabled, string? ArgumentIfDisabled, bool Default)
+		{
+			this.Group = Group;
+			this.Text = Text;
+			this.ArgumentIfEnabled = ArgumentIfEnabled;
+			this.ArgumentIfDisabled = ArgumentIfDisabled;
+			this.Default = Default;
+		}
+
+		/// <summary>
+		/// Converts this data to a model object
+		/// </summary>
+		/// <returns>New <see cref="ListParameterItem"/> object</returns>
+		public ListParameterItem ToModel()
+		{
+			return new ListParameterItem(Group, Text, ArgumentIfEnabled, ArgumentIfDisabled, Default);
+		}
+	}
+
+	/// <summary>
+	/// Allows the user to select a value from a constrained list of choices
+	/// </summary>
+	[JsonDiscriminator("List")]
+	public class ListParameterData : ParameterData
+	{
+		/// <summary>
+		/// Label to display next to this parameter. Defaults to the parameter name.
+		/// </summary>
+		public string Label { get; set; }
+
+		/// <summary>
+		/// The type of list parameter
+		/// </summary>
+		public ListParameterStyle Style { get; set; }
+
+		/// <summary>
+		/// List of values to display in the list
+		/// </summary>
+		public List<ListParameterItemData> Items { get; set; }
+
+		/// <summary>
+		/// Tool tip text to display
+		/// </summary>
+		public string? ToolTip { get; set; }
+
+		/// <summary>
+		/// Private constructor
+		/// </summary>
+		private ListParameterData()
+		{
+			Label = null!;
+			Items = null!;
+		}
+
+		/// <summary>
+		/// List of possible values
+		/// </summary>
+		/// <param name="Label">Label to show next to this parameter</param>
+		/// <param name="Style">Type of picker to show</param>
+		/// <param name="Items">Entries for this list</param>
+		/// <param name="ToolTip">Tool tip text to display</param>
+		public ListParameterData(string Label, ListParameterStyle Style, List<ListParameterItemData> Items, string? ToolTip)
+		{
+			this.Label = Label;
+			this.Style = Style;
+			this.Items = Items;
+			this.ToolTip = ToolTip;
+		}
+
+		/// <summary>
+		/// Converts this data to a model object
+		/// </summary>
+		/// <returns>New <see cref="ListParameter"/> object</returns>
+		public override Parameter ToModel()
+		{
+			return new ListParameter(Label, Style, Items.ConvertAll(x => x.ToModel()), ToolTip);
+		}
+	}
+
+	/// <summary>
+	/// Allows the user to toggle an option on or off
+	/// </summary>
+	[JsonDiscriminator("Bool")]
+	public class BoolParameterData : ParameterData
+	{
+		/// <summary>
+		/// Name of the parameter associated with this parameter.
+		/// </summary>
+		public string Label { get; set; }
+
+		/// <summary>
+		/// Value if enabled
+		/// </summary>
+		public string? ArgumentIfEnabled { get; set; }
+
+		/// <summary>
+		/// Value if disabled
+		/// </summary>
+		public string? ArgumentIfDisabled { get; set; }
+
+		/// <summary>
+		/// Whether this argument is enabled by default
+		/// </summary>
+		public bool Default { get; set; }
+
+		/// <summary>
+		/// Tool tip text to display
+		/// </summary>
+		public string? ToolTip { get; set; }
+
+		/// <summary>
+		/// Private constructor for serialization
+		/// </summary>
+		private BoolParameterData()
+		{
+			Label = null!;
+		}
+
+		/// <summary>
+		/// Constructor
+		/// </summary>
+		/// <param name="Label">Label to show next to this parameter</param>
+		/// <param name="ArgumentIfEnabled">Value if enabled</param>
+		/// <param name="ArgumentIfDisabled">Value if disabled</param>
+		/// <param name="Default">Whether this option is enabled by default</param>
+		/// <param name="ToolTip">The tool tip text to display</param>
+		public BoolParameterData(string Label, string? ArgumentIfEnabled, string? ArgumentIfDisabled, bool Default, string? ToolTip)
+		{
+			this.Label = Label;
+			this.ArgumentIfEnabled = ArgumentIfEnabled;
+			this.ArgumentIfDisabled = ArgumentIfDisabled;
+			this.Default = Default;
+			this.ToolTip = ToolTip;
+		}
+
+		/// <summary>
+		/// Converts this data to a model object
+		/// </summary>
+		/// <returns>New <see cref="BoolParameter"/> object</returns>
+		public override Parameter ToModel()
+		{
+			return new BoolParameter(Label, ArgumentIfEnabled, ArgumentIfDisabled, Default, ToolTip);
+		}
+	}
+	
+	/// <summary>
+	/// Counter for executing a job based on a template
+	/// </summary>
+	public class TemplateCounter
+	{
+		/// <summary>
+		/// The property name
+		/// </summary>
+		[Required]
+		public string PropertyName { get; set; }
+
+		/// <summary>
+		/// The counter name
+		/// </summary>
+		[Required]
+		public string CounterName { get; set; }
+
+		/// <summary>
+		/// Private constructor for serialization
+		/// </summary>
+		private TemplateCounter()
+		{
+			PropertyName = null!;
+			CounterName = null!;
+		}
+
+		/// <summary>
+		/// Constructor
+		/// </summary>
+		/// <param name="Counter">Counter to construct from</param>
+		public TemplateCounter(ITemplateCounter Counter)
+		{
+			this.PropertyName = Counter.PropertyName;
+			this.CounterName = Counter.CounterName;
+		}
+	}
+
+	/// <summary>
+	/// Parameters to create a new template
+	/// </summary>
+	public class CreateTemplateRequest
+	{
+		/// <summary>
+		/// Name for the new template
+		/// </summary>
+		[Required]
+		public string Name { get; set; } = null!;
+
+		/// <summary>
+		/// Default priority for this job
+		/// </summary>
+		public Priority? Priority { get; set; }
+
+		/// <summary>
+		/// Whether to allow preflights of this template
+		/// </summary>
+		public bool AllowPreflights { get; set; } = true;
+
+		/// <summary>
+		/// Initial agent type to parse the buildgraph script on
+		/// </summary>
+		public string? InitialAgentType { get; set; }
+
+		/// <summary>
+		/// Path to a file within the stream to submit to generate a new changelist for jobs
+		/// </summary>
+		public string? SubmitNewChange { get; set; }
+
+		/// <summary>
+		/// Counters for this template
+		/// </summary>
+		public List<TemplateCounter> Counters { get; set; } = new List<TemplateCounter>();
+
+		/// <summary>
+		/// Fixed arguments for the new job
+		/// </summary>
+		public List<string> Arguments { get; set; } = new List<string>();
+
+		/// <summary>
+		/// Parameters for this template
+		/// </summary>
+		public List<ParameterData> Parameters { get; set; } = new List<ParameterData>();
+	}
+
+	/// <summary>
+	/// Response describing a template
+	/// </summary>
+	public class GetTemplateResponseBase
+	{
+		/// <summary>
+		/// Name of the template
+		/// </summary>
+		public string Name { get; set; }
+
+		/// <summary>
+		/// Default priority for this job
+		/// </summary>
+		public Priority? Priority { get; set; }
+
+		/// <summary>
+		/// Whether to allow preflights of this template
+		/// </summary>
+		public bool AllowPreflights { get; set; }
+
+		/// <summary>
+		/// The initial agent type to parse the BuildGraph script on
+		/// </summary>
+		public string? InitialAgentType { get; set; }
+
+		/// <summary>
+		/// Path to a file within the stream to submit to generate a new changelist for jobs
+		/// </summary>
+		public string? SubmitNewChange { get; }
+
+		/// <summary>
+		/// List of template counters
+		/// </summary>
+		public List<TemplateCounter> Counters { get; set; }
+
+		/// <summary>
+		/// Parameters for the job.
+		/// </summary>
+		public List<string> Arguments { get; set; }
+
+		/// <summary>
+		/// List of parameters for this template
+		/// </summary>
+		public List<ParameterData> Parameters { get; set; }
+
+		/// <summary>
+		/// Parameterless constructor for serialization
+		/// </summary>
+		protected GetTemplateResponseBase()
+		{
+			this.Name = null!;
+			this.AllowPreflights = true;
+			this.Counters = new List<TemplateCounter>();
+			this.Arguments = new List<string>();
+			this.Parameters = new List<ParameterData>();
+		}
+
+		/// <summary>
+		/// Constructor
+		/// </summary>
+		/// <param name="Template">The template to construct from</param>
+		public GetTemplateResponseBase(ITemplate Template)
+		{
+			this.Name = Template.Name;
+			this.Priority = Template.Priority;
+			this.AllowPreflights = Template.AllowPreflights;
+			this.InitialAgentType = Template.InitialAgentType;
+			this.SubmitNewChange = Template.SubmitNewChange;
+			this.Counters = Template.Counters.ConvertAll(x => new TemplateCounter(x));
+			this.Arguments = new List<string>(Template.Arguments);
+			this.Parameters = Template.Parameters.ConvertAll(x => x.ToData());
+		}
+	}
+
+	/// <summary>
+	/// Response describing a template
+	/// </summary>
+	public class GetTemplateResponse : GetTemplateResponseBase
+	{
+		/// <summary>
+		/// Unique id of the template
+		/// </summary>
+		public string Id { get; set; }
+
+		/// <summary>
+		/// Parameterless constructor for serialization
+		/// </summary>
+		protected GetTemplateResponse()
+			: base()
+		{
+			this.Id = null!;
+		}
+
+		/// <summary>
+		/// Constructor
+		/// </summary>
+		/// <param name="Template">The template to construct from</param>
+		public GetTemplateResponse(ITemplate Template)
+			: base(Template)
+		{
+			this.Id = Template.Id.ToString();
+		}
+	}
+}
diff --git a/Engine/Source/Programs/Horde/HordeServer/Api/TestData.cs b/Engine/Source/Programs/Horde/HordeServer/Api/TestData.cs
new file mode 100644
--- /dev/null
+++ b/Engine/Source/Programs/Horde/HordeServer/Api/TestData.cs
@@ -0,0 +1,121 @@
+using HordeServer.Models;
+using MongoDB.Bson.Serialization;
+using System;
+using System.Collections.Generic;
+using System.ComponentModel.DataAnnotations;
+using System.Linq;
+using System.Threading.Tasks;
+
+namespace HordeServer.Api
+{
+	/// <summary>
+	/// Response object describing test data to store
+	/// </summary>
+	public class CreateTestDataRequest
+	{
+		/// <summary>
+		/// The job which produced the data
+		/// </summary>
+		[Required]
+		public string JobId { get; set; } = String.Empty;
+
+		/// <summary>
+		/// The step that ran
+		/// </summary>
+		[Required]
+		public string StepId { get; set; } = String.Empty;
+
+		/// <summary>
+		/// Key used to identify the particular data
+		/// </summary>
+		public string Key { get; set; } = String.Empty;
+
+		/// <summary>
+		/// The data stored for this test
+		/// </summary>
+		[Required]
+		public Dictionary<string, object> Data { get; set; } = new Dictionary<string, object>();
+	}
+
+	/// <summary>
+	/// Response object describing the created document
+	/// </summary>
+	public class CreateTestDataResponse
+	{
+		/// <summary>
+		/// The id for the new document
+		/// </summary>
+		public string Id { get; set; }
+
+		/// <summary>
+		/// Constructor
+		/// </summary>
+		/// <param name="Id">Id of the new document</param>
+		public CreateTestDataResponse(string Id)
+		{
+			this.Id = Id;
+		}
+	}
+
+	/// <summary>
+	/// Response object describing test results
+	/// </summary>
+	public class GetTestDataResponse
+	{
+		/// <summary>
+		/// Unique id of the test data
+		/// </summary>
+		public string Id { get; set; }
+
+		/// <summary>
+		/// Stream that generated the test data
+		/// </summary>
+		public string StreamId { get; set; }
+
+		/// <summary>
+		/// The template reference id
+		/// </summary>
+		public string TemplateRefId { get; set; }
+
+		/// <summary>
+		/// The job which produced the data
+		/// </summary>
+		public string JobId { get; set; }
+
+		/// <summary>
+		/// The step that ran
+		/// </summary>
+		public string StepId { get; set; }
+
+		/// <summary>
+		/// The changelist number that contained the data
+		/// </summary>
+		public int Change { get; set; }
+
+		/// <summary>
+		/// Key used to identify the particular data
+		/// </summary>
+		public string Key { get; set; }
+
+		/// <summary>
+		/// The data stored for this test
+		/// </summary>
+		public Dictionary<string, object> Data { get; set; }
+
+		/// <summary>
+		/// Constructor
+		/// </summary>
+		/// <param name="TestData">Test data to construct from</param>
+		internal GetTestDataResponse(ITestData TestData)
+		{
+			this.Id = TestData.Id.ToString();
+			this.StreamId = TestData.StreamId.ToString();
+			this.TemplateRefId = TestData.TemplateRefId.ToString();
+			this.JobId = TestData.JobId.ToString();
+			this.StepId = TestData.StepId.ToString();
+			this.Change = TestData.Change;
+			this.Key = TestData.Key;
+			this.Data = BsonSerializer.Deserialize<Dictionary<string, object>>(TestData.Data);
+		}
+	}
+}
diff --git a/Engine/Source/Programs/Horde/HordeServer/Api/UgsMetadata.cs b/Engine/Source/Programs/Horde/HordeServer/Api/UgsMetadata.cs
new file mode 100644
--- /dev/null
+++ b/Engine/Source/Programs/Horde/HordeServer/Api/UgsMetadata.cs
@@ -0,0 +1,556 @@
+using EpicGames.Core;
+using HordeServer.Models;
+using HordeServer.Services;
+using HordeServer.Utilities;
+using System;
+using System.Collections.Generic;
+using System.ComponentModel.DataAnnotations;
+using System.Linq;
+using System.Text.Json.Serialization;
+using System.Threading.Tasks;
+
+namespace HordeServer.Api
+{
+	/// <summary>
+	/// Review by a user of a particular change
+	/// </summary>
+	public enum UgsUserVote
+	{
+		/// <summary>
+		/// No vote for the current change
+		/// </summary>
+		None,
+
+		/// <summary>
+		/// Succesfully compiled the change
+		/// </summary>
+		CompileSuccess,
+
+		/// <summary>
+		/// Failed to compile the change
+		/// </summary>
+		CompileFailure,
+
+		/// <summary>
+		/// Manually marked the change as good
+		/// </summary>
+		Good,
+
+		/// <summary>
+		/// Manually marked the change as bad
+		/// </summary>
+		Bad
+	}
+
+	/// <summary>
+	/// State of a badge
+	/// </summary>
+	public enum UgsBadgeState
+	{
+		/// <summary>
+		/// Starting work on this badge, outcome currently unknown
+		/// </summary>
+		Starting,
+
+		/// <summary>
+		/// Badge failed
+		/// </summary>
+		Failure,
+
+		/// <summary>
+		/// Badge produced a warning
+		/// </summary>
+		Warning,
+
+		/// <summary>
+		/// Badge succeeded
+		/// </summary>
+		Success,
+
+		/// <summary>
+		/// Badge was skipped
+		/// </summary>
+		Skipped,
+	}
+
+	/// <summary>
+	/// Adds a new badge to the change
+	/// </summary>
+	public class AddUgsBadgeRequest
+	{
+		/// <summary>
+		/// Name of the badge
+		/// </summary>
+		[Required]
+		public string Name { get; set; } = String.Empty;
+
+		/// <summary>
+		/// Url for the badge
+		/// </summary>
+		public Uri? Url { get; set; }
+
+		/// <summary>
+		/// Current status of the badge
+		/// </summary>
+		public UgsBadgeState State { get; set; }
+	}
+
+	/// <summary>
+	/// Request object for adding new metadata to the server
+	/// </summary>
+	public class AddUgsMetadataRequest
+	{
+		/// <summary>
+		/// The stream name
+		/// </summary>
+		[Required]
+		public string Stream { get; set; } = null!;
+
+		/// <summary>
+		/// The changelist number
+		/// </summary>
+		[Required]
+		public int Change { get; set; }
+
+		/// <summary>
+		/// The project name
+		/// </summary>
+		public string? Project { get; set; }
+
+		/// <summary>
+		/// Name of the current user
+		/// </summary>
+		public string? UserName { get; set; }
+
+		/// <summary>
+		/// Whether this changelist has been synced by the user
+		/// </summary>
+		public bool? Synced { get; set; }
+
+		/// <summary>
+		/// State of the user
+		/// </summary>
+		public UgsUserVote? Vote { get; set; }
+
+		/// <summary>
+		/// New starred state for the issue
+		/// </summary>
+		public bool? Starred { get; set; }
+
+		/// <summary>
+		/// Whether the user is investigating
+		/// </summary>
+		public bool? Investigating { get; set; }
+
+		/// <summary>
+		/// Comment for this change
+		/// </summary>
+		public string? Comment { get; set; }
+
+		/// <summary>
+		/// List of badges to add
+		/// </summary>
+		public List<AddUgsBadgeRequest>? Badges { get; set; }
+	}
+
+	/// <summary>
+	/// Information about a user synced to a change
+	/// </summary>
+	public class GetUgsUserResponse
+	{
+		/// <summary>
+		/// Name of the user
+		/// </summary>
+		public string User { get; set; }
+
+		/// <summary>
+		/// Time that the change was synced
+		/// </summary>
+		public long? SyncTime { get; set; }
+
+		/// <summary>
+		/// State of the user
+		/// </summary>
+		public UgsUserVote? Vote { get; set; }
+
+		/// <summary>
+		/// Comment by this user
+		/// </summary>
+		public string? Comment { get; set; }
+
+		/// <summary>
+		/// Whether the user is investigating this change
+		/// </summary>
+		public bool? Investigating { get; set; }
+
+		/// <summary>
+		/// Whether this changelist is starred
+		/// </summary>
+		public bool? Starred { get; set; }
+
+		/// <summary>
+		/// Constructor
+		/// </summary>
+		/// <param name="UserData">The user data object</param>
+		internal GetUgsUserResponse(IUgsUserData UserData)
+		{
+			this.User = UserData.User;
+			this.SyncTime = UserData.SyncTime;
+			this.Vote = (UserData.Vote == UgsUserVote.None) ? (UgsUserVote?)null : UserData.Vote;
+			this.Comment = UserData.Comment;
+			this.Investigating = UserData.Investigating;
+			this.Starred = UserData.Starred;
+		}
+	}
+
+	/// <summary>
+	/// Information about a badge
+	/// </summary>
+	public class GetUgsBadgeResponse
+	{
+		/// <summary>
+		/// Name of the badge
+		/// </summary>
+		public string Name { get; set; }
+
+		/// <summary>
+		/// Url for the badge
+		/// </summary>
+		public Uri? Url { get; set; }
+
+		/// <summary>
+		/// Current status of the badge
+		/// </summary>
+		public UgsBadgeState State { get; set; }
+
+		/// <summary>
+		/// Constructor
+		/// </summary>
+		/// <param name="BadgeData">The badge data object</param>
+		internal GetUgsBadgeResponse(IUgsBadgeData BadgeData)
+		{
+			this.Name = BadgeData.Name;
+			this.Url = BadgeData.Url;
+			this.State = BadgeData.State;
+		}
+	}
+
+	/// <summary>
+	/// Response object for querying metadata
+	/// </summary>
+	public class GetUgsMetadataResponse
+	{
+		/// <summary>
+		/// Number of the changelist
+		/// </summary>
+		public int Change { get; set; }
+
+		/// <summary>
+		/// The project name
+		/// </summary>
+		public string? Project { get; set; }
+
+		/// <summary>
+		/// Information about a user synced to this change
+		/// </summary>
+		public List<GetUgsUserResponse>? Users { get; set; }
+
+		/// <summary>
+		/// Badges for this change
+		/// </summary>
+		public List<GetUgsBadgeResponse>? Badges { get; set; }
+
+		/// <summary>
+		/// Constructor
+		/// </summary>
+		/// <param name="Metadata">Metadata object</param>
+		internal GetUgsMetadataResponse(IUgsMetadata Metadata)
+		{
+			Change = Metadata.Change;
+			Project = Metadata.Project;
+			Users = Metadata.Users?.ConvertAll(x => new GetUgsUserResponse(x));
+			Badges = Metadata.Badges?.ConvertAll(x => new GetUgsBadgeResponse(x));
+		}
+	}
+
+	/// <summary>
+	/// Response object for querying metadata
+	/// </summary>
+	public class GetUgsMetadataListResponse
+	{
+		/// <summary>
+		/// Last time that the metadata was modified
+		/// </summary>
+		public long SequenceNumber { get; set; }
+
+		/// <summary>
+		/// List of changes matching the requested criteria
+		/// </summary>
+		public List<GetUgsMetadataResponse> Items { get; set; } = new List<GetUgsMetadataResponse>();
+	}
+
+	/// <summary>
+	/// Outcome of a particular build
+	/// </summary>
+	public enum UgsIssueBuildOutcome
+	{
+		/// <summary>
+		/// Unknown outcome
+		/// </summary>
+		Unknown,
+
+		/// <summary>
+		/// Build succeeded
+		/// </summary>
+		Success,
+
+		/// <summary>
+		/// Build failed
+		/// </summary>
+		Error,
+
+		/// <summary>
+		/// Build finished with warnings
+		/// </summary>
+		Warning,
+	}
+
+	/// <summary>
+	/// Legacy response describing a build
+	/// </summary>
+	public class GetUgsIssueBuildResponse
+	{
+		/// <summary>
+		/// Identifier for this build
+		/// </summary>
+		public long Id { get; set; }
+
+		/// <summary>
+		/// Path to the stream containing this build
+		/// </summary>
+		public string Stream { get; set; }
+
+		/// <summary>
+		/// The changelist number for this build
+		/// </summary>
+		public int Change { get; set; }
+
+		/// <summary>
+		/// Name of this job
+		/// </summary>
+		public string? JobName { get; set; }
+
+		/// <summary>
+		/// Link to the job
+		/// </summary>
+		public Uri? JobUrl { get; set; }
+
+		/// <summary>
+		/// Name of the job step
+		/// </summary>
+		public string? JobStepName { get; set; }
+
+		/// <summary>
+		/// Link to the job step
+		/// </summary>
+		public Uri? JobStepUrl { get; set; }
+
+		/// <summary>
+		/// Url for this particular error
+		/// </summary>
+		public Uri? ErrorUrl { get; set; }
+
+		/// <summary>
+		/// Outcome of this build (see <see cref="IssueBuildOutcome"/>)
+		/// </summary>
+		public int Outcome { get; set; }
+
+		/// <summary>
+		/// 
+		/// </summary>
+		/// <param name="Stream"></param>
+		/// <param name="Change"></param>
+		/// <param name="Outcome"></param>
+		public GetUgsIssueBuildResponse(string Stream, int Change, IssueBuildOutcome Outcome)
+		{
+			this.Stream = Stream;
+			this.Change = Change;
+			this.Outcome = (int)Outcome;
+		}
+	}
+
+	/// <summary>
+	/// Information about a diagnostic
+	/// </summary>
+	public class GetUgsIssueDiagnosticResponse
+	{
+		/// <summary>
+		/// The corresponding build id
+		/// </summary>
+		public long? BuildId { get; set; }
+
+		/// <summary>
+		/// Message for the diagnostic
+		/// </summary>
+		public string Message { get; set; }
+
+		/// <summary>
+		/// Link to the error
+		/// </summary>
+		public Uri Url { get; set; }
+
+		/// <summary>
+		/// Constructor
+		/// </summary>
+		/// <param name="BuildId">The corresponding build id</param>
+		/// <param name="Message">Message for the diagnostic</param>
+		/// <param name="Url">Link to the diagnostic</param>
+		public GetUgsIssueDiagnosticResponse(long? BuildId, string Message, Uri Url)
+		{
+			this.BuildId = BuildId;
+			this.Message = Message;
+			this.Url = Url;
+		}
+	}
+
+	/// <summary>
+	/// Stores information about a build health issue
+	/// </summary>
+	public class GetUgsIssueResponse
+	{
+		/// <summary>
+		/// Version number for this response
+		/// </summary>
+		public int Version { get; set; } = 1;
+
+		/// <summary>
+		/// The unique object id
+		/// </summary>
+		public int Id { get; set; }
+
+		/// <summary>
+		/// Time at which the issue was created
+		/// </summary>
+		public DateTime CreatedAt { get; set; }
+
+		/// <summary>
+		/// Time at which the issue was retrieved
+		/// </summary>
+		public DateTime RetrievedAt { get; set; }
+
+		/// <summary>
+		/// The associated project for the issue
+		/// </summary>
+		public string? Project { get; set; }
+
+		/// <summary>
+		/// The summary text for this issue
+		/// </summary>
+		public string Summary { get; set; }
+
+		/// <summary>
+		/// Owner of the issue
+		/// </summary>
+		public string? Owner { get; set; }
+
+		/// <summary>
+		/// User that nominated the current owner
+		/// </summary>
+		public string? NominatedBy { get; set; }
+
+		/// <summary>
+		/// Time that the issue was acknowledged
+		/// </summary>
+		public DateTime? AcknowledgedAt { get; set; }
+
+		/// <summary>
+		/// Changelist that fixed this issue
+		/// </summary>
+		public int? FixChange { get; set; }
+
+		/// <summary>
+		/// Time at which the issue was resolved
+		/// </summary>
+		public DateTime? ResolvedAt { get; set; }
+
+		/// <summary>
+		/// Whether to notify the user about this issue
+		/// </summary>
+		public bool bNotify { get; set; }
+
+		/// <summary>
+		/// Whether this issue just contains warnings
+		/// </summary>
+		public bool bWarning { get; set; }
+
+		/// <summary>
+		/// Link to the last build
+		/// </summary>
+		public Uri? BuildUrl { get; set; }
+
+		/// <summary>
+		/// List of streams affected by this issue
+		/// </summary>
+		public List<string> Streams { get; set; }
+
+		/// <summary>
+		/// Constructs a new issue
+		/// </summary>
+		/// <param name="Details">Issue to construct from</param>
+		/// <param name="bNotify">Whether to notify the user about this issue</param>
+		/// <param name="BuildUrl">Link to the last build for this issue</param>
+		public GetUgsIssueResponse(IIssueDetails Details, bool bNotify, Uri? BuildUrl)
+		{
+			IIssue Issue = Details.Issue;
+			this.Id = Issue.Id;
+			this.CreatedAt = Issue.CreatedAt;
+			this.RetrievedAt = DateTime.UtcNow;
+			this.Summary = Issue.Summary;
+			this.Owner = Issue.Owner;
+			this.NominatedBy = Issue.NominatedBy;
+			this.AcknowledgedAt = Issue.AcknowledgedAt;
+			this.FixChange = Issue.FixChange;
+			this.ResolvedAt = Issue.ResolvedAt;
+			this.bNotify = bNotify;
+			this.bWarning = Issue.Severity == IssueSeverity.Warning;
+			this.BuildUrl = BuildUrl;
+			this.Streams = Details.Spans.Select(x => x.StreamName ?? x.StreamId.ToString()).Distinct().ToList();
+		}
+	}
+
+	/// <summary>
+	/// Request an issue to be updated
+	/// </summary>
+	public class UpdateUgsIssueRequest
+	{
+		/// <summary>
+		/// New owner of the issue
+		/// </summary>
+		public string? Owner { get; set; }
+
+		/// <summary>
+		/// User than nominates the new owner
+		/// </summary>
+		public string? NominatedBy { get; set; }
+
+		/// <summary>
+		/// Whether the issue has been acknowledged
+		/// </summary>
+		public bool? Acknowledged { get; set; }
+
+		/// <summary>
+		/// Name of the user that declines the issue
+		/// </summary>
+		public string? DeclinedBy { get; set; }
+
+		/// <summary>
+		/// The change at which the issue is claimed fixed. 0 = not fixed, -1 = systemic issue.
+		/// </summary>
+		public int? FixChange { get; set; }
+
+		/// <summary>
+		/// Whether the issue should be marked as resolved
+		/// </summary>
+		public bool? Resolved { get; set; }
+	}
+}
\ No newline at end of file
diff --git a/Engine/Source/Programs/Horde/HordeServer/Api/Users.cs b/Engine/Source/Programs/Horde/HordeServer/Api/Users.cs
new file mode 100644
--- /dev/null
+++ b/Engine/Source/Programs/Horde/HordeServer/Api/Users.cs
@@ -0,0 +1,94 @@
+using EpicGames.Core;
+using HordeServer.Models;
+using MongoDB.Bson;
+using System;
+using System.Collections.Generic;
+using System.Linq;
+using System.Text.Json;
+using System.Threading.Tasks;
+
+namespace HordeServer.Api
+{
+	/// <summary>
+	/// Response describing the current user
+	/// </summary>
+	public class GetUserResponse
+	{
+		/// <summary>
+		/// Id of the user
+		/// </summary>
+		public string Id { get; set; }
+
+		/// <summary>
+		/// Claims for the user
+		/// </summary>
+		public List<UserClaim> Claims { get; set; }
+
+		/// <summary>
+		/// Whether to enable experimental features for this user
+		/// </summary>
+		public bool EnableExperimentalFeatures { get; set; }
+
+		/// <summary>
+		/// Whether to enable slack notifications for this user
+		/// </summary>
+		public bool EnableIssueNotifications { get; set; }
+
+		/// <summary>
+		/// Settings for the dashboard
+		/// </summary>
+		public object DashboardSettings { get; set; }
+
+		/// <summary>
+		/// List of pinned job ids
+		/// </summary>
+		public List<string> PinnedJobIds { get; set; }
+
+		/// <summary>
+		/// Constructor
+		/// </summary>
+		/// <param name="User"></param>
+		/// <param name="Claims"></param>
+		/// <param name="Settings"></param>
+		public GetUserResponse(IUser User, IUserClaims Claims, IUserSettings Settings)
+		{
+			this.Id = User.Id.ToString();
+			this.Claims = Claims.Claims.Select(x => new UserClaim(x)).ToList();
+			this.EnableExperimentalFeatures = Settings.EnableExperimentalFeatures;
+			this.EnableIssueNotifications = Settings.EnableIssueNotifications;
+			this.DashboardSettings = BsonTypeMapper.MapToDotNetValue(Settings.DashboardSettings);
+			this.PinnedJobIds = Settings.PinnedJobIds.ConvertAll(x => x.ToString());
+		}
+	}
+
+	/// <summary>
+	/// Request to update settings for a user
+	/// </summary>
+	public class UpdateUserRequest
+	{
+		/// <summary>
+		/// Whether to enable experimental features for this user
+		/// </summary>
+		public bool? EnableExperimentalFeatures { get; set; }
+
+		/// <summary>
+		/// Whether to enable Slack notifications for this user
+		/// </summary>
+		public bool? EnableIssueNotifications { get; set; }
+
+		/// <summary>
+		/// New dashboard settings
+		/// </summary>
+		public JsonElement? DashboardSettings { get; set; }
+
+		/// <summary>
+		/// Job ids to add to the pinned list
+		/// </summary>
+		public List<string>? AddPinnedJobIds { get; set; }
+
+		/// <summary>
+		/// Jobs ids to remove from the pinned list
+		/// </summary>
+		public List<string>? RemovePinnedJobIds { get; set; }
+	}
+}
diff --git a/Engine/Source/Programs/Horde/HordeServer/Api/Workers.cs b/Engine/Source/Programs/Horde/HordeServer/Api/Workers.cs
new file mode 100644
--- /dev/null
+++ b/Engine/Source/Programs/Horde/HordeServer/Api/Workers.cs
@@ -0,0 +1,229 @@
+// Copyright Epic Games, Inc. All Rights Reserved.
+
+using HordeServer.Models;
+using System;
+using System.Collections.Generic;
+using System.ComponentModel.DataAnnotations;
+using System.Linq;
+using System.Threading.Tasks;
+
+namespace HordeServer.Api
+{
+	/// <summary>
+	/// Information about a device attached to this agent
+	/// </summary>
+	public class GetDeviceCapabilitiesResponse
+	{
+		/// <summary>
+		/// Logical name of this device
+		/// </summary>
+		[Required]
+		public string Name { get; set; }
+
+		/// <summary>
+		/// Required properties for this device, in the form "KEY=VALUE"
+		/// </summary>
+		public List<string>? Properties { get; set; }
+
+		/// <summary>
+		/// Required resources for this node. If null, the node will assume exclusive access to the device.
+		/// </summary>
+		public Dictionary<string, int>? Resources { get; set; }
+
+		/// <summary>
+		/// Private constructor for serialization
+		/// </summary>
+		private GetDeviceCapabilitiesResponse()
+		{
+			this.Name = null!;
+		}
+
+		/// <summary>
+		/// Converts this object to an API response
+		/// </summary>
+		/// <returns>Response object</returns>
+		public GetDeviceCapabilitiesResponse(DeviceCapabilities DeviceCapabilities)
+		{
+			this.Name = DeviceCapabilities.Handle;
+			this.Properties = DeviceCapabilities.Properties.ToList();
+			this.Resources = DeviceCapabilities.Resources;
+		}
+	}
+
+	/// <summary>
+	/// Information about the capabilities of this agent
+	/// </summary>
+	public class GetAgentCapabilitiesResponse
+	{
+		/// <summary>
+		/// Information about the devices required for this node to run
+		/// </summary>
+		public List<GetDeviceCapabilitiesResponse>? Devices { get; set; }
+
+		/// <summary>
+		/// Global agent properties for this node
+		/// </summary>
+		public List<string>? Properties { get; set; }
+
+		/// <summary>
+		/// Default constructor
+		/// </summary>
+		public GetAgentCapabilitiesResponse()
+		{
+		}
+
+		/// <summary>
+		/// Constructor
+		/// </summary>
+		/// <param name="AgentCapabilities">The agent capabiilities to construct from</param>
+		public GetAgentCapabilitiesResponse(AgentCapabilities AgentCapabilities)
+		{
+			this.Devices = AgentCapabilities.Devices.ConvertAll(x => new GetDeviceCapabilitiesResponse(x));
+			this.Properties = AgentCapabilities.Properties?.ToList();
+		}
+	}
+
+	/// <summary>
+	/// Information about the device requirements of a node
+	/// </summary>
+	public class CreateDeviceRequirementsRequest
+	{
+		/// <summary>
+		/// Logical name of this device
+		/// </summary>
+		public string? Name { get; set; }
+
+		/// <summary>
+		/// Required properties for this device, in the form "KEY=VALUE"
+		/// </summary>
+		public HashSet<string>? Properties { get; set; }
+
+		/// <summary>
+		/// Required resources for this node. If null, the node will assume exclusive access to the device.
+		/// </summary>
+		public Dictionary<string, int>? Resources { get; set; }
+
+		/// <summary>
+		/// Private constructor for serialization
+		/// </summary>
+		private CreateDeviceRequirementsRequest()
+		{
+			this.Name = null!;
+		}
+
+		/// <summary>
+		/// Constructor
+		/// </summary>
+		/// <param name="Name">Logical name of this device</param>
+		public CreateDeviceRequirementsRequest(string Name)
+		{
+			this.Name = Name;
+		}
+	}
+
+	/// <summary>
+	/// Information about the agent requirements of node
+	/// </summary>
+	public class CreateAgentRequirementsRequest
+	{
+		/// <summary>
+		/// Information about the devices required for this node to run
+		/// </summary>
+		public List<CreateDeviceRequirementsRequest>? Devices { get; set; }
+
+		/// <summary>
+		/// Global agent properties for this node
+		/// </summary>
+		public List<string>? Properties { get; set; }
+
+		/// <summary>
+		/// Configurable settings for this agent
+		/// </summary>
+		public Dictionary<string, string>? Configs { get; set; }
+
+		/// <summary>
+		/// Whether the agent can be shared with another job
+		/// </summary>
+		public bool Shared { get; set; }
+	}
+
+	/// <summary>
+	/// Information about the device requirements of a node
+	/// </summary>
+	public class GetDeviceRequirementsResponse
+	{
+		/// <summary>
+		/// Logical name of this device
+		/// </summary>
+		[Required]
+		public string? Name { get; set; }
+
+		/// <summary>
+		/// Required properties for this device, in the form "KEY=VALUE"
+		/// </summary>
+		public HashSet<string>? Properties { get; set; }
+
+		/// <summary>
+		/// Required resources for this node. If null, the node will assume exclusive access to the device.
+		/// </summary>
+		public Dictionary<string, int>? Resources { get; set; }
+
+		/// <summary>
+		/// Private constructor for serialization
+		/// </summary>
+		private GetDeviceRequirementsResponse()
+		{
+			this.Name = null!;
+		}
+
+		/// <summary>
+		/// Constructor
+		/// </summary>
+		/// <param name="DeviceRequirements">The device requirements</param>
+		public GetDeviceRequirementsResponse(DeviceRequirements DeviceRequirements)
+		{
+			this.Name = DeviceRequirements.Handle;
+			this.Properties = DeviceRequirements.Properties;
+			this.Resources = DeviceRequirements.Resources;
+		}
+	}
+
+	/// <summary>
+	/// Information about the agent requirements of node
+	/// </summary>
+	public class GetAgentRequirementsResponse
+	{
+		/// <summary>
+		/// Information about the devices required for this node to run
+		/// </summary>
+		public List<GetDeviceRequirementsResponse>? Devices { get; set; }
+
+		/// <summary>
+		/// Global agent properties for this node
+		/// </summary>
+		public List<string>? Properties { get; set; }
+
+		/// <summary>
+		/// Whether the agent can be shared with another job
+		/// </summary>
+		public bool Shared { get; set; }
+
+		/// <summary>
+		/// Default constructor
+		/// </summary>
+		public GetAgentRequirementsResponse()
+		{
+		}
+
+		/// <summary>
+		/// Constructor
+		/// </summary>
+		/// <param name="AgentRequirements">Agent requirements to construct from</param>
+		public GetAgentRequirementsResponse(AgentRequirements AgentRequirements)
+		{
+			this.Devices = AgentRequirements.Devices?.ConvertAll(x => new GetDeviceRequirementsResponse(x));
+			this.Properties = AgentRequirements.Properties?.ToList();
+			this.Shared = AgentRequirements.Shared;
+		}
+	}
+}
diff --git a/Engine/Source/Programs/Horde/HordeServer/Authentication/AnonymousHandler.cs b/Engine/Source/Programs/Horde/HordeServer/Authentication/AnonymousHandler.cs
new file mode 100644
--- /dev/null
+++ b/Engine/Source/Programs/Horde/HordeServer/Authentication/AnonymousHandler.cs
@@ -0,0 +1,44 @@
+using Microsoft.AspNetCore.Authentication;
+using Microsoft.Extensions.Logging;
+using Microsoft.Extensions.Options;
+using System;
+using System.Collections.Generic;
+using System.Linq;
+using System.Security.Claims;
+using System.Text.Encodings.Web;
+using System.Threading.Tasks;
+
+namespace HordeServer.Authentication
+{
+	class AnonymousAuthenticationOptions : AuthenticationSchemeOptions
+	{
+		public string? AdminClaimType { get; set; }
+		public string? AdminClaimValue { get; set; }
+	}
+
+	class AnonymousAuthenticationHandler : AuthenticationHandler<AnonymousAuthenticationOptions>
+	{
+		public const string AuthenticationScheme = "Anonymous";
+
+		public AnonymousAuthenticationHandler(IOptionsMonitor<AnonymousAuthenticationOptions> Options,
+			ILoggerFactory Logger, UrlEncoder Encoder, ISystemClock Clock)
+			: base(Options, Logger, Encoder, Clock)
+		{
+		}
+
+		protected override Task<AuthenticateResult> HandleAuthenticateAsync()
+		{
+			Claim[] Claims =
+			{
+					new Claim(ClaimTypes.Name, AuthenticationScheme),
+					new Claim(Options.AdminClaimType, Options.AdminClaimValue)
+				};
+			ClaimsIdentity Identity = new ClaimsIdentity(Claims, Scheme.Name);
+
+			ClaimsPrincipal Principal = new ClaimsPrincipal(Identity);
+			AuthenticationTicket Ticket = new AuthenticationTicket(Principal, Scheme.Name);
+
+			return Task.FromResult(AuthenticateResult.Success(Ticket));
+		}
+	}
+}
diff --git a/Engine/Source/Programs/Horde/HordeServer/Authentication/HordeJwtBearerHandler.cs b/Engine/Source/Programs/Horde/HordeServer/Authentication/HordeJwtBearerHandler.cs
new file mode 100644
--- /dev/null
+++ b/Engine/Source/Programs/Horde/HordeServer/Authentication/HordeJwtBearerHandler.cs
@@ -0,0 +1,75 @@
+using HordeServer.Services;
+using HordeServer.Utilities;
+using Microsoft.AspNetCore.Authentication;
+using Microsoft.AspNetCore.Authentication.JwtBearer;
+using Microsoft.Extensions.Logging;
+using Microsoft.Extensions.Options;
+using System;
+using System.Collections.Generic;
+using System.IdentityModel.Tokens.Jwt;
+using System.Linq;
+using System.Text.Encodings.Web;
+using System.Threading.Tasks;
+
+namespace HordeServer.Authentication
+{
+	/// <summary>
+	/// JWT handler for server-issued bearer tokens. These tokens are signed using a randomly generated key per DB instance.
+	/// </summary>
+	class HordeJwtBearerHandler : JwtBearerHandler
+	{
+		/// <summary>
+		/// Default name of the authentication scheme
+		/// </summary>
+		public const string AuthenticationScheme = "ServerJwt";
+
+		DatabaseService DatabaseService;
+
+		public HordeJwtBearerHandler(ILoggerFactory Logger, UrlEncoder Encoder, ISystemClock Clock, DatabaseService DatabaseService, IOptionsMonitorCache<JwtBearerOptions> OptionsCache)
+			: base(GetOptionsMonitor(DatabaseService, OptionsCache), Logger, Encoder, Clock)
+		{
+			this.DatabaseService = DatabaseService;
+		}
+
+		private static IOptionsMonitor<JwtBearerOptions> GetOptionsMonitor(DatabaseService DatabaseService, IOptionsMonitorCache<JwtBearerOptions> OptionsCache)
+		{
+			ConfigureNamedOptions<JwtBearerOptions> NamedOptions = new ConfigureNamedOptions<JwtBearerOptions>(AuthenticationScheme, Options => Configure(Options, DatabaseService));
+			OptionsFactory<JwtBearerOptions> OptionsFactory = new OptionsFactory<JwtBearerOptions>(new[] { NamedOptions }, Array.Empty<IPostConfigureOptions<JwtBearerOptions>>());
+			return new OptionsMonitor<JwtBearerOptions>(OptionsFactory, Array.Empty<IOptionsChangeTokenSource<JwtBearerOptions>>(), OptionsCache);
+		}
+
+		private static void Configure(JwtBearerOptions Options, DatabaseService DatabaseService)
+		{
+			Options.TokenValidationParameters.ValidateAudience = false;
+
+			Options.TokenValidationParameters.RequireExpirationTime = false;
+			Options.TokenValidationParameters.ValidateLifetime = true;
+
+			Options.TokenValidationParameters.ValidIssuer = DatabaseService.JwtIssuer;
+			Options.TokenValidationParameters.ValidateIssuer = true;
+
+			Options.TokenValidationParameters.ValidateIssuerSigningKey = true;
+			Options.TokenValidationParameters.IssuerSigningKey = DatabaseService.JwtSigningKey;
+		}
+
+		protected override Task<AuthenticateResult> HandleAuthenticateAsync()
+		{
+			// Silent fail if this JWT is not issued by the server
+			string? Token;
+			if (!JwtUtils.TryGetBearerToken(Request, "Bearer ", out Token))
+			{
+				return Task.FromResult(AuthenticateResult.NoResult());
+			}
+
+			// Validate that it's from the correct issuer
+			JwtSecurityToken? JwtToken;
+			if (!JwtUtils.TryParseJwt(Token, out JwtToken) || !String.Equals(JwtToken.Issuer, DatabaseService.JwtIssuer, StringComparison.Ordinal))
+			{
+				return Task.FromResult(AuthenticateResult.NoResult());
+			}
+
+			// Pass it to the base class
+			return base.HandleAuthenticateAsync();
+		}
+	}
+}
diff --git a/Engine/Source/Programs/Horde/HordeServer/Authentication/OktaHandler.cs b/Engine/Source/Programs/Horde/HordeServer/Authentication/OktaHandler.cs
new file mode 100644
--- /dev/null
+++ b/Engine/Source/Programs/Horde/HordeServer/Authentication/OktaHandler.cs
@@ -0,0 +1,131 @@
+using HordeServer.Collections;
+using HordeServer.Models;
+using HordeServer.Utilities;
+using Microsoft.AspNetCore.Authentication;
+using Microsoft.AspNetCore.Authentication.OAuth.Claims;
+using Microsoft.AspNetCore.Authentication.OpenIdConnect;
+using Microsoft.Extensions.DependencyInjection;
+using Microsoft.Extensions.DependencyInjection.Extensions;
+using Microsoft.Extensions.Logging;
+using Microsoft.Extensions.Options;
+using Microsoft.IdentityModel.Protocols.OpenIdConnect;
+using System;
+using System.Collections.Generic;
+using System.Linq;
+using System.Security.Claims;
+using System.Text.Encodings.Web;
+using System.Text.Json;
+using System.Threading.Tasks;
+
+namespace HordeServer.Authentication
+{
+	class OktaDefaults
+	{
+		public const string AuthenticationScheme = "Okta" + OpenIdConnectDefaults.AuthenticationScheme;
+	}
+
+	class OktaHandler : OpenIdConnectHandler
+	{
+		IUserCollection UserCollection;
+
+		public OktaHandler(IOptionsMonitor<OpenIdConnectOptions> options, ILoggerFactory logger, HtmlEncoder htmlEncoder, UrlEncoder encoder, ISystemClock clock, IUserCollection UserCollection)
+			: base(options, logger, htmlEncoder, encoder, clock)
+		{
+			this.UserCollection = UserCollection;
+		}
+
+		public static void AddUserInfoClaims(JsonElement UserInfo, ClaimsIdentity Identity)
+		{
+			JsonElement UserElement;
+			if (UserInfo.TryGetProperty("preferred_username", out UserElement))
+			{
+				if (Identity.FindFirst(ClaimTypes.Name) == null)
+				{
+					Identity.AddClaim(new Claim(ClaimTypes.Name, UserElement.ToString()));
+				}
+				if (Identity.FindFirst(HordeClaimTypes.User) == null)
+				{
+					Identity.AddClaim(new Claim(HordeClaimTypes.User, UserElement.ToString()));
+				}
+				if (Identity.FindFirst(HordeClaimTypes.PerforceUser) == null)
+				{
+					Identity.AddClaim(new Claim(HordeClaimTypes.PerforceUser, UserElement.ToString()));
+				}
+			}
+
+			JsonElement GroupsElement;
+			if (UserInfo.TryGetProperty("groups", out GroupsElement) && GroupsElement.ValueKind == JsonValueKind.Array)
+			{
+				for (int Idx = 0; Idx < GroupsElement.GetArrayLength(); Idx++)
+				{
+					Identity.AddClaim(new Claim(ClaimTypes.Role, GroupsElement[Idx].ToString()));
+				}
+			}
+
+			JsonElement EmailElement;
+			if (UserInfo.TryGetProperty("email", out EmailElement) && Identity.FindFirst(ClaimTypes.Email) == null)
+			{
+				Identity.AddClaim(new Claim(ClaimTypes.Email, EmailElement.GetString()));
+			}
+		}
+
+		protected override async Task<HandleRequestResult> HandleRemoteAuthenticateAsync()
+		{
+			// Authenticate with the OIDC provider
+			HandleRequestResult Result = await base.HandleRemoteAuthenticateAsync();
+			if (!Result.Succeeded)
+			{
+				return Result;
+			}
+
+			ClaimsIdentity Identity = (ClaimsIdentity)Result.Principal.Identity;
+
+			string Login = Identity.FindFirst(ClaimTypes.Name)!.Value;
+			string? Name = Identity.FindFirst("name")?.Value;
+			string? Email = Identity.FindFirst(ClaimTypes.Email)?.Value;
+
+			IUser User = await UserCollection.FindOrAddUserByLoginAsync(Login, Name, Email);
+			Identity.AddClaim(new Claim(HordeClaimTypes.UserId, User.Id.ToString()));
+
+			await UserCollection.UpdateClaimsAsync(User.Id, Identity.Claims.Select(x => new UserClaim(x.Type, x.Value)));
+
+			return Result;
+		}
+	}
+
+	static class OktaExtensions
+	{
+		class MapRolesClaimAction : ClaimAction
+		{
+			public MapRolesClaimAction()
+				: base(ClaimTypes.Role, ClaimTypes.Role)
+			{
+			}
+
+			public override void Run(JsonElement UserData, ClaimsIdentity Identity, string Issuer)
+			{
+				OktaHandler.AddUserInfoClaims(UserData, Identity);
+			}
+		}
+
+		static void ApplyDefaultOktaOptions(OpenIdConnectOptions Options, Action<OpenIdConnectOptions> Handler)
+		{
+			Options.Scope.Add("profile");
+			Options.Scope.Add("groups");
+			Options.Scope.Add("email");
+			Options.ResponseType = OpenIdConnectResponseType.Code;
+			Options.GetClaimsFromUserInfoEndpoint = true;
+			Options.SaveTokens = true;
+			Options.TokenValidationParameters.NameClaimType = "name";
+			Options.ClaimActions.Add(new MapRolesClaimAction());
+
+			Handler(Options);
+		}
+
+		public static void AddOkta(this AuthenticationBuilder Builder, string AuthenticationScheme, string DisplayName, Action<OpenIdConnectOptions> Handler)
+		{
+			Builder.Services.TryAddEnumerable(ServiceDescriptor.Singleton<IPostConfigureOptions<OpenIdConnectOptions>, OpenIdConnectPostConfigureOptions>());
+			Builder.AddRemoteScheme<OpenIdConnectOptions, OktaHandler>(AuthenticationScheme, DisplayName, Options => ApplyDefaultOktaOptions(Options, Handler));
+		}
+	}
+}
diff --git a/Engine/Source/Programs/Horde/HordeServer/Authentication/ServiceAccountAuthHandler.cs b/Engine/Source/Programs/Horde/HordeServer/Authentication/ServiceAccountAuthHandler.cs
new file mode 100644
--- /dev/null
+++ b/Engine/Source/Programs/Horde/HordeServer/Authentication/ServiceAccountAuthHandler.cs
@@ -0,0 +1,70 @@
+﻿using Microsoft.AspNetCore.Authentication;
+using Microsoft.Extensions.Logging;
+using Microsoft.Extensions.Options;
+using System;
+using System.Collections.Generic;
+using System.Linq;
+using System.Security.Claims;
+using System.Text.Encodings.Web;
+using System.Threading.Tasks;
+using HordeServer.Collections;
+using HordeServer.Models;
+using Microsoft.Net.Http.Headers;
+
+namespace HordeServer.Authentication
+{
+	class ServiceAccountAuthOptions : AuthenticationSchemeOptions
+	{
+	}
+
+	class ServiceAccountAuthHandler : AuthenticationHandler<ServiceAccountAuthOptions>
+	{
+		public const string AuthenticationScheme = "ServiceAccount";
+		public const string Prefix = "ServiceAccount";
+
+		private readonly IServiceAccountCollection ServiceAccounts;
+
+		public ServiceAccountAuthHandler(IOptionsMonitor<ServiceAccountAuthOptions> Options,
+			ILoggerFactory Logger, UrlEncoder Encoder, ISystemClock Clock, IServiceAccountCollection ServiceAccounts)
+			: base(Options, Logger, Encoder, Clock)
+		{
+			this.ServiceAccounts = ServiceAccounts;
+		}
+		
+		protected override async Task<AuthenticateResult> HandleAuthenticateAsync()
+		{
+			if (!Context.Request.Headers.TryGetValue(HeaderNames.Authorization, out var HeaderValue))
+			{
+				return AuthenticateResult.NoResult();
+			}
+
+			if (HeaderValue.Count < 1)
+			{
+				return AuthenticateResult.NoResult();
+			}
+
+			if (!HeaderValue[0].StartsWith(Prefix, StringComparison.Ordinal))
+			{
+				return AuthenticateResult.NoResult();
+			}
+			
+			string Token = HeaderValue[0].Replace(Prefix, "", StringComparison.Ordinal).Trim();
+			IServiceAccount? ServiceAccount = await ServiceAccounts.GetBySecretTokenAsync(Token);
+
+			if (ServiceAccount == null)
+			{
+				return AuthenticateResult.Fail($"Service account for token {Token} not found");
+			}
+
+			List<Claim> Claims = new List<Claim>(10);
+			Claims.Add(new Claim(ClaimTypes.Name, AuthenticationScheme));
+			Claims.AddRange(ServiceAccount.Claims.Select(ClaimPair => new Claim(ClaimPair.Key, ClaimPair.Value)));
+
+			ClaimsIdentity Identity = new ClaimsIdentity(Claims, Scheme.Name);
+			ClaimsPrincipal Principal = new ClaimsPrincipal(Identity);
+			AuthenticationTicket Ticket = new AuthenticationTicket(Principal, Scheme.Name);
+
+			return AuthenticateResult.Success(Ticket);
+		}
+	}
+}
