Initial drop of Horde Server/Agent for UE5/Main. Still missing some Epic-specific functionality required to build/run correctly.

diff --git a/Engine/Source/Programs/Horde/HordeServer/Collections/DatadogTracingCollection.cs b/Engine/Source/Programs/Horde/HordeServer/Collections/DatadogTracingCollection.cs
new file mode 100644
--- /dev/null
+++ b/Engine/Source/Programs/Horde/HordeServer/Collections/DatadogTracingCollection.cs
@@ -0,0 +1,693 @@
+﻿using System.Collections.Generic;
+using System.Threading;
+using System.Threading.Tasks;
+using Datadog.Trace;
+using MongoDB.Bson;
+using MongoDB.Bson.Serialization;
+using MongoDB.Driver;
+
+namespace HordeServer.Collections
+{
+	/// <summary>
+	/// Wrap a IMongoCollection with Datadog tracing
+	///
+	/// Will capture the entire invocation of MongoDB calls, including serialization.
+	/// The other command logging (not in this file) only deals with queries sent to the server at the protocol level.
+	/// </summary>
+	/// <typeparam name="T">A MongoDB document</typeparam>
+	public class DatadogTracingCollection<T> : IMongoCollection<T>
+	{
+		private readonly IMongoCollection<T> Collection;
+		private const string Prefix = "MongoCollection.";
+
+		/// <summary>
+		/// Constructor
+		/// </summary>
+		/// <param name="Collection">Collection to wrap with tracing</param>
+		public DatadogTracingCollection(IMongoCollection<T> Collection)
+		{
+			this.Collection = Collection;
+		}
+
+#pragma warning disable CS0618		
+		/// <inheritdoc />
+		public IAsyncCursor<TResult> Aggregate<TResult>(PipelineDefinition<T, TResult> pipeline, AggregateOptions Options = null!,
+			CancellationToken  CancellationToken = new CancellationToken())
+		{
+			return Collection.Aggregate(pipeline, Options,  CancellationToken);
+		}
+
+		/// <inheritdoc />
+		public IAsyncCursor<TResult> Aggregate<TResult>(IClientSessionHandle Session, PipelineDefinition<T, TResult> pipeline,
+			AggregateOptions Options = null!, CancellationToken  CancellationToken = new CancellationToken())
+		{
+			return Collection.Aggregate(Session, pipeline, Options,  CancellationToken);
+		}
+
+		/// <inheritdoc />
+		public async Task<IAsyncCursor<TResult>> AggregateAsync<TResult>(PipelineDefinition<T, TResult> pipeline, AggregateOptions Options = null!,
+			CancellationToken  CancellationToken = new CancellationToken())
+		{
+			using Scope _ = Tracer.Instance.StartActive(Prefix + "AggregateAsync");
+			return await Collection.AggregateAsync(pipeline, Options,  CancellationToken);
+		}
+
+		/// <inheritdoc />
+		public async Task<IAsyncCursor<TResult>> AggregateAsync<TResult>(IClientSessionHandle Session, PipelineDefinition<T, TResult> pipeline, AggregateOptions Options = null!,
+			CancellationToken  CancellationToken = new CancellationToken())
+		{
+			using Scope _ = Tracer.Instance.StartActive(Prefix + "AggregateAsync");
+			return await Collection.AggregateAsync(Session, pipeline, Options, CancellationToken);
+		}
+
+		/// <inheritdoc />
+		public BulkWriteResult<T> BulkWrite(IEnumerable<WriteModel<T>> requests, BulkWriteOptions Options = null!,
+			CancellationToken  CancellationToken = new CancellationToken())
+		{
+			return Collection.BulkWrite(requests, Options,  CancellationToken);
+		}
+
+		/// <inheritdoc />
+		public BulkWriteResult<T> BulkWrite(IClientSessionHandle Session, IEnumerable<WriteModel<T>> requests, BulkWriteOptions Options = null!,
+			CancellationToken  CancellationToken = new CancellationToken())
+		{
+			return Collection.BulkWrite(Session, requests, Options,  CancellationToken);
+		}
+
+		/// <inheritdoc />
+		public async Task<BulkWriteResult<T>> BulkWriteAsync(IEnumerable<WriteModel<T>> requests, BulkWriteOptions Options = null!,
+			CancellationToken  CancellationToken = new CancellationToken())
+		{
+			using Scope _ = Tracer.Instance.StartActive(Prefix + "BulkWriteAsync");
+			return await Collection.BulkWriteAsync(requests, Options,  CancellationToken);
+		}
+
+		/// <inheritdoc />
+		public async Task<BulkWriteResult<T>> BulkWriteAsync(IClientSessionHandle Session, IEnumerable<WriteModel<T>> requests, BulkWriteOptions Options = null!,
+			CancellationToken  CancellationToken = new CancellationToken())
+		{
+			using Scope _ = Tracer.Instance.StartActive(Prefix + "BulkWriteAsync");
+			return await Collection.BulkWriteAsync(Session, requests, Options,  CancellationToken);
+		}
+
+		/// <inheritdoc />
+		public long Count(FilterDefinition<T> filter, CountOptions Options = null!,
+			CancellationToken  CancellationToken = new CancellationToken())
+		{
+			return Collection.Count(filter, Options,  CancellationToken);
+		}
+
+		/// <inheritdoc />
+		public long Count(IClientSessionHandle Session, FilterDefinition<T> filter, CountOptions Options = null!,
+			CancellationToken  CancellationToken = new CancellationToken())
+		{
+			return Collection.Count(Session, filter, Options,  CancellationToken);
+		}
+
+		/// <inheritdoc />
+		public async Task<long> CountAsync(FilterDefinition<T> filter, CountOptions Options = null!,
+			CancellationToken  CancellationToken = new CancellationToken())
+		{
+			using Scope _ = Tracer.Instance.StartActive(Prefix + "CountAsync");
+			return await Collection.CountAsync(filter, Options,  CancellationToken);
+		}
+
+		/// <inheritdoc />
+		public async Task<long> CountAsync(IClientSessionHandle Session, FilterDefinition<T> filter, CountOptions Options = null!,
+			CancellationToken  CancellationToken = new CancellationToken())
+		{
+			using Scope _ = Tracer.Instance.StartActive(Prefix + "CountAsync");
+			return await Collection.CountAsync(Session, filter, Options,  CancellationToken);
+		}
+
+		/// <inheritdoc />
+		public long CountDocuments(FilterDefinition<T> filter, CountOptions Options = null!,
+			CancellationToken  CancellationToken = new CancellationToken())
+		{
+			return Collection.CountDocuments(filter, Options,  CancellationToken);
+		}
+
+		/// <inheritdoc />
+		public long CountDocuments(IClientSessionHandle Session, FilterDefinition<T> filter, CountOptions Options = null!,
+			CancellationToken  CancellationToken = new CancellationToken())
+		{
+			using Scope _ = Tracer.Instance.StartActive(Prefix + "CountDocumentsAsync");
+			return Collection.CountDocuments(Session, filter, Options,  CancellationToken);
+		}
+
+		/// <inheritdoc />
+		public async Task<long> CountDocumentsAsync(FilterDefinition<T> filter, CountOptions Options = null!,
+			CancellationToken  CancellationToken = new CancellationToken())
+		{
+			using Scope _ = Tracer.Instance.StartActive(Prefix + "CountDocumentsAsync");
+			return await Collection.CountDocumentsAsync(filter, Options,  CancellationToken);
+		}
+
+		/// <inheritdoc />
+		public async Task<long> CountDocumentsAsync(IClientSessionHandle Session, FilterDefinition<T> filter, CountOptions Options = null!,
+			CancellationToken  CancellationToken = new CancellationToken())
+		{
+			using Scope _ = Tracer.Instance.StartActive(Prefix + "CountDocumentsAsync");
+			return await Collection.CountDocumentsAsync(Session, filter, Options,  CancellationToken);
+		}
+
+		/// <inheritdoc />
+		public DeleteResult DeleteMany(FilterDefinition<T> filter, CancellationToken  CancellationToken = new CancellationToken())
+		{
+			return Collection.DeleteMany(filter,  CancellationToken);
+		}
+
+		/// <inheritdoc />
+		public DeleteResult DeleteMany(FilterDefinition<T> filter, DeleteOptions Options,
+			CancellationToken  CancellationToken = new CancellationToken())
+		{
+			return Collection.DeleteMany(filter, Options,  CancellationToken);
+		}
+
+		/// <inheritdoc />
+		public DeleteResult DeleteMany(IClientSessionHandle Session, FilterDefinition<T> filter, DeleteOptions Options = null!,
+			CancellationToken  CancellationToken = new CancellationToken())
+		{
+			return Collection.DeleteMany(Session, filter, Options,  CancellationToken);
+		}
+
+		/// <inheritdoc />
+		public async Task<DeleteResult> DeleteManyAsync(FilterDefinition<T> filter, CancellationToken  CancellationToken = new CancellationToken())
+		{
+			using Scope _ = Tracer.Instance.StartActive(Prefix + "DeleteManyAsync");
+			return await Collection.DeleteManyAsync(filter,  CancellationToken);
+		}
+
+		/// <inheritdoc />
+		public async Task<DeleteResult> DeleteManyAsync(FilterDefinition<T> filter, DeleteOptions Options,
+			CancellationToken  CancellationToken = new CancellationToken())
+		{
+			using Scope _ = Tracer.Instance.StartActive(Prefix + "DeleteManyAsync");
+			return await Collection.DeleteManyAsync(filter, Options,  CancellationToken);
+		}
+
+		/// <inheritdoc />
+		public async Task<DeleteResult> DeleteManyAsync(IClientSessionHandle Session, FilterDefinition<T> filter, DeleteOptions Options = null!,
+			CancellationToken  CancellationToken = new CancellationToken())
+		{
+			using Scope _ = Tracer.Instance.StartActive(Prefix + "DeleteManyAsync");
+			return await Collection.DeleteManyAsync(Session, filter, Options,  CancellationToken);
+		}
+
+		/// <inheritdoc />
+		public DeleteResult DeleteOne(FilterDefinition<T> filter, CancellationToken  CancellationToken = new CancellationToken())
+		{
+			return Collection.DeleteOne(filter,  CancellationToken);
+		}
+
+		/// <inheritdoc />
+		public DeleteResult DeleteOne(FilterDefinition<T> filter, DeleteOptions Options,
+			CancellationToken  CancellationToken = new CancellationToken())
+		{
+			return Collection.DeleteOne(filter, Options,  CancellationToken);
+		}
+
+		/// <inheritdoc />
+		public DeleteResult DeleteOne(IClientSessionHandle Session, FilterDefinition<T> filter, DeleteOptions Options = null!,
+			CancellationToken  CancellationToken = new CancellationToken())
+		{
+			return Collection.DeleteOne(Session, filter, Options,  CancellationToken);
+		}
+
+		/// <inheritdoc />
+		public async Task<DeleteResult> DeleteOneAsync(FilterDefinition<T> filter, CancellationToken  CancellationToken = new CancellationToken())
+		{
+			using Scope _ = Tracer.Instance.StartActive(Prefix + "DeleteOneAsync");
+			return await Collection.DeleteOneAsync(filter,  CancellationToken);
+		}
+
+		/// <inheritdoc />
+		public async Task<DeleteResult> DeleteOneAsync(FilterDefinition<T> filter, DeleteOptions Options,
+			CancellationToken  CancellationToken = new CancellationToken())
+		{
+			using Scope _ = Tracer.Instance.StartActive(Prefix + "DeleteOneAsync");
+			return await Collection.DeleteOneAsync(filter, Options,  CancellationToken);
+		}
+
+		/// <inheritdoc />
+		public Task<DeleteResult> DeleteOneAsync(IClientSessionHandle Session, FilterDefinition<T> filter, DeleteOptions Options = null!,
+			CancellationToken  CancellationToken = new CancellationToken())
+		{
+			using Scope _ = Tracer.Instance.StartActive(Prefix + "DeleteOneAsync");
+			return Collection.DeleteOneAsync(Session, filter, Options,  CancellationToken);
+		}
+
+		/// <inheritdoc />
+		public IAsyncCursor<TField> Distinct<TField>(FieldDefinition<T, TField> field, FilterDefinition<T> filter, DistinctOptions Options = null!,
+			CancellationToken  CancellationToken = new CancellationToken())
+		{
+			return Collection.Distinct(field, filter, Options,  CancellationToken);
+		}
+
+		/// <inheritdoc />
+		public IAsyncCursor<TField> Distinct<TField>(IClientSessionHandle Session, FieldDefinition<T, TField> field, FilterDefinition<T> filter,
+			DistinctOptions Options = null!, CancellationToken  CancellationToken = new CancellationToken())
+		{
+			return Collection.Distinct(Session, field, filter, Options,  CancellationToken);
+		}
+
+		/// <inheritdoc />
+		public async Task<IAsyncCursor<TField>> DistinctAsync<TField>(FieldDefinition<T, TField> field, FilterDefinition<T> filter, DistinctOptions Options = null!,
+			CancellationToken  CancellationToken = new CancellationToken())
+		{
+			using Scope _ = Tracer.Instance.StartActive(Prefix + "DistinctAsync");
+			return await Collection.DistinctAsync(field, filter, Options,  CancellationToken);
+		}
+
+		/// <inheritdoc />
+		public async Task<IAsyncCursor<TField>> DistinctAsync<TField>(IClientSessionHandle Session, FieldDefinition<T, TField> field, FilterDefinition<T> filter,
+			DistinctOptions Options = null!, CancellationToken  CancellationToken = new CancellationToken())
+		{
+			using Scope _ = Tracer.Instance.StartActive(Prefix + "DistinctAsync");
+			return await Collection.DistinctAsync(Session, field, filter, Options,  CancellationToken);
+		}
+
+		/// <inheritdoc />
+		public long EstimatedDocumentCount(EstimatedDocumentCountOptions Options = null!,
+			CancellationToken  CancellationToken = new CancellationToken())
+		{
+			return Collection.EstimatedDocumentCount(Options,  CancellationToken);
+		}
+
+		/// <inheritdoc />
+		public async Task<long> EstimatedDocumentCountAsync(EstimatedDocumentCountOptions Options = null!,
+			CancellationToken  CancellationToken = new CancellationToken())
+		{
+			using Scope _ = Tracer.Instance.StartActive(Prefix + "EstimatedDocumentCountAsync");
+			return await Collection.EstimatedDocumentCountAsync(Options,  CancellationToken);
+		}
+
+		/// <inheritdoc />
+		public IAsyncCursor<TProjection> FindSync<TProjection>(FilterDefinition<T> filter, FindOptions<T, TProjection> Options = null!,
+			CancellationToken  CancellationToken = new CancellationToken())
+		{
+			return Collection.FindSync(filter, Options,  CancellationToken);
+		}
+
+		/// <inheritdoc />
+		public IAsyncCursor<TProjection> FindSync<TProjection>(IClientSessionHandle Session, FilterDefinition<T> filter, FindOptions<T, TProjection> Options = null!,
+			CancellationToken  CancellationToken = new CancellationToken())
+		{
+			return Collection.FindSync(Session, filter, Options,  CancellationToken);
+		}
+
+		/// <inheritdoc />
+		public async Task<IAsyncCursor<TProjection>> FindAsync<TProjection>(FilterDefinition<T> filter, FindOptions<T, TProjection> Options = null!,
+			CancellationToken  CancellationToken = new CancellationToken())
+		{
+			using Scope _ = Tracer.Instance.StartActive(Prefix + "FindAsync");
+			return await Collection.FindAsync(filter, Options,  CancellationToken);
+		}
+
+		/// <inheritdoc />
+		public async Task<IAsyncCursor<TProjection>> FindAsync<TProjection>(IClientSessionHandle Session, FilterDefinition<T> filter, FindOptions<T, TProjection> Options = null!,
+			CancellationToken  CancellationToken = new CancellationToken())
+		{
+			using Scope _ = Tracer.Instance.StartActive(Prefix + "FindAsync");
+			return await Collection.FindAsync(Session, filter, Options,  CancellationToken);
+		}
+
+		/// <inheritdoc />
+		public TProjection FindOneAndDelete<TProjection>(FilterDefinition<T> filter, FindOneAndDeleteOptions<T, TProjection> Options = null!,
+			CancellationToken  CancellationToken = new CancellationToken())
+		{
+			return Collection.FindOneAndDelete(filter, Options,  CancellationToken);
+		}
+
+		/// <inheritdoc />
+		public TProjection FindOneAndDelete<TProjection>(IClientSessionHandle Session, FilterDefinition<T> filter,
+			FindOneAndDeleteOptions<T, TProjection> Options = null!, CancellationToken  CancellationToken = new CancellationToken())
+		{
+			return Collection.FindOneAndDelete(Session, filter, Options,  CancellationToken);
+		}
+
+		/// <inheritdoc />
+		public async Task<TProjection> FindOneAndDeleteAsync<TProjection>(FilterDefinition<T> filter, FindOneAndDeleteOptions<T, TProjection> Options = null!,
+			CancellationToken  CancellationToken = new CancellationToken())
+		{
+			using Scope _ = Tracer.Instance.StartActive(Prefix + "FindOneAndDeleteAsync");
+			return await Collection.FindOneAndDeleteAsync(filter, Options,  CancellationToken);
+		}
+
+		/// <inheritdoc />
+		public async Task<TProjection> FindOneAndDeleteAsync<TProjection>(IClientSessionHandle Session, FilterDefinition<T> filter,
+			FindOneAndDeleteOptions<T, TProjection> Options = null!, CancellationToken  CancellationToken = new CancellationToken())
+		{
+			using Scope _ = Tracer.Instance.StartActive(Prefix + "FindOneAndDeleteAsync");
+			return await Collection.FindOneAndDeleteAsync(Session, filter, Options,  CancellationToken);
+		}
+
+		/// <inheritdoc />
+		public TProjection FindOneAndReplace<TProjection>(FilterDefinition<T> filter, T replacement,
+			FindOneAndReplaceOptions<T, TProjection> Options = null!, CancellationToken  CancellationToken = new CancellationToken())
+		{
+			return Collection.FindOneAndReplace(filter, replacement, Options,  CancellationToken);
+		}
+
+		/// <inheritdoc />
+		public TProjection FindOneAndReplace<TProjection>(IClientSessionHandle Session, FilterDefinition<T> filter, T replacement,
+			FindOneAndReplaceOptions<T, TProjection> Options = null!, CancellationToken  CancellationToken = new CancellationToken())
+		{
+			return Collection.FindOneAndReplace(Session, filter, replacement, Options,  CancellationToken);
+		}
+
+		/// <inheritdoc />
+		public async Task<TProjection> FindOneAndReplaceAsync<TProjection>(FilterDefinition<T> filter, T replacement,
+			FindOneAndReplaceOptions<T, TProjection> Options = null!, CancellationToken  CancellationToken = new CancellationToken())
+		{
+			using Scope _ = Tracer.Instance.StartActive(Prefix + "FindOneAndReplaceAsync");
+			return await Collection.FindOneAndReplaceAsync(filter, replacement, Options,  CancellationToken);
+		}
+
+		/// <inheritdoc />
+		public async Task<TProjection> FindOneAndReplaceAsync<TProjection>(IClientSessionHandle Session, FilterDefinition<T> filter, T replacement,
+			FindOneAndReplaceOptions<T, TProjection> Options = null!, CancellationToken  CancellationToken = new CancellationToken())
+		{
+			using Scope _ = Tracer.Instance.StartActive(Prefix + "FindOneAndReplaceAsync");
+			return await Collection.FindOneAndReplaceAsync(Session, filter, replacement, Options,  CancellationToken);
+		}
+
+		/// <inheritdoc />
+		public TProjection FindOneAndUpdate<TProjection>(FilterDefinition<T> filter, UpdateDefinition<T> update,
+			FindOneAndUpdateOptions<T, TProjection> Options = null!, CancellationToken  CancellationToken = new CancellationToken())
+		{
+			return Collection.FindOneAndUpdate(filter, update, Options,  CancellationToken);
+		}
+
+		/// <inheritdoc />
+		public TProjection FindOneAndUpdate<TProjection>(IClientSessionHandle Session, FilterDefinition<T> filter,
+			UpdateDefinition<T> update, FindOneAndUpdateOptions<T, TProjection> Options = null!,
+			CancellationToken  CancellationToken = new CancellationToken())
+		{
+			return Collection.FindOneAndUpdate(Session, filter, update, Options,  CancellationToken);
+		}
+
+		/// <inheritdoc />
+		public async Task<TProjection> FindOneAndUpdateAsync<TProjection>(FilterDefinition<T> filter, UpdateDefinition<T> update,
+			FindOneAndUpdateOptions<T, TProjection> Options = null!, CancellationToken  CancellationToken = new CancellationToken())
+		{
+			using Scope _ = Tracer.Instance.StartActive(Prefix + "FindOneAndUpdateAsync");
+			return await Collection.FindOneAndUpdateAsync(filter, update, Options,  CancellationToken);
+		}
+
+		/// <inheritdoc />
+		public async Task<TProjection> FindOneAndUpdateAsync<TProjection>(IClientSessionHandle Session, FilterDefinition<T> filter, UpdateDefinition<T> update,
+			FindOneAndUpdateOptions<T, TProjection> Options = null!, CancellationToken  CancellationToken = new CancellationToken())
+		{
+			using Scope _ = Tracer.Instance.StartActive(Prefix + "FindOneAndUpdateAsync");
+			return await Collection.FindOneAndUpdateAsync(Session, filter, update, Options,  CancellationToken);
+		}
+
+		/// <inheritdoc />
+		public void InsertOne( T Document, InsertOneOptions Options = null!,
+			CancellationToken  CancellationToken = new CancellationToken())
+		{
+			Collection.InsertOne(Document, Options,  CancellationToken);
+		}
+
+		/// <inheritdoc />
+		public void InsertOne(IClientSessionHandle Session, T Document, InsertOneOptions Options = null!,
+			CancellationToken  CancellationToken = new CancellationToken())
+		{
+			Collection.InsertOne(Session, Document, Options,  CancellationToken);
+		}
+
+		/// <inheritdoc />
+		public async Task InsertOneAsync( T Document, CancellationToken CancellationToken)
+		{
+			using Scope _ = Tracer.Instance.StartActive(Prefix + "InsertOneAsync");
+			await Collection.InsertOneAsync(Document, CancellationToken);
+		}
+
+		/// <inheritdoc />
+		public async Task InsertOneAsync( T Document, InsertOneOptions Options = null!,
+			CancellationToken  CancellationToken = new CancellationToken())
+		{
+			using Scope _ = Tracer.Instance.StartActive(Prefix + "InsertOneAsync");
+			await Collection.InsertOneAsync(Document, Options,  CancellationToken);
+		}
+
+		/// <inheritdoc />
+		public async Task InsertOneAsync(IClientSessionHandle Session, T Document, InsertOneOptions Options = null!,
+			CancellationToken CancellationToken = new CancellationToken())
+		{
+			using Scope _ = Tracer.Instance.StartActive(Prefix + "InsertOneAsync");
+			await Collection.InsertOneAsync(Session, Document, Options,  CancellationToken);
+		}
+
+		/// <inheritdoc />
+		public void InsertMany(IEnumerable<T> Documents, InsertManyOptions Options = null!,
+			CancellationToken  CancellationToken = new CancellationToken())
+		{
+			Collection.InsertMany(Documents, Options,  CancellationToken);
+		}
+
+		/// <inheritdoc />
+		public void InsertMany(IClientSessionHandle Session, IEnumerable<T> Documents, InsertManyOptions Options = null!,
+			CancellationToken  CancellationToken = new CancellationToken())
+		{
+			Collection.InsertMany(Session, Documents, Options,  CancellationToken);
+		}
+
+		/// <inheritdoc />
+		public async Task InsertManyAsync(IEnumerable<T> Documents, InsertManyOptions Options = null!,
+			CancellationToken  CancellationToken = new CancellationToken())
+		{
+			using Scope _ = Tracer.Instance.StartActive(Prefix + "InsertManyAsync");
+			await Collection.InsertManyAsync(Documents, Options,  CancellationToken);
+		}
+
+		/// <inheritdoc />
+		public async Task InsertManyAsync(IClientSessionHandle Session, IEnumerable<T> documents, InsertManyOptions Options = null!,
+			CancellationToken  CancellationToken = new CancellationToken())
+		{
+			using Scope _ = Tracer.Instance.StartActive(Prefix + "InsertManyAsync");
+			await Collection.InsertManyAsync(Session, documents, Options,  CancellationToken);
+		}
+
+		/// <inheritdoc />
+		public IAsyncCursor<TResult> MapReduce<TResult>(BsonJavaScript map, BsonJavaScript reduce, MapReduceOptions<T, TResult> Options = null!,
+			CancellationToken  CancellationToken = new CancellationToken())
+		{
+			return Collection.MapReduce(map, reduce, Options,  CancellationToken);
+		}
+
+		/// <inheritdoc />
+		public IAsyncCursor<TResult> MapReduce<TResult>(IClientSessionHandle Session, BsonJavaScript map, BsonJavaScript reduce,
+			MapReduceOptions<T, TResult> Options = null!, CancellationToken  CancellationToken = new CancellationToken())
+		{
+			return Collection.MapReduce(Session, map, reduce, Options,  CancellationToken);
+		}
+
+		/// <inheritdoc />
+		public async Task<IAsyncCursor<TResult>> MapReduceAsync<TResult>(BsonJavaScript map, BsonJavaScript reduce, MapReduceOptions<T, TResult> Options = null!,
+			CancellationToken  CancellationToken = new CancellationToken())
+		{
+			using Scope _ = Tracer.Instance.StartActive(Prefix + "MapReduceAsync");
+			return await Collection.MapReduceAsync(map, reduce, Options,  CancellationToken);
+		}
+
+		/// <inheritdoc />
+		public async Task<IAsyncCursor<TResult>> MapReduceAsync<TResult>(IClientSessionHandle Session, BsonJavaScript map, BsonJavaScript reduce,
+			MapReduceOptions<T, TResult> Options = null!, CancellationToken  CancellationToken = new CancellationToken())
+		{
+			using Scope _ = Tracer.Instance.StartActive(Prefix + "MapReduceAsync");
+			return await Collection.MapReduceAsync(Session, map, reduce, Options,  CancellationToken);
+		}
+
+		/// <inheritdoc />
+		public IFilteredMongoCollection<TDerivedDocument> OfType<TDerivedDocument>() where TDerivedDocument : T
+		{
+			return Collection.OfType<TDerivedDocument>();
+		}
+
+		/// <inheritdoc />
+		public ReplaceOneResult ReplaceOne(FilterDefinition<T> filter, T replacement, ReplaceOptions Options = null!,
+			CancellationToken  CancellationToken = new CancellationToken())
+		{
+			return Collection.ReplaceOne(filter, replacement, Options,  CancellationToken);
+		}
+
+		/// <inheritdoc />
+		public ReplaceOneResult ReplaceOne(FilterDefinition<T> filter, T replacement, UpdateOptions Options,
+			CancellationToken  CancellationToken = new CancellationToken())
+		{
+			return Collection.ReplaceOne(filter, replacement, Options,  CancellationToken);
+		}
+
+		/// <inheritdoc />
+		public ReplaceOneResult ReplaceOne(IClientSessionHandle Session, FilterDefinition<T> filter, T replacement,
+			ReplaceOptions Options = null!, CancellationToken  CancellationToken = new CancellationToken())
+		{
+			return Collection.ReplaceOne(Session, filter, replacement, Options,  CancellationToken);
+		}
+
+		/// <inheritdoc />
+		public ReplaceOneResult ReplaceOne(IClientSessionHandle Session, FilterDefinition<T> filter, T replacement, UpdateOptions Options,
+			CancellationToken  CancellationToken = new CancellationToken())
+		{
+			return Collection.ReplaceOne(Session, filter, replacement, Options,  CancellationToken);
+		}
+
+		/// <inheritdoc />
+		public async Task<ReplaceOneResult> ReplaceOneAsync(FilterDefinition<T> filter, T replacement, ReplaceOptions Options = null!,
+			CancellationToken  CancellationToken = new CancellationToken())
+		{
+			using Scope _ = Tracer.Instance.StartActive(Prefix + "ReplaceOneAsync");
+			return await Collection.ReplaceOneAsync(filter, replacement, Options,  CancellationToken);
+		}
+
+		/// <inheritdoc />
+		public async Task<ReplaceOneResult> ReplaceOneAsync(FilterDefinition<T> filter, T replacement, UpdateOptions Options,
+			CancellationToken  CancellationToken = new CancellationToken())
+		{
+			using Scope _ = Tracer.Instance.StartActive(Prefix + "ReplaceOneAsync");
+			return await Collection.ReplaceOneAsync(filter, replacement, Options,  CancellationToken);
+		}
+
+		/// <inheritdoc />
+		public async Task<ReplaceOneResult> ReplaceOneAsync(IClientSessionHandle Session, FilterDefinition<T> filter, T replacement,
+			ReplaceOptions Options = null!, CancellationToken  CancellationToken = new CancellationToken())
+		{
+			using Scope _ = Tracer.Instance.StartActive(Prefix + "ReplaceOneAsync");
+			return await Collection.ReplaceOneAsync(Session, filter, replacement, Options,  CancellationToken);
+		}
+
+		/// <inheritdoc />
+		public async Task<ReplaceOneResult> ReplaceOneAsync(IClientSessionHandle Session, FilterDefinition<T> filter, T replacement, UpdateOptions Options,
+			CancellationToken  CancellationToken = new CancellationToken())
+		{
+			using Scope _ = Tracer.Instance.StartActive(Prefix + "ReplaceOneAsync");
+			return await Collection.ReplaceOneAsync(Session, filter, replacement, Options,  CancellationToken);
+		}
+
+		/// <inheritdoc />
+		public UpdateResult UpdateMany(FilterDefinition<T> filter, UpdateDefinition<T> update, UpdateOptions Options = null!,
+			CancellationToken  CancellationToken = new CancellationToken())
+		{
+			return Collection.UpdateMany(filter, update, Options,  CancellationToken);
+		}
+
+		/// <inheritdoc />
+		public UpdateResult UpdateMany(IClientSessionHandle Session, FilterDefinition<T> filter, UpdateDefinition<T> update,
+			UpdateOptions Options = null!, CancellationToken  CancellationToken = new CancellationToken())
+		{
+			return Collection.UpdateMany(Session, filter, update, Options,  CancellationToken);
+		}
+
+		/// <inheritdoc />
+		public async Task<UpdateResult> UpdateManyAsync(FilterDefinition<T> filter, UpdateDefinition<T> update, UpdateOptions Options = null!,
+			CancellationToken  CancellationToken = new CancellationToken())
+		{
+			using Scope _ = Tracer.Instance.StartActive(Prefix + "UpdateManyAsync");
+			return await Collection.UpdateManyAsync(filter, update, Options,  CancellationToken);
+		}
+
+		/// <inheritdoc />
+		public async Task<UpdateResult> UpdateManyAsync(IClientSessionHandle Session, FilterDefinition<T> filter, UpdateDefinition<T> update,
+			UpdateOptions Options = null!, CancellationToken  CancellationToken = new CancellationToken())
+		{
+			using Scope _ = Tracer.Instance.StartActive(Prefix + "UpdateManyAsync");
+			return await Collection.UpdateManyAsync(Session, filter, update, Options,  CancellationToken);
+		}
+
+		/// <inheritdoc />
+		public UpdateResult UpdateOne(FilterDefinition<T> filter, UpdateDefinition<T> update, UpdateOptions Options = null!,
+			CancellationToken  CancellationToken = new CancellationToken())
+		{
+			return Collection.UpdateOne(filter, update, Options,  CancellationToken);
+		}
+
+		/// <inheritdoc />
+		public UpdateResult UpdateOne(IClientSessionHandle Session, FilterDefinition<T> filter, UpdateDefinition<T> update,
+			UpdateOptions Options = null!, CancellationToken  CancellationToken = new CancellationToken())
+		{
+			return Collection.UpdateOne(Session, filter, update, Options,  CancellationToken);
+		}
+
+		/// <inheritdoc />
+		public async Task<UpdateResult> UpdateOneAsync(FilterDefinition<T> filter, UpdateDefinition<T> update, UpdateOptions Options = null!,
+			CancellationToken  CancellationToken = new CancellationToken())
+		{
+			using Scope _ = Tracer.Instance.StartActive(Prefix + "UpdateOneAsync");
+			return await Collection.UpdateOneAsync(filter, update, Options,  CancellationToken);
+		}
+
+		/// <inheritdoc />
+		public async Task<UpdateResult> UpdateOneAsync(IClientSessionHandle Session, FilterDefinition<T> filter, UpdateDefinition<T> update,
+			UpdateOptions Options = null!, CancellationToken  CancellationToken = new CancellationToken())
+		{
+			using Scope _ = Tracer.Instance.StartActive(Prefix + "UpdateOneAsync");
+			return await Collection.UpdateOneAsync(Session, filter, update, Options,  CancellationToken);
+		}
+
+		/// <inheritdoc />
+		public IChangeStreamCursor<TResult> Watch<TResult>(PipelineDefinition<ChangeStreamDocument<T>, TResult> pipeline, ChangeStreamOptions Options = null!,
+			CancellationToken  CancellationToken = new CancellationToken())
+		{
+			return Collection.Watch(pipeline, Options,  CancellationToken);
+		}
+
+		/// <inheritdoc />
+		public IChangeStreamCursor<TResult> Watch<TResult>(IClientSessionHandle Session, PipelineDefinition<ChangeStreamDocument<T>, TResult> pipeline,
+			ChangeStreamOptions Options = null!, CancellationToken  CancellationToken = new CancellationToken())
+		{
+			return Collection.Watch(Session, pipeline, Options,  CancellationToken);
+		}
+
+		/// <inheritdoc />
+		public async Task<IChangeStreamCursor<TResult>> WatchAsync<TResult>(PipelineDefinition<ChangeStreamDocument<T>, TResult> pipeline, ChangeStreamOptions Options = null!,
+			CancellationToken  CancellationToken = new CancellationToken())
+		{
+			using Scope _ = Tracer.Instance.StartActive(Prefix + "WatchAsync");
+			return await Collection.WatchAsync(pipeline, Options,  CancellationToken);
+		}
+
+		/// <inheritdoc />
+		public async Task<IChangeStreamCursor<TResult>> WatchAsync<TResult>(IClientSessionHandle Session, PipelineDefinition<ChangeStreamDocument<T>, TResult> pipeline, ChangeStreamOptions Options = null!,
+			CancellationToken  CancellationToken = new CancellationToken())
+		{
+			using Scope _ = Tracer.Instance.StartActive(Prefix + "WatchAsync");
+			return await Collection.WatchAsync(Session, pipeline, Options,  CancellationToken);
+		}
+
+		/// <inheritdoc />
+		public IMongoCollection<T> WithReadConcern(ReadConcern ReadConcern)
+		{
+			return Collection.WithReadConcern(ReadConcern);
+		}
+
+		/// <inheritdoc />
+		public IMongoCollection<T> WithReadPreference(ReadPreference ReadPreference)
+		{
+			return Collection.WithReadPreference(ReadPreference);
+		}
+
+		/// <inheritdoc />
+		public IMongoCollection<T> WithWriteConcern(WriteConcern WriteConcern)
+		{
+			return Collection.WithWriteConcern(WriteConcern);
+		}
+#pragma warning restore CS0618
+		
+		/// <inheritdoc />
+		public CollectionNamespace CollectionNamespace => Collection.CollectionNamespace;
+
+		/// <inheritdoc />
+		public IMongoDatabase Database => Collection.Database;
+
+		/// <inheritdoc />
+		public IBsonSerializer<T> DocumentSerializer => Collection.DocumentSerializer;
+
+		/// <inheritdoc />
+		public IMongoIndexManager<T> Indexes => Collection.Indexes;
+
+		/// <inheritdoc />
+		public MongoCollectionSettings Settings => Collection.Settings;
+	}
+}
\ No newline at end of file
diff --git a/Engine/Source/Programs/Horde/HordeServer/Collections/IAgentSoftwareCollection.cs b/Engine/Source/Programs/Horde/HordeServer/Collections/IAgentSoftwareCollection.cs
new file mode 100644
--- /dev/null
+++ b/Engine/Source/Programs/Horde/HordeServer/Collections/IAgentSoftwareCollection.cs
@@ -0,0 +1,47 @@
+using EpicGames.Core;
+using HordeCommon;
+using HordeServer.Models;
+using HordeServer.Services;
+using MongoDB.Bson;
+using MongoDB.Bson.Serialization.Attributes;
+using MongoDB.Driver;
+using System;
+using System.Collections.Generic;
+using System.Globalization;
+using System.IO;
+using System.IO.Compression;
+using System.Linq;
+using System.Text.Json;
+using System.Threading.Tasks;
+
+using AgentSoftwareVersion = HordeServer.Utilities.StringId<HordeServer.Collections.IAgentSoftwareCollection>;
+
+namespace HordeServer.Collections
+{
+	/// <summary>
+	/// Interface for a collection of template documents
+	/// </summary>
+	public interface IAgentSoftwareCollection
+	{
+		/// <summary>
+		/// Adds a new software revision
+		/// </summary>
+		/// <param name="Version">The version number</param>
+		/// <param name="Data">Zip file containing the new software</param>
+		/// <returns>New software instance</returns>
+		Task AddAsync(AgentSoftwareVersion Version, byte[] Data);
+
+		/// <summary>
+		/// Removes a software archive 
+		/// </summary>
+		/// <param name="Version">Version of the software to delete</param>
+		Task<bool> RemoveAsync(AgentSoftwareVersion Version);
+
+		/// <summary>
+		/// Downloads software of a given revision
+		/// </summary>
+		/// <param name="Version">Version of the software</param>
+		/// <returns>Data for the given software</returns>
+		Task<byte[]?> GetAsync(AgentSoftwareVersion Version);
+	}
+}
diff --git a/Engine/Source/Programs/Horde/HordeServer/Collections/IGraphCollection.cs b/Engine/Source/Programs/Horde/HordeServer/Collections/IGraphCollection.cs
new file mode 100644
--- /dev/null
+++ b/Engine/Source/Programs/Horde/HordeServer/Collections/IGraphCollection.cs
@@ -0,0 +1,57 @@
+using EpicGames.Core;
+using HordeServer.Api;
+using HordeCommon;
+using HordeServer.Models;
+using HordeServer.Services;
+using MongoDB.Bson;
+using MongoDB.Bson.Serialization.Attributes;
+using MongoDB.Driver;
+using System;
+using System.Collections.Concurrent;
+using System.Collections.Generic;
+using System.Diagnostics;
+using System.Linq;
+using System.Threading;
+using System.Threading.Tasks;
+
+namespace HordeServer.Collections
+{
+	/// <summary>
+	/// Interface for a collection of graph documents
+	/// </summary>
+	public interface IGraphCollection
+	{
+		/// <summary>
+		/// Adds a graph from a template
+		/// </summary>
+		/// <param name="Template">The template</param>
+		/// <returns>New graph</returns>
+		Task<IGraph> AddAsync(ITemplate Template);
+
+		/// <summary>
+		/// Creates a graph by appending groups and aggregates to an existing graph.
+		/// </summary>
+		/// <param name="BaseGraph">The base graph</param>
+		/// <param name="NewGroupRequests">List of group requests</param>
+		/// <param name="NewAggregateRequests">List of aggregate requests</param>
+		/// <param name="NewLabelRequests">List of label requests</param>
+		/// <returns>The new graph definition</returns>
+		Task<IGraph> AppendAsync(IGraph? BaseGraph, List<CreateGroupRequest>? NewGroupRequests = null, List<CreateAggregateRequest>? NewAggregateRequests = null, List<CreateLabelRequest>? NewLabelRequests = null);
+
+		/// <summary>
+		/// Gets the graph for a job
+		/// </summary>
+		/// <param name="Hash">Hash of the graph to retrieve</param>
+		/// <returns>The graph for this job</returns>
+		Task<IGraph> GetAsync(ContentHash Hash);
+
+		/// <summary>
+		/// Finds all graphs stored in the collection
+		/// </summary>
+		/// <param name="Hashes">Hashes to filter by</param>
+		/// <param name="Index">Starting index of the graph to return</param>
+		/// <param name="Count">Number of results to return</param>
+		/// <returns>List of graphs</returns>
+		Task<List<IGraph>> FindAllAsync(ContentHash[]? Hashes = null, int? Index = null, int? Count = null);
+	}
+}
diff --git a/Engine/Source/Programs/Horde/HordeServer/Collections/IIssueCollection.cs b/Engine/Source/Programs/Horde/HordeServer/Collections/IIssueCollection.cs
new file mode 100644
--- /dev/null
+++ b/Engine/Source/Programs/Horde/HordeServer/Collections/IIssueCollection.cs
@@ -0,0 +1,520 @@
+using HordeCommon;
+using HordeServer.Models;
+using HordeServer.Services;
+using HordeServer.Utilities;
+using MessagePack.Formatters;
+using Microsoft.AspNetCore.Connections.Features;
+using Microsoft.Extensions.Logging;
+using MongoDB.Bson;
+using MongoDB.Bson.Serialization;
+using MongoDB.Bson.Serialization.Attributes;
+using MongoDB.Driver;
+using System;
+using System.Collections.Generic;
+using System.Diagnostics.CodeAnalysis;
+using System.IO;
+using System.Linq;
+using System.Linq.Expressions;
+using System.Text.RegularExpressions;
+using System.Threading.Tasks;
+using System.Transactions;
+using Datadog.Trace;
+
+using StreamId = HordeServer.Utilities.StringId<HordeServer.Models.IStream>;
+using TemplateRefId = HordeServer.Utilities.StringId<HordeServer.Models.TemplateRef>;
+
+namespace HordeServer.Collections
+{
+	/// <summary>
+	/// Fingerprint for an issue
+	/// </summary>
+	public class NewIssueFingerprint : IIssueFingerprint
+	{
+		/// <summary>
+		/// The type of issue
+		/// </summary>
+		public string Type { get; }
+
+		/// <summary>
+		/// List of keys which identify this issue.
+		/// </summary>
+		public CaseInsensitiveStringSet Keys { get; set; }
+
+		/// <summary>
+		/// Set of keys which should trigger a negative match
+		/// </summary>
+		public CaseInsensitiveStringSet? RejectKeys { get; set; }
+
+		/// <summary>
+		/// Constructor
+		/// </summary>
+		/// <param name="Type">The type of issue</param>
+		/// <param name="Keys">Keys which uniquely identify this issue</param>
+		/// <param name="RejectKeys">Keys which should not match with this issue</param>
+		public NewIssueFingerprint(string Type, IEnumerable<string> Keys, IEnumerable<string>? RejectKeys)
+		{
+			this.Type = Type;
+			this.Keys = new CaseInsensitiveStringSet(Keys);
+
+			if (RejectKeys != null && RejectKeys.Any())
+			{
+				this.RejectKeys = new CaseInsensitiveStringSet(RejectKeys);
+			}
+		}
+
+		/// <summary>
+		/// Copy constructor
+		/// </summary>
+		/// <param name="Other">The fingerprint to copy from</param>
+		public NewIssueFingerprint(IIssueFingerprint Other)
+			: this(Other.Type, Other.Keys, Other.RejectKeys)
+		{
+		}
+
+		/// <summary>
+		/// Merges another fingerprint into this one
+		/// </summary>
+		/// <param name="Other">The other fingerprint to merge with</param>
+		public NewIssueFingerprint MergeWith(IIssueFingerprint Other)
+		{
+			NewIssueFingerprint Result = new NewIssueFingerprint(this);
+
+			Result.Keys.UnionWith(Other.Keys);
+			if (Other.RejectKeys != null)
+			{
+				if (Result.RejectKeys == null)
+				{
+					Result.RejectKeys = new CaseInsensitiveStringSet();
+				}
+				Result.RejectKeys.UnionWith(Other.RejectKeys);
+			}
+
+			return Result;
+		}
+
+		/// <inheritdoc/>
+		public override bool Equals(object? Other)
+		{
+			NewIssueFingerprint? OtherFingerprint = Other as NewIssueFingerprint;
+			if(OtherFingerprint == null || !Type.Equals(OtherFingerprint.Type, StringComparison.Ordinal) || !Keys.SetEquals(OtherFingerprint.Keys))
+			{
+				return false;
+			}
+
+			if (RejectKeys == null)
+			{
+				if (OtherFingerprint.RejectKeys != null && OtherFingerprint.RejectKeys.Count > 0)
+				{
+					return false;
+				}
+			}
+			else
+			{
+				if (OtherFingerprint.RejectKeys == null || !RejectKeys.SetEquals(OtherFingerprint.RejectKeys))
+				{
+					return false;
+				}
+			}
+
+			return true;
+		}
+
+		/// <summary>
+		/// Gets a hashcode for this fingerprint
+		/// </summary>
+		/// <returns>The hashcode value</returns>
+		public override int GetHashCode()
+		{
+			int Result = string.GetHashCode(Type, StringComparison.Ordinal);
+			foreach (string Key in Keys)
+			{
+				Result = HashCode.Combine(Result, Key);
+			}
+			if (RejectKeys != null)
+			{
+				foreach (string RejectKey in RejectKeys)
+				{
+					Result = HashCode.Combine(Result, RejectKey);
+				}
+			}
+			return Result;
+		}
+
+		/// <summary>
+		/// Merges two fingerprints togetherthis issue fingerprint with another
+		/// </summary>
+		/// <param name="A">The first fingerprint</param>
+		/// <param name="B">The second fingerprint to merge with</param>
+		/// <returns>Merged fingerprint</returns>
+		public static NewIssueFingerprint Merge(IIssueFingerprint A, IIssueFingerprint B)
+		{
+			NewIssueFingerprint NewFingerprint = new NewIssueFingerprint(A);
+			NewFingerprint.Keys.UnionWith(B.Keys);
+
+			if (B.RejectKeys != null)
+			{
+				if (NewFingerprint.RejectKeys == null)
+				{
+					NewFingerprint.RejectKeys = new CaseInsensitiveStringSet();
+				}
+				NewFingerprint.RejectKeys.UnionWith(B.RejectKeys);
+			}
+
+			return NewFingerprint;
+		}
+
+		/// <inheritdoc/>
+		public override string ToString()
+		{
+			return $"{Type}: \"{String.Join("\", \"", Keys)}\"";
+		}
+	}
+
+	/// <summary>
+	/// Identifies a particular changelist and job
+	/// </summary>
+	public class NewIssueStepData
+	{
+		/// <summary>
+		/// The changelist number
+		/// </summary>
+		public int Change { get; set; }
+
+		/// <summary>
+		/// Name of the job
+		/// </summary>
+		public string JobName { get; set; }
+
+		/// <summary>
+		/// The unique job id
+		/// </summary>
+		public ObjectId JobId { get; set; }
+
+		/// <summary>
+		/// Batch id for the step
+		/// </summary>
+		public SubResourceId BatchId { get; set; }
+
+		/// <summary>
+		/// Id of the step
+		/// </summary>
+		public SubResourceId StepId { get; set; }
+
+		/// <summary>
+		/// Time that the step started
+		/// </summary>
+		public DateTime StepTime { get; set; }
+
+		/// <summary>
+		/// The log id for this step
+		/// </summary>
+		public ObjectId? LogId { get; set; }
+
+		/// <summary>
+		/// Whether to notify suspects for this step
+		/// </summary>
+		public bool NotifySuspects { get; set; }
+
+		/// <summary>
+		/// Constructor
+		/// </summary>
+		/// <param name="Change">The changelist number for this job</param>
+		/// <param name="JobName">The job name</param>
+		/// <param name="JobId">The unique job id</param>
+		/// <param name="BatchId">The batch id</param>
+		/// <param name="StepId">The step id</param>
+		/// <param name="StepTime">Time that the step started</param>
+		/// <param name="LogId">Unique id of the log file for this step</param>
+		/// <param name="NotifySuspects">Whether to notify suspects for this step</param>
+		public NewIssueStepData(int Change, string JobName, ObjectId JobId, SubResourceId BatchId, SubResourceId StepId, DateTime StepTime, ObjectId? LogId, bool NotifySuspects)
+		{
+			this.Change = Change;
+			this.JobName = JobName;
+			this.JobId = JobId;
+			this.BatchId = BatchId;
+			this.StepId = StepId;
+			this.StepTime = StepTime;
+			this.LogId = LogId;
+			this.NotifySuspects = NotifySuspects;
+		}
+
+		/// <summary>
+		/// Constructor
+		/// </summary>
+		/// <param name="Job">The job being built</param>
+		/// <param name="Batch">Batch of the job for the step</param>
+		/// <param name="Step">The step being built</param>
+		/// <param name="NotifySuspects">Whether to notify suspects for this step failing</param>
+		public NewIssueStepData(IJob Job, IJobStepBatch Batch, IJobStep Step, bool NotifySuspects)
+			: this(Job.Change, Job.Name, Job.Id, Batch.Id, Step.Id, Step.StartTimeUtc ?? default, Step.LogId, NotifySuspects)
+		{
+		}
+
+		/// <summary>
+		/// Constructor
+		/// </summary>
+		/// <param name="JobStepRef">The jobstep reference</param>
+		public NewIssueStepData(IJobStepRef JobStepRef)
+			: this(JobStepRef.Change, JobStepRef.JobName, JobStepRef.Id.JobId, JobStepRef.Id.BatchId, JobStepRef.Id.StepId, JobStepRef.StartTimeUtc, JobStepRef.LogId, false)
+		{
+		}
+	}
+
+	/// <summary>
+	/// Information about a suspect changelist that may have caused an issue
+	/// </summary>
+	public class NewIssueSuspectData
+	{
+		/// <summary>
+		/// Author of the changelist
+		/// </summary>
+		public string Author { get; set; }
+
+		/// <summary>
+		/// The submitted changelist
+		/// </summary>
+		public int Change { get; set; }
+
+		/// <summary>
+		/// The time that the user declined causing this issue
+		/// </summary>
+		public DateTime? DeclinedAt { get; set; }
+
+		/// <summary>
+		/// Constructor
+		/// </summary>
+		/// <param name="Author">Author of the change</param>
+		/// <param name="Change">The changelist number</param>
+		/// <param name="DeclinedAt">The time that the user declined this issue</param>
+		public NewIssueSuspectData(string Author, int Change, DateTime? DeclinedAt)
+		{
+			this.Author = Author;
+			this.Change = Change;
+			this.DeclinedAt = DeclinedAt;
+		}
+	}
+
+	/// <summary>
+	/// Information about a suspect changelist that may have caused an issue
+	/// </summary>
+	public class NewIssueSpanSuspectData
+	{
+		/// <summary>
+		/// The submitted changelist
+		/// </summary>
+		public int Change { get; set; }
+
+		/// <summary>
+		/// Author of the changelist
+		/// </summary>
+		public string Author { get; set; }
+
+		/// <summary>
+		/// The original changelist number, if merged from another branch. For changes merged between several branches, this is the originally submitted change.
+		/// </summary>
+		public int? OriginatingChange { get; set; }
+
+		/// <summary>
+		/// Constructor
+		/// </summary>
+		/// <param name="Change">The changelist number</param>
+		/// <param name="Author">Author of the change</param>
+		public NewIssueSpanSuspectData(int Change, string Author)
+		{
+			this.Change = Change;
+			this.Author = Author;
+		}
+	}
+
+	/// <summary>
+	/// Interface for a token used to indicate a position in the transaction queue. Transactions are opportunistic, and will fail if the token is longer valid at the time that it is used.
+	/// </summary>
+	[SuppressMessage("Design", "CA1040:Avoid empty interfaces")]
+	public interface IIssueSequenceToken
+	{
+		/// <summary>
+		/// Reset the token to the last operation
+		/// </summary>
+		/// <returns>Async task</returns>
+		public Task ResetAsync();
+	}
+
+	/// <summary>
+	/// Interface for an issue collection
+	/// </summary>
+	public interface IIssueCollection
+	{
+		/// <summary>
+		/// Gets a token used to sequence operations to the issue collection
+		/// </summary>
+		/// <returns>Sequence token instance</returns>
+		Task<IIssueSequenceToken> GetSequenceTokenAsync();
+
+		/// <summary>
+		/// Assign a unique issue id. For compatibility with the legacy issue system.
+		/// </summary>
+		/// <param name="Token">The current sequence token</param>
+		/// <returns>Sequence token</returns>
+		Task<int?> ReserveUniqueIdAsync(IIssueSequenceToken Token);
+
+		#region Issues
+
+		/// <summary>
+		/// Creates a new issue
+		/// </summary>
+		/// <param name="Token">Token for the transaction</param>
+		/// <param name="Summary">Summary text for the issue</param>
+		/// <param name="Span">The initial span for the issue</param>
+		/// <returns>The new issue instance</returns>
+		Task<IIssue?> AddIssueAsync(IIssueSequenceToken Token, string Summary, IIssueSpan Span);
+
+		/// <summary>
+		/// Retrieves and issue by id
+		/// </summary>
+		/// <param name="IssueId">Unique id of the issue</param>
+		/// <returns>The issue matching the given id, or null</returns>
+		Task<IIssue?> GetIssueAsync(int IssueId);
+
+		/// <summary>
+		/// Searches for open issues
+		/// </summary>
+		/// <param name="Ids">Set of issue ids to find</param>
+		/// <param name="StreamId">The stream affected by the issue</param>
+		/// <param name="MinChange">Minimum changelist affected by the issue</param>
+		/// <param name="MaxChange">Maximum changelist affected by the issue</param>
+		/// <param name="Resolved">Include issues that are now resolved</param>
+		/// <param name="Index">Index within the results to return</param>
+		/// <param name="Count">Number of results</param>
+		/// <returns>List of streams open in the given stream at the given changelist</returns>
+		Task<List<IIssue>> FindIssuesAsync(IEnumerable<int>? Ids = null, StreamId? StreamId = null, int? MinChange = null, int? MaxChange = null, bool? Resolved = null, int? Index = null, int? Count = null);
+
+		/// <summary>
+		/// Searches for open issues
+		/// </summary>
+		/// <param name="Changes">List of suspect changes</param>
+		/// <returns>List of issues that are affected by the given changes</returns>
+		Task<List<IIssue>> FindIssuesForSuspectsAsync(List<int> Changes);
+
+		/// <summary>
+		/// Updates the state of an issue
+		/// </summary>
+		/// <param name="Issue">The issue to update</param>
+		/// <param name="NewSeverity">New severity for the issue</param>
+		/// <param name="NewSummary">New summary for the issue</param>
+		/// <param name="NewUserSummary">New user summary for the issue</param>
+		/// <param name="NewOwner">New owner of the issue</param>
+		/// <param name="NewNominatedBy">Person that nominated the new owner</param>
+		/// <param name="NewAcknowledged">Whether the issue has been acknowledged</param>
+		/// <param name="NewDeclinedBy">Name of a user that has declined the issue</param>
+		/// <param name="NewFixChange">Fix changelist for the issue. Pass 0 to clear the fix changelist, -1 for systemic issue.</param>
+		/// <param name="NewResolved">Whether the issue has been resolved</param>
+		/// <param name="NewNotifySuspects">Whether all suspects should be notified about this issue</param>
+		/// <param name="NewSuspects">New list of suspects for this issue</param>
+		/// <returns>True if the issue was updated</returns>
+		Task<IIssue?> UpdateIssueAsync(IIssue Issue, IssueSeverity? NewSeverity = null, string? NewSummary = null, string? NewUserSummary = null, string? NewOwner = null, string? NewNominatedBy = null, bool? NewAcknowledged = null, string? NewDeclinedBy = null, int? NewFixChange = null, bool? NewResolved = null, bool? NewNotifySuspects = null, List<NewIssueSuspectData>? NewSuspects = null);
+
+		#endregion
+
+		#region Spans
+
+		/// <summary>
+		/// Creates a new span from the given failure
+		/// </summary>
+		/// <param name="Token">Token to sequence the addition of new spans</param>
+		/// <param name="Stream">The stream containing the span</param>
+		/// <param name="TemplateRefId">The template being executed</param>
+		/// <param name="NodeName">Name of the step</param>
+		/// <param name="Severity">Severity of this span</param>
+		/// <param name="Fingerprint">Fingerprint for this issue</param>
+		/// <param name="LastSuccess">Last time the build succeeded before the span</param>
+		/// <param name="Failure">List of failing steps</param>
+		/// <param name="NextSuccess">First time the build succeeded after the span</param>
+		/// <param name="Suspects">Suspects for this span</param>
+		/// <returns>New span, or null if the sequence token is not valid</returns>
+		Task<IIssueSpan?> AddSpanAsync(IIssueSequenceToken Token, IStream Stream, TemplateRefId TemplateRefId, string NodeName, IssueSeverity Severity, NewIssueFingerprint Fingerprint, NewIssueStepData? LastSuccess, NewIssueStepData Failure, NewIssueStepData? NextSuccess, List<NewIssueSpanSuspectData> Suspects);
+
+		/// <summary>
+		/// Gets a particular span
+		/// </summary>
+		/// <param name="SpanId">Unique id of the span</param>
+		/// <returns>New span, or null if the sequence token is not valid</returns>
+		Task<IIssueSpan?> GetSpanAsync(ObjectId SpanId);
+
+		/// <summary>
+		/// Updates the given span. Note that data in the span's issue may be derived from this, and the issue should be updated afterwards.
+		/// </summary>
+		/// <param name="Span">Span to update</param>
+		/// <param name="NewSeverity">New severity for the span</param>
+		/// <param name="NewLastSuccess">New last successful step</param>
+		/// <param name="NewFailure">New failed step</param>
+		/// <param name="NewNextSuccess">New next successful step</param>
+		/// <param name="NewSuspects">New suspects for the span</param>
+		/// <param name="NewModified">New modified flag for the span</param>
+		/// <returns>The updated span, or null on failure</returns>
+		Task<IIssueSpan?> UpdateSpanAsync(IIssueSpan Span, IssueSeverity? NewSeverity = null, NewIssueStepData? NewLastSuccess = null, NewIssueStepData? NewFailure = null, NewIssueStepData? NewNextSuccess = null, List<NewIssueSpanSuspectData>? NewSuspects = null, bool? NewModified = null);
+
+		/// <summary>
+		/// Updates the given span
+		/// </summary>
+		/// <param name="Token">The token</param>
+		/// <param name="Span">The span to update</param>
+		/// <param name="Issue">The issue to attach to</param>
+		/// <param name="NewSeverity">New severity for the span</param>
+		/// <param name="NewFingerprint">New fingerprint for the span</param>
+		/// <param name="NewSummary">New summary for the issue</param>
+		/// <returns>The updated span, or null on failure</returns>
+		Task<bool> AttachSpanToIssueAsync(IIssueSequenceToken Token, IIssueSpan Span, IIssue Issue, IssueSeverity? NewSeverity = null, NewIssueFingerprint? NewFingerprint = null, string? NewSummary = null);
+
+		/// <summary>
+		/// Gets all the spans for a particular issue
+		/// </summary>
+		/// <param name="IssueId">Issue id</param>
+		/// <returns>List of spans</returns>
+		Task<List<IIssueSpan>> FindSpansAsync(int IssueId);
+
+		/// <summary>
+		/// Retrieves multiple spans
+		/// </summary>
+		/// <param name="SpanIds">The span ids</param>
+		/// <returns>List of spans</returns>
+		Task<List<IIssueSpan>> FindSpansAsync(IEnumerable<ObjectId> SpanIds);
+
+		/// <summary>
+		/// Finds the open issues for a given stream
+		/// </summary>
+		/// <param name="StreamId">The stream id</param>
+		/// <param name="TemplateId">The template id</param>
+		/// <param name="Name">Name of the node</param>
+		/// <param name="Change">Changelist number to query</param>
+		/// <returns>List of open issues</returns>
+		Task<List<IIssueSpan>> FindOpenSpansAsync(StreamId StreamId, TemplateRefId TemplateId, string Name, int Change);
+
+		/// <summary>
+		/// Finds a span marked as modified
+		/// </summary>
+		/// <returns></returns>
+		Task<IIssueSpan?> FindModifiedSpanAsync();
+
+		#endregion
+
+		#region Steps
+			
+		/// <summary>
+		/// Find steps for the given spans
+		/// </summary>
+		/// <param name="SpanIds">Span ids</param>
+		/// <returns>List of steps</returns>
+		Task<List<IIssueStep>> FindStepsAsync(IEnumerable<ObjectId> SpanIds);
+
+		/// <summary>
+		/// Find steps for the given spans
+		/// </summary>
+		/// <param name="JobId">The job id</param>
+		/// <param name="BatchId">The batch id</param>
+		/// <param name="StepId">The step id</param>
+		/// <returns>List of steps</returns>
+		Task<List<IIssueStep>> FindStepsAsync(ObjectId JobId, SubResourceId? BatchId, SubResourceId? StepId);
+
+		#endregion
+	}
+}
diff --git a/Engine/Source/Programs/Horde/HordeServer/Collections/IJobCollection.cs b/Engine/Source/Programs/Horde/HordeServer/Collections/IJobCollection.cs
new file mode 100644
--- /dev/null
+++ b/Engine/Source/Programs/Horde/HordeServer/Collections/IJobCollection.cs
@@ -0,0 +1,232 @@
+using Datadog.Trace;
+using EpicGames.Core;
+using HordeServer.Api;
+using HordeCommon;
+using HordeServer.Models;
+using HordeServer.Services;
+using HordeServer.Utilities;
+using Microsoft.Extensions.Logging;
+using MongoDB.Bson;
+using MongoDB.Bson.Serialization.Attributes;
+using MongoDB.Driver;
+using System;
+using System.Collections.Generic;
+using System.Diagnostics.CodeAnalysis;
+using System.IO;
+using System.Linq;
+using System.Threading.Tasks;
+
+using StreamId = HordeServer.Utilities.StringId<HordeServer.Models.IStream>;
+using TemplateRefId = HordeServer.Utilities.StringId<HordeServer.Models.TemplateRef>;
+using PoolId = HordeServer.Utilities.StringId<HordeServer.Models.IPool>;
+
+namespace HordeServer.Collections
+{
+	/// <summary>
+	/// Interface for a collection of job documents
+	/// </summary>
+	[SuppressMessage("Compiler", "CA1054:URI properties should not be strings")]
+	public interface IJobCollection
+	{
+		/// <summary>
+		/// Creates a new job
+		/// </summary>
+		/// <param name="JobId">A requested job id</param>
+		/// <param name="StreamId">Unique id of the stream that this job belongs to</param>
+		/// <param name="TemplateRefId">Name of the template ref</param>
+		/// <param name="TemplateHash">Template for this job</param>
+		/// <param name="Graph">The graph for the new job</param>
+		/// <param name="Name">Name of the job</param>
+		/// <param name="Change">The change to build</param>
+		/// <param name="CodeChange">The corresponding code changelist number</param>
+		/// <param name="PreflightChange">Optional changelist to preflight</param>
+		/// <param name="ClonedPreflightChange">Optional cloned preflight changelist</param>
+		/// <param name="StartedByUserId">User that started the job</param>
+		/// <param name="StartedByUserName">User that started the job</param>
+		/// <param name="Priority">Priority of the job</param>
+		/// <param name="AutoSubmit">Whether to automatically submit the preflighted change on completion</param>
+		/// <param name="JobTriggers">List of downstream job triggers</param>
+		/// <param name="ShowUgsBadges">Whether to show badges in UGS for this job</param>
+		/// <param name="ShowUgsAlerts">Whether to show alerts in UGS for this job</param>
+		/// <param name="NotificationChannel">Notification channel for this job</param>
+		/// <param name="NotificationChannelFilter">Notification channel filter for this job</param>
+		/// <param name="HelixSwarmCallbackUrl">Callback URL to a Helix Server review, if any</param>
+		/// <param name="Arguments">Arguments for the job</param>
+		/// <returns>The new job document</returns>
+		Task<IJob> AddAsync(ObjectId JobId, StreamId StreamId, TemplateRefId TemplateRefId, ContentHash TemplateHash, IGraph Graph, string Name, int Change, int CodeChange, int? PreflightChange, int? ClonedPreflightChange, ObjectId? StartedByUserId, string? StartedByUserName, Priority? Priority, bool? AutoSubmit, List<ChainedJobTemplate>? JobTriggers, bool ShowUgsBadges, bool ShowUgsAlerts, string? NotificationChannel, string? NotificationChannelFilter, string? HelixSwarmCallbackUrl, List<string>? Arguments);
+
+		/// <summary>
+		/// Gets a job with the given unique id
+		/// </summary>
+		/// <param name="JobId">Job id to search for</param>
+		/// <returns>Information about the given job</returns>
+		Task<IJob?> GetAsync(ObjectId JobId);
+
+		/// <summary>
+		/// Deletes a job
+		/// </summary>
+		/// <param name="Job">The job to remove</param>
+		Task<bool> RemoveAsync(IJob Job);
+
+		/// <summary>
+		/// Delete all the jobs for a stream
+		/// </summary>
+		/// <param name="StreamId">Unique id of the stream</param>
+		/// <returns>Async task</returns>
+		Task RemoveStreamAsync(StreamId StreamId);
+
+		/// <summary>
+		/// Gets a job's permissions info by ID
+		/// </summary>
+		/// <param name="JobId">Unique id of the job</param>
+		/// <returns>The job document</returns>
+		Task<IJobPermissions?> GetPermissionsAsync(ObjectId JobId);
+
+		/// <summary>
+		/// Searches for jobs matching the given criteria
+		/// </summary>
+		/// <param name="JobIds">List of job ids to return</param>
+		/// <param name="StreamId">The stream containing the job</param>
+		/// <param name="Name">Name of the job</param>
+		/// <param name="Templates">Templates to look for</param>
+		/// <param name="MinChange">The minimum changelist number</param>
+		/// <param name="MaxChange">The maximum changelist number</param>
+		/// <param name="PreflightChange">Preflight change to find</param>
+		/// <param name="PreflightStartedByUser">User for which to include preflight jobs</param>
+		/// <param name="MinCreateTime">The minimum creation time</param>
+		/// <param name="MaxCreateTime">The maximum creation time</param>
+		/// <param name="ModifiedBefore">Filter the results by modified time</param>
+		/// <param name="ModifiedAfter">Filter the results by modified time</param>
+		/// <param name="Index">Index of the first result to return</param>
+		/// <param name="Count">Number of results to return</param>
+		/// <returns>List of jobs matching the given criteria</returns>
+		Task<List<IJob>> FindAsync(ObjectId[]? JobIds = null, StreamId? StreamId = null, string? Name = null, TemplateRefId[]? Templates = null, int? MinChange = null, int? MaxChange = null, int? PreflightChange = null, string? PreflightStartedByUser = null, DateTimeOffset? MinCreateTime = null, DateTimeOffset? MaxCreateTime = null, DateTimeOffset? ModifiedBefore = null, DateTimeOffset? ModifiedAfter = null, int? Index = null, int? Count = null);
+
+		/// <summary>
+		/// Updates a new job
+		/// </summary>
+		/// <param name="Job">The job document to update</param>
+		/// <param name="Graph">Graph for the job</param>
+		/// <param name="Name">Name of the job</param>
+		/// <param name="Priority">Priority of the job</param>
+		/// <param name="AutoSubmit">Automatically submit the job on completion</param>
+		/// <param name="AutoSubmitChange">Changelist that was automatically submitted</param>
+		/// <param name="AutoSubmitMessage"></param>
+		/// <param name="AbortedByUser">Name of the user that aborted the job</param>
+		/// <param name="NotificationTriggerId">Id for a notification trigger</param>
+		/// <param name="Reports">New reports</param>
+		/// <param name="Arguments">New arguments for the job</param>
+		/// <param name="LabelIdxToTriggerId">New trigger ID for a label in the job</param>
+		/// <param name="JobTrigger">New downstream job id</param>
+		Task<bool> TryUpdateJobAsync(IJob Job, IGraph Graph, string? Name = null, Priority? Priority = null, bool? AutoSubmit = null, int? AutoSubmitChange = null, string? AutoSubmitMessage = null, string? AbortedByUser = null, ObjectId? NotificationTriggerId = null, List<Report>? Reports = null, List<string>? Arguments = null, KeyValuePair<int, ObjectId>? LabelIdxToTriggerId = null, KeyValuePair<TemplateRefId, ObjectId>? JobTrigger = null);
+
+		/// <summary>
+		/// Updates the state of a batch
+		/// </summary>
+		/// <param name="Job">Job to update</param>
+		/// <param name="Graph">The graph for this job</param>
+		/// <param name="BatchId">Unique id of the batch to update</param>
+		/// <param name="NewLogId">The new log file id</param>
+		/// <param name="NewState">New state of the jobstep</param>
+		/// <param name="NewError">Error code for the batch</param>
+		/// <returns>True if the job was updated, false if it was deleted</returns>
+		Task<bool> TryUpdateBatchAsync(IJob Job, IGraph Graph, SubResourceId BatchId, ObjectId? NewLogId, JobStepBatchState? NewState, JobStepBatchError? NewError);
+
+		/// <summary>
+		/// Update a jobstep state
+		/// </summary>
+		/// <param name="Job">Job to update</param>
+		/// <param name="Graph">The graph for this job</param>
+		/// <param name="BatchId">Unique id of the batch containing the step</param>
+		/// <param name="StepId">Unique id of the step to update</param>
+		/// <param name="NewState">New state of the jobstep</param>
+		/// <param name="NewOutcome">New outcome of the jobstep</param>
+		/// <param name="NewAbortRequested">New state of request abort</param>
+		/// <param name="NewAbortByUser">New name of user that requested the abort</param>
+		/// <param name="NewLogId">New log id for the jobstep</param>
+		/// <param name="NewNotificationTriggerId">New id for a notification trigger</param>
+		/// <param name="NewRetryByUser">Whether the step should be retried</param>
+		/// <param name="NewPriority">New priority for this step</param>
+		/// <param name="NewReports">New report documents</param>
+		/// <param name="NewProperties">Property changes. Any properties with a null value will be removed.</param>
+		/// <returns>True if the job was updated, false if it was deleted in the meantime</returns>
+		Task<bool> TryUpdateStepAsync(IJob Job, IGraph Graph, SubResourceId BatchId, SubResourceId StepId, JobStepState NewState = default, JobStepOutcome NewOutcome = default, bool? NewAbortRequested = null, string? NewAbortByUser = null, ObjectId? NewLogId = null, ObjectId? NewNotificationTriggerId = null, string? NewRetryByUser = null, Priority? NewPriority = null, List<Report>? NewReports = null, Dictionary<string, string?>? NewProperties = null);
+
+		/// <summary>
+		/// Attempts to update the node groups to be executed for a job. Fails if another write happens in the meantime.
+		/// </summary>
+		/// <param name="Job">The job to update</param>
+		/// <param name="NewGraph">New graph for this job</param>
+		/// <returns>True if the groups were updated to the given list. False if another write happened first.</returns>
+		Task<bool> TryUpdateGraphAsync(IJob Job, IGraph NewGraph);
+
+		/// <summary>
+		/// Adds an issue to a job
+		/// </summary>
+		/// <param name="JobId">The job id</param>
+		/// <param name="IssueId">The issue to add</param>
+		/// <returns>Async task</returns>
+		Task AddIssueToJobAsync(ObjectId JobId, int IssueId);
+
+		/// <summary>
+		/// Gets a queue of jobs to consider for execution
+		/// </summary>
+		/// <returns>Sorted list of jobs to execute</returns>
+		Task<List<IJob>> GetDispatchQueueAsync();
+
+		/// <summary>
+		/// Marks a job as skipped
+		/// </summary>
+		/// <param name="Job">The job to update</param>
+		/// <param name="Graph">Graph for the job</param>
+		/// <param name="Reason">Reason for this batch being failed</param>
+		/// <returns>Updated version of the job</returns>
+		Task<IJob?> SkipAllBatchesAsync(IJob? Job, IGraph Graph, JobStepBatchError Reason);
+
+		/// <summary>
+		/// Marks a batch as skipped
+		/// </summary>
+		/// <param name="Job">The job to update</param>
+		/// <param name="BatchIdx">The batch to mark as skipped</param>
+		/// <param name="Graph">Graph for the job</param>
+		/// <param name="Reason">Reason for this batch being failed</param>
+		/// <returns>Updated version of the job</returns>
+		Task<IJob?> SkipBatchAsync(IJob? Job, int BatchIdx, IGraph Graph, JobStepBatchError Reason);
+
+		/// <summary>
+		/// Abort an agent's lease, and update the payload accordingly
+		/// </summary>
+		/// <param name="Job">The job to update</param>
+		/// <param name="BatchIdx">Index of the batch to cancel</param>
+		/// <param name="Graph">Graph for the job</param>
+		/// <returns>True if the job is updated</returns>
+		Task<bool> TryFailBatchAsync(IJob Job, int BatchIdx, IGraph Graph);
+
+		/// <summary>
+		/// Attempt to assign a lease to execute a batch
+		/// </summary>
+		/// <param name="Job">The job containing the batch</param>
+		/// <param name="BatchIdx">Index of the batch</param>
+		/// <param name="PoolId">The pool id</param>
+		/// <param name="AgentId">New agent to execute the batch</param>
+		/// <param name="SessionId">Session of the agent that is to execute the batch</param>
+		/// <param name="LeaseId">The lease unique id</param>
+		/// <param name="LogId">Unique id of the log for the batch</param>
+		/// <returns>True if the batch is updated</returns>
+		Task<bool> TryAssignLeaseAsync(IJob Job, int BatchIdx, PoolId PoolId, AgentId AgentId, ObjectId SessionId, ObjectId LeaseId, ObjectId LogId);
+
+		/// <summary>
+		/// Cancel a lease reservation on a batch (before it has started)
+		/// </summary>
+		/// <param name="Job">The job containing the lease</param>
+		/// <param name="BatchIdx">Index of the batch to cancel</param>
+		/// <returns>True if the job is updated</returns>
+		Task<bool> TryCancelLeaseAsync(IJob Job, int BatchIdx);
+
+		/// <summary>
+		/// Upgrade all documents in the collection
+		/// </summary>
+		/// <returns>Async task</returns>
+		Task UpgradeDocumentsAsync();
+	}
+}
diff --git a/Engine/Source/Programs/Horde/HordeServer/Collections/IJobStepRefCollection.cs b/Engine/Source/Programs/Horde/HordeServer/Collections/IJobStepRefCollection.cs
new file mode 100644
--- /dev/null
+++ b/Engine/Source/Programs/Horde/HordeServer/Collections/IJobStepRefCollection.cs
@@ -0,0 +1,96 @@
+using HordeCommon;
+using HordeServer.Models;
+using HordeServer.Services;
+using MongoDB.Bson;
+using MongoDB.Bson.Serialization.Attributes;
+using MongoDB.Driver;
+using System;
+using System.Collections.Generic;
+using System.Linq;
+using System.Threading.Tasks;
+
+using StreamId = HordeServer.Utilities.StringId<HordeServer.Models.IStream>;
+using PoolId = HordeServer.Utilities.StringId<HordeServer.Models.IPool>;
+using TemplateRefId = HordeServer.Utilities.StringId<HordeServer.Models.TemplateRef>;
+using HordeServer.Utilities;
+
+namespace HordeServer.Collections
+{
+	/// <summary>
+	/// Interface for a collection of JobStepRef documents
+	/// </summary>
+	public interface IJobStepRefCollection
+	{
+		/// <summary>
+		/// Constructor
+		/// </summary>
+		/// <param name="Id">Unique id for the step being referenced</param>
+		/// <param name="JobName">Name of the job</param>
+		/// <param name="StepName">Name of the step</param>
+		/// <param name="StreamId">Unique id for the stream containing the job</param>
+		/// <param name="TemplateId"></param>
+		/// <param name="Change">The change number being built</param>
+		/// <param name="LogId">The log file id</param>
+		/// <param name="PoolId">The pool id</param>
+		/// <param name="AgentId">The agent id</param>
+		/// <param name="Outcome">Outcome of this step, if known</param>
+		/// <param name="WaitTime">Time taken for the batch containing this step to start</param>
+		/// <param name="InitTime">Time taken for the batch containing this step to initializer</param>
+		/// <param name="StartTimeUtc">Start time</param>
+		/// <param name="FinishTimeUtc">Finish time for the step, if known</param>
+		Task<IJobStepRef> InsertOrReplaceAsync(JobStepRefId Id, string JobName, string StepName, StreamId StreamId, TemplateRefId TemplateId, int Change, ObjectId? LogId, PoolId? PoolId, AgentId? AgentId, JobStepOutcome? Outcome, float WaitTime, float InitTime, DateTime StartTimeUtc, DateTime? FinishTimeUtc);
+
+		/// <summary>
+		/// Gets the history of a given node
+		/// </summary>
+		/// <param name="StreamId">Unique id for a stream</param>
+		/// <param name="TemplateId"></param>
+		/// <param name="NodeName">Name of the node</param>
+		/// <param name="Change">The current change</param>
+		/// <param name="IncludeFailed">Whether to include failed nodes</param>
+		/// <param name="Count">Number of results to return</param>
+		/// <returns>List of step references</returns>
+		Task<List<IJobStepRef>> GetStepsForNodeAsync(StreamId StreamId, TemplateRefId TemplateId, string NodeName, int? Change, bool IncludeFailed, int Count);
+
+		/// <summary>
+		/// Gets the previous job that ran a given step
+		/// </summary>
+		/// <param name="StreamId">Id of the stream to search</param>
+		/// <param name="TemplateId">The template id</param>
+		/// <param name="NodeName">Name of the step to find</param>
+		/// <param name="Change">The current changelist number</param>
+		/// <returns>The previous job, or null.</returns>
+		Task<IJobStepRef?> GetPrevStepForNodeAsync(StreamId StreamId, TemplateRefId TemplateId, string NodeName, int Change);
+
+		/// <summary>
+		/// Gets the next job that ran a given step
+		/// </summary>
+		/// <param name="StreamId">Id of the stream to search</param>
+		/// <param name="TemplateId">The template id</param>
+		/// <param name="NodeName">Name of the step to find</param>
+		/// <param name="Change">The current changelist number</param>
+		/// <returns>The previous job, or null.</returns>
+		Task<IJobStepRef?> GetNextStepForNodeAsync(StreamId StreamId, TemplateRefId TemplateId, string NodeName, int Change);
+	}
+
+	static class JobStepRefCollectionExtensions
+	{
+		public static Task UpdateAsync(this IJobStepRefCollection JobStepRefs, IJob Job, IJobStepBatch Batch, IJobStep Step, IGraph Graph)
+		{
+			if (Job.PreflightChange == 0)
+			{
+				float WaitTime = (float)(Batch.GetWaitTime() ?? TimeSpan.Zero).TotalSeconds;
+				float InitTime = (float)(Batch.GetInitTime() ?? TimeSpan.Zero).TotalSeconds;
+
+				string NodeName = Graph.Groups[Batch.GroupIdx].Nodes[Step.NodeIdx].Name;
+				JobStepOutcome? Outcome = Step.IsPending() ? (JobStepOutcome?)null : Step.Outcome;
+
+				return JobStepRefs.InsertOrReplaceAsync(new JobStepRefId(Job.Id, Batch.Id, Step.Id), Job.Name, NodeName, Job.StreamId, Job.TemplateId, Job.Change, Step.LogId, Batch.PoolId, Batch.AgentId, Outcome, WaitTime, InitTime, Step.StartTimeUtc ?? DateTime.UtcNow, Step.FinishTimeUtc);
+			}
+			else
+			{
+				return Task.CompletedTask;
+			}
+		}
+	}
+}
diff --git a/Engine/Source/Programs/Horde/HordeServer/Collections/IJobTimingCollection.cs b/Engine/Source/Programs/Horde/HordeServer/Collections/IJobTimingCollection.cs
new file mode 100644
--- /dev/null
+++ b/Engine/Source/Programs/Horde/HordeServer/Collections/IJobTimingCollection.cs
@@ -0,0 +1,98 @@
+using Amazon.Runtime.Internal.Util;
+using Amazon.S3.Model.Internal.MarshallTransformations;
+using HordeServer.Models;
+using HordeServer.Services;
+using HordeServer.Utilities;
+using MongoDB.Bson;
+using MongoDB.Bson.Serialization.Attributes;
+using MongoDB.Driver;
+using System;
+using System.Collections.Generic;
+using System.Diagnostics.CodeAnalysis;
+using System.Linq;
+using System.Threading.Tasks;
+
+namespace HordeServer.Collections
+{
+	/// <summary>
+	/// Average timing information for a node
+	/// </summary>
+	public class JobStepTimingData
+	{
+		/// <summary>
+		/// Name of the node
+		/// </summary>
+		public string Name { get; set; }
+
+		/// <summary>
+		/// Wait time before executing the group containing this node
+		/// </summary>
+		[BsonIgnoreIfNull]
+		public float? AverageWaitTime { get; set; }
+
+		/// <summary>
+		/// Time taken for the group containing this node to initialize
+		/// </summary>
+		[BsonIgnoreIfNull]
+		public float? AverageInitTime { get; set; }
+
+		/// <summary>
+		/// Time spent executing this node
+		/// </summary>
+		[BsonIgnoreIfNull]
+		public float? AverageDuration { get; set; }
+
+		/// <summary>
+		/// Constructor for serialization
+		/// </summary>
+		[BsonConstructor]
+		private JobStepTimingData()
+		{
+			this.Name = null!;
+		}
+
+		/// <summary>
+		/// Constructor
+		/// </summary>
+		/// <param name="Name">Name of the node</param>
+		/// <param name="AverageWaitTime">Wait time before executing the group containing this node</param>
+		/// <param name="AverageInitTime">Time taken for the group containing this node to initialize</param>
+		/// <param name="AverageDuration">Time spent executing this node</param>
+		public JobStepTimingData(string Name, float? AverageWaitTime, float? AverageInitTime, float? AverageDuration)
+		{
+			this.Name = Name;
+			this.AverageWaitTime = AverageWaitTime;
+			this.AverageInitTime = AverageInitTime;
+			this.AverageDuration = AverageDuration;
+		}
+	}
+
+	/// <summary>
+	/// Interface for a collection of JobTiming documents
+	/// </summary>
+	public interface IJobTimingCollection
+	{
+		/// <summary>
+		/// Add timing information for the given job
+		/// </summary>
+		/// <param name="JobId">The job id</param>
+		/// <param name="Steps">List of timing info for each step</param>
+		/// <returns>New timing document</returns>
+		Task<IJobTiming?> TryAddAsync(ObjectId JobId, List<JobStepTimingData> Steps);
+
+		/// <summary>
+		/// Attempts to get the timing information for a particular job
+		/// </summary>
+		/// <param name="JobId">The unique job id</param>
+		/// <returns>Timing info for the requested jbo</returns>
+		Task<IJobTiming?> TryGetAsync(ObjectId JobId);
+
+		/// <summary>
+		/// Adds timing information for the particular job
+		/// </summary>
+		/// <param name="JobTiming">The current timing info for the job</param>
+		/// <param name="Steps">List of steps to add</param>
+		/// <returns>New timing document. Null if the document could not be updated.</returns>
+		Task<IJobTiming?> TryAddStepsAsync(IJobTiming JobTiming, List<JobStepTimingData> Steps);
+	}
+}
diff --git a/Engine/Source/Programs/Horde/HordeServer/Collections/ILeaseCollection.cs b/Engine/Source/Programs/Horde/HordeServer/Collections/ILeaseCollection.cs
new file mode 100644
--- /dev/null
+++ b/Engine/Source/Programs/Horde/HordeServer/Collections/ILeaseCollection.cs
@@ -0,0 +1,81 @@
+using HordeCommon;
+using HordeServer.Models;
+using HordeServer.Services;
+using HordeServer.Utilities;
+using MongoDB.Bson;
+using MongoDB.Bson.Serialization.Attributes;
+using MongoDB.Driver;
+using System;
+using System.Collections.Generic;
+using System.Linq;
+using System.Threading.Tasks;
+
+using PoolId = HordeServer.Utilities.StringId<HordeServer.Models.IPool>;
+using StreamId = HordeServer.Utilities.StringId<HordeServer.Models.IStream>;
+
+namespace HordeServer.Collections
+{
+	/// <summary>
+	/// Interface for a collection of lease documents
+	/// </summary>
+	public interface ILeaseCollection
+	{
+		/// <summary>
+		/// Adds a lease to the collection
+		/// </summary>
+		/// <param name="Id">The lease id</param>
+		/// <param name="Name">Name of the lease</param>
+		/// <param name="AgentId">The agent id</param>
+		/// <param name="SessionId">The agent session handling the lease</param>
+		/// <param name="StreamId">Stream for the payload</param>
+		/// <param name="PoolId">The pool for the lease</param>
+		/// <param name="LogId">Log id for the lease</param>
+		/// <param name="StartTime">Start time of the lease</param>
+		/// <param name="Payload">Payload for the lease</param>
+		/// <returns>Async task</returns>
+		Task<ILease> AddAsync(ObjectId Id, string Name, AgentId AgentId, ObjectId SessionId, StreamId? StreamId, PoolId? PoolId, ObjectId? LogId, DateTime StartTime, byte[] Payload);
+
+		/// <summary>
+		/// Deletes a lease from the collection
+		/// </summary>
+		/// <param name="LeaseId">Unique id of the lease</param>
+		/// <returns>Async task</returns>
+		Task DeleteAsync(ObjectId LeaseId);
+
+		/// <summary>
+		/// Gets a specific lease
+		/// </summary>
+		/// <param name="LeaseId">Unique id of the lease</param>
+		/// <returns>The lease that was found, or null if it does not exist</returns>
+		Task<ILease?> GetAsync(ObjectId LeaseId);
+
+		/// <summary>
+		/// Finds all leases matching a set of criteria
+		/// </summary>
+		/// <param name="AgentId">Unqiue id of the agent executing this lease</param>
+		/// <param name="SessionId">Unique id of the agent session</param>
+		/// <param name="MinTime">Start of the window to include leases</param>
+		/// <param name="MaxTime">End of the window to include leases</param>
+		/// <param name="Index">Index of the first result to return</param>
+		/// <param name="Count">Number of results to return</param>
+		/// <returns>List of leases matching the given criteria</returns>
+		Task<List<ILease>> FindLeasesAsync(AgentId? AgentId = null, ObjectId? SessionId = null, DateTime? MinTime = null, DateTime? MaxTime = null, int? Index = null, int? Count = null);
+
+		/// <summary>
+		/// Finds all active leases
+		/// </summary>
+		/// <param name="Index">Index of the first result to return</param>
+		/// <param name="Count">Number of results to return</param>
+		/// <returns>List of leases</returns>
+		Task<List<ILease>> FindActiveLeasesAsync(int? Index = null, int? Count = null);
+
+		/// <summary>
+		/// Sets the outcome of a lease
+		/// </summary>
+		/// <param name="LeaseId">The lease to update</param>
+		/// <param name="FinishTime">Time at which the lease finished</param>
+		/// <param name="Outcome">Outcome of the lease</param>
+		/// <returns>True if the lease was updated, false otherwise</returns>
+		Task<bool> TrySetOutcomeAsync(ObjectId LeaseId, DateTime FinishTime, LeaseOutcome Outcome);
+	}
+}
diff --git a/Engine/Source/Programs/Horde/HordeServer/Collections/IPoolCollection.cs b/Engine/Source/Programs/Horde/HordeServer/Collections/IPoolCollection.cs
new file mode 100644
--- /dev/null
+++ b/Engine/Source/Programs/Horde/HordeServer/Collections/IPoolCollection.cs
@@ -0,0 +1,78 @@
+using HordeServer.Models;
+using HordeServer.Services;
+using HordeServer.Utilities;
+using MongoDB.Bson;
+using MongoDB.Bson.Serialization;
+using MongoDB.Bson.Serialization.Attributes;
+using MongoDB.Driver;
+using System;
+using System.Collections.Generic;
+using System.Linq;
+using System.Linq.Expressions;
+using System.Threading.Tasks;
+
+using PoolId = HordeServer.Utilities.StringId<HordeServer.Models.IPool>;
+
+namespace HordeServer.Collections
+{
+	/// <summary>
+	/// Collection of pool documents
+	/// </summary>
+	public interface IPoolCollection
+	{
+		/// <summary>
+		/// Creates a new pool
+		/// </summary>
+		/// <param name="Id">Unique id for the new pool</param>
+		/// <param name="Name">Name of the new pool</param>
+		/// <param name="Requirements">Requirements for agents to include in this pool</param>
+		/// <param name="EnableAutoscaling">Whether to enable autoscaling for this pool</param>
+		/// <param name="MinAgents">Minimum number of agents in the pool</param>
+		/// <param name="NumReserveAgents">Minimum number of idle agents to maintain</param>
+		/// <param name="Properties">Properties for the pool</param>
+		/// <returns>The new pool document</returns>
+		Task<IPool> AddAsync(PoolId Id, string Name, AgentRequirements? Requirements = null, bool? EnableAutoscaling = null, int? MinAgents = null, int? NumReserveAgents = null, IEnumerable<KeyValuePair<string, string>>? Properties = null);
+
+		/// <summary>
+		/// Enumerates all the pools
+		/// </summary>
+		/// <returns>The pool documents</returns>
+		Task<List<IPool>> GetAsync();
+
+		/// <summary>
+		/// Gets a pool by ID
+		/// </summary>
+		/// <param name="Id">Unique id of the pool</param>
+		/// <returns>The pool document</returns>
+		Task<IPool?> GetAsync(PoolId Id);
+
+		/// <summary>
+		/// Gets a list of all valid pool ids
+		/// </summary>
+		/// <returns>List of pool ids</returns>
+		Task<List<PoolId>> GetPoolIdsAsync();
+
+		/// <summary>
+		/// Gets a pool by ID
+		/// </summary>
+		/// <param name="Id">Unique id of the pool</param>
+		/// <returns>The pool document</returns>
+		Task<bool> DeleteAsync(PoolId Id);
+
+		/// <summary>
+		/// Updates a pool
+		/// </summary>
+		/// <param name="Pool">The pool to update</param>
+		/// <param name="NewName">New name for the pool</param>
+		/// <param name="NewRequirements">New requirements for the pool</param>
+		/// <param name="NewEnableAutoscaling">New setting for whether to enable autoscaling</param>
+		/// <param name="NewMinAgents">Minimum number of agents in the pool</param>
+		/// <param name="NewNumReserveAgents">Minimum number of idle agents to maintain</param>
+		/// <param name="NewWorkspaces">New workspaces for the pool</param>
+		/// <param name="NewProperties">New properties for the pool</param>
+		/// <param name="LastScaleUpTime">New time for last (auto) scale up</param>
+		/// <param name="LastScaleDownTime">New time for last (auto) scale down</param>
+		/// <returns>Async task</returns>
+		Task<IPool?> TryUpdateAsync(IPool Pool, string? NewName = null, AgentRequirements? NewRequirements = null, bool? NewEnableAutoscaling = null, int? NewMinAgents = null, int? NewNumReserveAgents = null, List<AgentWorkspace>? NewWorkspaces = null, Dictionary<string, string?>? NewProperties = null, DateTime? LastScaleUpTime = null, DateTime? LastScaleDownTime = null);
+	}
+}
diff --git a/Engine/Source/Programs/Horde/HordeServer/Collections/IProjectCollection.cs b/Engine/Source/Programs/Horde/HordeServer/Collections/IProjectCollection.cs
new file mode 100644
--- /dev/null
+++ b/Engine/Source/Programs/Horde/HordeServer/Collections/IProjectCollection.cs
@@ -0,0 +1,77 @@
+// Copyright Epic Games, Inc. All Rights Reserved.
+
+using HordeServer.Api;
+using HordeServer.Models;
+using HordeServer.Services;
+using HordeServer.Utilities;
+using Microsoft.Extensions.Options;
+using MongoDB.Bson;
+using MongoDB.Bson.Serialization.Attributes;
+using MongoDB.Driver;
+using System;
+using System.Collections.Generic;
+using System.Diagnostics.CodeAnalysis;
+using System.Linq;
+using System.Security.Claims;
+using System.Threading.Tasks;
+
+using ProjectId = HordeServer.Utilities.StringId<HordeServer.Models.IProject>;
+
+namespace HordeServer.Collections
+{
+	/// <summary>
+	/// Interface for a collection of project documents
+	/// </summary>
+	public interface IProjectCollection
+	{
+		/// <summary>
+		/// Attempts to create a new project
+		/// </summary>
+		/// <param name="Id">Identifier for the new project</param>
+		/// <param name="Name">Name of the new project</param>
+		/// <param name="Order">Order to display this project on the dashboard</param>
+		/// <param name="Categories">Categories for the new project</param>
+		/// <param name="Properties">Properties for the new project</param>
+		/// <returns>The new project document</returns>
+		Task<IProject?> TryAddAsync(ProjectId Id, string Name, int? Order = null, List<StreamCategory>? Categories = null, Dictionary<string, string>? Properties = null);
+
+		/// <summary>
+		/// Updates an existing project
+		/// </summary>
+		/// <param name="ProjectId">Unique id of the project</param>
+		/// <param name="NewName">The new name for the project</param>
+		/// <param name="NewOrder">The new order for this project on the dashboard</param>
+		/// <param name="NewCategories">List of categories for the project</param>
+		/// <param name="NewProperties">Properties on the project to update. Any properties with a value of null will be removed.</param>
+		/// <param name="NewAcl">New acl for the project</param>
+		/// <returns>Async task object</returns>
+		Task UpdateAsync(ProjectId ProjectId, string? NewName, int? NewOrder, List<StreamCategory>? NewCategories, Dictionary<string, string>? NewProperties, Acl? NewAcl);
+
+		/// <summary>
+		/// Gets all the available projects
+		/// </summary>
+		/// <returns>List of project documents</returns>
+		Task<List<IProject>> FindAllAsync();
+
+		/// <summary>
+		/// Gets a project by ID
+		/// </summary>
+		/// <param name="ProjectId">Unique id of the project</param>
+		/// <returns>The project document</returns>
+		Task<IProject?> GetAsync(ProjectId ProjectId);
+
+		/// <summary>
+		/// Gets a project's permissions info by ID
+		/// </summary>
+		/// <param name="ProjectId">Unique id of the project</param>
+		/// <returns>The project document</returns>
+		Task<IProjectPermissions?> GetPermissionsAsync(ProjectId ProjectId);
+
+		/// <summary>
+		/// Deletes a project by id
+		/// </summary>
+		/// <param name="ProjectId">Unique id of the project</param>
+		/// <returns>True if the project was deleted</returns>
+		Task DeleteAsync(ProjectId ProjectId);
+	}
+}
diff --git a/Engine/Source/Programs/Horde/HordeServer/Collections/ISessionCollection.cs b/Engine/Source/Programs/Horde/HordeServer/Collections/ISessionCollection.cs
new file mode 100644
--- /dev/null
+++ b/Engine/Source/Programs/Horde/HordeServer/Collections/ISessionCollection.cs
@@ -0,0 +1,74 @@
+using HordeServer.Models;
+using HordeServer.Services;
+using HordeServer.Utilities;
+using Microsoft.Extensions.DependencyInjection;
+using MongoDB.Bson;
+using MongoDB.Bson.Serialization.Attributes;
+using MongoDB.Driver;
+using System;
+using System.Collections.Generic;
+using System.Linq;
+using System.Threading.Tasks;
+
+using AgentSoftwareVersion = HordeServer.Utilities.StringId<HordeServer.Collections.IAgentSoftwareCollection>;
+
+namespace HordeServer.Collections
+{
+	/// <summary>
+	/// Interface for a collection of session documents
+	/// </summary>
+	public interface ISessionCollection
+	{
+		/// <summary>
+		/// Adds a new session to the collection
+		/// </summary>
+		/// <param name="Id">The session id</param>
+		/// <param name="AgentId">The agent this session is for</param>
+		/// <param name="StartTime">Start time of this session</param>
+		/// <param name="Capabilities">Properties of this agent at the time the session started</param>
+		/// <param name="Version">Version of the agent software</param>
+		Task<ISession> AddAsync(ObjectId Id, AgentId AgentId, DateTime StartTime, AgentCapabilities? Capabilities, AgentSoftwareVersion? Version);
+
+		/// <summary>
+		/// Gets information about a particular session
+		/// </summary>
+		/// <param name="SessionId">The unique session id</param>
+		/// <returns>The session information</returns>
+		Task<ISession?> GetAsync(ObjectId SessionId);
+
+		/// <summary>
+		/// Find sessions for the given agent
+		/// </summary>
+		/// <param name="AgentId">The unique agent id</param>
+		/// <param name="StartTime">Start time to include in the search</param>
+		/// <param name="FinishTime">Finish time to include in the search</param>
+		/// <param name="Index">Index of the first result to return</param>
+		/// <param name="Count">Number of results to return</param>
+		/// <returns>List of sessions matching the given criteria</returns>
+		Task<List<ISession>> FindAsync(AgentId AgentId, DateTime? StartTime, DateTime? FinishTime, int Index, int Count);
+
+		/// <summary>
+		/// Finds any active sessions
+		/// </summary>
+		/// <param name="Index">Index of the first result to return</param>
+		/// <param name="Count">Number of results to return</param>
+		/// <returns>List of sessions</returns>
+		Task<List<ISession>> FindActiveSessionsAsync(int? Index = null, int? Count = null);
+
+		/// <summary>
+		/// Update a session from the collection
+		/// </summary>
+		/// <param name="SessionId">The session to update</param>
+		/// <param name="FinishTime">Time at which the session finished</param>
+		/// <param name="Capabilities">Capabilities of the agent at the end of the session</param>
+		/// <returns>Async task</returns>
+		Task UpdateAsync(ObjectId SessionId, DateTime FinishTime, AgentCapabilities Capabilities);
+
+		/// <summary>
+		/// Delete a session from the collection
+		/// </summary>
+		/// <param name="SessionId">The session id</param>
+		/// <returns>Async task</returns>
+		Task DeleteAsync(ObjectId SessionId);
+	}
+}
diff --git a/Engine/Source/Programs/Horde/HordeServer/Collections/ISoftwareCollection.cs b/Engine/Source/Programs/Horde/HordeServer/Collections/ISoftwareCollection.cs
new file mode 100644
--- /dev/null
+++ b/Engine/Source/Programs/Horde/HordeServer/Collections/ISoftwareCollection.cs
@@ -0,0 +1,71 @@
+using EpicGames.Core;
+using HordeCommon;
+using HordeServer.Models;
+using HordeServer.Services;
+using MongoDB.Bson;
+using MongoDB.Bson.Serialization.Attributes;
+using MongoDB.Driver;
+using System;
+using System.Collections.Generic;
+using System.IO;
+using System.IO.Compression;
+using System.Linq;
+using System.Text.Json;
+using System.Threading.Tasks;
+
+namespace HordeServer.Collections
+{
+	/// <summary>
+	/// Interface for a collection of template documents
+	/// </summary>
+	public interface ISoftwareCollection
+	{
+		/// <summary>
+		/// Adds a new software revision
+		/// </summary>
+		/// <param name="Data">Zip file containing the new software</param>
+		/// <param name="User">User that uploaded the software</param>
+		/// <param name="Default">Whether to make the software the default</param>
+		/// <returns>New software instance</returns>
+		Task<ISoftware> AddAsync(byte[] Data, string? User, bool Default);
+
+		/// <summary>
+		/// Upload a software archive to the backend
+		/// </summary>
+		/// <param name="Id">Unique id of the software to update</param>
+		/// <param name="User">Name of the current user</param>
+		/// <param name="Default">Whether the software should be the default</param>
+		/// <returns>Async task</returns>
+		Task<bool> UpdateAsync(ObjectId Id, string? User, bool Default);
+
+		/// <summary>
+		/// Finds all software archives matching a set of criteria
+		/// </summary>
+		/// <param name="UploadedByUser">The user that uploaded the software</param>
+		/// <param name="MadeDefaultByUser">The user that made the software the default</param>
+		/// <param name="MadeDefault">Whether the software was made default</param>
+		/// <param name="Offset">Index of the first result to return</param>
+		/// <param name="Count">Number of results to return</param>
+		Task<List<ISoftware>> FindAsync(string? UploadedByUser, string? MadeDefaultByUser, bool? MadeDefault, int Offset, int Count);
+
+		/// <summary>
+		/// Gets a single software archive
+		/// </summary>
+		/// <param name="Id">Unique id of the file to retrieve</param>
+		/// <param name="bIncludeData">Whether to include the data</param>
+		Task<ISoftware?> GetAsync(ObjectId Id, bool bIncludeData);
+
+		/// <summary>
+		/// Finds the current default software
+		/// </summary>
+		/// <param name="bIncludeData">Whether to include the data</param>
+		/// <returns>The current default software</returns>
+		Task<ISoftware?> GetDefaultAsync(bool bIncludeData);
+
+		/// <summary>
+		/// Removes a software archive 
+		/// </summary>
+		/// <param name="Id">Unique id of the software to delete</param>
+		Task<bool> DeleteAsync(ObjectId Id);
+	}
+}
