Initial drop of Horde Server/Agent for UE5/Main. Still missing some Epic-specific functionality required to build/run correctly.

diff --git a/Engine/Source/Programs/Horde/HordeServer/Collections/IStreamCollection.cs b/Engine/Source/Programs/Horde/HordeServer/Collections/IStreamCollection.cs
new file mode 100644
--- /dev/null
+++ b/Engine/Source/Programs/Horde/HordeServer/Collections/IStreamCollection.cs
@@ -0,0 +1,143 @@
+using HordeServer.Api;
+using HordeServer.Models;
+using HordeServer.Services;
+using HordeServer.Utilities;
+using MongoDB.Bson;
+using MongoDB.Bson.Serialization.Attributes;
+using MongoDB.Driver;
+using System;
+using System.Collections.Generic;
+using System.Diagnostics.CodeAnalysis;
+using System.Linq;
+using System.Threading.Tasks;
+
+using ProjectId = HordeServer.Utilities.StringId<HordeServer.Models.IProject>;
+using StreamId = HordeServer.Utilities.StringId<HordeServer.Models.IStream>;
+using TemplateRefId = HordeServer.Utilities.StringId<HordeServer.Models.TemplateRef>;
+
+namespace HordeServer.Collections
+{
+	/// <summary>
+	/// Collection of stream documents
+	/// </summary>
+	public interface IStreamCollection
+	{
+		/// <summary>
+		/// Creates a new stream
+		/// </summary>
+		/// <param name="Id">Unique id for the new stream</param>
+		/// <param name="Name">Name of the new stream</param>
+		/// <param name="ProjectId">Unique id of the project</param>
+		/// <param name="ConfigPath">Path to the config file for this stream in Perforce</param>
+		/// <param name="Order">Order for this stream</param>
+		/// <param name="NotificationChannel">Notification channel for this stream</param>
+		/// <param name="NotificationChannelFilter">Notification channel filter for this stream</param>
+		/// <param name="TriageChannel">Notification channel for triage requests</param>
+		/// <param name="DefaultPreflight">Default template for preflights</param>
+		/// <param name="Tabs">Tabs for the new stream</param>
+		/// <param name="AgentTypes">Map of new agent types to machine attributes</param>
+		/// <param name="WorkspaceTypes">Map of new workspace types</param>
+		/// <param name="TemplateRefs">New job templates for the stream</param>
+		/// <param name="Properties">New properties for the stream</param>
+		/// <param name="Acl">The ACL for this stream</param>
+		/// <returns>The new stream document</returns>
+		Task<IStream?> TryCreateAsync(StreamId Id, string Name, ProjectId ProjectId, string? ConfigPath = null, int? Order = null, string? NotificationChannel = null, string? NotificationChannelFilter = null, string? TriageChannel = null, DefaultPreflight? DefaultPreflight = null, List<StreamTab>? Tabs = null, Dictionary<string, AgentType>? AgentTypes = null, Dictionary<string, WorkspaceType>? WorkspaceTypes = null, Dictionary<TemplateRefId, TemplateRef>? TemplateRefs = null, Dictionary<string, string>? Properties = null, Acl? Acl = null);
+
+		/// <summary>
+		/// Replaces an existing stream
+		/// </summary>
+		/// <param name="StreamInterface">Unique id for the new stream</param>
+		/// <param name="Name">Name of the new stream</param>
+		/// <param name="ConfigPath">Path to the config file for this stream in Perforce</param>
+		/// <param name="ConfigChange">Changelist number for the config file</param>
+		/// <param name="Order">Order for this stream</param>
+		/// <param name="NotificationChannel">Notification channel for this stream</param>
+		/// <param name="NotificationChannelFilter">Notification channel filter for this stream</param>
+		/// <param name="TriageChannel"></param>
+		/// <param name="DefaultPreflight">Default template for preflights</param>
+		/// <param name="Tabs">Tabs for the new stream</param>
+		/// <param name="AgentTypes">Map of new agent types to machine attributes</param>
+		/// <param name="WorkspaceTypes">Map of new workspace types</param>
+		/// <param name="TemplateRefs">New job templates for the stream</param>
+		/// <param name="Properties">New properties for the stream</param>
+		/// <param name="Acl">The ACL for this stream</param>
+		/// <param name="PausedUntil">The new datetime for pausing builds</param>
+		/// <param name="PauseComment">The reason for pausing</param>
+		/// <returns>The new stream document</returns>
+		Task<IStream?> TryReplaceAsync(IStream StreamInterface, string Name, string? ConfigPath = null, int? ConfigChange = null, int? Order = null, string? NotificationChannel = null, string? NotificationChannelFilter = null, string? TriageChannel = null, DefaultPreflight? DefaultPreflight = null, List<StreamTab>? Tabs = null, Dictionary<string, AgentType>? AgentTypes = null, Dictionary<string, WorkspaceType>? WorkspaceTypes = null, Dictionary<TemplateRefId, TemplateRef>? TemplateRefs = null, Dictionary<string, string>? Properties = null, Acl? Acl = null, DateTime? PausedUntil = null, string? PauseComment = null);
+
+		/// <summary>
+		/// Gets a stream by ID
+		/// </summary>
+		/// <param name="StreamId">Unique id of the stream</param>
+		/// <returns>The stream document</returns>
+		Task<IStream?> GetAsync(StreamId StreamId);
+
+		/// <summary>
+		/// Gets a stream's permissions by ID
+		/// </summary>
+		/// <param name="StreamId">Unique id of the stream</param>
+		/// <returns>The stream document</returns>
+		Task<IStreamPermissions?> GetPermissionsAsync(StreamId StreamId);
+
+		/// <summary>
+		/// Enumerates all streams
+		/// </summary>
+		/// <returns></returns>
+		Task<List<IStream>> FindAllAsync();
+
+		/// <summary>
+		/// Gets all the available streams for a project
+		/// </summary>
+		/// <param name="ProjectIds">Unique id of the projects to query</param>
+		/// <returns>List of stream documents</returns>
+		Task<List<IStream>> FindForProjectsAsync(ProjectId[] ProjectIds);
+
+		/// <summary>
+		/// Updates user-facing properties for an existing stream
+		/// </summary>
+		/// <param name="Stream">The stream to update</param>
+		/// <param name="NewName">The new name for the stream</param>
+		/// <param name="NewOrder">New order for the stream</param>
+		/// <param name="NewNotificationChannel">New notification channel for the stream</param>
+		/// <param name="NewNotificationChannelFilter">New notification channel filter for the stream</param>
+		/// <param name="NewTriageChannel">New triage channel</param>
+		/// <param name="NewTabs">New tabs for the stream</param>
+		/// <param name="NewAgentTypes">Map of agent types to update. Anything with a value of null will be removed.</param>
+		/// <param name="NewWorkspaceTypes">Map of workspace types to update. Anything with a value of null will be removed.</param>
+		/// <param name="NewTemplateRefs">New template references for this stream</param>
+		/// <param name="NewProperties">Properties on the stream to update. Anything with a value of null will be removed.</param>
+		/// <param name="NewAcl">The new ACL object</param>
+		/// <param name="UpdatePauseFields">Must be set to true to update the pause fields</param>
+		/// <param name="NewPausedUntil">The new datetime for pausing builds</param>
+		/// <param name="NewPauseComment">The reason for pausing</param>
+		/// <returns>Async task object</returns>
+		Task<bool> TryUpdatePropertiesAsync(IStream Stream, string? NewName, int? NewOrder, string? NewNotificationChannel, string? NewNotificationChannelFilter, string? NewTriageChannel, List<StreamTab>? NewTabs, Dictionary<string, AgentType?>? NewAgentTypes, Dictionary<string, WorkspaceType?>? NewWorkspaceTypes, Dictionary<TemplateRefId, TemplateRef>? NewTemplateRefs, Dictionary<string, string?>? NewProperties, Acl? NewAcl, bool? UpdatePauseFields, DateTime? NewPausedUntil, string? NewPauseComment);
+
+		/// <summary>
+		/// Attempts to update the last commit time for a streams
+		/// </summary>
+		/// <param name="Stream"></param>
+		/// <param name="LastCommitTime"></param>
+		/// <returns></returns>
+		Task<bool> TryUpdateCommitTimeAsync(IStream Stream, DateTime LastCommitTime);
+
+		/// <summary>
+		/// Attempts to update the last trigger time for a schedule
+		/// </summary>
+		/// <param name="Stream">The stream to update</param>
+		/// <param name="TemplateRefId">The template ref id</param>
+		/// <param name="LastTriggerTime">New last trigger time for the schedule</param>
+		/// <param name="LastTriggerChange">New last trigger changelist for the schedule</param>
+		/// <param name="NewActiveJobs">New list of active jobs</param>
+		/// <returns>True if the stream was updated</returns>
+		Task<bool> TryUpdateScheduleTriggerAsync(IStream Stream, TemplateRefId TemplateRefId, DateTimeOffset? LastTriggerTime, int? LastTriggerChange, List<ObjectId> NewActiveJobs);
+
+		/// <summary>
+		/// Delete a stream
+		/// </summary>
+		/// <param name="StreamId">Unique id of the stream</param>
+		/// <returns>Async task</returns>
+		Task DeleteAsync(StreamId StreamId);
+	}
+}
diff --git a/Engine/Source/Programs/Horde/HordeServer/Collections/ITelemetryCollection.cs b/Engine/Source/Programs/Horde/HordeServer/Collections/ITelemetryCollection.cs
new file mode 100644
--- /dev/null
+++ b/Engine/Source/Programs/Horde/HordeServer/Collections/ITelemetryCollection.cs
@@ -0,0 +1,39 @@
+using HordeServer.Models;
+using HordeServer.Utilities;
+using System;
+using System.Collections.Generic;
+using System.Linq;
+using System.Threading.Tasks;
+
+using PoolId = HordeServer.Utilities.StringId<HordeServer.Models.IPool>;
+using StreamId = HordeServer.Utilities.StringId<HordeServer.Models.IStream>;
+
+namespace HordeServer.Collections
+{
+	/// <summary>
+	/// Collection of utilization collection
+	/// </summary>
+	public interface ITelemetryCollection
+	{
+		/// <summary>
+		/// Adds entries for the given utilization
+		/// </summary>
+		/// <param name="Telemetry">Telemetry data to add</param>
+		/// <returns>Async task</returns>
+		Task AddUtilizationTelemetryAsync(IUtilizationTelemetry Telemetry);
+
+		/// <summary>
+		/// Finds utilization data matching the given criteria
+		/// </summary>
+		/// <param name="StartTimeUtc">Start time to query utilization for</param>
+		/// <param name="FinishTimeUtc">Finish time to query utilization for</param>
+		/// <returns>The utilization data</returns>
+		Task<List<IUtilizationTelemetry>> GetUtilizationTelemetryAsync(DateTime StartTimeUtc, DateTime FinishTimeUtc);
+
+		/// <summary>
+		/// Finds the latest utilization data
+		/// </summary>
+		/// <returns>The utilization data</returns>
+		Task<IUtilizationTelemetry?> GetLatestUtilizationTelemetryAsync();
+	}
+}
diff --git a/Engine/Source/Programs/Horde/HordeServer/Collections/ITestDataCollection.cs b/Engine/Source/Programs/Horde/HordeServer/Collections/ITestDataCollection.cs
new file mode 100644
--- /dev/null
+++ b/Engine/Source/Programs/Horde/HordeServer/Collections/ITestDataCollection.cs
@@ -0,0 +1,61 @@
+using HordeServer.Models;
+using HordeServer.Services;
+using MongoDB.Bson;
+using MongoDB.Bson.Serialization.Attributes;
+using MongoDB.Driver;
+using System;
+using System.Collections.Generic;
+using System.Diagnostics.CodeAnalysis;
+using System.Linq;
+using System.Threading.Tasks;
+
+using ProjectId = HordeServer.Utilities.StringId<HordeServer.Models.IProject>;
+using StreamId = HordeServer.Utilities.StringId<HordeServer.Models.IStream>;
+using TemplateRefId = HordeServer.Utilities.StringId<HordeServer.Models.TemplateRef>;
+
+namespace HordeServer.Collections
+{
+	/// <summary>
+	/// Collection of test data documents
+	/// </summary>
+	public interface ITestDataCollection
+	{
+		/// <summary>
+		/// Creates a new test data document
+		/// </summary>
+		/// <param name="Job">The job containing the step</param>
+		/// <param name="Step">The step producing the data</param>
+		/// <param name="Key">Key identifying the test</param>
+		/// <param name="Value">The data to store</param>
+		/// <returns>The new stream document</returns>
+		Task<ITestData> AddAsync(IJob Job, IJobStep Step, string Key, BsonDocument Value);
+
+		/// <summary>
+		/// Gets a stream by ID
+		/// </summary>
+		/// <param name="Id">Unique id of the stream</param>
+		/// <returns>The stream document</returns>
+		Task<ITestData?> GetAsync(ObjectId Id);
+
+		/// <summary>
+		/// Searches for test data that matches a set of criteria
+		/// </summary>
+		/// <param name="StreamId">The stream id</param>
+		/// <param name="MinChange">The minimum changelist number to return (inclusive)</param>
+		/// <param name="MaxChange">The maximum changelist number to return (inclusive)</param>
+		/// <param name="JobId">The job id</param>
+		/// <param name="StepId">The unique step id</param>
+		/// <param name="Key">Key identifying the result to return</param>
+		/// <param name="Index">Offset within the results to return</param>
+		/// <param name="Count">Number of results to return</param>
+		/// <returns>The stream document</returns>
+		Task<List<ITestData>> FindAsync(StreamId? StreamId, int? MinChange, int? MaxChange, ObjectId? JobId, SubResourceId? StepId, string? Key = null, int Index = 0, int Count = 10);
+
+		/// <summary>
+		/// Delete the test data
+		/// </summary>
+		/// <param name="Id">Unique id of the test data</param>
+		/// <returns>Async task</returns>
+		Task DeleteAsync(ObjectId Id);
+	}
+}
diff --git a/Engine/Source/Programs/Horde/HordeServer/Collections/IUserCollection.cs b/Engine/Source/Programs/Horde/HordeServer/Collections/IUserCollection.cs
new file mode 100644
--- /dev/null
+++ b/Engine/Source/Programs/Horde/HordeServer/Collections/IUserCollection.cs
@@ -0,0 +1,109 @@
+using HordeServer.Models;
+using HordeServer.Services;
+using HordeServer.Utilities;
+using MongoDB.Bson;
+using System;
+using System.Collections.Generic;
+using System.Linq;
+using System.Security.Claims;
+using System.Threading.Tasks;
+
+namespace HordeServer.Collections
+{
+	/// <summary>
+	/// Manages user documents
+	/// </summary>
+	public interface IUserCollection
+	{
+		/// <summary>
+		/// Gets a user by unique id
+		/// </summary>
+		/// <param name="Id">Id of the user</param>
+		/// <returns>The user information</returns>
+		Task<IUser?> GetUserAsync(ObjectId Id);
+
+		/// <summary>
+		/// Gets a user by unique id
+		/// </summary>
+		/// <param name="Ids">Ids of the users</param>
+		/// <param name="Index">Maximum number of results</param>
+		/// <param name="Count">Number of results to return</param>
+		/// <returns>The user information</returns>
+		Task<List<IUser>> FindUsersAsync(IEnumerable<ObjectId>? Ids = null, int? Index = null, int? Count = null);
+
+		/// <summary>
+		/// Gets a user by login
+		/// </summary>
+		/// <param name="Login">Login for the user</param>
+		/// <returns>The user information</returns>
+		Task<IUser?> FindUserByLoginAsync(string Login);
+
+		/// <summary>
+		/// Find or add a user with the given claims. Claims will be updated if the user exists.
+		/// </summary>
+		/// <param name="Login">Login id of the user</param>
+		/// <param name="Name">Full name of the user</param>
+		/// <param name="Email">Email address of the user</param>
+		/// <returns>The user document</returns>
+		Task<IUser> FindOrAddUserByLoginAsync(string Login, string? Name = null, string? Email = null);
+
+		/// <summary>
+		/// Gets the claims for a user
+		/// </summary>
+		/// <param name="UserId"></param>
+		/// <returns></returns>
+		Task<IUserClaims> GetClaimsAsync(ObjectId UserId);
+
+		/// <summary>
+		/// Update the claims for a user
+		/// </summary>
+		/// <param name="UserId"></param>
+		/// <param name="Claims"></param>
+		/// <returns></returns>
+		Task UpdateClaimsAsync(ObjectId UserId, IEnumerable<IUserClaim> Claims);
+
+		/// <summary>
+		/// Get settings for a user
+		/// </summary>
+		/// <param name="UserId"></param>
+		/// <returns></returns>
+		Task<IUserSettings> GetSettingsAsync(ObjectId UserId);
+
+		/// <summary>
+		/// Update a user
+		/// </summary>
+		/// <param name="UserId">The user to update</param>
+		/// <param name="EnableExperimentalFeatures"></param>
+		/// <param name="EnableIssueNotifications"></param>
+		/// <param name="DashboardSettings">Opaque settings object for the dashboard</param>
+		/// <param name="AddPinnedJobIds"></param>
+		/// <param name="RemovePinnedJobIds"></param>
+		/// <returns>Updated user object</returns>
+		Task UpdateSettingsAsync(ObjectId UserId, bool? EnableExperimentalFeatures = null, bool? EnableIssueNotifications = null, BsonValue? DashboardSettings = null, IEnumerable<ObjectId>? AddPinnedJobIds = null, IEnumerable<ObjectId>? RemovePinnedJobIds = null);
+	}
+
+	/// <summary>
+	/// Extension methods for the user collect
+	/// </summary>
+	static class UserCollectionExtensions
+	{
+		/// <summary>
+		/// Gets a particular user info from the collection
+		/// </summary>
+		/// <param name="UserCollection"></param>
+		/// <param name="Principal"></param>
+		/// <returns></returns>
+		public static Task<IUser?> GetUserAsync(this IUserCollection UserCollection, ClaimsPrincipal Principal)
+		{
+			ObjectId? UserId = Principal.GetUserId();
+			if (UserId != null)
+			{
+				return UserCollection.GetUserAsync(UserId.Value);
+			}
+			else
+			{
+				return Task.FromResult<IUser?>(null);
+			}
+		}
+	}
+}
diff --git a/Engine/Source/Programs/Horde/HordeServer/Collections/Impl/AgentCollection.cs b/Engine/Source/Programs/Horde/HordeServer/Collections/Impl/AgentCollection.cs
new file mode 100644
--- /dev/null
+++ b/Engine/Source/Programs/Horde/HordeServer/Collections/Impl/AgentCollection.cs
@@ -0,0 +1,517 @@
+using Google.Protobuf.WellKnownTypes;
+using HordeServer.Api;
+using HordeCommon;
+using HordeServer.Models;
+using HordeCommon.Rpc.Tasks;
+using HordeServer.Services;
+using HordeServer.Utilities;
+using MongoDB.Bson;
+using MongoDB.Bson.Serialization.Attributes;
+using MongoDB.Driver;
+using System;
+using System.Collections.Generic;
+using System.Linq;
+using System.Linq.Expressions;
+using System.Threading.Tasks;
+
+using PoolId = HordeServer.Utilities.StringId<HordeServer.Models.IPool>;
+using AgentSoftwareVersion = HordeServer.Utilities.StringId<HordeServer.Collections.IAgentSoftwareCollection>;
+using AgentSoftwareChannelName = HordeServer.Utilities.StringId<HordeServer.Services.AgentSoftwareChannels>;
+
+namespace HordeServer.Collections.Impl
+{
+	/// <summary>
+	/// Collection of agent documents
+	/// </summary>
+	public class AgentCollection : IAgentCollection
+	{
+		/// <summary>
+		/// Concrete implementation of an agent document
+		/// </summary>
+		class AgentDocument : IAgent
+		{
+			static IReadOnlyList<AgentLease> EmptyLeases = new List<AgentLease>();
+
+			[BsonRequired, BsonId]
+			public AgentId Id { get; set; }
+
+			public ObjectId? SessionId { get; set; }
+			public DateTime? SessionExpiresAt { get; set; }
+
+			public AgentStatus Status { get; set; }
+
+			[BsonRequired]
+			public bool Enabled { get; set; }
+
+			public bool Ephemeral { get; set; }
+
+			[BsonIgnoreIfDefault, BsonDefaultValue(false)]
+			public bool Deleted { get; set; }
+
+			[BsonElement("Version2")]
+			public AgentSoftwareVersion? Version { get; set; }
+
+			[BsonIgnoreIfNull]
+			public AgentSoftwareChannelName? Channel { get; set; }
+
+			[BsonIgnoreIfNull]
+			public string? LastUpgradeVersion { get; set; }
+
+			[BsonIgnoreIfNull]
+			public DateTime? LastUpgradeTime { get; set; }
+
+			public List<PoolId> Pools { get; set; } = new List<PoolId>();
+			public bool RequestConform { get; set; }
+
+			[BsonIgnoreIfNull]
+			public bool RequestRestart { get; set; }
+
+			[BsonIgnoreIfNull]
+			public bool RequestShutdown { get; set; }
+
+			public List<AgentWorkspace> Workspaces { get; set; } = new List<AgentWorkspace>();
+			public DateTime LastConformTime { get; set; }
+
+			[BsonIgnoreIfNull]
+			public int? ConformAttemptCount { get; set; }
+
+			public AgentCapabilities Capabilities { get; set; } = new AgentCapabilities();
+			public List<AgentLease>? Leases { get; set; }
+			public Acl? Acl { get; set; }
+			public DateTime UpdateTime { get; set; }
+			public uint UpdateIndex { get; set; }
+			public string? Comment { get; set; }
+
+			IReadOnlyList<PoolId> IAgent.ExplicitPools => Pools;
+			IReadOnlyList<AgentWorkspace> IAgent.Workspaces => Workspaces;
+			IReadOnlyList<AgentLease> IAgent.Leases => Leases ?? EmptyLeases;
+
+			[BsonConstructor]
+			private AgentDocument()
+			{
+			}
+
+			public AgentDocument(AgentId Id, bool bEnabled, bool bEphemeral, AgentSoftwareChannelName? Channel, List<PoolId> Pools)
+			{
+				this.Id = Id;
+				this.Acl = new Acl();
+				this.Enabled = bEnabled;
+				this.Ephemeral = bEphemeral;
+				this.Deleted = bEphemeral;
+				this.Channel = Channel;
+				this.Pools = Pools;
+			}
+		}
+
+		/// <summary>
+		/// Collection of agent documents
+		/// </summary>
+		readonly IMongoCollection<AgentDocument> Agents;
+
+		/// <summary>
+		/// Constructor
+		/// </summary>
+		/// <param name="DatabaseService">The database service instance</param>
+		public AgentCollection(DatabaseService DatabaseService)
+		{
+			Agents = DatabaseService.GetCollection<AgentDocument>("Agents");
+			if (!DatabaseService.ReadOnlyMode)
+			{
+				Agents.Indexes.CreateOne(new CreateIndexModel<AgentDocument>(Builders<AgentDocument>.IndexKeys.Ascending(x => x.Deleted).Ascending(x => x.Id).Ascending(x => x.Pools)));
+			}
+		}
+
+		/// <inheritdoc/>
+		public async Task<IAgent> AddAsync(AgentId Id, bool bEnabled, bool bEphemeral, AgentSoftwareChannelName? Channel, List<PoolId>? Pools)
+		{
+			AgentDocument Agent = new AgentDocument(Id, bEnabled, bEphemeral, Channel, Pools ?? new List<PoolId>());
+			await Agents.InsertOneAsync(Agent);
+			return Agent;
+		}
+
+		/// <inheritdoc/>
+		public async Task<IAgent?> TryDeleteAsync(IAgent AgentInterface)
+		{
+			AgentDocument Agent = (AgentDocument)AgentInterface;
+
+			UpdateDefinition<AgentDocument> Update = Builders<AgentDocument>.Update.Set(x => x.Deleted, true);
+			if (await TryUpdateAsync(Agent, Update))
+			{
+				Agent.Deleted = true;
+				return Agent;
+			}
+			return null;
+		}
+
+		/// <inheritdoc/>
+		public async Task ForceDeleteAsync(AgentId AgentId)
+		{
+			await Agents.DeleteOneAsync(x => x.Id == AgentId);
+		}
+
+		/// <inheritdoc/>
+		public async Task<IAgent?> GetAsync(AgentId AgentId)
+		{
+			return await Agents.Find<AgentDocument>(x => x.Id == AgentId).FirstOrDefaultAsync();
+		}
+
+		/// <inheritdoc/>
+		public async Task<List<IAgent>> FindAsync(ObjectId? Pool, string? PoolId, DateTime? ModifiedAfter, AgentStatus? Status, int? Index, int? Count)
+		{
+			FilterDefinitionBuilder<AgentDocument> FilterBuilder = new FilterDefinitionBuilder<AgentDocument>();
+
+			FilterDefinition<AgentDocument> Filter = FilterBuilder.Ne(x => x.Deleted, true);
+			if (Pool != null)
+			{
+				Filter &= FilterBuilder.Eq(nameof(AgentDocument.Pools), Pool);
+			}
+			
+			if (PoolId != null)
+			{
+				Filter &= FilterBuilder.Eq(nameof(AgentDocument.Pools), PoolId);
+			}
+			
+			if (ModifiedAfter != null)
+			{
+				Filter &= FilterBuilder.Gt(x => x.UpdateTime, ModifiedAfter.Value);
+			}
+			
+			if (Status != null)
+			{
+				Filter &= FilterBuilder.Eq(x => x.Status, Status.Value);
+			}
+
+			IFindFluent<AgentDocument, AgentDocument> Search = Agents.Find(Filter);
+			if (Index != null)
+			{
+				Search = Search.Skip(Index.Value);
+			}
+			if (Count != null)
+			{
+				Search = Search.Limit(Count.Value);
+			}
+
+			List<AgentDocument> Results = await Search.ToListAsync();
+			return Results.ConvertAll<IAgent>(x => x);
+		}
+
+		/// <inheritdoc/>
+		public async Task<List<IAgent>> FindExpiredAsync(DateTime UtcNow, int MaxAgents)
+		{
+			List<AgentDocument> Results = await Agents.Find(x => x.SessionId.HasValue && !(x.SessionExpiresAt > UtcNow)).Limit(MaxAgents).ToListAsync();
+			return Results.ConvertAll<IAgent>(x => x);
+		}
+
+		/// <summary>
+		/// Update a single document
+		/// </summary>
+		/// <param name="Current">The document to update</param>
+		/// <param name="Update">The update definition</param>
+		/// <returns>True if the agent was updated</returns>
+		private async Task<bool> TryUpdateAsync(AgentDocument Current, UpdateDefinition<AgentDocument> Update)
+		{
+			uint PrevUpdateIndex = Current.UpdateIndex++;
+			Current.UpdateTime = DateTime.UtcNow;
+
+			Expression<Func<AgentDocument, bool>> Filter = x => x.Id == Current.Id && x.UpdateIndex == PrevUpdateIndex;
+			UpdateDefinition<AgentDocument> UpdateWithIndex = Update.Set(x => x.UpdateIndex, Current.UpdateIndex).Set(x => x.UpdateTime, Current.UpdateTime);
+
+			UpdateResult Result = await Agents.UpdateOneAsync<AgentDocument>(Filter, UpdateWithIndex);
+			return Result.ModifiedCount == 1;
+		}
+
+		/// <inheritdoc/>
+		public async Task<IAgent?> TryUpdateSettingsAsync(IAgent AgentInterface, bool? Enabled = null, bool? RequestConform = null, bool? RequestRestart = null, bool? RequestShutdown = null, AgentSoftwareChannelName? Channel = null, List<PoolId>? Pools = null, Acl? Acl = null, string? Comment = null)
+		{
+			AgentDocument Agent = (AgentDocument)AgentInterface;
+
+			// Update the database
+			UpdateDefinitionBuilder<AgentDocument> UpdateBuilder = new UpdateDefinitionBuilder<AgentDocument>();
+
+			List<UpdateDefinition<AgentDocument>> Updates = new List<UpdateDefinition<AgentDocument>>();
+			if (Pools != null)
+			{
+				Updates.Add(UpdateBuilder.Set(x => x.Pools, Pools));
+			}
+			if (Enabled != null)
+			{
+				Updates.Add(UpdateBuilder.Set(x => x.Enabled, Enabled.Value));
+			}
+			if (RequestConform != null)
+			{
+				Updates.Add(UpdateBuilder.Set(x => x.RequestConform, RequestConform.Value));
+				Updates.Add(UpdateBuilder.Unset(x => x.ConformAttemptCount));
+			}
+			if (RequestRestart != null)
+			{
+				if (RequestRestart.Value)
+				{
+					Updates.Add(UpdateBuilder.Set(x => x.RequestRestart, true));
+				}
+				else
+				{
+					Updates.Add(UpdateBuilder.Unset(x => x.RequestRestart));
+				}
+			}
+			if (RequestShutdown != null)
+			{
+				if (RequestShutdown.Value)
+				{
+					Updates.Add(UpdateBuilder.Set(x => x.RequestShutdown, true));
+				}
+				else
+				{
+					Updates.Add(UpdateBuilder.Unset(x => x.RequestShutdown));
+				}
+			}
+			if (Channel != null)
+			{
+				if (Channel.Value == AgentSoftwareService.DefaultChannelName)
+				{
+					Updates.Add(UpdateBuilder.Unset(x => x.Channel));
+				}
+				else
+				{
+					Updates.Add(UpdateBuilder.Set(x => x.Channel, Channel));
+				}
+			}
+			if (Acl != null)
+			{
+				Updates.Add(Acl.CreateUpdate<AgentDocument>(x => x.Acl!, Acl));
+			}
+			if (Comment != null)
+			{
+				Updates.Add(UpdateBuilder.Set(x => x.Comment, Comment));
+			}
+
+			// Apply the update
+			if (Updates.Count != 0 && !await TryUpdateAsync(Agent, UpdateBuilder.Combine(Updates)))
+			{
+				return null;
+			}
+
+			// Update the document
+			if (Pools != null)
+			{
+				Agent.Pools = Pools;
+			}
+			if (Enabled != null)
+			{
+				Agent.Enabled = Enabled.Value;
+			}
+			if (RequestConform != null)
+			{
+				Agent.RequestConform = RequestConform.Value;
+			}
+			if (RequestRestart != null)
+			{
+				Agent.RequestRestart = RequestRestart.Value;
+			}
+			if (RequestShutdown != null)
+			{
+				Agent.RequestShutdown = RequestShutdown.Value;
+			}
+			if (Channel != null)
+			{
+				Agent.Channel = (Channel.Value == AgentSoftwareService.DefaultChannelName) ? null : Channel;
+			}
+			if (Acl != null)
+			{
+				Agent.Acl = Acl;
+			}
+			return Agent;
+		}
+
+		/// <inheritdoc/>
+		public async Task<IAgent?> TryUpdateSessionAsync(IAgent AgentInterface, AgentStatus? Status, DateTime? SessionExpiresAt, AgentCapabilities? Capabilities, List<AgentLease>? Leases)
+		{
+			AgentDocument Agent = (AgentDocument)AgentInterface;
+
+			// Create an update definition for the agent
+			UpdateDefinitionBuilder<AgentDocument> UpdateBuilder = Builders<AgentDocument>.Update;
+			List<UpdateDefinition<AgentDocument>> Updates = new List<UpdateDefinition<AgentDocument>>();
+
+			if (Status != null && Agent.Status != Status.Value)
+			{
+				Agent.Status = Status.Value;
+				Updates.Add(UpdateBuilder.Set(x => x.Status, Agent.Status));
+			}
+			if (SessionExpiresAt != null)
+			{
+				Agent.SessionExpiresAt = SessionExpiresAt.Value;
+				Updates.Add(UpdateBuilder.Set(x => x.SessionExpiresAt, Agent.SessionExpiresAt));
+			}
+			if (Capabilities != null)
+			{
+				Agent.Capabilities = Capabilities;
+				Updates.Add(UpdateBuilder.Set(x => x.Capabilities, Agent.Capabilities));
+			}
+			if (Leases != null)
+			{
+				foreach (AgentLease Lease in Leases)
+				{
+					if (Lease.Payload != null && (Agent.Leases == null || !Agent.Leases.Any(x => x.Id == Lease.Id)))
+					{
+						Any Payload = Any.Parser.ParseFrom(Lease.Payload.ToArray());
+						if (Payload.TryUnpack(out ConformTask ConformTask))
+						{
+							Agent.ConformAttemptCount = (Agent.ConformAttemptCount ?? 0) + 1;
+							Updates.Add(UpdateBuilder.Set(x => x.ConformAttemptCount, Agent.ConformAttemptCount));
+
+							Agent.LastConformTime = DateTime.UtcNow;
+							Updates.Add(UpdateBuilder.Set(x => x.LastConformTime, Agent.LastConformTime));
+						}
+						else if (Payload.TryUnpack(out UpgradeTask UpgradeTask))
+						{
+							Agent.LastUpgradeVersion = UpgradeTask.SoftwareId;
+							Updates.Add(UpdateBuilder.Set(x => x.LastUpgradeVersion, Agent.LastUpgradeVersion));
+
+							Agent.LastUpgradeTime = DateTime.UtcNow;
+							Updates.Add(UpdateBuilder.Set(x => x.LastUpgradeTime, Agent.LastUpgradeTime));
+						}
+					}
+				}
+
+				Agent.Leases = Leases;
+				Updates.Add(UpdateBuilder.Set(x => x.Leases, Agent.Leases));
+			}
+
+			// Update the agent, and try to create new lease documents if we succeed
+			if (Updates.Count > 0 && !await TryUpdateAsync(Agent, UpdateBuilder.Combine(Updates)))
+			{
+				return null;
+			}
+
+			return Agent;
+		}
+
+		/// <inheritdoc/>
+		public async Task<IAgent?> TryUpdateWorkspacesAsync(IAgent AgentInterface, List<AgentWorkspace> Workspaces, bool RequestConform)
+		{
+			AgentDocument Agent = (AgentDocument)AgentInterface;
+			DateTime LastConformTime = DateTime.UtcNow;
+
+			// Set the new workspaces
+			UpdateDefinition<AgentDocument> Update = Builders<AgentDocument>.Update.Set(x => x.Workspaces, Workspaces);
+			Update = Update.Set(x => x.LastConformTime, LastConformTime);
+			Update = Update.Unset(x => x.ConformAttemptCount);
+			Update = Update.Set(x => x.RequestConform, RequestConform);
+
+			// Update the agent
+			if (!await TryUpdateAsync(Agent, Update))
+			{
+				return null;
+			}
+
+			// Update the document
+			Agent.Workspaces = new List<AgentWorkspace>(Workspaces);
+			Agent.LastConformTime = LastConformTime;
+			Agent.ConformAttemptCount = null;
+			Agent.RequestConform = RequestConform;
+
+			return Agent;
+		}
+
+		/// <inheritdoc/>
+		public async Task<IAgent?> TryStartSessionAsync(IAgent AgentInterface, ObjectId SessionId, DateTime SessionExpiresAt, AgentStatus Status, AgentCapabilities Capabilities, AgentSoftwareVersion? Version)
+		{
+			AgentDocument Agent = (AgentDocument)AgentInterface;
+
+			// Reset the agent to use the new session
+			UpdateDefinitionBuilder<AgentDocument> UpdateBuilder = Builders<AgentDocument>.Update;
+
+			List<UpdateDefinition<AgentDocument>> Updates = new List<UpdateDefinition<AgentDocument>>();
+			Updates.Add(UpdateBuilder.Set(x => x.SessionId, SessionId));
+			Updates.Add(UpdateBuilder.Set(x => x.SessionExpiresAt, SessionExpiresAt));
+			Updates.Add(UpdateBuilder.Set(x => x.Status, Status));
+			Updates.Add(UpdateBuilder.Unset(x => x.Leases));
+			Updates.Add(UpdateBuilder.Unset(x => x.Deleted));
+			Updates.Add(UpdateBuilder.Set(x => x.Capabilities, Capabilities));
+			Updates.Add(UpdateBuilder.Set(x => x.Version, Version));
+			Updates.Add(UpdateBuilder.Unset(x => x.RequestRestart));
+			Updates.Add(UpdateBuilder.Unset(x => x.RequestShutdown));
+
+			// Apply the update
+			if (!await TryUpdateAsync(Agent, UpdateBuilder.Combine(Updates)))
+			{
+				return null;
+			}
+
+			// Update the document
+			Agent.SessionId = SessionId;
+			Agent.SessionExpiresAt = SessionExpiresAt;
+			Agent.Status = Status;
+			Agent.Leases = null;
+			Agent.Deleted = false;
+			Agent.Capabilities = Capabilities;
+			Agent.Version = Version;
+
+			return Agent;
+		}
+
+		/// <inheritdoc/>
+		public async Task<IAgent?> TryTerminateSessionAsync(IAgent AgentInterface)
+		{
+			AgentDocument Agent = (AgentDocument)AgentInterface;
+			UpdateDefinition<AgentDocument> Update = new BsonDocument();
+
+			Update = Update.Unset(x => x.SessionId);
+			Update = Update.Unset(x => x.SessionExpiresAt);
+			Update = Update.Unset(x => x.Leases);
+
+			bool bDeleted = Agent.Deleted || Agent.Ephemeral;
+			if (bDeleted != Agent.Deleted)
+			{
+				Update = Update.Set(x => x.Deleted, Agent.Deleted);
+			}
+
+			if (!await TryUpdateAsync(Agent, Update))
+			{
+				return null;
+			}
+
+			Agent.SessionId = null;
+			Agent.SessionExpiresAt = null;
+			Agent.Deleted = bDeleted;
+			Agent.Leases = new List<AgentLease>();
+			return Agent;
+		}
+
+		/// <inheritdoc/>
+		public async Task<IAgent?> TryAddLeaseAsync(IAgent AgentInterface, AgentLease NewLease)
+		{
+			AgentDocument Agent = (AgentDocument)AgentInterface;
+
+			List<AgentLease> Leases = new List<AgentLease>();
+			if (Agent.Leases != null)
+			{
+				Leases.AddRange(Agent.Leases);
+			}
+			Leases.Add(NewLease);
+
+			UpdateDefinition<AgentDocument> Update = Builders<AgentDocument>.Update.Set(x => x.Leases, Leases);
+			if (!await TryUpdateAsync(Agent, Update))
+			{
+				return null;
+			}
+
+			Agent.Leases = Leases;
+			return Agent;
+		}
+
+		/// <inheritdoc/>
+		public async Task<IAgent?> TryCancelLeaseAsync(IAgent AgentInterface, int LeaseIdx)
+		{
+			AgentDocument Agent = (AgentDocument)AgentInterface;
+
+			UpdateDefinition<AgentDocument> Update = Builders<AgentDocument>.Update.Set(x => x.Leases![LeaseIdx].State, LeaseState.Cancelled);
+			if (!await TryUpdateAsync(Agent, Update))
+			{
+				return null;
+			}
+
+			Agent.Leases![LeaseIdx].State = LeaseState.Cancelled;
+			return Agent;
+		}
+	}
+}
diff --git a/Engine/Source/Programs/Horde/HordeServer/Collections/Impl/AgentSoftwareCollection.cs b/Engine/Source/Programs/Horde/HordeServer/Collections/Impl/AgentSoftwareCollection.cs
new file mode 100644
--- /dev/null
+++ b/Engine/Source/Programs/Horde/HordeServer/Collections/Impl/AgentSoftwareCollection.cs
@@ -0,0 +1,138 @@
+using EpicGames.Core;
+using HordeCommon;
+using HordeServer.Models;
+using HordeServer.Services;
+using MongoDB.Bson;
+using MongoDB.Bson.Serialization.Attributes;
+using MongoDB.Driver;
+using System;
+using System.Collections.Generic;
+using System.Globalization;
+using System.IO;
+using System.IO.Compression;
+using System.Linq;
+using System.Text.Json;
+using System.Threading.Tasks;
+
+using AgentSoftwareVersion = HordeServer.Utilities.StringId<HordeServer.Collections.IAgentSoftwareCollection>;
+
+namespace HordeServer.Collections.Impl
+{
+	/// <summary>
+	/// Collection of agent software documents
+	/// </summary>
+	public class AgentSoftwareCollection : IAgentSoftwareCollection
+	{
+		/// <summary>
+		/// Information about agent software
+		/// </summary>
+		class AgentSoftwareDocument
+		{
+			[BsonId]
+			public AgentSoftwareVersion Id { get; set; }
+			public byte[] Data { get; set; } = null!;
+
+			[BsonConstructor]
+			private AgentSoftwareDocument()
+			{
+				this.Data = null!;
+			}
+
+			public AgentSoftwareDocument(AgentSoftwareVersion Id, byte[] Data)
+			{
+				this.Id = Id;
+				this.Data = Data;
+			}
+		}
+
+		/// <summary>
+		/// Template documents
+		/// </summary>
+		IMongoCollection<AgentSoftwareDocument> Collection;
+
+		/// <summary>
+		/// Constructor
+		/// </summary>
+		/// <param name="DatabaseService">The database service singleton</param>
+		public AgentSoftwareCollection(DatabaseService DatabaseService)
+		{
+			Collection = DatabaseService.GetCollection<AgentSoftwareDocument>("AgentSoftware");
+		}
+
+		/// <inheritdoc/>
+		public async Task AddAsync(AgentSoftwareVersion Version, byte[] Data)
+		{
+			byte[] NewData = UpdateVersion(Data, Version);
+
+			AgentSoftwareDocument Document = new AgentSoftwareDocument(Version, NewData);
+			await Collection.InsertOneAsync(Document);
+		}
+
+		/// <summary>
+		/// Updates the client id within the archive in the data array
+		/// </summary>
+		/// <param name="Data">Data for the zip archive</param>
+		/// <param name="Version">The new version identifier</param>
+		/// <returns>New agent app data</returns>
+		public static byte[] UpdateVersion(byte[] Data, AgentSoftwareVersion Version)
+		{
+			bool bWrittenClientId = false;
+
+			MemoryStream OutputStream = new MemoryStream();
+			using (ZipArchive OutputArchive = new ZipArchive(OutputStream, ZipArchiveMode.Create, true))
+			{
+				MemoryStream InputStream = new MemoryStream(Data);
+				using (ZipArchive InputArchive = new ZipArchive(InputStream, ZipArchiveMode.Read, true))
+				{
+					foreach (ZipArchiveEntry InputEntry in InputArchive.Entries)
+					{
+						ZipArchiveEntry OutputEntry = OutputArchive.CreateEntry(InputEntry.FullName);
+
+						using System.IO.Stream InputEntryStream = InputEntry.Open();
+						using System.IO.Stream OutputEntryStream = OutputEntry.Open();
+
+						if (InputEntry.FullName.Equals("appsettings.json", StringComparison.OrdinalIgnoreCase))
+						{
+							using MemoryStream MemoryStream = new MemoryStream();
+							InputEntryStream.CopyTo(MemoryStream);
+
+							Dictionary<string, Dictionary<string, object>> Document = JsonSerializer.Deserialize<Dictionary<string, Dictionary<string, object>>>(MemoryStream.ToArray());
+							Document["Horde"]["Version"] = Version.ToString();
+
+							using Utf8JsonWriter Writer = new Utf8JsonWriter(OutputEntryStream, new JsonWriterOptions { Indented = true });
+							JsonSerializer.Serialize<Dictionary<string, Dictionary<string, object>>>(Writer, Document, new JsonSerializerOptions { WriteIndented = true });
+
+							bWrittenClientId = true;
+						}
+						else
+						{
+							InputEntryStream.CopyTo(OutputEntryStream);
+						}
+					}
+				}
+			}
+
+			if (!bWrittenClientId)
+			{
+				throw new InvalidDataException("Missing appsettings.json file from zip archive");
+			}
+
+			return OutputStream.ToArray();
+		}
+
+		/// <inheritdoc/>
+		public async Task<bool> RemoveAsync(AgentSoftwareVersion Version)
+		{
+			DeleteResult Result = await Collection.DeleteOneAsync(x => x.Id == Version);
+			return Result.DeletedCount > 0;
+		}
+
+		/// <inheritdoc/>
+		public async Task<byte[]?> GetAsync(AgentSoftwareVersion Version)
+		{
+			IFindFluent<AgentSoftwareDocument, AgentSoftwareDocument> Query = Collection.Find(x => x.Id == Version);
+			AgentSoftwareDocument? Software = await Query.FirstOrDefaultAsync();
+			return Software?.Data;
+		}
+	}
+}
diff --git a/Engine/Source/Programs/Horde/HordeServer/Collections/Impl/GraphCollection.cs b/Engine/Source/Programs/Horde/HordeServer/Collections/Impl/GraphCollection.cs
new file mode 100644
--- /dev/null
+++ b/Engine/Source/Programs/Horde/HordeServer/Collections/Impl/GraphCollection.cs
@@ -0,0 +1,468 @@
+using EpicGames.Core;
+using HordeServer.Api;
+using HordeCommon;
+using HordeServer.Models;
+using HordeServer.Services;
+using MongoDB.Bson;
+using MongoDB.Bson.Serialization.Attributes;
+using MongoDB.Driver;
+using System;
+using System.Collections.Concurrent;
+using System.Collections.Generic;
+using System.Diagnostics;
+using System.Linq;
+using System.Threading;
+using System.Threading.Tasks;
+
+namespace HordeServer.Collections.Impl
+{
+	/// <summary>
+	/// Collection of graph documents
+	/// </summary>
+	public class GraphCollection : IGraphCollection
+	{
+		/// <summary>
+		/// Represents a node in the graph
+		/// </summary>
+		[DebuggerDisplay("{Name}")]
+		class Node : INode
+		{
+			[BsonRequired]
+			public string Name { get; set; }
+
+			[BsonIgnoreIfNull]
+			public NodeRef[] InputDependencies { get; set; }
+
+			[BsonIgnoreIfNull]
+			public NodeRef[] OrderDependencies { get; set; }
+
+			public Priority Priority { get; set; }
+			public bool AllowRetry { get; set; } = true;
+			public bool RunEarly { get; set; }
+			public bool Warnings { get; set; } = true;
+
+			[BsonIgnoreIfNull]
+			public Dictionary<string, string>? Credentials { get; set; }
+
+			[BsonIgnoreIfNull]
+			public Dictionary<string, string>? Properties { get; set; }
+
+			IReadOnlyDictionary<string, string>? INode.Credentials => Credentials;
+			IReadOnlyDictionary<string, string>? INode.Properties => Properties;
+
+			[BsonConstructor]
+			private Node()
+			{
+				Name = null!;
+				InputDependencies = null!;
+				OrderDependencies = null!;
+			}
+
+			public Node(string Name, NodeRef[] InputDependencies, NodeRef[] OrderDependencies, Priority Priority, bool AllowRetry, bool RunEarly, bool Warnings, Dictionary<string, string>? Credentials, Dictionary<string, string>? Properties)
+			{
+				this.Name = Name;
+				this.InputDependencies = InputDependencies;
+				this.OrderDependencies = OrderDependencies;
+				this.Priority = Priority;
+				this.AllowRetry = AllowRetry;
+				this.RunEarly = RunEarly;
+				this.Warnings = Warnings;
+				this.Credentials = Credentials;
+				this.Properties = Properties;
+			}
+		}
+
+		class NodeGroup : INodeGroup
+		{
+			public string AgentType { get; set; }
+			public List<Node> Nodes { get; set; }
+
+			IReadOnlyList<INode> INodeGroup.Nodes => Nodes;
+
+			/// <summary>
+			/// Private constructor for BSON serializer
+			/// </summary>
+			[BsonConstructor]
+			private NodeGroup()
+			{
+				AgentType = null!;
+				Nodes = null!;
+			}
+
+			/// <summary>
+			/// Constructor
+			/// </summary>
+			/// <param name="AgentType">The type of agent to execute this group</param>
+			/// <param name="Nodes">Nodes to execute</param>
+			public NodeGroup(string AgentType, List<Node> Nodes)
+			{
+				this.AgentType = AgentType;
+				this.Nodes = Nodes;
+			}
+		}
+
+		class Aggregate : IAggregate
+		{
+			public string Name { get; set; }
+			public List<NodeRef> Nodes { get; set; }
+
+			IReadOnlyList<NodeRef> IAggregate.Nodes => Nodes;
+
+			private Aggregate()
+			{
+				Name = null!;
+				Nodes = new List<NodeRef>();
+			}
+
+			public Aggregate(string Name, List<NodeRef> Nodes)
+			{
+				this.Name = Name;
+				this.Nodes = Nodes;
+			}
+		}
+
+		class Label : ILabel
+		{
+			[BsonIgnoreIfNull]
+			public string? Name { get; set; }
+
+			[BsonIgnoreIfNull]
+			public string? Category { get; set; }
+
+			[BsonIgnoreIfNull]
+			public string? DashboardName { get; set; }
+
+			[BsonIgnoreIfNull]
+			public string? DashboardCategory { get; set; }
+
+			[BsonIgnoreIfNull]
+			public string? UgsName { get; set; }
+
+			[BsonIgnoreIfNull]
+			public string? UgsProject { get; set; }
+
+			[BsonIgnoreIfNull]
+			public LabelChange Change { get; set; }
+
+			public List<NodeRef> RequiredNodes { get; set; }
+			public List<NodeRef> IncludedNodes { get; set; }
+
+			string? ILabel.DashboardName => DashboardName ?? Name;
+			string? ILabel.DashboardCategory => DashboardCategory ?? Category;
+
+			private Label()
+			{
+				RequiredNodes = new List<NodeRef>();
+				IncludedNodes = new List<NodeRef>();
+			}
+
+			public Label(string? DashboardName, string? DashboardCategory, string? UgsName, string? UgsProject, LabelChange Change, List<NodeRef> RequiredNodes, List<NodeRef> IncludedNodes)
+			{
+				this.DashboardName = DashboardName;
+				this.DashboardCategory = DashboardCategory;
+				this.UgsName = UgsName;
+				this.UgsProject = UgsProject;
+				this.Change = Change;
+				this.RequiredNodes = RequiredNodes;
+				this.IncludedNodes = IncludedNodes;
+			}
+		}
+
+		class GraphDocument : IGraph
+		{
+			public static GraphDocument Empty { get; } = new GraphDocument(new List<NodeGroup>(), new List<Aggregate>(), new List<Label>());
+
+			[BsonRequired, BsonId]
+			public ContentHash Id { get; private set; }
+
+			public int Schema { get; set; }
+			public List<NodeGroup> Groups { get; private set; } = new List<NodeGroup>();
+			public List<Aggregate> Aggregates { get; private set; } = new List<Aggregate>();
+			public List<Label> Labels { get; private set; } = new List<Label>();
+
+			[BsonIgnore]
+			IReadOnlyDictionary<string, NodeRef>? CachedNodeNameToRef;
+
+			IReadOnlyList<INodeGroup> IGraph.Groups => Groups;
+			IReadOnlyList<IAggregate> IGraph.Aggregates => Aggregates;
+			IReadOnlyList<ILabel> IGraph.Labels => Labels;
+
+			[BsonConstructor]
+			private GraphDocument()
+			{
+				Id = ContentHash.Empty;
+			}
+
+			public GraphDocument(List<NodeGroup> Groups, List<Aggregate> Aggregates, List<Label> Labels)
+			{
+				this.Groups = Groups;
+				this.Aggregates = Aggregates;
+				this.Labels = Labels;
+				this.Id = ContentHash.SHA1(BsonExtensionMethods.ToBson(this));
+			}
+
+			public GraphDocument(GraphDocument BaseGraph, List<CreateGroupRequest>? NewGroupRequests, List<CreateAggregateRequest>? NewAggregateRequests, List<CreateLabelRequest>? NewLabelRequests)
+			{
+				Dictionary<string, NodeRef> NodeNameToRef = new Dictionary<string, NodeRef>(BaseGraph.GetNodeNameToRef(), StringComparer.OrdinalIgnoreCase);
+
+				// Update the new list of groups
+				List<NodeGroup> NewGroups = new List<NodeGroup>(BaseGraph.Groups);
+				if (NewGroupRequests != null)
+				{
+					foreach (CreateGroupRequest NewGroupRequest in NewGroupRequests)
+					{
+						List<Node> Nodes = new List<Node>();
+						foreach (CreateNodeRequest NewNodeRequest in NewGroupRequest.Nodes)
+						{
+							int NodeIdx = Nodes.Count;
+
+							Priority Priority = NewNodeRequest.Priority ?? Priority.Normal;
+							bool bAllowRetry = NewNodeRequest.AllowRetry ?? true;
+							bool bRunEarly = NewNodeRequest.RunEarly ?? false;
+							bool bWarnings = NewNodeRequest.Warnings ?? true;
+
+							NodeRef[] InputDependencies = (NewNodeRequest.InputDependencies == null) ? Array.Empty<NodeRef>() : NewNodeRequest.InputDependencies.Select(x => NodeNameToRef[x]).ToArray();
+							NodeRef[] OrderDependencies = (NewNodeRequest.OrderDependencies == null) ? Array.Empty<NodeRef>() : NewNodeRequest.OrderDependencies.Select(x => NodeNameToRef[x]).ToArray();
+							OrderDependencies = OrderDependencies.Union(InputDependencies).ToArray();
+							Nodes.Add(new Node(NewNodeRequest.Name, InputDependencies, OrderDependencies, Priority, bAllowRetry, bRunEarly, bWarnings, NewNodeRequest.Credentials, NewNodeRequest.Properties));
+
+							NodeNameToRef.Add(NewNodeRequest.Name, new NodeRef(NewGroups.Count, NodeIdx));
+						}
+						NewGroups.Add(new NodeGroup(NewGroupRequest.AgentType, Nodes));
+					}
+				}
+
+				// Update the list of aggregates
+				List<Aggregate> NewAggregates = new List<Aggregate>(BaseGraph.Aggregates);
+				if (NewAggregateRequests != null)
+				{
+					foreach (CreateAggregateRequest NewAggregateRequest in NewAggregateRequests)
+					{
+						List<NodeRef> Nodes = NewAggregateRequest.Nodes.ConvertAll(x => NodeNameToRef[x]);
+						NewAggregates.Add(new Aggregate(NewAggregateRequest.Name, Nodes));
+					}
+				}
+
+				// Update the list of labels
+				List<Label> NewLabels = new List<Label>(BaseGraph.Labels);
+				if (NewLabelRequests != null)
+				{
+					foreach (CreateLabelRequest NewLabelRequest in NewLabelRequests)
+					{
+						List<NodeRef> RequiredNodes = NewLabelRequest.RequiredNodes.ConvertAll(x => NodeNameToRef[x]);
+						List<NodeRef> IncludedNodes = NewLabelRequest.IncludedNodes.ConvertAll(x => NodeNameToRef[x]);
+						NewLabels.Add(new Label(NewLabelRequest.DashboardName, NewLabelRequest.DashboardCategory, NewLabelRequest.UgsName, NewLabelRequest.UgsProject, NewLabelRequest.Change, RequiredNodes, IncludedNodes));
+					}
+				}
+
+				// Create the new arrays
+				Groups = NewGroups;
+				Aggregates = NewAggregates;
+				Labels = NewLabels;
+
+				// Create the new graph, and save the generated node lookup into it
+				CachedNodeNameToRef = NodeNameToRef;
+
+				// Compute the hash
+				Id = ContentHash.SHA1(BsonExtensionMethods.ToBson(this));
+			}
+
+			public IReadOnlyDictionary<string, NodeRef> GetNodeNameToRef()
+			{
+				if (CachedNodeNameToRef == null)
+				{
+					Dictionary<string, NodeRef> NodeNameToRef = new Dictionary<string, NodeRef>(StringComparer.OrdinalIgnoreCase);
+					for (int GroupIdx = 0; GroupIdx < Groups.Count; GroupIdx++)
+					{
+						List<Node> Nodes = Groups[GroupIdx].Nodes;
+						for (int NodeIdx = 0; NodeIdx < Nodes.Count; NodeIdx++)
+						{
+							Node Node = Nodes[NodeIdx];
+							NodeNameToRef[Node.Name] = new NodeRef(GroupIdx, NodeIdx);
+						}
+					}
+					CachedNodeNameToRef = NodeNameToRef;
+				}
+				return CachedNodeNameToRef;
+			}
+		}
+
+		/// <summary>
+		/// Stores information about a cached graph
+		/// </summary>
+		class CachedGraph
+		{
+			/// <summary>
+			/// Time at which the graph was last accessed
+			/// </summary>
+			public long LastAccessTime
+			{
+				get { return LastAccessTimePrivate; }
+			}
+
+			/// <summary>
+			/// Backing value for <see cref="LastAccessTime"/>
+			/// </summary>
+			private long LastAccessTimePrivate;
+
+			/// <summary>
+			/// The graph instance
+			/// </summary>
+			public IGraph Graph { get; }
+
+			/// <summary>
+			/// Constructor
+			/// </summary>
+			/// <param name="Graph">The graph to store</param>
+			public CachedGraph(IGraph Graph)
+			{
+				this.LastAccessTimePrivate = Stopwatch.GetTimestamp();
+				this.Graph = Graph;
+			}
+
+			/// <summary>
+			/// Update the last access time
+			/// </summary>
+			public void Touch()
+			{
+				for (; ; )
+				{
+					long Time = Stopwatch.GetTimestamp();
+					long LastAccessTimeCopy = LastAccessTimePrivate;
+					if (Time < LastAccessTimeCopy || Interlocked.CompareExchange(ref LastAccessTimePrivate, Time, LastAccessTimeCopy) == LastAccessTimeCopy)
+					{
+						break;
+					}
+				}
+			}
+		}
+
+		/// <summary>
+		/// The jobs collection
+		/// </summary>
+		IMongoCollection<GraphDocument> Graphs;
+
+		/// <summary>
+		/// Maximum number of graphs to keep in the cache
+		/// </summary>
+		const int MaxGraphs = 1000;
+
+		/// <summary>
+		/// Cache of recently accessed graphs
+		/// </summary>
+		ConcurrentDictionary<ContentHash, CachedGraph> CachedGraphs = new ConcurrentDictionary<ContentHash, CachedGraph>();
+
+		/// <summary>
+		/// Constructor
+		/// </summary>
+		/// <param name="DatabaseService">The database service singleton</param>
+		public GraphCollection(DatabaseService DatabaseService)
+		{
+			Graphs = DatabaseService.GetCollection<GraphDocument>("Graphs");
+		}
+
+		/// <summary>
+		/// Adds a new graph document
+		/// </summary>
+		/// <param name="Graph">The graph to add</param>
+		/// <returns>Async task</returns>
+		async Task AddAsync(GraphDocument Graph)
+		{
+			if (!await Graphs.Find(x => x.Id == Graph.Id).AnyAsync())
+			{
+				try
+				{
+					await Graphs.InsertOneAsync(Graph);
+				}
+				catch (MongoWriteException Ex)
+				{
+					if (Ex.WriteError.Category != ServerErrorCategory.DuplicateKey)
+					{
+						throw;
+					}
+				}
+			}
+		}
+
+		/// <inheritdoc/>
+		public async Task<IGraph> AddAsync(ITemplate Template)
+		{
+			Node Node = new Node(IJob.SetupNodeName, Array.Empty<NodeRef>(), Array.Empty<NodeRef>(), Priority.High, true, false, true, null, null);
+			NodeGroup Group = new NodeGroup(Template.InitialAgentType ?? "Win64", new List<Node> { Node });
+
+			GraphDocument Graph = new GraphDocument(new List<NodeGroup> { Group }, new List<Aggregate>(), new List<Label>());
+			await AddAsync(Graph);
+			return Graph;
+		}
+
+		/// <inheritdoc/>
+		public async Task<IGraph> AppendAsync(IGraph? BaseGraph, List<CreateGroupRequest>? NewGroupRequests, List<CreateAggregateRequest>? NewAggregateRequests, List<CreateLabelRequest>? NewLabelRequests)
+		{
+			GraphDocument Graph = new GraphDocument((GraphDocument?)BaseGraph ?? GraphDocument.Empty, NewGroupRequests, NewAggregateRequests, NewLabelRequests);
+			await AddAsync(Graph);
+			return Graph;
+		}
+
+		/// <inheritdoc/>
+		public async Task<IGraph> GetAsync(ContentHash Hash)
+		{
+			// Special case for an empty graph request
+			if (Hash == ContentHash.Empty || Hash == GraphDocument.Empty.Id)
+			{
+				return GraphDocument.Empty;
+			}
+
+			// Try to read the graph from the cache
+			CachedGraph? CachedGraph;
+			if (CachedGraphs.TryGetValue(Hash, out CachedGraph))
+			{
+				// Update the last access time
+				CachedGraph.Touch();
+			}
+			else
+			{
+				// Trim the cache
+				while (CachedGraphs.Count > MaxGraphs)
+				{
+					ContentHash RemoveHash = CachedGraphs.OrderBy(x => x.Value.LastAccessTime).Select(x => x.Key).FirstOrDefault();
+					if (RemoveHash == ContentHash.Empty)
+					{
+						break;
+					}
+					CachedGraphs.TryRemove(RemoveHash, out _);
+				}
+
+				// Create the new entry
+				CachedGraph = new CachedGraph(await Graphs.Find<GraphDocument>(x => x.Id == Hash).FirstAsync());
+				CachedGraphs.TryAdd(Hash, CachedGraph);
+			}
+			return CachedGraph.Graph;
+		}
+
+		/// <inheritdoc/>
+		public async Task<List<IGraph>> FindAllAsync(ContentHash[]? Hashes, int? Index, int? Count)
+		{
+			FilterDefinitionBuilder<GraphDocument> FilterBuilder = Builders<GraphDocument>.Filter;
+
+			FilterDefinition<GraphDocument> Filter = FilterBuilder.Empty;
+			if (Hashes != null)
+			{
+				Filter &= FilterBuilder.In(x => x.Id, Hashes);
+			}
+
+			List<GraphDocument> Results;
+			IFindFluent<GraphDocument, GraphDocument> Search = Graphs.Find(Filter);
+			if (Index != null)
+			{
+				Search = Search.Skip(Index.Value);
+			}
+			if (Count != null)
+			{
+				Search = Search.Limit(Count.Value);
+			}
+
+			Results = await Search.ToListAsync();
+			return Results.ConvertAll<IGraph>(x => x);
+		}
+	}
+}
diff --git a/Engine/Source/Programs/Horde/HordeServer/Collections/Impl/IssueCollection.cs b/Engine/Source/Programs/Horde/HordeServer/Collections/Impl/IssueCollection.cs
new file mode 100644
--- /dev/null
+++ b/Engine/Source/Programs/Horde/HordeServer/Collections/Impl/IssueCollection.cs
@@ -0,0 +1,1049 @@
+using EpicGames.Core;
+using HordeServer.Models;
+using HordeServer.Services;
+using HordeServer.Utilities;
+using Microsoft.Extensions.Caching.Memory;
+using Microsoft.Extensions.Logging;
+using MongoDB.Bson;
+using MongoDB.Bson.Serialization;
+using MongoDB.Bson.Serialization.Attributes;
+using MongoDB.Driver;
+using System;
+using System.Collections.Concurrent;
+using System.Collections.Generic;
+using System.Linq;
+using System.Linq.Expressions;
+using System.Security.Claims;
+using System.Threading.Tasks;
+
+using StreamId = HordeServer.Utilities.StringId<HordeServer.Models.IStream>;
+using TemplateRefId = HordeServer.Utilities.StringId<HordeServer.Models.TemplateRef>;
+
+namespace HordeServer.Collections.Impl
+{
+	class IssueCollection : IIssueCollection
+	{
+		class UpdateTransaction
+		{
+			class Statics<TDocument>
+			{
+				public static IBsonSerializer<TDocument> Serializer { get; } = BsonSerializer.LookupSerializer<TDocument>();
+			}
+
+			[BsonRequired]
+			public string Filter { get; set; }
+
+			[BsonRequired]
+			public string Update { get; set; }
+
+			private UpdateTransaction()
+			{
+				this.Filter = null!;
+				this.Update = null!;
+			}
+
+			public UpdateTransaction(BsonDocument Filter, BsonDocument Update)
+			{
+				this.Filter = Filter.ToJson();
+				this.Update = Update.ToJson();
+			}
+
+			public static UpdateTransaction Create<TDocument>(Expression<Func<TDocument, bool>> FilterFunc, UpdateDefinition<TDocument> Update)
+			{
+				return Create((FilterDefinition<TDocument>)FilterFunc, Update);
+			}
+
+			public static UpdateTransaction Create<TDocument>(FilterDefinition<TDocument> Filter, UpdateDefinition<TDocument> Update)
+			{
+				IBsonSerializer<TDocument> Serializer = Statics<TDocument>.Serializer;
+
+				BsonDocument FilterDocument = Filter.Render(Serializer, BsonSerializer.SerializerRegistry);
+				BsonDocument UpdateDocument = (BsonDocument)Update.Render(Serializer, BsonSerializer.SerializerRegistry);
+
+				return new UpdateTransaction(FilterDocument, UpdateDocument);
+			}
+		}
+
+		[SingletonDocument("5e4c226440ce25fa3207a9af")]
+		class IssueLedger : SingletonBase
+		{
+			public int NextId { get; set; }
+
+			[BsonIgnoreIfNull]
+			public Issue? InsertIssue { get; set; }
+
+			[BsonIgnoreIfNull]
+			public IssueSpan? InsertIssueSpan { get; set; }
+
+			[BsonIgnoreIfNull]
+			public UpdateTransaction? UpdateIssue { get; set; }
+
+			[BsonIgnoreIfNull]
+			public UpdateTransaction? UpdateIssueSpan { get; set; }
+
+			public bool HasPendingTransaction => InsertIssue != null || InsertIssueSpan != null || UpdateIssue != null || UpdateIssueSpan != null;
+		}
+
+		class IssueSequenceToken : IIssueSequenceToken
+		{
+			IssueCollection Collection;
+			public IssueLedger Ledger { get; private set; }
+
+			public IssueSequenceToken(IssueCollection Collection, IssueLedger Ledger)
+			{
+				this.Collection = Collection;
+				this.Ledger = Ledger;
+			}
+
+			public async Task ResetAsync()
+			{
+				Ledger = await Collection.GetLedgerAsync();
+			}
+
+			public async Task FlushAsync()
+			{
+				Ledger = await Collection.FlushLedgerAsync(Ledger);
+			}
+		}
+
+		class Issue : IIssue
+		{
+			[BsonId]
+			public int Id { get; set; }
+
+			public string Summary { get; set; }
+
+			[BsonIgnoreIfNull]
+			public string? UserSummary { get; set; }
+
+			[BsonRequired]
+			public IssueFingerprint Fingerprint { get; set; }
+			public IssueSeverity Severity { get; set; }
+
+			[BsonIgnoreIfNull]
+			public string? Owner { get; set; }
+
+			[BsonIgnoreIfNull]
+			public ObjectId? OwnerId { get; set; }
+
+			[BsonIgnoreIfNull]
+			public string? NominatedBy { get; set; }
+
+			[BsonIgnoreIfNull]
+			public ObjectId? NominatedById { get; set; }
+
+			public DateTime CreatedAt { get; set; }
+
+			[BsonIgnoreIfNull]
+			public DateTime? NominatedAt { get; set; }
+
+			[BsonIgnoreIfNull]
+			public DateTime? AcknowledgedAt { get; set; }
+
+			public bool Resolved { get; set; }
+
+			[BsonIgnoreIfNull]
+			public DateTime? ResolvedAt { get; set; }
+
+			[BsonIgnoreIfNull]
+			public int? FixChange { get; set; }
+
+			public int MinSuspectChange { get; set; }
+			public int MaxSuspectChange { get; set; }
+
+			public bool NotifySuspects { get; set; }
+
+			[BsonRequired]
+			public List<IssueSuspect> Suspects { get; set; }
+
+			public int UpdateIndex { get; set; }
+
+			IIssueFingerprint IIssue.Fingerprint => Fingerprint;
+			IReadOnlyList<IIssueSuspect> IIssue.Suspects => Suspects;
+			string? IIssue.Owner => Owner ?? GetDefaultOwner()?.Author;
+			ObjectId? IIssue.OwnerId => OwnerId ?? GetDefaultOwner()?.AuthorId;
+
+			[BsonConstructor]
+			private Issue()
+			{
+				Summary = String.Empty;
+				Fingerprint = null!;
+				Suspects = null!;
+			}
+
+			public Issue(int Id, string Summary, IIssueSpan Span)
+			{
+				this.Id = Id;
+				this.Summary = Summary;
+				this.CreatedAt = DateTime.UtcNow;
+				this.Fingerprint = new IssueFingerprint(Span.Fingerprint);
+				this.Severity = Span.Severity;
+				this.NotifySuspects = Span.NotifySuspects;
+				this.Suspects = Span.Suspects.ConvertAll(x => new IssueSuspect(x.Author, x.AuthorId, x.OriginatingChange ?? x.Change, null));
+
+				if (Suspects.Count > 0)
+				{
+					MinSuspectChange = Suspects.Min(x => x.Change);
+					MaxSuspectChange = Suspects.Max(x => x.Change);
+				}
+			}
+
+			IssueSuspect? GetDefaultOwner()
+			{
+				if (Suspects.Count > 0)
+				{
+					string PossibleOwner = Suspects[0].Author;
+					if (Suspects.All(x => x.Author.Equals(PossibleOwner, StringComparison.OrdinalIgnoreCase)) && Suspects.Any(x => x.DeclinedAt == null))
+					{
+						return Suspects[0];
+					}
+				}
+				return null;
+			}
+		}
+
+		class IssueSuspect : IIssueSuspect
+		{
+			[BsonRequired]
+			public string Author { get; set; }
+			public ObjectId AuthorId { get; set; }
+			public int Change { get; set; }
+			public DateTime? DeclinedAt { get; set; }
+
+			private IssueSuspect()
+			{
+				Author = String.Empty;
+			}
+
+			public IssueSuspect(string Author, ObjectId AuthorId, int Change, DateTime? DeclinedAt)
+			{
+				this.Author = Author;
+				this.AuthorId = AuthorId;
+				this.Change = Change;
+				this.DeclinedAt = DeclinedAt;
+			}
+
+			public IssueSuspect(IIssueSuspect Other)
+			{
+				this.Author = Other.Author;
+				this.AuthorId = Other.AuthorId;
+				this.Change = Other.Change;
+				this.DeclinedAt = Other.DeclinedAt;
+			}
+		}
+
+		class IssueFingerprint : IIssueFingerprint
+		{
+			public string Type { get; set; }
+			public CaseInsensitiveStringSet Keys { get; set; }
+			public CaseInsensitiveStringSet? RejectKeys { get; set; }
+
+			[BsonConstructor]
+			private IssueFingerprint()
+			{
+				Type = String.Empty;
+				Keys = new CaseInsensitiveStringSet();
+			}
+
+			public IssueFingerprint(IIssueFingerprint Fingerprint)
+			{
+				this.Type = Fingerprint.Type;
+				this.Keys = Fingerprint.Keys;
+				this.RejectKeys = Fingerprint.RejectKeys;
+			}
+		}
+
+		class IssueSpan : IIssueSpan
+		{
+			public ObjectId Id { get; set; }
+
+			[BsonRequired]
+			public StreamId StreamId { get; set; }
+
+			[BsonRequired]
+			public string StreamName { get; set; }
+
+			[BsonRequired]
+			public TemplateRefId TemplateRefId { get; set; }
+
+			[BsonRequired]
+			public string NodeName { get; set; }
+			public IssueSeverity Severity { get; set; }
+
+			[BsonRequired]
+			public IssueFingerprint Fingerprint { get; set; }
+
+			public int MinChange { get; set; }
+			public int MaxChange { get; set; } = int.MaxValue;
+
+			public IssueStep? LastSuccess { get; set; }
+
+			[BsonRequired]
+			public IssueStep FirstFailure { get; set; }
+
+			[BsonRequired]
+			public IssueStep LastFailure { get; set; }
+
+			public IssueStep? NextSuccess { get; set; }
+
+			public bool NotifySuspects { get; set; }
+			public List<IssueSpanSuspect> Suspects { get; set; }
+			public int? IssueId { get; set; }
+			public bool Modified { get; set; }
+			public int UpdateIndex { get; set; }
+
+			IIssueStep? IIssueSpan.LastSuccess => LastSuccess;
+			IIssueStep IIssueSpan.FirstFailure => FirstFailure;
+			IIssueStep IIssueSpan.LastFailure => LastFailure;
+			IIssueStep? IIssueSpan.NextSuccess => NextSuccess;
+			IReadOnlyList<IIssueSpanSuspect> IIssueSpan.Suspects => Suspects;
+			IIssueFingerprint IIssueSpan.Fingerprint => Fingerprint;
+
+			private IssueSpan()
+			{
+				this.StreamName = null!;
+				this.NodeName = null!;
+				this.Fingerprint = null!;
+				this.FirstFailure = null!;
+				this.LastFailure = null!;
+				this.Suspects = new List<IssueSpanSuspect>();
+			}
+
+			public IssueSpan(StreamId StreamId, string StreamName, TemplateRefId TemplateRefId, string NodeName, IssueSeverity Severity, IIssueFingerprint Fingerprint, NewIssueStepData? LastSuccess, NewIssueStepData Failure, NewIssueStepData? NextSuccess, List<IssueSpanSuspect> Suspects)
+			{
+				this.Id = ObjectId.GenerateNewId();
+				this.StreamId = StreamId;
+				this.StreamName = StreamName;
+				this.TemplateRefId = TemplateRefId;
+				this.NodeName = NodeName;
+				this.Severity = Severity;
+				this.Fingerprint = new IssueFingerprint(Fingerprint);
+				if (LastSuccess != null)
+				{
+					this.MinChange = LastSuccess.Change;
+					this.LastSuccess = new IssueStep(Id, LastSuccess);
+				}
+				this.FirstFailure = new IssueStep(Id, Failure);
+				this.LastFailure = new IssueStep(Id, Failure);
+				if (NextSuccess != null)
+				{
+					this.MaxChange = NextSuccess.Change;
+					this.NextSuccess = new IssueStep(Id, NextSuccess);
+				}
+				this.NotifySuspects = Failure.NotifySuspects;
+				this.Suspects = Suspects;
+				this.Modified = true;
+			}
+		}
+
+		class IssueSpanSuspect : IIssueSpanSuspect
+		{
+			public int Change { get; set; }
+			public string Author { get; set; }
+			public ObjectId AuthorId { get; set; }
+			public int? OriginatingChange { get; set; }
+
+			[BsonConstructor]
+			private IssueSpanSuspect()
+			{
+				Author = String.Empty;
+			}
+
+			public IssueSpanSuspect(int Change, string Author, ObjectId AuthorId, int? OriginatingChange)
+			{
+				this.Change = Change;
+				this.Author = Author;
+				this.AuthorId = AuthorId;
+				this.OriginatingChange = OriginatingChange;
+			}
+		}
+
+		class IssueStep : IIssueStep
+		{
+			public ObjectId Id { get; set; }
+			public ObjectId SpanId { get; set; }
+
+			public int Change { get; set; }
+
+			[BsonRequired]
+			public string JobName { get; set; }
+
+			[BsonRequired]
+			public ObjectId JobId { get; set; }
+
+			[BsonRequired]
+			public SubResourceId BatchId { get; set; }
+
+			[BsonRequired]
+			public SubResourceId StepId { get; set; }
+
+			public DateTime StepTime { get; set; }
+
+			public ObjectId? LogId { get; set; }
+
+			public bool NotifySuspects { get; set; }
+
+			[BsonConstructor]
+			private IssueStep()
+			{
+				this.JobName = null!;
+			}
+
+			public IssueStep(ObjectId SpanId, NewIssueStepData StepData)
+			{
+				this.Id = ObjectId.GenerateNewId();
+				this.SpanId = SpanId;
+				this.Change = StepData.Change;
+				this.JobName = StepData.JobName;
+				this.JobId = StepData.JobId;
+				this.BatchId = StepData.BatchId;
+				this.StepId = StepData.StepId;
+				this.StepTime = StepData.StepTime;
+				this.LogId = StepData.LogId;
+				this.NotifySuspects = StepData.NotifySuspects;
+			}
+		}
+
+		ISingletonDocument<IssueLedger> LedgerSingleton;
+		IMongoCollection<Issue> Issues;
+		IMongoCollection<IssueSpan> IssueSpans;
+		IMongoCollection<IssueStep> IssueSteps;
+		ILogger Logger;
+
+		// TODO: Temp upgrade path for switching user references to ids. Remove after resaving all issues.
+		IUserCollection UserCollection;
+		ConcurrentDictionary<string, ObjectId> UserNameToIdCache = new ConcurrentDictionary<string, ObjectId>(StringComparer.OrdinalIgnoreCase);
+		IPerforceService PerforceService;
+
+		public IssueCollection(DatabaseService DatabaseService, IUserCollection UserCollection, IPerforceService PerforceService, ILogger<IssueCollection> Logger)
+		{
+			this.UserCollection = UserCollection;
+			this.PerforceService = PerforceService;
+			this.Logger = Logger;
+
+			LedgerSingleton = new SingletonDocument<IssueLedger>(DatabaseService);
+
+			Issues = DatabaseService.GetCollection<Issue>("IssuesV2");
+			IssueSpans = DatabaseService.GetCollection<IssueSpan>("IssuesV2.Spans");
+			IssueSteps = DatabaseService.GetCollection<IssueStep>("IssuesV2.Steps");
+			
+			if (!DatabaseService.ReadOnlyMode)
+			{
+				Issues.Indexes.CreateOne(new CreateIndexModel<Issue>(Builders<Issue>.IndexKeys.Ascending(x => x.Resolved)));
+				
+				IssueSpans.Indexes.CreateOne(new CreateIndexModel<IssueSpan>(Builders<IssueSpan>.IndexKeys.Ascending(x => x.Modified)));
+				IssueSpans.Indexes.CreateOne(new CreateIndexModel<IssueSpan>(Builders<IssueSpan>.IndexKeys.Ascending(x => x.IssueId)));
+				IssueSpans.Indexes.CreateOne(new CreateIndexModel<IssueSpan>(Builders<IssueSpan>.IndexKeys.Ascending(x => x.StreamId).Ascending(x => x.MinChange).Ascending(x => x.MaxChange)));
+				IssueSpans.Indexes.CreateOne(new CreateIndexModel<IssueSpan>(Builders<IssueSpan>.IndexKeys.Ascending(x => x.StreamId).Ascending(x => x.TemplateRefId).Ascending(x => x.NodeName).Ascending(x => x.MinChange).Ascending(x => x.MaxChange), new CreateIndexOptions { Name = "StreamChanges" }));
+				
+				IssueSteps.Indexes.CreateOne(new CreateIndexModel<IssueStep>(Builders<IssueStep>.IndexKeys.Ascending(x => x.SpanId)));
+				IssueSteps.Indexes.CreateOne(new CreateIndexModel<IssueStep>(Builders<IssueStep>.IndexKeys.Ascending(x => x.JobId).Ascending(x => x.BatchId).Ascending(x => x.StepId)));
+			}
+		}
+
+		async Task<IIssue?> ApplyTransactionAsync(IIssue Issue, TransactionBuilder<Issue> Transaction)
+		{
+			Issue IssueDocument = (Issue)Issue;
+
+			int PrevUpdateIndex = IssueDocument.UpdateIndex;
+			Transaction.Set(x => x.UpdateIndex, PrevUpdateIndex + 1);
+
+			UpdateResult Result = await Issues.UpdateOneAsync(x => x.Id == IssueDocument.Id && x.UpdateIndex == PrevUpdateIndex, Transaction.ToUpdateDefinition());
+			if (Result.ModifiedCount == 0)
+			{
+				return null;
+			}
+
+			Transaction.ApplyTo(IssueDocument);
+			return IssueDocument;
+		}
+
+		async Task<IIssueSpan?> ApplyTransactionAsync(IIssueSpan IssueSpan, TransactionBuilder<IssueSpan> Transaction)
+		{
+			IssueSpan IssueSpanDocument = (IssueSpan)IssueSpan;
+
+			int PrevUpdateIndex = IssueSpanDocument.UpdateIndex;
+			Transaction.Set(x => x.UpdateIndex, PrevUpdateIndex + 1);
+
+			UpdateResult Result = await IssueSpans.UpdateOneAsync(x => x.Id == IssueSpanDocument.Id && x.UpdateIndex == PrevUpdateIndex, Transaction.ToUpdateDefinition());
+			if (Result.ModifiedCount > 0)
+			{
+				Transaction.ApplyTo(IssueSpanDocument);
+				return IssueSpanDocument;
+			}
+			return null;
+		}
+
+		/// <inheritdoc/>
+		public async Task<IIssueSequenceToken> GetSequenceTokenAsync()
+		{
+			IssueLedger Ledger = await GetLedgerAsync();
+			return new IssueSequenceToken(this, Ledger);
+		}
+
+		async Task<IssueLedger> GetLedgerAsync()
+		{
+			IssueLedger Ledger = await LedgerSingleton.GetAsync();
+			Ledger = await FlushLedgerAsync(Ledger);
+			return Ledger;
+		}
+
+		async Task<IssueLedger> FlushLedgerAsync(IssueLedger Ledger)
+		{
+			while (Ledger.HasPendingTransaction)
+			{
+				if (Ledger.InsertIssue != null)
+				{
+					await Issues.InsertOneIgnoreDuplicatesAsync(Ledger.InsertIssue);
+					Ledger.InsertIssue = null;
+				}
+				if (Ledger.InsertIssueSpan != null)
+				{
+					await IssueSpans.InsertOneIgnoreDuplicatesAsync(Ledger.InsertIssueSpan);
+					Ledger.InsertIssueSpan = null;
+				}
+				if (Ledger.UpdateIssue != null)
+				{
+					await Issues.UpdateOneAsync(Ledger.UpdateIssue.Filter, Ledger.UpdateIssue.Update);
+					Ledger.UpdateIssue = null;
+				}
+				if (Ledger.UpdateIssueSpan != null)
+				{
+					await IssueSpans.UpdateOneAsync(Ledger.UpdateIssueSpan.Filter, Ledger.UpdateIssueSpan.Update);
+					Ledger.UpdateIssueSpan = null;
+				}
+
+				if (await LedgerSingleton.TryUpdateAsync(Ledger))
+				{
+					break;
+				}
+
+				Ledger = await LedgerSingleton.GetAsync();
+			}
+			return Ledger;
+		}
+
+		#region Issues
+
+		/// <inheritdoc/>
+		public async Task<int?> ReserveUniqueIdAsync(IIssueSequenceToken Token)
+		{
+			IssueSequenceToken TokenValue = ((IssueSequenceToken)Token);
+			int NewIssueId = ++TokenValue.Ledger.NextId;
+
+			if (!await LedgerSingleton.TryUpdateAsync(TokenValue.Ledger))
+			{
+				return null;
+			}
+
+			return NewIssueId;
+		}
+
+		/// <inheritdoc/>
+		public async Task<IIssue?> AddIssueAsync(IIssueSequenceToken Token, string Summary, IIssueSpan Span)
+		{
+			IssueSequenceToken TokenValue = ((IssueSequenceToken)Token);
+			Issue NewIssue = new Issue(++TokenValue.Ledger.NextId, Summary, Span);
+			TokenValue.Ledger.InsertIssue = NewIssue;
+			TokenValue.Ledger.UpdateIssueSpan = UpdateTransaction.Create<IssueSpan>(x => x.Id == Span.Id, Builders<IssueSpan>.Update.Set(x => x.IssueId, NewIssue.Id));
+
+			if (!await LedgerSingleton.TryUpdateAsync(TokenValue.Ledger))
+			{
+				return null;
+			}
+
+			await TokenValue.FlushAsync();
+			return NewIssue;
+		}
+
+		/// <inheritdoc/>
+		public async Task<IIssue?> GetIssueAsync(int IssueId)
+		{
+			Issue Issue = await Issues.Find(x => x.Id == IssueId).FirstOrDefaultAsync();
+			await UpgradeIssueAsync(Issue);
+			return Issue;
+		}
+
+
+		class ProjectedIssueId
+		{
+			public int? IssueId { get; set; }
+		}
+
+		class ProjectedIssueIdWithIssue
+		{
+			public int? IssueId { get; set; }
+			public IEnumerable<Issue> Issues { get; set; } = null!;
+		}
+
+		/// <inheritdoc/>
+		public async Task<List<IIssue>> FindIssuesAsync(IEnumerable<int>? Ids = null, StreamId? StreamId = null, int? MinChange = null, int? MaxChange = null, bool? Resolved = null, int? Index = null, int? Count = null)
+		{
+			// Filtering against a stream ID needs to be done via a join on the spans collection
+			if (StreamId != null)
+			{
+				ProjectionDefinition<IssueSpan> Projection = Builders<IssueSpan>.Projection.Include(x => x.IssueId);
+
+				FilterDefinition<IssueSpan> SpanFilter = Builders<IssueSpan>.Filter.Eq(x => x.StreamId, StreamId.Value);
+				if (Ids != null && Ids.Any())
+				{
+					SpanFilter &= Builders<IssueSpan>.Filter.In(x => x.IssueId, Ids.Select<int, int?>(x => x));
+				}
+				else
+				{
+					SpanFilter &= Builders<IssueSpan>.Filter.Exists(x => x.IssueId);
+				}
+
+				if (MinChange != null)
+				{
+					SpanFilter &= Builders<IssueSpan>.Filter.Not(Builders<IssueSpan>.Filter.Lt(x => x.MaxChange, MinChange.Value));
+				}
+				if (MaxChange != null)
+				{
+					SpanFilter &= Builders<IssueSpan>.Filter.Not(Builders<IssueSpan>.Filter.Gt(x => x.MinChange, MaxChange.Value));
+				}
+				if (Resolved != null)
+				{
+					if (Resolved.Value)
+					{
+						SpanFilter &= Builders<IssueSpan>.Filter.Ne(x => x.MaxChange, int.MaxValue);
+					}
+					else
+					{
+						SpanFilter &= Builders<IssueSpan>.Filter.Eq(x => x.MaxChange, int.MaxValue);
+					}
+				}
+
+				IAggregateFluent<ProjectedIssueId> IssueIds = IssueSpans.Aggregate().Match(SpanFilter).Group(x => x.IssueId, x => new ProjectedIssueId { IssueId = x.First().IssueId }).SortByDescending(x => x.IssueId);
+				if (Index != null)
+				{
+					IssueIds = IssueIds.Skip(Index.Value);
+				}
+				if (Count != null)
+				{
+					IssueIds = IssueIds.Limit(Count.Value);
+				}
+
+				List<ProjectedIssueIdWithIssue> ProjectedResults = await IssueIds.Lookup(Issues, x => x.IssueId, x => x.Id, (ProjectedIssueIdWithIssue x) => x.Issues).ToListAsync();
+				return await UpgradeIssuesAsync(ProjectedResults.Where(x => x.Issues.Any()).Select(x => x.Issues.First()).ToList());
+			}
+			else
+			{
+				FilterDefinition<Issue> Filter = FilterDefinition<Issue>.Empty;
+				if (Ids != null)
+				{
+					Filter &= Builders<Issue>.Filter.In(x => x.Id, Ids);
+				}
+				if (Resolved != null)
+				{
+					Filter &= Builders<Issue>.Filter.Eq(x => x.Resolved, Resolved.Value);
+				}
+
+				List<Issue> Results = await Issues.Find(Filter).Range(Index, Count).ToListAsync();
+				return await UpgradeIssuesAsync(Results);
+			}
+		}
+
+		/// <inheritdoc/>
+		public async Task<List<IIssue>> FindIssuesForSuspectsAsync(List<int> Changes)
+		{
+			List<IIssue> Results = new List<IIssue>();
+			if (Changes.Count > 0)
+			{
+				int MinChange = Changes.Min();
+				int MaxChange = Changes.Max();
+
+				List<Issue> BaseResults = await Issues.Find(x => x.MinSuspectChange <= MaxChange && x.MaxSuspectChange >= MinChange).ToListAsync();
+				await UpgradeIssuesAsync(BaseResults);
+				Results.AddRange(BaseResults.Where(x => x.Suspects.Any(y => Changes.Contains(y.Change))));
+			}
+			return Results;
+		}
+
+		async Task<ObjectId> FindOrAddUserIdAsync(string UserName)
+		{
+			ObjectId UserId;
+			if (!UserNameToIdCache.TryGetValue(UserName, out UserId))
+			{
+				PerforceUserInfo? Info = await PerforceService.GetUserInfoAsync(UserName);
+				if (Info == null)
+				{
+					Logger.LogWarning("No user '{UserName}' found", UserName);
+				}
+
+				IUser User = await UserCollection.FindOrAddUserByLoginAsync(UserName, Info?.Name, Info?.Email);
+				UserId = User.Id;
+
+				UserNameToIdCache.TryAdd(UserName, UserId);
+			}
+			return UserId;
+		}
+
+		async Task<IIssue> UpgradeIssueAsync(Issue Issue)
+		{
+			if (Issue.Owner != null && Issue.OwnerId == null)
+			{
+				Issue.OwnerId = await FindOrAddUserIdAsync(Issue.Owner);
+			}
+			if (Issue.NominatedBy != null && Issue.NominatedById == null)
+			{
+				Issue.NominatedById = await FindOrAddUserIdAsync(Issue.NominatedBy);
+			}
+			foreach (IssueSuspect Suspect in Issue.Suspects)
+			{
+				if (Suspect.AuthorId == ObjectId.Empty)
+				{
+					Suspect.AuthorId = await FindOrAddUserIdAsync(Suspect.Author);
+				}
+			}
+			return Issue;
+		}
+
+		async Task<List<IIssue>> UpgradeIssuesAsync(List<Issue> Issues)
+		{
+			List<IIssue> Results = new List<IIssue>();
+			foreach(Issue Issue in Issues)
+			{
+				Results.Add(await UpgradeIssueAsync(Issue));
+			}
+			return Results;
+		}
+
+		async Task<IIssueSpan> UpgradeIssueSpanAsync(IssueSpan IssueSpan)
+		{
+			foreach (IssueSpanSuspect Suspect in IssueSpan.Suspects)
+			{
+				if (Suspect.AuthorId == ObjectId.Empty)
+				{
+					Suspect.AuthorId = await FindOrAddUserIdAsync(Suspect.Author);
+				}
+			}
+			return IssueSpan;
+		}
+
+		async Task<List<IIssueSpan>> UpgradeIssueSpansAsync(List<IssueSpan> IssueSpans)
+		{
+			List<IIssueSpan> Results = new List<IIssueSpan>();
+			foreach (IssueSpan IssueSpan in IssueSpans)
+			{
+				Results.Add(await UpgradeIssueSpanAsync(IssueSpan));
+			}
+			return Results;
+		}
+
+		async Task<List<IssueSpanSuspect>> CreateSpanSuspectsAsync(List<NewIssueSpanSuspectData> Suspects)
+		{
+			List<IssueSpanSuspect> NewSuspects = new List<IssueSpanSuspect>();
+			foreach (NewIssueSpanSuspectData Suspect in Suspects)
+			{
+				ObjectId AuthorId = await FindOrAddUserIdAsync(Suspect.Author);
+				NewSuspects.Add(new IssueSpanSuspect(Suspect.Change, Suspect.Author, AuthorId, Suspect.OriginatingChange));
+			}
+			return NewSuspects;
+		}
+
+		/// <inheritdoc/>
+		public async Task<IIssue?> UpdateIssueAsync(IIssue Issue, IssueSeverity? NewSeverity = null, string? NewSummary = null, string? NewUserSummary = null, string? NewOwner = null, string? NewNominatedBy = null, bool? NewAcknowledged = null, string? NewDeclinedBy = null, int? NewFixChange = null, bool? NewResolved = null, bool? NewNotifySuspects = null, List<NewIssueSuspectData>? NewSuspects = null)
+		{
+			Issue IssueDocument = (Issue)Issue;
+
+			DateTime UtcNow = DateTime.UtcNow;
+
+			TransactionBuilder<Issue> Transaction = new TransactionBuilder<Issue>();
+			if (NewSeverity != null)
+			{
+				Transaction.Set(x => x.Severity, NewSeverity.Value);
+			}
+			if (NewSummary != null)
+			{
+				Transaction.Set(x => x.Summary, NewSummary);
+			}
+			if (NewUserSummary != null)
+			{
+				if (NewUserSummary.Length == 0)
+				{
+					Transaction.Unset(x => x.UserSummary!);
+				}
+				else
+				{
+					Transaction.Set(x => x.UserSummary, NewUserSummary);
+				}
+			}
+			if (NewOwner != null)
+			{
+				if (NewOwner.Length == 0)
+				{
+					Transaction.Unset(x => x.Owner!);
+					Transaction.Unset(x => x.OwnerId!);
+					Transaction.Unset(x => x.NominatedAt!);
+					Transaction.Unset(x => x.NominatedBy!);
+					Transaction.Unset(x => x.NominatedById!);
+				}
+				else
+				{
+					Transaction.Set(x => x.Owner!, NewOwner);
+					Transaction.Set(x => x.OwnerId!, await FindOrAddUserIdAsync(NewOwner));
+					Transaction.Set(x => x.NominatedAt, DateTime.UtcNow);
+					if (String.IsNullOrEmpty(NewNominatedBy))
+					{
+						Transaction.Unset(x => x.NominatedBy!);
+						Transaction.Unset(x => x.NominatedById!);
+					}
+					else
+					{
+						Transaction.Set(x => x.NominatedBy, NewNominatedBy);
+						Transaction.Set(x => x.NominatedById, await FindOrAddUserIdAsync(NewNominatedBy));
+					}
+					NewAcknowledged = NewAcknowledged ?? false;
+				}
+			}
+			if (NewAcknowledged != null)
+			{
+				if (NewAcknowledged.Value)
+				{
+					if (IssueDocument.AcknowledgedAt == null)
+					{
+						Transaction.Set(x => x.AcknowledgedAt, UtcNow);
+					}
+				}
+				else
+				{
+					if (IssueDocument.AcknowledgedAt != null)
+					{
+						Transaction.Unset(x => x.AcknowledgedAt!);
+					}
+				}
+			}
+			if (NewDeclinedBy != null)
+			{
+				int SuspectIdx = IssueDocument.Suspects.FindIndex(x => x.Author.Equals(NewDeclinedBy, StringComparison.OrdinalIgnoreCase));
+				if (SuspectIdx != -1 && IssueDocument.Suspects[SuspectIdx].DeclinedAt == null)
+				{
+					Transaction.Set(x => x.Suspects[SuspectIdx].DeclinedAt, UtcNow);
+				}
+			}
+			if (NewFixChange != null)
+			{
+				if (NewFixChange == 0)
+				{
+					Transaction.Unset(x => x.FixChange!);
+				}
+				else
+				{
+					Transaction.Set(x => x.FixChange, NewFixChange);
+				}
+			}
+			if (NewResolved != null)
+			{
+				Transaction.Set(x => x.Resolved, NewResolved.Value);
+
+				if (NewResolved.Value)
+				{
+					if (IssueDocument.ResolvedAt == null)
+					{
+						Transaction.Set(x => x.ResolvedAt, UtcNow);
+					}
+				}
+				else
+				{
+					if (IssueDocument.ResolvedAt != null)
+					{
+						Transaction.Unset(x => x.ResolvedAt!);
+					}
+				}
+			}
+			if (NewNotifySuspects != null)
+			{
+				Transaction.Set(x => x.NotifySuspects, NewNotifySuspects.Value);
+			}
+			if (NewSuspects != null)
+			{
+				Transaction.Set(x => x.MinSuspectChange, (NewSuspects.Count > 0) ? NewSuspects.Min(x => x.Change) : 0);
+				Transaction.Set(x => x.MaxSuspectChange, (NewSuspects.Count > 0) ? NewSuspects.Max(x => x.Change) : 0);
+
+				List<IssueSuspect> Suspects = new List<IssueSuspect>();
+				foreach (NewIssueSuspectData NewSuspect in NewSuspects)
+				{
+					ObjectId AuthorId = await FindOrAddUserIdAsync(NewSuspect.Author);
+					Suspects.Add(new IssueSuspect(NewSuspect.Author, AuthorId, NewSuspect.Change, NewSuspect.DeclinedAt));
+				}
+
+				Transaction.Set(x => x.Suspects, Suspects);
+			}
+
+			return await ApplyTransactionAsync(Issue, Transaction);
+		}
+
+		#endregion
+
+		#region Spans
+
+		/// <inheritdoc/>
+		public async Task<IIssueSpan?> AddSpanAsync(IIssueSequenceToken Token, IStream Stream, TemplateRefId TemplateRefId, string NodeName, IssueSeverity Severity, NewIssueFingerprint Fingerprint, NewIssueStepData? LastSuccess, NewIssueStepData Failure, NewIssueStepData? NextSuccess, List<NewIssueSpanSuspectData> Suspects)
+		{
+			List<IssueSpanSuspect> NewSuspects = await CreateSpanSuspectsAsync(Suspects);
+			IssueSpan NewSpan = new IssueSpan(Stream.Id, Stream.Name, TemplateRefId, NodeName, Severity, Fingerprint, LastSuccess, Failure, NextSuccess, NewSuspects);
+
+			IssueLedger Ledger = ((IssueSequenceToken)Token).Ledger;
+			Ledger.InsertIssueSpan = NewSpan;
+
+			if (!await LedgerSingleton.TryUpdateAsync(Ledger))
+			{
+				return null;
+			}
+
+			await FlushLedgerAsync(Ledger);
+			await AddStepAsync(NewSpan.Id, Failure);
+			return NewSpan;
+		}
+
+		/// <inheritdoc/>
+		public async Task<IIssueSpan?> GetSpanAsync(ObjectId SpanId)
+		{
+			return await IssueSpans.Find(Builders<IssueSpan>.Filter.Eq(x => x.Id, SpanId)).FirstOrDefaultAsync();
+		}
+
+		/// <inheritdoc/>
+		public async Task<IIssueSpan?> UpdateSpanAsync(IIssueSpan Span, IssueSeverity? NewSeverity = null, NewIssueStepData? NewLastSuccess = null, NewIssueStepData? NewFailure = null, NewIssueStepData? NewNextSuccess = null, List<NewIssueSpanSuspectData>? NewSuspects = null, bool? NewModified = null)
+		{
+			TransactionBuilder<IssueSpan> Transaction = new TransactionBuilder<IssueSpan>();
+			if (NewSeverity != null)
+			{
+				Transaction.Set(x => x.Severity, NewSeverity.Value);
+			}
+			if (NewLastSuccess != null)
+			{
+				Transaction.Set(x => x.MinChange, NewLastSuccess.Change);
+				Transaction.Set(x => x.LastSuccess, new IssueStep(Span.Id, NewLastSuccess));
+			}
+			if (NewFailure != null)
+			{
+				if (NewFailure.Change < Span.FirstFailure.Change)
+				{
+					Transaction.Set(x => x.FirstFailure, new IssueStep(Span.Id, NewFailure));
+				}
+				if (NewFailure.Change > Span.LastFailure.Change)
+				{
+					Transaction.Set(x => x.LastFailure, new IssueStep(Span.Id, NewFailure));
+				}
+				if (NewFailure.NotifySuspects != Span.NotifySuspects && NewFailure.Change >= Span.LastFailure.Change)
+				{
+					Transaction.Set(x => x.NotifySuspects, NewFailure.NotifySuspects);
+				}
+			}
+			if (NewNextSuccess != null)
+			{
+				Transaction.Set(x => x.MaxChange, NewNextSuccess.Change);
+				Transaction.Set(x => x.NextSuccess, new IssueStep(Span.Id, NewNextSuccess));
+			}
+			if (NewSuspects != null)
+			{
+				Transaction.Set(x => x.Suspects, await CreateSpanSuspectsAsync(NewSuspects));
+			}
+			if (NewModified != null)
+			{
+				Transaction.Set(x => x.Modified, NewModified.Value);
+			}
+			
+			IIssueSpan? NewSpan = await ApplyTransactionAsync(Span, Transaction);
+			if (NewSpan != null && NewFailure != null)
+			{
+				await AddStepAsync(NewSpan.Id, NewFailure);
+			}
+			return NewSpan;
+		}
+
+		/// <inheritdoc/>
+		public async Task<List<IIssueSpan>> FindSpansAsync(int IssueId)
+		{
+			return await UpgradeIssueSpansAsync(await IssueSpans.Find(x => x.IssueId == IssueId).ToListAsync());
+		}
+
+		/// <inheritdoc/>
+		public Task<List<IIssueSpan>> FindSpansAsync(IEnumerable<ObjectId> SpanIds)
+		{
+			return IssueSpans.Find(Builders<IssueSpan>.Filter.In(x => x.Id, SpanIds)).ToListAsync<IssueSpan, IIssueSpan>();
+		}
+
+		/// <inheritdoc/>
+		public async Task<List<IIssueSpan>> FindOpenSpansAsync(StreamId StreamId, TemplateRefId TemplateId, string NodeName, int Change)
+		{
+			List<IssueSpan> Spans = await IssueSpans.Find(x => x.StreamId == StreamId && x.TemplateRefId == TemplateId && x.NodeName == NodeName && Change >= x.MinChange && Change <= x.MaxChange).ToListAsync();
+			return Spans.ConvertAll<IIssueSpan>(x => x);
+		}
+
+		/// <inheritdoc/>
+		public async Task<bool> AttachSpanToIssueAsync(IIssueSequenceToken Token, IIssueSpan Span, IIssue Issue, IssueSeverity? NewSeverity = null, NewIssueFingerprint? NewFingerprint = null, string? NewSummary = null)
+		{
+			List<UpdateDefinition<Issue>> IssueUpdates = new List<UpdateDefinition<Issue>>();
+			if (NewSeverity != null)
+			{
+				IssueUpdates.Add(Builders<Issue>.Update.Set(x => x.Severity, NewSeverity.Value));
+			}
+			if (NewFingerprint != null)
+			{
+				IssueUpdates.Add(Builders<Issue>.Update.Set(x => x.Fingerprint, new IssueFingerprint(NewFingerprint)));
+			}
+			if (NewSummary != null)
+			{
+				IssueUpdates.Add(Builders<Issue>.Update.Set(x => x.Summary, NewSummary));
+			}
+
+			IssueSequenceToken TokenValue = ((IssueSequenceToken)Token);
+			if (IssueUpdates.Count > 0)
+			{
+				TokenValue.Ledger.UpdateIssue = UpdateTransaction.Create(Builders<Issue>.Filter.Eq(x => x.Id, Issue.Id), Builders<Issue>.Update.Combine(IssueUpdates));
+			}
+			TokenValue.Ledger.UpdateIssueSpan = UpdateTransaction.Create(Builders<IssueSpan>.Filter.Eq(x => x.Id, Span.Id), Builders<IssueSpan>.Update.Set(x => x.IssueId, Issue.Id));
+
+			if (!await LedgerSingleton.TryUpdateAsync(TokenValue.Ledger))
+			{
+				return false;
+			}
+
+			await TokenValue.FlushAsync();
+			return true;
+		}
+
+		/// <inheritdoc/>
+		public async Task<IIssueSpan?> FindModifiedSpanAsync()
+		{
+			IssueSpan? IssueSpan = await IssueSpans.Find(x => x.Modified).FirstOrDefaultAsync();
+			return IssueSpan;
+		}
+
+		#endregion
+
+		#region Steps
+
+		/// <inheritdoc/>
+		async Task AddStepAsync(ObjectId SpanId, NewIssueStepData NewStep)
+		{
+			await IssueSteps.InsertOneIgnoreDuplicatesAsync(new IssueStep(SpanId, NewStep));
+		}
+
+		/// <inheritdoc/>
+		public Task<List<IIssueStep>> FindStepsAsync(IEnumerable<ObjectId> SpanIds)
+		{
+			FilterDefinition<IssueStep> Filter = Builders<IssueStep>.Filter.In(x => x.SpanId, SpanIds);
+			return IssueSteps.Find(Filter).ToListAsync<IssueStep, IIssueStep>();
+		}
+
+		/// <inheritdoc/>
+		public Task<List<IIssueStep>> FindStepsAsync(ObjectId JobId, SubResourceId? BatchId, SubResourceId? StepId)
+		{
+			FilterDefinition<IssueStep> Filter = Builders<IssueStep>.Filter.Eq(x => x.JobId, JobId);
+			if (BatchId != null)
+			{
+				Filter &= Builders<IssueStep>.Filter.Eq(x => x.BatchId, BatchId.Value);
+			}
+			if (StepId != null)
+			{
+				Filter &= Builders<IssueStep>.Filter.Eq(x => x.StepId, StepId.Value);
+			}
+			return IssueSteps.Find(Filter).ToListAsync<IssueStep, IIssueStep>();
+		}
+
+		#endregion
+	}
+}
