Initial drop of Horde Server/Agent for UE5/Main. Still missing some Epic-specific functionality required to build/run correctly.

diff --git a/Engine/Source/Programs/Horde/HordeServer/Collections/Impl/JobCollection.cs b/Engine/Source/Programs/Horde/HordeServer/Collections/Impl/JobCollection.cs
new file mode 100644
--- /dev/null
+++ b/Engine/Source/Programs/Horde/HordeServer/Collections/Impl/JobCollection.cs
@@ -0,0 +1,1598 @@
+using Datadog.Trace;
+using EpicGames.Core;
+using HordeServer.Api;
+using HordeCommon;
+using HordeServer.Models;
+using HordeServer.Services;
+using HordeServer.Utilities;
+using Microsoft.Extensions.Logging;
+using MongoDB.Bson;
+using MongoDB.Bson.Serialization.Attributes;
+using MongoDB.Driver;
+using System;
+using System.Collections.Generic;
+using System.Diagnostics.CodeAnalysis;
+using System.IO;
+using System.Linq;
+using System.Threading.Tasks;
+
+using StreamId = HordeServer.Utilities.StringId<HordeServer.Models.IStream>;
+using TemplateRefId = HordeServer.Utilities.StringId<HordeServer.Models.TemplateRef>;
+using PoolId = HordeServer.Utilities.StringId<HordeServer.Models.IPool>;
+
+namespace HordeServer.Collections.Impl
+{
+	/// <summary>
+	/// Wrapper around the jobs collection in a mongo DB
+	/// </summary>
+	public class JobCollection : IJobCollection
+	{
+		/// <summary>
+		/// Embedded jobstep document
+		/// </summary>
+		[BsonIgnoreExtraElements]
+		class JobStepDocument : IJobStep
+		{
+			[BsonRequired]
+			public SubResourceId Id { get; set; }
+
+			[BsonRequired]
+			public int NodeIdx { get; set; }
+
+			[BsonRequired]
+			public JobStepState State { get; set; } = JobStepState.Waiting;
+
+			public JobStepOutcome Outcome { get; set; } = JobStepOutcome.Success;
+
+			[BsonIgnoreIfNull]
+			public ObjectId? LogId { get; set; }
+
+			[BsonIgnoreIfNull]
+			public ObjectId? NotificationTriggerId { get; set; }
+
+			[BsonIgnoreIfNull]
+			public DateTimeOffset? ReadyTime { get; set; }
+
+			[BsonIgnoreIfNull]
+			public DateTimeOffset? StartTime { get; set; }
+
+			[BsonIgnoreIfNull]
+			public DateTimeOffset? FinishTime { get; set; }
+
+			[BsonIgnoreIfNull]
+			public Priority? Priority { get; set; }
+
+			[BsonIgnoreIfDefault, BsonDefaultValue(false)]
+			public bool Retry { get; set; }
+
+			public string? RetryByUser { get; set; }
+
+			public bool AbortRequested { get; set; } = false;
+
+			public string? AbortByUser { get; set; }
+
+			[BsonIgnoreIfNull]
+			public List<Report>? Reports { get; set; }
+			IReadOnlyList<IReport>? IJobStep.Reports => Reports;
+
+			[BsonIgnoreIfNull]
+			public Dictionary<string, string>? Properties { get; set; }
+
+			string? IJobStep.RetryByUser => RetryByUser ?? (Retry ? (string?)"Unknown" : null);
+			DateTime? IJobStep.ReadyTimeUtc => ReadyTime?.UtcDateTime;
+			DateTime? IJobStep.StartTimeUtc => StartTime?.UtcDateTime;
+			DateTime? IJobStep.FinishTimeUtc => FinishTime?.UtcDateTime;
+
+			[BsonConstructor]
+			private JobStepDocument()
+			{
+			}
+
+			public JobStepDocument(SubResourceId Id, int NodeIdx)
+			{
+				this.Id = Id;
+				this.NodeIdx = NodeIdx;
+			}
+		}
+
+		class JobStepBatchDocument : IJobStepBatch
+		{
+			[BsonRequired]
+			public SubResourceId Id { get; set; }
+
+			public ObjectId? LogId { get; set; }
+
+			[BsonRequired]
+			public int GroupIdx { get; set; }
+
+			[BsonRequired]
+			public JobStepBatchState State { get; set; }
+
+			[BsonIgnoreIfDefault, BsonDefaultValue(JobStepBatchError.None)]
+			public JobStepBatchError Error { get; set; }
+
+			public List<JobStepDocument> Steps { get; set; } = new List<JobStepDocument>();
+
+			[BsonIgnoreIfNull]
+			public PoolId? PoolId { get; set; }
+
+			[BsonIgnoreIfNull]
+			public AgentId? AgentId { get; set; }
+
+			[BsonIgnoreIfNull]
+			public ObjectId? SessionId { get; set; }
+
+			[BsonIgnoreIfNull]
+			public ObjectId? LeaseId { get; set; }
+
+			public int SchedulePriority { get; set; }
+
+			[BsonIgnoreIfNull]
+			public DateTimeOffset? ReadyTime { get; set; }
+
+			[BsonIgnoreIfNull]
+			public DateTimeOffset? StartTime { get; set; }
+
+			[BsonIgnoreIfNull]
+			public DateTimeOffset? FinishTime { get; set; }
+
+			IReadOnlyList<IJobStep> IJobStepBatch.Steps => Steps;
+			DateTime? IJobStepBatch.ReadyTimeUtc => ReadyTime?.UtcDateTime;
+			DateTime? IJobStepBatch.StartTimeUtc => StartTime?.UtcDateTime;
+			DateTime? IJobStepBatch.FinishTimeUtc => FinishTime?.UtcDateTime;
+
+			[BsonConstructor]
+			private JobStepBatchDocument()
+			{
+			}
+
+			public JobStepBatchDocument(SubResourceId Id, int GroupIdx)
+			{
+				this.Id = Id;
+				this.GroupIdx = GroupIdx;
+			}
+		}
+
+		class ChainedJobDocument : IChainedJob
+		{
+			public string Target { get; set; }
+			public TemplateRefId TemplateRefId { get; set; }
+			public ObjectId? JobId { get; set; }
+
+			[BsonConstructor]
+			private ChainedJobDocument()
+			{
+				Target = String.Empty;
+			}
+
+			public ChainedJobDocument(ChainedJobTemplate Trigger)
+			{
+				this.Target = Trigger.Trigger;
+				this.TemplateRefId = Trigger.TemplateRefId;
+			}
+		}
+
+		class LabelNotificationDocument
+		{
+			public int LabelIdx;
+			public ObjectId TriggerId;
+		}
+
+		class JobDocument : IJob
+		{
+			[BsonRequired, BsonId]
+			public ObjectId Id { get; set; }
+
+			public StreamId StreamId { get; set; }
+			public TemplateRefId TemplateId { get; set; }
+			public ContentHash? TemplateHash { get; set; }
+			public ContentHash GraphHash { get; set; } = ContentHash.Empty;
+
+			[BsonIgnoreIfNull]
+			public ObjectId? StartedByUserId { get; set; }
+
+			[BsonIgnoreIfNull]
+			public string? StartedByUser { get; set; }
+
+			[BsonIgnoreIfNull]
+			public string? AbortedByUser { get; set; }
+
+			[BsonRequired]
+			public string Name { get; set; }
+
+			public int Change { get; set; }
+			public int CodeChange { get; set; }
+			public int PreflightChange { get; set; }
+			public int ClonedPreflightChange { get; set; }
+			public Priority Priority { get; set; }
+
+			[BsonIgnoreIfDefault]
+			public bool AutoSubmit { get; set; }
+
+			[BsonIgnoreIfNull]
+			public int? AutoSubmitChange { get; set; }
+
+			[BsonIgnoreIfNull]
+			public string? AutoSubmitMessage { get; set; }
+
+			[BsonIgnoreIfNull]
+			public DateTimeOffset? CreateTime { get; set; }
+
+			[BsonIgnoreIfNull]
+			public DateTime? CreateTimeUtc { get; set; }
+
+			public int SchedulePriority { get; set; }
+			public List<JobStepBatchDocument> Batches { get; set; } = new List<JobStepBatchDocument>();
+			public List<Report>? Reports { get; set; }
+			public List<string> Arguments { get; set; } = new List<string>();
+			public List<int> ReferencedByIssues { get; set; } = new List<int>();
+			public ObjectId? NotificationTriggerId { get; set; }
+			public bool ShowUgsBadges { get; set; }
+			public bool ShowUgsAlerts { get; set; }
+			public string? NotificationChannel { get; set; }
+			public string? NotificationChannelFilter { get; set; }
+			public string? HelixSwarmCallbackUrl { get; set; }
+			public List<LabelNotificationDocument> LabelNotifications = new List<LabelNotificationDocument>();
+			public List<ChainedJobDocument> ChainedJobs { get; set; } = new List<ChainedJobDocument>();
+
+			[BsonIgnoreIfNull]
+			public List<NodeRef>? RetriedNodes { get; set; }
+			public SubResourceId NextSubResourceId { get; set; }
+
+			[BsonIgnoreIfNull]
+			public DateTimeOffset? UpdateTime { get; set; }
+
+			[BsonIgnoreIfNull]
+			public DateTime? UpdateTimeUtc { get; set; }
+
+			public Acl? Acl { get; set; }
+
+			[BsonRequired]
+			public int UpdateIndex { get; set; }
+
+			DateTime IJob.CreateTimeUtc => CreateTimeUtc ?? CreateTime?.UtcDateTime ?? DateTime.UnixEpoch;
+			DateTime IJob.UpdateTimeUtc => UpdateTimeUtc ?? UpdateTime?.UtcDateTime ?? DateTime.UnixEpoch;
+			IReadOnlyList<IJobStepBatch> IJob.Batches => (IReadOnlyList<JobStepBatchDocument>)Batches;
+			IReadOnlyList<IReport>? IJob.Reports => Reports;
+			IReadOnlyList<string> IJob.Arguments => Arguments;
+			IReadOnlyList<int> IJob.Issues => ReferencedByIssues;
+			IReadOnlyDictionary<int, ObjectId> IJob.LabelIdxToTriggerId => LabelNotifications.ToDictionary(x => x.LabelIdx, x => x.TriggerId);
+			IReadOnlyList<IChainedJob> IJob.ChainedJobs => ChainedJobs;
+
+			[BsonConstructor]
+			private JobDocument()
+			{
+				Name = null!;
+			}
+
+			public JobDocument(ObjectId Id, StreamId StreamId, TemplateRefId TemplateId, ContentHash TemplateHash, ContentHash GraphHash, string Name, int Change, int CodeChange, int PreflightChange, int ClonedPreflightChange, ObjectId? StartedByUserId, string? StartedByUserName, Priority? Priority, bool? AutoSubmit, DateTime CreateTimeUtc, List<ChainedJobDocument> ChainedJobs, bool ShowUgsBadges, bool ShowUgsAlerts, string? NotificationChannel, string? NotificationChannelFilter, string? HelixSwarmCallbackUrl, List<string>? Arguments)
+			{
+				this.Id = Id;
+				this.StreamId = StreamId;
+				this.TemplateId = TemplateId;
+				this.TemplateHash = TemplateHash;
+				this.GraphHash = GraphHash;
+				this.Name = Name;
+				this.Change = Change;
+				this.CodeChange = CodeChange;
+				this.PreflightChange = PreflightChange;
+				this.ClonedPreflightChange = ClonedPreflightChange;
+				this.StartedByUserId = StartedByUserId;
+				this.StartedByUser = StartedByUserName;
+				this.Priority = Priority ?? HordeCommon.Priority.Normal;
+				this.AutoSubmit = AutoSubmit ?? false;
+				this.CreateTimeUtc = CreateTimeUtc;
+				this.ChainedJobs = ChainedJobs;
+				this.ShowUgsBadges = ShowUgsBadges;
+				this.ShowUgsAlerts = ShowUgsAlerts;
+				this.NotificationChannel = NotificationChannel;
+				this.NotificationChannelFilter = NotificationChannelFilter;
+				this.HelixSwarmCallbackUrl = HelixSwarmCallbackUrl;
+				this.Arguments = Arguments ?? this.Arguments;
+				this.NextSubResourceId = SubResourceId.Random();
+				this.UpdateTimeUtc = CreateTimeUtc;
+			}
+
+			public void PostLoad()
+			{
+				if (GraphHash == ContentHash.Empty)
+				{
+					Batches.Clear();
+				}
+			}
+		}
+
+		/// <summary>
+		/// Projection of a job definition to just include permissions info
+		/// </summary>
+		[SuppressMessage("Design", "CA1812: Class is never instantiated")]
+		class JobPermissions : IJobPermissions
+		{
+			public static ProjectionDefinition<JobDocument> Projection { get; } = Builders<JobDocument>.Projection.Include(x => x.Acl).Include(x => x.StreamId);
+
+			public Acl? Acl { get; set; }
+			public StreamId StreamId { get; set; }
+		}
+
+		/// <summary>
+		/// The jobs collection
+		/// </summary>
+		IMongoCollection<JobDocument> Jobs;
+
+		/// <summary>
+		/// Logger for output
+		/// </summary>
+		ILogger<JobCollection> Logger;
+
+		/// <summary>
+		/// Constructor
+		/// </summary>
+		/// <param name="DatabaseService">The database service singleton</param>
+		/// <param name="Logger">The logger instance</param>
+		public JobCollection(DatabaseService DatabaseService, ILogger<JobCollection> Logger)
+		{
+			this.Logger = Logger;
+
+			Jobs = DatabaseService.GetCollection<JobDocument>("Jobs");
+
+			if (!DatabaseService.ReadOnlyMode)
+			{
+				Jobs.Indexes.CreateOne(new CreateIndexModel<JobDocument>(Builders<JobDocument>.IndexKeys.Ascending(x => x.StreamId)));
+				Jobs.Indexes.CreateOne(new CreateIndexModel<JobDocument>(Builders<JobDocument>.IndexKeys.Ascending(x => x.Change)));
+				Jobs.Indexes.CreateOne(new CreateIndexModel<JobDocument>(Builders<JobDocument>.IndexKeys.Ascending(x => x.PreflightChange)));
+				Jobs.Indexes.CreateOne(new CreateIndexModel<JobDocument>(Builders<JobDocument>.IndexKeys.Ascending(x => x.CreateTimeUtc)));
+				Jobs.Indexes.CreateOne(new CreateIndexModel<JobDocument>(Builders<JobDocument>.IndexKeys.Ascending(x => x.UpdateTimeUtc)));
+				Jobs.Indexes.CreateOne(new CreateIndexModel<JobDocument>(Builders<JobDocument>.IndexKeys.Ascending(x => x.Name)));
+				Jobs.Indexes.CreateOne(new CreateIndexModel<JobDocument>(Builders<JobDocument>.IndexKeys.Ascending(x => x.StartedByUser)));
+				Jobs.Indexes.CreateOne(new CreateIndexModel<JobDocument>(Builders<JobDocument>.IndexKeys.Ascending(x => x.TemplateId)));
+				Jobs.Indexes.CreateOne(new CreateIndexModel<JobDocument>(Builders<JobDocument>.IndexKeys.Descending(x => x.SchedulePriority)));
+			}
+		}
+
+		/// <inheritdoc/>
+		[SuppressMessage("Compiler", "CA1054:URI parameters should not be strings")]
+		public async Task<IJob> AddAsync(ObjectId JobId, StreamId StreamId, TemplateRefId TemplateRefId, ContentHash TemplateHash, IGraph Graph, string Name, int Change, int CodeChange, int? PreflightChange, int? ClonedPreflightChange, ObjectId? StartedByUserId, string? StartedByUserName, Priority? Priority, bool? AutoSubmit, List<ChainedJobTemplate>? ChainedJobs, bool ShowUgsBadges, bool ShowUgsAlerts, string? NotificationChannel, string? NotificationChannelFilter, string? HelixSwarmCallbackUrl, List<string>? Arguments)
+		{
+			List<ChainedJobDocument> JobTriggers = new List<ChainedJobDocument>();
+			if (ChainedJobs == null)
+			{
+				JobTriggers = new List<ChainedJobDocument>();
+			}
+			else
+			{
+				JobTriggers = ChainedJobs.ConvertAll(x => new ChainedJobDocument(x));
+			}
+
+			JobDocument NewJob = new JobDocument(JobId, StreamId, TemplateRefId, TemplateHash, Graph.Id, Name, Change, CodeChange, PreflightChange ?? 0, ClonedPreflightChange ?? 0, StartedByUserId, StartedByUserName, Priority, AutoSubmit, DateTime.UtcNow, JobTriggers, ShowUgsBadges, ShowUgsAlerts, NotificationChannel, NotificationChannelFilter, HelixSwarmCallbackUrl, Arguments);
+			CreateBatches(NewJob, Graph, Logger);
+
+			await Jobs.InsertOneAsync(NewJob);
+
+			return NewJob;
+		}
+
+		/// <inheritdoc/>
+		public async Task<IJob?> GetAsync(ObjectId JobId)
+		{
+			JobDocument? Job = await Jobs.Find<JobDocument>(x => x.Id == JobId).FirstOrDefaultAsync();
+			if (Job != null)
+			{
+				Job.PostLoad();
+			}
+			return Job;
+		}
+
+		/// <inheritdoc/>
+		public async Task<bool> RemoveAsync(IJob Job)
+		{
+			DeleteResult Result = await Jobs.DeleteOneAsync(x => x.Id == Job.Id && x.UpdateIndex == Job.UpdateIndex);
+			return Result.DeletedCount > 0;
+		}
+
+		/// <inheritdoc/>
+		public async Task RemoveStreamAsync(StreamId StreamId)
+		{
+			await Jobs.DeleteManyAsync(x => x.StreamId == StreamId);
+		}
+
+		/// <inheritdoc/>
+		public async Task<IJobPermissions?> GetPermissionsAsync(ObjectId JobId)
+		{
+			return await Jobs.Find<JobDocument>(x => x.Id == JobId).Project<JobPermissions>(JobPermissions.Projection).FirstOrDefaultAsync();
+		}
+
+		/// <inheritdoc/>
+		public async Task<List<IJob>> FindAsync(ObjectId[]? JobIds, StreamId? StreamId, string? Name, TemplateRefId[]? Templates, int? MinChange, int? MaxChange, int? PreflightChange, string? PreflightStartedByUser, DateTimeOffset? MinCreateTime, DateTimeOffset? MaxCreateTime, DateTimeOffset? ModifiedBefore, DateTimeOffset? ModifiedAfter, int? Index, int? Count)
+		{
+			FilterDefinitionBuilder<JobDocument> FilterBuilder = Builders<JobDocument>.Filter;
+
+			FilterDefinition<JobDocument> Filter = FilterBuilder.Empty;
+			if (JobIds != null && JobIds.Length > 0)
+			{
+				Filter &= FilterBuilder.In(x => x.Id, JobIds);
+			}
+			if (StreamId != null)
+			{
+				Filter &= FilterBuilder.Eq(x => x.StreamId, StreamId.Value);
+			}
+			if (Name != null)
+			{
+				Filter &= FilterBuilder.Eq(x => x.Name, Name);
+			}
+			if (Templates != null)
+			{
+				Filter &= FilterBuilder.In(x => x.TemplateId, Templates);
+			}
+			if (MinChange != null)
+			{
+				Filter &= FilterBuilder.Gte(x => x.Change, MinChange);
+			}
+			if (MaxChange != null)
+			{
+				Filter &= FilterBuilder.Lte(x => x.Change, MaxChange);
+			}
+			if (PreflightChange != null)
+			{
+				Filter &= FilterBuilder.Eq(x => x.PreflightChange, PreflightChange);
+			}
+			if (PreflightStartedByUser != null)
+			{
+				Filter &= FilterBuilder.Or(FilterBuilder.Eq(x => x.PreflightChange, 0), FilterBuilder.Eq(x => x.StartedByUser, PreflightStartedByUser));
+			}
+			if (MinCreateTime != null)
+			{
+				Filter &= FilterBuilder.Gte(x => x.CreateTimeUtc!, MinCreateTime.Value.UtcDateTime);
+			}
+			if (MaxCreateTime != null)
+			{
+				Filter &= FilterBuilder.Lte(x => x.CreateTimeUtc!, MaxCreateTime.Value.UtcDateTime);
+			}
+			if (ModifiedBefore != null)
+			{
+				Filter &= FilterBuilder.Lte(x => x.UpdateTimeUtc!, ModifiedBefore.Value.UtcDateTime);
+			}
+			if (ModifiedAfter != null)
+			{
+				Filter &= FilterBuilder.Gte(x => x.UpdateTimeUtc!, ModifiedAfter.Value.UtcDateTime);
+			}
+
+			List<JobDocument> Results;
+			using (Scope Scope = Tracer.Instance.StartActive("Jobs.Find"))
+			{
+				IFindFluent<JobDocument, JobDocument> Query = Jobs.Find<JobDocument>(Filter).SortByDescending(x => x.CreateTimeUtc!);
+				if (Index != null)
+				{
+					Query = Query.Skip(Index.Value);
+				}
+				if (Count != null)
+				{
+					Query = Query.Limit(Count.Value);
+				}
+				Results = await Query.ToListAsync();
+			}
+			foreach (JobDocument Result in Results)
+			{
+				Result.PostLoad();
+			}
+			return Results.ConvertAll<JobDocument, IJob>(x => x);
+		}
+
+		/// <inheritdoc/>
+		public async Task<bool> TryUpdateJobAsync(IJob Job, IGraph Graph, string? Name, Priority? Priority, bool? AutoSubmit, int? AutoSubmitChange, string? AutoSubmitMessage, string? AbortedByUser, ObjectId? NotificationTriggerId, List<Report>? Reports, List<string>? Arguments, KeyValuePair<int, ObjectId>? LabelIdxToTriggerId, KeyValuePair<TemplateRefId, ObjectId>? JobTrigger)
+		{
+			// Create the update 
+			UpdateDefinitionBuilder<JobDocument> UpdateBuilder = Builders<JobDocument>.Update;
+			List<UpdateDefinition<JobDocument>> Updates = new List<UpdateDefinition<JobDocument>>();
+
+			// Flag for whether to update batches
+			bool bUpdateBatches = false;
+
+			// Build the update list
+			JobDocument JobDocument = (JobDocument)Job;
+			if (Name != null)
+			{
+				JobDocument.Name = Name;
+				Updates.Add(UpdateBuilder.Set(x => x.Name, Job.Name));
+			}
+			if (Priority != null)
+			{
+				JobDocument.Priority = Priority.Value;
+				Updates.Add(UpdateBuilder.Set(x => x.Priority, Job.Priority));
+			}
+			if (AutoSubmit != null)
+			{
+				JobDocument.AutoSubmit = AutoSubmit.Value;
+				Updates.Add(UpdateBuilder.Set(x => x.AutoSubmit, Job.AutoSubmit));
+			}
+			if (AutoSubmitChange != null)
+			{
+				JobDocument.AutoSubmitChange = AutoSubmitChange.Value;
+				Updates.Add(UpdateBuilder.Set(x => x.AutoSubmitChange, Job.AutoSubmitChange));
+			}
+			if (AutoSubmitMessage != null)
+			{
+				JobDocument.AutoSubmitMessage = (AutoSubmitMessage.Length == 0)? null : AutoSubmitMessage;
+				Updates.Add(UpdateBuilder.SetOrUnsetNullRef(x => x.AutoSubmitMessage, Job.AutoSubmitMessage));
+			}
+			if (AbortedByUser != null && Job.AbortedByUser == null)
+			{
+				JobDocument.AbortedByUser = AbortedByUser;
+				Updates.Add(UpdateBuilder.Set(x => x.AbortedByUser, Job.AbortedByUser));
+				bUpdateBatches = true;
+			}
+			if (NotificationTriggerId != null)
+			{
+				JobDocument.NotificationTriggerId = NotificationTriggerId.Value;
+				Updates.Add(UpdateBuilder.Set(x => x.NotificationTriggerId, NotificationTriggerId));
+			}
+			if (LabelIdxToTriggerId != null)
+			{
+				if (JobDocument.LabelNotifications.Any(x => x.LabelIdx == LabelIdxToTriggerId.Value.Key))
+				{
+					throw new ArgumentException("Cannot update label trigger that already exists");
+				}
+				JobDocument.LabelNotifications.Add(new LabelNotificationDocument { LabelIdx = LabelIdxToTriggerId.Value.Key, TriggerId = LabelIdxToTriggerId.Value.Value });
+				Updates.Add(UpdateBuilder.Set(x => x.LabelNotifications, JobDocument.LabelNotifications));
+			}
+			if (JobTrigger != null)
+			{
+				for (int Idx = 0; Idx < JobDocument.ChainedJobs.Count; Idx++)
+				{
+					ChainedJobDocument JobTriggerDocument = JobDocument.ChainedJobs[Idx];
+					if (JobTriggerDocument.TemplateRefId == JobTrigger.Value.Key)
+					{
+						int LocalIdx = Idx;
+						JobTriggerDocument.JobId = JobTrigger.Value.Value;
+						Updates.Add(UpdateBuilder.Set(x => x.ChainedJobs[LocalIdx].JobId, JobTrigger.Value.Value));
+					}
+				}
+			}
+			if (Reports != null)
+			{
+				JobDocument.Reports ??= new List<Report>();
+				JobDocument.Reports.RemoveAll(x => Reports.Any(y => y.Name == x.Name));
+				JobDocument.Reports.AddRange(Reports);
+				Updates.Add(UpdateBuilder.Set(x => x.Reports, JobDocument.Reports));
+			}
+			if (Arguments != null)
+			{
+				HashSet<string> ModifiedArguments = new HashSet<string>(Job.Arguments);
+				ModifiedArguments.SymmetricExceptWith(Arguments);
+
+				foreach (string ModifiedArgument in ModifiedArguments)
+				{
+					if (ModifiedArgument.StartsWith(IJob.TargetArgumentPrefix, StringComparison.OrdinalIgnoreCase))
+					{
+						bUpdateBatches = true;
+					}
+				}
+
+				JobDocument.Arguments = Arguments.ToList();
+				Updates.Add(UpdateBuilder.Set(x => x.Arguments, JobDocument.Arguments));
+			}
+
+			// Update the batches
+			if (bUpdateBatches)
+			{
+				UpdateBatches(JobDocument, Graph, Updates, Logger);
+			}
+
+			// Update the new list of job steps
+			return (Updates.Count == 0 || await TryUpdateAsync(JobDocument, UpdateBuilder.Combine(Updates)));
+		}
+
+		/// <inheritdoc/>
+		public async Task<bool> TryUpdateBatchAsync(IJob Job, IGraph Graph, SubResourceId BatchId, ObjectId? NewLogId, JobStepBatchState? NewState, JobStepBatchError? NewError)
+		{
+			JobDocument JobDocument = (JobDocument)Job;
+
+			// Find the index of the appropriate batch
+			int BatchIdx = JobDocument.Batches.FindIndex(x => x.Id == BatchId);
+			if (BatchIdx == -1)
+			{
+				return false;
+			}
+
+			// If we're marking the batch as complete and there are still steps to run (eg. because the agent crashed), we need to mark all the steps as complete first
+			JobStepBatchDocument Batch = JobDocument.Batches[BatchIdx];
+
+			// Create the update 
+			UpdateDefinitionBuilder<JobDocument> UpdateBuilder = Builders<JobDocument>.Update;
+			List<UpdateDefinition<JobDocument>> Updates = new List<UpdateDefinition<JobDocument>>();
+
+			// Update the batch
+			if (NewLogId != null)
+			{
+				Batch.LogId = NewLogId.Value;
+				Updates.Add(UpdateBuilder.Set(x => x.Batches[BatchIdx].LogId, Batch.LogId));
+			}
+			if (NewState != null)
+			{
+				Batch.State = NewState.Value;
+				Updates.Add(UpdateBuilder.Set(x => x.Batches[BatchIdx].State, Batch.State));
+
+				if (NewState == JobStepBatchState.Running)
+				{
+					Batch.StartTime = DateTimeOffset.Now;
+					Updates.Add(UpdateBuilder.Set(x => x.Batches[BatchIdx].StartTime, Batch.StartTime));
+				}
+				if (NewState == JobStepBatchState.Complete)
+				{
+					Batch.FinishTime = DateTimeOffset.Now;
+					Updates.Add(UpdateBuilder.Set(x => x.Batches[BatchIdx].FinishTime, Batch.FinishTime));
+				}
+			}
+			if (NewError != null)
+			{
+				Batch.Error = NewError.Value;
+				Updates.Add(UpdateBuilder.Set(x => x.Batches[BatchIdx].Error, Batch.Error));
+			}
+
+			// If the batch is being marked as incomplete, see if we can reschedule any of the work
+			if (NewError == JobStepBatchError.Incomplete)
+			{
+				// Get the new list of retried nodes
+				List<NodeRef> RetriedNodes = JobDocument.RetriedNodes ?? new List<NodeRef>();
+
+				// Check if there are any steps that need to be run again
+				int InitialRetriedNodes = RetriedNodes.Count;
+				foreach (JobStepDocument Step in Batch.Steps)
+				{
+					if (Step.State == JobStepState.Ready || Step.State == JobStepState.Waiting)
+					{
+						NodeRef NodeRef = new NodeRef(Batch.GroupIdx, Step.NodeIdx);
+						if (JobDocument.RetriedNodes != null && JobDocument.RetriedNodes.Contains(NodeRef))
+						{
+							Step.State = JobStepState.Skipped;
+						}
+						else
+						{
+							RetriedNodes.Add(NodeRef);
+						}
+					}
+				}
+
+				// Update the steps
+				if (RetriedNodes.Count > InitialRetriedNodes)
+				{
+					Updates.Clear();
+					UpdateBatches(JobDocument, Graph, Updates, Logger);
+
+					JobDocument.RetriedNodes = RetriedNodes;
+					Updates.Add(UpdateBuilder.Set(x => x.RetriedNodes, JobDocument.RetriedNodes));
+				}
+			}
+
+			// Update the new list of job steps
+			return (Updates.Count == 0 || await TryUpdateAsync(JobDocument, UpdateBuilder.Combine(Updates)));
+		}
+
+		/// <inheritdoc/>
+		public async Task<bool> TryUpdateStepAsync(IJob Job, IGraph Graph, SubResourceId BatchId, SubResourceId StepId, JobStepState NewState, JobStepOutcome NewOutcome, bool? NewAbortRequested, string? NewAbortByUser, ObjectId? NewLogId, ObjectId? NewNotificationTriggerId, string? NewRetryByUser, Priority? NewPriority, List<Report>? NewReports, Dictionary<string, string?>? NewProperties)
+		{
+			JobDocument JobDocument = (JobDocument)Job;
+
+			// Create the update 
+			UpdateDefinitionBuilder<JobDocument> UpdateBuilder = Builders<JobDocument>.Update;
+			List<UpdateDefinition<JobDocument>> Updates = new List<UpdateDefinition<JobDocument>>();
+
+			// Update the appropriate batch
+			bool bRefreshBatches = false;
+			bool bRefreshDependentJobSteps = false;
+			for (int LoopBatchIdx = 0; LoopBatchIdx < JobDocument.Batches.Count; LoopBatchIdx++)
+			{
+				int BatchIdx = LoopBatchIdx; // For lambda capture
+				JobStepBatchDocument Batch = JobDocument.Batches[BatchIdx];
+				if (Batch.Id == BatchId)
+				{
+					for (int LoopStepIdx = 0; LoopStepIdx < Batch.Steps.Count; LoopStepIdx++)
+					{
+						int StepIdx = LoopStepIdx; // For lambda capture
+						JobStepDocument Step = Batch.Steps[StepIdx];
+						if (Step.Id == StepId)
+						{
+							// Update the state
+							if (NewState != JobStepState.Unspecified && Step.State != NewState)
+							{
+								Step.State = NewState;
+								Updates.Add(UpdateBuilder.Set(x => x.Batches[BatchIdx].Steps[StepIdx].State, Step.State));
+
+								if (Step.State == JobStepState.Running)
+								{
+									Step.StartTime = DateTimeOffset.Now;
+									Updates.Add(UpdateBuilder.Set(x => x.Batches[BatchIdx].Steps[StepIdx].StartTime, Step.StartTime));
+								}
+								else if (Step.State == JobStepState.Completed || Step.State == JobStepState.Aborted)
+								{
+									Step.FinishTime = DateTimeOffset.Now;
+									Updates.Add(UpdateBuilder.Set(x => x.Batches[BatchIdx].Steps[StepIdx].FinishTime, Step.FinishTime));
+								}
+
+								bRefreshDependentJobSteps = true;
+							}
+
+							// Update the job outcome
+							if (NewOutcome != JobStepOutcome.Unspecified && Step.Outcome != NewOutcome)
+							{
+								Step.Outcome = NewOutcome;
+								Updates.Add(UpdateBuilder.Set(x => x.Batches[BatchIdx].Steps[StepIdx].Outcome, Step.Outcome));
+
+								bRefreshDependentJobSteps = true;
+							}
+
+							// Update the request abort status
+							if (NewAbortRequested != null && Step.AbortRequested == false)
+							{
+								Step.AbortRequested = NewAbortRequested.Value;
+								Updates.Add(UpdateBuilder.Set(x => x.Batches[BatchIdx].Steps[StepIdx].AbortRequested, Step.AbortRequested));
+
+								bRefreshDependentJobSteps = true;
+							}
+
+							// Update the user that requested the abort
+							if (NewAbortByUser != null && Step.AbortByUser == null)
+							{
+								Step.AbortByUser = NewAbortByUser;
+								Updates.Add(UpdateBuilder.Set(x => x.Batches[BatchIdx].Steps[StepIdx].AbortByUser, Step.AbortByUser));
+
+								bRefreshDependentJobSteps = true;
+							}
+
+							// Update the log id
+							if (NewLogId != null && Step.LogId != NewLogId.Value)
+							{
+								Step.LogId = NewLogId.Value;
+								Updates.Add(UpdateBuilder.Set(x => x.Batches[BatchIdx].Steps[StepIdx].LogId, Step.LogId));
+
+								bRefreshDependentJobSteps = true;
+							}
+
+							// Update the notification trigger id
+							if (NewNotificationTriggerId != null && Step.NotificationTriggerId == null)
+							{
+								Step.NotificationTriggerId = NewNotificationTriggerId.Value;
+								Updates.Add(UpdateBuilder.Set(x => x.Batches[BatchIdx].Steps[StepIdx].NotificationTriggerId, Step.NotificationTriggerId));
+							}
+
+							// Update the retry flag
+							if (NewRetryByUser != null && Step.RetryByUser == null)
+							{
+								Step.RetryByUser = NewRetryByUser;
+								Updates.Add(UpdateBuilder.Set(x => x.Batches[BatchIdx].Steps[StepIdx].RetryByUser, Step.RetryByUser));
+
+								bRefreshBatches = true;
+							}
+
+							// Update the priority
+							if (NewPriority != null && NewPriority.Value != Step.Priority)
+							{
+								Step.Priority = NewPriority.Value;
+								Updates.Add(UpdateBuilder.Set(x => x.Batches[BatchIdx].Steps[StepIdx].Priority, Step.Priority));
+
+								bRefreshBatches = true;
+							}
+
+							// Add any new reports
+							if (NewReports != null)
+							{
+								Step.Reports ??= new List<Report>();
+								Step.Reports.RemoveAll(x => NewReports.Any(y => y.Name == x.Name));
+								Step.Reports.AddRange(NewReports);
+								Updates.Add(UpdateBuilder.Set(x => x.Batches[BatchIdx].Steps[StepIdx].Reports, Step.Reports));
+							}
+
+							// Apply any property updates
+							if (NewProperties != null)
+							{
+								if (Step.Properties == null)
+								{
+									Step.Properties = new Dictionary<string, string>(StringComparer.Ordinal);
+								}
+
+								foreach (KeyValuePair<string, string?> Pair in NewProperties)
+								{
+									if (Pair.Value == null)
+									{
+										Step.Properties.Remove(Pair.Key);
+									}
+									else
+									{
+										Step.Properties[Pair.Key] = Pair.Value;
+									}
+								}
+
+								if (Step.Properties.Count == 0)
+								{
+									Step.Properties = null;
+									Updates.Add(UpdateBuilder.Unset(x => x.Batches[BatchIdx].Steps[StepIdx].Properties));
+								}
+								else
+								{
+									foreach (KeyValuePair<string, string?> Pair in NewProperties)
+									{
+										if (Pair.Value == null)
+										{
+											Updates.Add(UpdateBuilder.Unset(x => x.Batches[BatchIdx].Steps[StepIdx].Properties![Pair.Key]));
+										}
+										else
+										{
+											Updates.Add(UpdateBuilder.Set(x => x.Batches[BatchIdx].Steps[StepIdx].Properties![Pair.Key], Pair.Value));
+										}
+									}
+								}
+							}
+							break;
+						}
+					}
+					break;
+				}
+			}
+
+			// Update the batches
+			if (bRefreshBatches)
+			{
+				Updates.Clear(); // UpdateBatches will update the entire batches list. We need to remove all the individual step updates to avoid an exception.
+				UpdateBatches(JobDocument, Graph, Updates, Logger);
+			}
+
+			// Update the state of dependent jobsteps
+			if (bRefreshDependentJobSteps)
+			{
+				RefreshDependentJobSteps(JobDocument, Graph, Updates, Logger);
+			}
+
+			// Update the new list of job steps
+			return (Updates.Count == 0 || await TryUpdateAsync(JobDocument, UpdateBuilder.Combine(Updates)));
+		}
+
+		/// <inheritdoc/>
+		async Task<bool> TryUpdateAsync(JobDocument Job, UpdateDefinition<JobDocument> Update)
+		{
+			int NewUpdateIndex = Job.UpdateIndex + 1;
+			Update = Update.Set(x => x.UpdateIndex, NewUpdateIndex);
+
+			DateTime NewUpdateTimeUtc = DateTime.UtcNow;
+			Update = Update.Set(x => x.UpdateTimeUtc, NewUpdateTimeUtc);
+
+			UpdateResult Result = await Jobs.UpdateOneAsync<JobDocument>(x => x.Id == Job.Id && x.UpdateIndex == Job.UpdateIndex, Update);
+			if (Result.ModifiedCount > 0)
+			{
+				Job.UpdateIndex = NewUpdateIndex;
+				Job.UpdateTimeUtc = NewUpdateTimeUtc;
+				return true;
+			}
+			return false;
+		}
+
+		/// <inheritdoc/>
+		public Task<bool> TryUpdateGraphAsync(IJob Job, IGraph NewGraph)
+		{
+			JobDocument JobDocument = (JobDocument)Job;
+
+			// Create the update 
+			UpdateDefinitionBuilder<JobDocument> UpdateBuilder = Builders<JobDocument>.Update;
+			List<UpdateDefinition<JobDocument>> Updates = new List<UpdateDefinition<JobDocument>>();
+
+			JobDocument.GraphHash = NewGraph.Id;
+			Updates.Add(UpdateBuilder.Set(x => x.GraphHash, Job.GraphHash));
+
+			UpdateBatches(JobDocument, NewGraph, Updates, Logger);
+
+			// Update the new list of job steps
+			return TryUpdateAsync(JobDocument, UpdateBuilder.Combine(Updates));
+		}
+
+		/// <inheritdoc/>
+		public async Task AddIssueToJobAsync(ObjectId JobId, int IssueId)
+		{
+			FilterDefinition<JobDocument> JobFilter = Builders<JobDocument>.Filter.Eq(x => x.Id, JobId);
+			UpdateDefinition<JobDocument> JobUpdate = Builders<JobDocument>.Update.AddToSet(x => x.ReferencedByIssues, IssueId).Inc(x => x.UpdateIndex, 1).Max(x => x.UpdateTimeUtc, DateTime.UtcNow);
+			await Jobs.UpdateOneAsync(JobFilter, JobUpdate);
+		}
+
+		/// <inheritdoc/>
+		public async Task<List<IJob>> GetDispatchQueueAsync()
+		{
+			List<JobDocument> NewJobs = await Jobs.Find(x => x.SchedulePriority > 0).SortByDescending(x => x.SchedulePriority).ToListAsync();
+			foreach (JobDocument Result in NewJobs)
+			{
+				Result.PostLoad();
+			}
+			return NewJobs.ConvertAll<JobDocument, IJob>(x => x);
+		}
+
+		/// <summary>
+		/// Marks a job as skipped
+		/// </summary>
+		/// <param name="Job">The job to update</param>
+		/// <param name="Graph">Graph for the job</param>
+		/// <param name="Reason">Reason for this batch being failed</param>
+		/// <returns>Updated version of the job</returns>
+		public async Task<IJob?> SkipAllBatchesAsync(IJob? Job, IGraph Graph, JobStepBatchError Reason)
+		{
+			while (Job != null)
+			{
+				JobDocument JobDocument = (JobDocument)Job;
+
+				for (int BatchIdx = 0; BatchIdx < JobDocument.Batches.Count; BatchIdx++)
+				{
+					JobStepBatchDocument Batch = JobDocument.Batches[BatchIdx];
+					if (Batch.State == JobStepBatchState.Ready || Batch.State == JobStepBatchState.Waiting)
+					{
+						Batch.State = JobStepBatchState.Complete;
+						Batch.Error = Reason;
+						Batch.FinishTime = DateTimeOffset.UtcNow;
+
+						for (int StepIdx = 0; StepIdx < Batch.Steps.Count; StepIdx++)
+						{
+							JobStepDocument Step = Batch.Steps[StepIdx];
+							if (Step.State == JobStepState.Ready || Step.State == JobStepState.Waiting)
+							{
+								Step.State = JobStepState.Completed;
+								Step.Outcome = JobStepOutcome.Failure;
+							}
+						}
+					}
+				}
+
+				List<UpdateDefinition<JobDocument>> Updates = new List<UpdateDefinition<JobDocument>>();
+				UpdateBatches(JobDocument, Graph, Updates, Logger);
+
+				if (await TryUpdateAsync(JobDocument, Builders<JobDocument>.Update.Combine(Updates)))
+				{
+					break;
+				}
+
+				Job = await GetAsync(Job.Id);
+			}
+			return Job;
+		}
+
+		/// <inheritdoc/>
+		public async Task<IJob?> SkipBatchAsync(IJob? Job, int BatchIdx, IGraph Graph, JobStepBatchError Reason)
+		{
+			while (Job != null)
+			{
+				JobDocument JobDocument = (JobDocument)Job;
+				JobStepBatchDocument Batch = JobDocument.Batches[BatchIdx];
+
+				Batch.State = JobStepBatchState.Complete;
+				Batch.Error = Reason;
+				Batch.FinishTime = DateTimeOffset.UtcNow;
+
+				foreach (JobStepDocument Step in Batch.Steps)
+				{
+					if (Step.State != JobStepState.Skipped)
+					{
+						Step.State = JobStepState.Skipped;
+						Step.Outcome = JobStepOutcome.Failure;
+					}
+				}
+
+				List<UpdateDefinition<JobDocument>> Updates = new List<UpdateDefinition<JobDocument>>();
+				UpdateBatches(JobDocument, Graph, Updates, Logger);
+
+				if (await TryUpdateAsync(JobDocument, Builders<JobDocument>.Update.Combine(Updates)))
+				{
+					break;
+				}
+
+				Job = await GetAsync(Job.Id);
+			}
+			return Job;
+		}
+
+		/// <inheritdoc/>
+		public async Task<bool> TryAssignLeaseAsync(IJob Job, int BatchIdx, PoolId PoolId, AgentId AgentId, ObjectId SessionId, ObjectId LeaseId, ObjectId LogId)
+		{
+			// Try to update the job with this agent id
+			UpdateDefinitionBuilder<JobDocument> UpdateBuilder = Builders<JobDocument>.Update;
+			List<UpdateDefinition<JobDocument>> Updates = new List<UpdateDefinition<JobDocument>>();
+
+			Updates.Add(UpdateBuilder.Set(x => x.Batches[BatchIdx].PoolId, PoolId));
+			Updates.Add(UpdateBuilder.Set(x => x.Batches[BatchIdx].AgentId, AgentId));
+			Updates.Add(UpdateBuilder.Set(x => x.Batches[BatchIdx].SessionId, SessionId));
+			Updates.Add(UpdateBuilder.Set(x => x.Batches[BatchIdx].LeaseId, LeaseId));
+			Updates.Add(UpdateBuilder.Set(x => x.Batches[BatchIdx].LogId, LogId));
+
+			// Extra logs for catching why session IDs sometimes doesn't match in prod. Resulting in PermissionDenied errors.
+			if (BatchIdx < Job.Batches.Count && Job.Batches[BatchIdx].SessionId != null)
+			{
+				string CurrentSessionId = Job.Batches[BatchIdx].SessionId!.Value.ToString();
+				Logger.LogError("Attempt to replace current session ID {CurrSessionId} with {NewSessionId} for batch {JobId}:{BatchId}", CurrentSessionId, SessionId.ToString(), Job.Id.ToString(), Job.Batches[BatchIdx].Id);
+				return false;
+			}
+
+			JobDocument JobDocument = (JobDocument)Job;
+			if (!await TryUpdateAsync(JobDocument, UpdateBuilder.Combine(Updates)))
+			{
+				return false;
+			}
+
+			JobStepBatchDocument Batch = JobDocument.Batches[BatchIdx];
+			Batch.AgentId = AgentId;
+			Batch.SessionId = SessionId;
+			Batch.LeaseId = LeaseId;
+			Batch.LogId = LogId;
+
+			return true;
+		}
+
+		/// <inheritdoc/>
+		public async Task<bool> TryCancelLeaseAsync(IJob Job, int BatchIdx)
+		{
+			JobDocument JobDocument = (JobDocument)Job;
+
+			UpdateDefinition<JobDocument> Update = Builders<JobDocument>.Update.Unset(x => x.Batches[BatchIdx].AgentId).Unset(x => x.Batches[BatchIdx].SessionId).Unset(x => x.Batches[BatchIdx].LeaseId);
+			if (await TryUpdateAsync(JobDocument, Update))
+			{
+				JobStepBatchDocument Batch = JobDocument.Batches[BatchIdx];
+				Batch.AgentId = null;
+				Batch.SessionId = null;
+				Batch.LeaseId = null;
+				return true;
+			}
+			return false;
+		}
+
+		/// <inheritdoc/>
+		public async Task<bool> TryFailBatchAsync(IJob Job, int BatchIdx, IGraph Graph)
+		{
+			JobDocument JobDocument = (JobDocument)Job;
+			JobStepBatchDocument Batch = JobDocument.Batches[BatchIdx];
+
+			UpdateDefinitionBuilder<JobDocument> UpdateBuilder = new UpdateDefinitionBuilder<JobDocument>();
+			List<UpdateDefinition<JobDocument>> Updates = new List<UpdateDefinition<JobDocument>>();
+
+			if (Batch.State != JobStepBatchState.Complete)
+			{
+				Batch.State = JobStepBatchState.Complete;
+				Updates.Add(UpdateBuilder.Set(x => x.Batches[BatchIdx].State, Batch.State));
+
+				Batch.Error = JobStepBatchError.LostConnection;
+				Updates.Add(UpdateBuilder.Set(x => x.Batches[BatchIdx].Error, Batch.Error));
+
+				Batch.FinishTime = DateTimeOffset.UtcNow;
+				Updates.Add(UpdateBuilder.Set(x => x.Batches[BatchIdx].FinishTime, Batch.FinishTime));
+			}
+
+			for (int StepIdx = 0; StepIdx < Batch.Steps.Count; StepIdx++)
+			{
+				JobStepDocument Step = Batch.Steps[StepIdx];
+				if (Step.State == JobStepState.Running)
+				{
+					int StepIdxCopy = StepIdx;
+
+					Step.State = JobStepState.Aborted;
+					Updates.Add(UpdateBuilder.Set(x => x.Batches[BatchIdx].Steps[StepIdxCopy].State, Step.State));
+
+					Step.Outcome = JobStepOutcome.Failure;
+					Updates.Add(UpdateBuilder.Set(x => x.Batches[BatchIdx].Steps[StepIdxCopy].Outcome, Step.Outcome));
+
+					Step.FinishTime = DateTimeOffset.Now;
+					Updates.Add(UpdateBuilder.Set(x => x.Batches[BatchIdx].Steps[StepIdxCopy].FinishTime, Step.FinishTime));
+				}
+				else if (Step.State == JobStepState.Ready)
+				{
+					int StepIdxCopy = StepIdx;
+
+					Step.State = JobStepState.Skipped;
+					Updates.Add(UpdateBuilder.Set(x => x.Batches[BatchIdx].Steps[StepIdxCopy].State, Step.State));
+
+					Step.Outcome = JobStepOutcome.Failure;
+					Updates.Add(UpdateBuilder.Set(x => x.Batches[BatchIdx].Steps[StepIdxCopy].Outcome, Step.Outcome));
+				}
+			}
+
+			RefreshDependentJobSteps(JobDocument, Graph, Updates, Logger);
+
+			return (Updates.Count == 0 || await TryUpdateAsync(JobDocument, UpdateBuilder.Combine(Updates)));
+		}
+
+		/// <summary>
+		/// Populate the list of batches for a job
+		/// </summary>
+		/// <param name="Job">The job to update</param>
+		/// <param name="Graph">The graph definition for this job</param>
+		/// <param name="Logger">Logger for output messages</param>
+		private static void CreateBatches(JobDocument Job, IGraph Graph, ILogger Logger)
+		{
+			UpdateBatches(Job, Graph, new List<UpdateDefinition<JobDocument>>(), Logger);
+		}
+
+		/// <summary>
+		/// Updates the list of batches for a job
+		/// </summary>
+		/// <param name="Job">Job to update</param>
+		/// <param name="Graph">Graph definition for this job</param>
+		/// <param name="Updates">List of updates for the job</param>
+		/// <param name="Logger">Logger for updates</param>
+		static void UpdateBatches(JobDocument Job, IGraph Graph, List<UpdateDefinition<JobDocument>> Updates, ILogger Logger)
+		{
+			UpdateDefinitionBuilder<JobDocument> UpdateBuilder = Builders<JobDocument>.Update;
+
+			// Update the list of batches
+			CreateOrUpdateBatches(Job, Graph);
+			Updates.Add(UpdateBuilder.Set(x => x.Batches, Job.Batches));
+			Updates.Add(UpdateBuilder.Set(x => x.NextSubResourceId, Job.NextSubResourceId));
+
+			// Update all the dependencies
+			RefreshDependentJobSteps(Job, Graph, new List<UpdateDefinition<JobDocument>>(), Logger);
+		}
+
+		/// <summary>
+		/// Update the jobsteps for the given node graph 
+		/// </summary>
+		/// <param name="Job">The job to update</param>
+		/// <param name="Graph">The graph for this job</param>
+		private static void CreateOrUpdateBatches(JobDocument Job, IGraph Graph)
+		{
+			// Find the priorities of each node, incorporating all the per-step overrides
+			Dictionary<INode, Priority> NodePriorities = new Dictionary<INode, Priority>();
+			foreach (INodeGroup Group in Graph.Groups)
+			{
+				foreach (INode Node in Group.Nodes)
+				{
+					NodePriorities[Node] = Node.Priority;
+				}
+			}
+			foreach (JobStepBatchDocument Batch in Job.Batches)
+			{
+				INodeGroup Group = Graph.Groups[Batch.GroupIdx];
+				foreach (JobStepDocument Step in Batch.Steps)
+				{
+					if (Step.Priority != null)
+					{
+						INode Node = Group.Nodes[Step.NodeIdx];
+						NodePriorities[Node] = Step.Priority.Value;
+					}
+				}
+			}
+
+			// Remove any steps and batches that haven't started yet
+			foreach (JobStepBatchDocument Batch in Job.Batches)
+			{
+				Batch.Steps.RemoveAll(x => x.State == JobStepState.Waiting || x.State == JobStepState.Ready);
+			}
+
+			// Remove any skipped nodes whose skipped state is no longer valid
+			HashSet<INode> FailedNodes = new HashSet<INode>();
+			foreach (JobStepBatchDocument Batch in Job.Batches)
+			{
+				INodeGroup Group = Graph.Groups[Batch.GroupIdx];
+				foreach (JobStepDocument Step in Batch.Steps)
+				{
+					INode Node = Group.Nodes[Step.NodeIdx];
+					if (Step.Retry)
+					{
+						FailedNodes.Remove(Node);
+					}
+					else if (Step.State == JobStepState.Skipped && Node.InputDependencies.Any(x => FailedNodes.Contains(Graph.GetNode(x))))
+					{
+						FailedNodes.Add(Node);
+					}
+					else if (Step.Outcome == JobStepOutcome.Failure)
+					{
+						FailedNodes.Add(Node);
+					}
+					else
+					{
+						FailedNodes.Remove(Node);
+					}
+				}
+				Batch.Steps.RemoveAll(x => x.State == JobStepState.Skipped && !FailedNodes.Contains(Group.Nodes[x.NodeIdx]));
+			}
+
+			// Remove any batches which are now empty
+			Job.Batches.RemoveAll(x => x.Steps.Count == 0 && x.Error == JobStepBatchError.None);
+
+			// Find all the targets in this job
+			HashSet<string> Targets = new HashSet<string>(StringComparer.OrdinalIgnoreCase);
+			if (Job.AbortedByUser == null)
+			{
+				foreach (string Argument in Job.Arguments)
+				{
+					if (Argument.StartsWith(IJob.TargetArgumentPrefix, StringComparison.OrdinalIgnoreCase))
+					{
+						Targets.UnionWith(Argument.Substring(IJob.TargetArgumentPrefix.Length).Split(';'));
+					}
+				}
+				Targets.Add(IJob.SetupNodeName);
+			}
+
+			// Add all the referenced aggregates
+			HashSet<INode> NewNodesToExecute = new HashSet<INode>();
+			foreach (IAggregate Aggregate in Graph.Aggregates)
+			{
+				if (Targets.Contains(Aggregate.Name))
+				{
+					NewNodesToExecute.UnionWith(Aggregate.Nodes.Select(x => Graph.GetNode(x)));
+				}
+			}
+
+			// Add any individual nodes
+			foreach (INode Node in Graph.Groups.SelectMany(x => x.Nodes))
+			{
+				if (Targets.Contains(Node.Name))
+				{
+					NewNodesToExecute.Add(Node);
+				}
+			}
+
+			// Also add any dependencies of these nodes
+			for (int GroupIdx = Graph.Groups.Count - 1; GroupIdx >= 0; GroupIdx--)
+			{
+				INodeGroup Group = Graph.Groups[GroupIdx];
+				for (int NodeIdx = Group.Nodes.Count - 1; NodeIdx >= 0; NodeIdx--)
+				{
+					INode Node = Group.Nodes[NodeIdx];
+					if (NewNodesToExecute.Contains(Node))
+					{
+						foreach (NodeRef Dependency in Node.InputDependencies)
+						{
+							NewNodesToExecute.Add(Graph.GetNode(Dependency));
+						}
+					}
+				}
+			}
+
+			// Cancel any batches which are still running but are no longer required
+			foreach (JobStepBatchDocument Batch in Job.Batches)
+			{
+				if (Batch.State == JobStepBatchState.Running)
+				{
+					INodeGroup Group = Graph.Groups[Batch.GroupIdx];
+					if (!Batch.Steps.Any(x => NewNodesToExecute.Contains(Group.Nodes[x.NodeIdx])))
+					{
+						Batch.Error = JobStepBatchError.Cancelled;
+					}
+				}
+			}
+
+			// Remove all the nodes which have already succeeded
+			foreach (JobStepBatchDocument Batch in Job.Batches)
+			{
+				foreach (IJobStep Step in Batch.Steps)
+				{
+					if ((Step.State == JobStepState.Running && Step.RetryByUser == null) || (Step.State == JobStepState.Completed && Step.RetryByUser == null) || (Step.State == JobStepState.Aborted && Step.RetryByUser == null) || Step.State == JobStepState.Skipped)
+					{
+						NewNodesToExecute.Remove(Graph.Groups[Batch.GroupIdx].Nodes[Step.NodeIdx]);
+					}
+				}
+			}
+
+			// Re-add all the nodes that have input dependencies in the same group.
+			for (int GroupIdx = Graph.Groups.Count - 1; GroupIdx >= 0; GroupIdx--)
+			{
+				INodeGroup Group = Graph.Groups[GroupIdx];
+				for (int NodeIdx = Group.Nodes.Count - 1; NodeIdx >= 0; NodeIdx--)
+				{
+					INode Node = Group.Nodes[NodeIdx];
+					if (NewNodesToExecute.Contains(Node))
+					{
+						foreach (NodeRef Dependency in Node.InputDependencies)
+						{
+							if (Dependency.GroupIdx == GroupIdx)
+							{
+								NewNodesToExecute.Add(Node);
+							}
+						}
+					}
+				}
+			}
+
+			// Build a list of nodes which are currently set to be executed
+			HashSet<INode> ExistingNodesToExecute = new HashSet<INode>();
+			foreach (JobStepBatchDocument Batch in Job.Batches)
+			{
+				INodeGroup Group = Graph.Groups[Batch.GroupIdx];
+				foreach (IJobStep Step in Batch.Steps)
+				{
+					INode Node = Group.Nodes[Step.NodeIdx];
+					ExistingNodesToExecute.Add(Node);
+				}
+			}
+
+			// Figure out the existing batch for each group
+			JobStepBatchDocument?[] AppendToBatches = new JobStepBatchDocument?[Graph.Groups.Count];
+			foreach (JobStepBatchDocument Batch in Job.Batches)
+			{
+				if (Batch.CanBeAppendedTo())
+				{
+					INodeGroup Group = Graph.Groups[Batch.GroupIdx];
+					INode FirstNode = Group.Nodes[Batch.Steps[0].NodeIdx];
+					AppendToBatches[Batch.GroupIdx] = Batch;
+				}
+			}
+
+			// Invalidate all the entries for groups where we're too late to append new entries (ie. we need to execute an earlier node that wasn't executed previously)
+			for (int GroupIdx = 0; GroupIdx < Graph.Groups.Count; GroupIdx++)
+			{
+				INodeGroup Group = Graph.Groups[GroupIdx];
+				for (int NodeIdx = 0; NodeIdx < Group.Nodes.Count; NodeIdx++)
+				{
+					INode Node = Group.Nodes[NodeIdx];
+					if (NewNodesToExecute.Contains(Node) && !ExistingNodesToExecute.Contains(Node))
+					{
+						IJobStepBatch? Batch = AppendToBatches[GroupIdx];
+						if (Batch != null)
+						{
+							IJobStep LastStep = Batch.Steps[Batch.Steps.Count - 1];
+							if (NodeIdx < LastStep.NodeIdx)
+							{
+								AppendToBatches[GroupIdx] = null;
+							}
+						}
+					}
+				}
+			}
+
+			// Create all the new jobsteps
+			for (int GroupIdx = 0; GroupIdx < Graph.Groups.Count; GroupIdx++)
+			{
+				INodeGroup Group = Graph.Groups[GroupIdx];
+				for (int NodeIdx = 0; NodeIdx < Group.Nodes.Count; NodeIdx++)
+				{
+					INode Node = Group.Nodes[NodeIdx];
+					if (NewNodesToExecute.Contains(Node))
+					{
+						JobStepBatchDocument? Batch = AppendToBatches[GroupIdx];
+						if (Batch == null)
+						{
+							Job.NextSubResourceId = Job.NextSubResourceId.Next();
+
+							Batch = new JobStepBatchDocument(Job.NextSubResourceId, GroupIdx);
+							Job.Batches.Add(Batch);
+
+							AppendToBatches[GroupIdx] = Batch;
+						}
+						if (Batch.Steps.Count == 0 || NodeIdx > Batch.Steps[Batch.Steps.Count - 1].NodeIdx)
+						{
+							Job.NextSubResourceId = Job.NextSubResourceId.Next();
+
+							JobStepDocument Step = new JobStepDocument(Job.NextSubResourceId, NodeIdx);
+							Batch.Steps.Add(Step);
+						}
+					}
+				}
+			}
+
+			// Find the priority of each node, propagating dependencies from dependent nodes
+			for (int GroupIdx = 0; GroupIdx < Graph.Groups.Count; GroupIdx++)
+			{
+				INodeGroup Group = Graph.Groups[GroupIdx];
+				for (int NodeIdx = 0; NodeIdx < Group.Nodes.Count; NodeIdx++)
+				{
+					INode Node = Group.Nodes[NodeIdx];
+					Priority NodePriority = NodePriorities[Node];
+
+					foreach (NodeRef DependencyRef in Node.OrderDependencies)
+					{
+						INode Dependency = Graph.Groups[DependencyRef.GroupIdx].Nodes[DependencyRef.NodeIdx];
+						if (NodePriorities[Node] > NodePriority)
+						{
+							NodePriorities[Dependency] = NodePriority;
+						}
+					}
+				}
+			}
+			foreach (JobStepBatchDocument Batch in Job.Batches)
+			{
+				if (Batch.Steps.Count > 0)
+				{
+					INodeGroup Group = Graph.Groups[Batch.GroupIdx];
+					Priority NodePriority = Batch.Steps.Max(x => NodePriorities[Group.Nodes[x.NodeIdx]]);
+					Batch.SchedulePriority = ((int)Job.Priority * 10) + (int)NodePriority + 1; // Reserve '0' for none.
+				}
+			}
+
+			// Check we're not running a node which doesn't allow retries more than once
+			Dictionary<INode, int> NodeExecutionCount = new Dictionary<INode, int>();
+			foreach (IJobStepBatch Batch in Job.Batches)
+			{
+				INodeGroup Group = Graph.Groups[Batch.GroupIdx];
+				foreach (IJobStep Step in Batch.Steps)
+				{
+					INode Node = Group.Nodes[Step.NodeIdx];
+
+					int Count;
+					NodeExecutionCount.TryGetValue(Node, out Count);
+
+					if (!Node.AllowRetry && Count > 0)
+					{
+						throw new RetryNotAllowedException(Node.Name);
+					}
+
+					NodeExecutionCount[Node] = Count + 1;
+				}
+			}
+		}
+
+		/// <summary>
+		/// Gets the scheduling priority of this job
+		/// </summary>
+		/// <param name="Job">Job to consider</param>
+		public static int GetSchedulePriority(IJob Job)
+		{
+			int NewSchedulePriority = 0;
+			foreach (IJobStepBatch Batch in Job.Batches)
+			{
+				if (Batch.State == JobStepBatchState.Ready)
+				{
+					NewSchedulePriority = Math.Max(Batch.SchedulePriority, NewSchedulePriority);
+				}
+			}
+			return NewSchedulePriority;
+		}
+
+		/// <summary>
+		/// Update the state of any jobsteps that are dependent on other jobsteps (eg. transition them from waiting to ready based on other steps completing)
+		/// </summary>
+		/// <param name="Job">The job to update</param>
+		/// <param name="Graph">The graph for this job</param>
+		/// <param name="Updates">List of updates to the job</param>
+		/// <param name="Logger">Logger instance</param>
+		static void RefreshDependentJobSteps(JobDocument Job, IGraph Graph, List<UpdateDefinition<JobDocument>> Updates, ILogger Logger)
+		{
+			// Update the batches
+			UpdateDefinitionBuilder<JobDocument> UpdateBuilder = Builders<JobDocument>.Update;
+			if (Job.Batches != null)
+			{
+				Dictionary<INode, IJobStep> StepForNode = new Dictionary<INode, IJobStep>();
+				for (int LoopBatchIdx = 0; LoopBatchIdx < Job.Batches.Count; LoopBatchIdx++)
+				{
+					int BatchIdx = LoopBatchIdx; // For lambda capture
+					JobStepBatchDocument Batch = Job.Batches[BatchIdx];
+
+					for (int LoopStepIdx = 0; LoopStepIdx < Batch.Steps.Count; LoopStepIdx++)
+					{
+						int StepIdx = LoopStepIdx; // For lambda capture
+						JobStepDocument Step = Batch.Steps[StepIdx];
+
+						JobStepState NewState = Step.State;
+						JobStepOutcome NewOutcome = Step.Outcome;
+
+						INode Node = Graph.Groups[Batch.GroupIdx].Nodes[Step.NodeIdx];
+						if (NewState == JobStepState.Waiting)
+						{
+							List<IJobStep> Steps = GetDependentSteps(Graph, Node, StepForNode);
+							if (Steps.Any(x => x.IsFailedOrSkipped()))
+							{
+								NewState = JobStepState.Skipped;
+								NewOutcome = JobStepOutcome.Failure;
+							}
+							else if (!Steps.Any(x => x.IsPending()))
+							{
+								Logger.LogDebug("Transitioning job {JobId}, batch {BatchId}, step {StepId} to ready state ({Dependencies})", Job.Id, Batch.Id, Step.Id, String.Join(", ", Steps.Select(x => x.Id.ToString())));
+								NewState = JobStepState.Ready;
+							}
+						}
+
+						if (NewState != Step.State)
+						{
+							Step.State = NewState;
+							Updates.Add(UpdateBuilder.Set(x => x.Batches[BatchIdx].Steps[StepIdx].State, NewState));
+						}
+
+						if (NewOutcome != Step.Outcome)
+						{
+							Step.Outcome = NewOutcome;
+							Updates.Add(UpdateBuilder.Set(x => x.Batches[BatchIdx].Steps[StepIdx].Outcome, NewOutcome));
+						}
+
+						StepForNode[Node] = Step;
+					}
+
+					if (Batch.State == JobStepBatchState.Waiting || Batch.State == JobStepBatchState.Ready)
+					{
+						DateTime? NewReadyTime;
+						JobStepBatchState NewState = GetBatchState(Job, Graph, Batch, StepForNode, out NewReadyTime);
+						if (Batch.State != NewState)
+						{
+							Batch.State = NewState;
+							Updates.Add(UpdateBuilder.Set(x => x.Batches[BatchIdx].State, Batch.State));
+						}
+						if (Batch.ReadyTime != NewReadyTime)
+						{
+							Batch.ReadyTime = NewReadyTime;
+							Updates.Add(UpdateBuilder.Set(x => x.Batches[BatchIdx].ReadyTime, Batch.ReadyTime));
+						}
+					}
+				}
+			}
+
+			// Update the weighted priority for the job
+			int NewSchedulePriority = GetSchedulePriority(Job);
+			if (Job.SchedulePriority != NewSchedulePriority)
+			{
+				Job.SchedulePriority = NewSchedulePriority;
+				Updates.Add(UpdateBuilder.Set(x => x.SchedulePriority, NewSchedulePriority));
+			}
+		}
+
+		/// <summary>
+		/// Gets the steps that a node depends on
+		/// </summary>
+		/// <param name="Graph">The graph for this job</param>
+		/// <param name="Node">The node to test</param>
+		/// <param name="StepForNode">Map of node to step</param>
+		/// <returns></returns>
+		static List<IJobStep> GetDependentSteps(IGraph Graph, INode Node, Dictionary<INode, IJobStep> StepForNode)
+		{
+			List<IJobStep> Steps = new List<IJobStep>();
+			foreach (NodeRef OrderDependencyRef in Node.OrderDependencies)
+			{
+				IJobStep? Step;
+				if (StepForNode.TryGetValue(Graph.GetNode(OrderDependencyRef), out Step))
+				{
+					Steps.Add(Step);
+				}
+			}
+			return Steps;
+		}
+
+		/// <summary>
+		/// Gets the new state for a batch
+		/// </summary>
+		/// <param name="Job">The job being executed</param>
+		/// <param name="Graph">Graph for the job</param>
+		/// <param name="Batch">List of nodes in the job</param>
+		/// <param name="StepForNode">Array mapping each index to the appropriate step for that node</param>
+		/// <param name="OutReadyTimeUtc">Receives the time at which the batch was ready to execute</param>
+		/// <returns>True if the batch is ready, false otherwise</returns>
+		static JobStepBatchState GetBatchState(IJob Job, IGraph Graph, IJobStepBatch Batch, Dictionary<INode, IJobStep> StepForNode, out DateTime? OutReadyTimeUtc)
+		{
+			// Check if the batch is already complete
+			if (Batch.Steps.All(x => x.State == JobStepState.Skipped || x.State == JobStepState.Completed || x.State == JobStepState.Aborted))
+			{
+				OutReadyTimeUtc = Batch.ReadyTimeUtc;
+				return JobStepBatchState.Complete;
+			}
+
+			// Get the dependencies for this batch to start. Some steps may be "after" dependencies that are optional parts of the graph.
+			List<INode> NodeDependencies = Batch.GetStartDependencies(Graph.Groups).ToList();
+
+			// Check if we're still waiting on anything
+			DateTime ReadyTimeUtc = Job.CreateTimeUtc;
+			foreach (INode NodeDependency in NodeDependencies)
+			{
+				IJobStep? StepDependency;
+				if (StepForNode.TryGetValue(NodeDependency, out StepDependency))
+				{
+					if (StepDependency.State != JobStepState.Completed && StepDependency.State != JobStepState.Skipped && StepDependency.State != JobStepState.Aborted)
+					{
+						OutReadyTimeUtc = null;
+						return JobStepBatchState.Waiting;
+					}
+
+					if (StepDependency.FinishTimeUtc != null && StepDependency.FinishTimeUtc.Value > ReadyTimeUtc)
+					{
+						ReadyTimeUtc = StepDependency.FinishTimeUtc.Value;
+					}
+				}
+			}
+
+			// Otherwise return the ready state
+			OutReadyTimeUtc = ReadyTimeUtc;
+			return JobStepBatchState.Ready;
+		}
+
+		/// <inheritdoc/>
+		public async Task UpgradeDocumentsAsync()
+		{
+			IAsyncCursor<JobDocument> Cursor = await Jobs.Find(Builders<JobDocument>.Filter.Eq(x => x.UpdateTimeUtc, null)).ToCursorAsync();
+
+			int NumUpdated = 0;
+			while (await Cursor.MoveNextAsync())
+			{
+				foreach (JobDocument Document in Cursor.Current)
+				{
+					UpdateDefinition<JobDocument> Update = Builders<JobDocument>.Update.Set(x => x.CreateTimeUtc, ((IJob)Document).CreateTimeUtc).Set(x => x.UpdateTimeUtc, ((IJob)Document).UpdateTimeUtc);
+					await Jobs.UpdateOneAsync(Builders<JobDocument>.Filter.Eq(x => x.Id, Document.Id), Update);
+					NumUpdated++;
+				}
+				Logger.LogInformation("Updated {NumDocuments} documents", NumUpdated);
+			}
+		}
+	}
+}
diff --git a/Engine/Source/Programs/Horde/HordeServer/Collections/Impl/JobStepRefCollection.cs b/Engine/Source/Programs/Horde/HordeServer/Collections/Impl/JobStepRefCollection.cs
new file mode 100644
--- /dev/null
+++ b/Engine/Source/Programs/Horde/HordeServer/Collections/Impl/JobStepRefCollection.cs
@@ -0,0 +1,139 @@
+using HordeCommon;
+using HordeServer.Models;
+using HordeServer.Services;
+using MongoDB.Bson;
+using MongoDB.Bson.Serialization.Attributes;
+using MongoDB.Driver;
+using System;
+using System.Collections.Generic;
+using System.Linq;
+using System.Threading.Tasks;
+
+using StreamId = HordeServer.Utilities.StringId<HordeServer.Models.IStream>;
+using TemplateRefId = HordeServer.Utilities.StringId<HordeServer.Models.TemplateRef>;
+using PoolId = HordeServer.Utilities.StringId<HordeServer.Models.IPool>;
+using HordeServer.Utilities;
+
+namespace HordeServer.Collections.Impl
+{
+	/// <summary>
+	/// Collection of JobStepRef documents
+	/// </summary>
+	public class JobStepRefCollection : IJobStepRefCollection
+	{
+		class JobStepRef : IJobStepRef
+		{
+			[BsonId]
+			public JobStepRefId Id { get; set; }
+
+			public string JobName { get; set; } = "Unknown";
+
+			public string Name { get; set; }
+
+			public StreamId StreamId { get; set; }
+			public TemplateRefId TemplateId { get; set; }
+			public int Change { get; set; }
+			public ObjectId? LogId { get; set; }
+			public PoolId? PoolId { get; set; }
+			public AgentId? AgentId { get; set; }
+			public JobStepOutcome? Outcome { get; set; }
+			public float BatchWaitTime { get; set; }
+			public float BatchInitTime { get; set; }
+
+			[BsonIgnoreIfNull]
+			public DateTimeOffset? StartTime { get; set; }
+
+			[BsonIgnoreIfNull]
+			public DateTime? StartTimeUtc { get; set; }
+
+			[BsonIgnoreIfNull]
+			public DateTimeOffset? FinishTime { get; set; }
+
+			[BsonIgnoreIfNull]
+			public DateTime? FinishTimeUtc { get; set; }
+
+			DateTime IJobStepRef.StartTimeUtc => StartTimeUtc ?? StartTime?.UtcDateTime ?? default;
+			DateTime? IJobStepRef.FinishTimeUtc => FinishTimeUtc ?? FinishTime?.UtcDateTime;
+			string IJobStepRef.NodeName => Name;
+
+			public JobStepRef(JobStepRefId Id, string JobName, string NodeName, StreamId StreamId, TemplateRefId TemplateId, int Change, ObjectId? LogId, PoolId? PoolId, AgentId? AgentId, JobStepOutcome? Outcome, float BatchWaitTime, float BatchInitTime, DateTime StartTimeUtc, DateTime? FinishTimeUtc)
+			{
+				this.Id = Id;
+				this.JobName = JobName;
+				this.Name = NodeName;
+				this.StreamId = StreamId;
+				this.TemplateId = TemplateId;
+				this.Change = Change;
+				this.LogId = LogId;
+				this.PoolId = PoolId;
+				this.AgentId = AgentId;
+				this.Outcome = Outcome;
+				this.BatchWaitTime = BatchWaitTime;
+				this.BatchInitTime = BatchInitTime;
+				this.StartTime = StartTimeUtc;
+				this.StartTimeUtc = StartTimeUtc;
+				this.FinishTime = FinishTimeUtc;
+				this.FinishTimeUtc = FinishTimeUtc;
+			}
+		}
+
+		IMongoCollection<JobStepRef> JobStepRefs;
+
+		/// <summary>
+		/// Constructor
+		/// </summary>
+		/// <param name="DatabaseService">The database service instance</param>
+		public JobStepRefCollection(DatabaseService DatabaseService)
+		{
+			JobStepRefs = DatabaseService.GetCollection<JobStepRef>("JobStepRefs");
+
+			if (!DatabaseService.ReadOnlyMode)
+			{
+				JobStepRefs.Indexes.CreateOne(new CreateIndexModel<JobStepRef>(Builders<JobStepRef>.IndexKeys.Ascending(x => x.StreamId).Ascending(x => x.TemplateId).Ascending(x => x.Name).Descending(x => x.Change)));
+			}
+		}
+
+		/// <inheritdoc/>
+		public async Task<IJobStepRef> InsertOrReplaceAsync(JobStepRefId Id, string JobName, string StepName, StreamId StreamId, TemplateRefId TemplateId, int Change, ObjectId? LogId, PoolId? PoolId, AgentId? AgentId, JobStepOutcome? Outcome, float WaitTime, float InitTime, DateTime StartTimeUtc, DateTime? FinishTimeUtc)
+		{
+			JobStepRef NewJobStepRef = new JobStepRef(Id, JobName, StepName, StreamId, TemplateId, Change, LogId, PoolId, AgentId, Outcome, WaitTime, InitTime, StartTimeUtc, FinishTimeUtc);
+			await JobStepRefs.ReplaceOneAsync(Builders<JobStepRef>.Filter.Eq(x => x.Id, NewJobStepRef.Id), NewJobStepRef, new ReplaceOptions { IsUpsert = true });
+			return NewJobStepRef;
+		}
+
+		/// <inheritdoc/>
+		public async Task<List<IJobStepRef>> GetStepsForNodeAsync(StreamId StreamId, TemplateRefId TemplateId, string NodeName, int? Change, bool IncludeFailed, int Count)
+		{
+			// Find all the steps matching the given criteria
+			FilterDefinitionBuilder<JobStepRef> FilterBuilder = Builders<JobStepRef>.Filter;
+
+			FilterDefinition<JobStepRef> Filter = FilterDefinition<JobStepRef>.Empty;
+			Filter &= FilterBuilder.Eq(x => x.StreamId, StreamId);
+			Filter &= FilterBuilder.Eq(x => x.TemplateId, TemplateId);
+			Filter &= FilterBuilder.Eq(x => x.Name, NodeName);
+			if (Change != null)
+			{
+				Filter &= FilterBuilder.Lte(x => x.Change, Change.Value);
+			}
+			if (!IncludeFailed)
+			{
+				Filter &= FilterBuilder.Ne(x => x.Outcome, JobStepOutcome.Failure);
+			}
+
+			List<JobStepRef> Steps = await JobStepRefs.Find(Filter).SortByDescending(x => x.Change).ThenByDescending(x => x.StartTime).Limit(Count).ToListAsync();
+			return Steps.ConvertAll<IJobStepRef>(x => x);
+		}
+
+		/// <inheritdoc/>
+		public async Task<IJobStepRef?> GetPrevStepForNodeAsync(StreamId StreamId, TemplateRefId TemplateId, string NodeName, int Change)
+		{
+			return await JobStepRefs.Find(x => x.StreamId == StreamId && x.TemplateId == TemplateId && x.Name == NodeName && x.Change < Change && x.Outcome != null).SortByDescending(x => x.Change).FirstOrDefaultAsync();
+		}
+
+		/// <inheritdoc/>
+		public async Task<IJobStepRef?> GetNextStepForNodeAsync(StreamId StreamId, TemplateRefId TemplateId, string NodeName, int Change)
+		{
+			return await JobStepRefs.Find(x => x.StreamId == StreamId && x.TemplateId == TemplateId && x.Name == NodeName && x.Change > Change && x.Outcome != null).SortBy(x => x.Change).FirstOrDefaultAsync();
+		}
+	}
+}
diff --git a/Engine/Source/Programs/Horde/HordeServer/Collections/Impl/JobTimingCollection.cs b/Engine/Source/Programs/Horde/HordeServer/Collections/Impl/JobTimingCollection.cs
new file mode 100644
--- /dev/null
+++ b/Engine/Source/Programs/Horde/HordeServer/Collections/Impl/JobTimingCollection.cs
@@ -0,0 +1,143 @@
+using Amazon.Runtime.Internal.Util;
+using Amazon.S3.Model.Internal.MarshallTransformations;
+using HordeServer.Models;
+using HordeServer.Services;
+using HordeServer.Utilities;
+using MongoDB.Bson;
+using MongoDB.Bson.Serialization.Attributes;
+using MongoDB.Driver;
+using System;
+using System.Collections.Generic;
+using System.Diagnostics.CodeAnalysis;
+using System.Linq;
+using System.Threading.Tasks;
+
+namespace HordeServer.Collections.Impl
+{
+	/// <summary>
+	/// Concrete implementation of IJobTimingCollection
+	/// </summary>
+	public class JobTimingCollection : IJobTimingCollection
+	{
+		/// <summary>
+		/// Timing information for a particular step
+		/// </summary>
+		class JobStepTimingDocument : IJobStepTiming
+		{
+			public string Name { get; set; }
+			public float? AverageWaitTime { get; set; }
+			public float? AverageInitTime { get; set; }
+			public float? AverageDuration { get; set; }
+
+			[BsonConstructor]
+			private JobStepTimingDocument()
+			{
+				this.Name = null!;
+			}
+
+			public JobStepTimingDocument(string Name, float? AverageWaitTime, float? AverageInitTime, float? AverageDuration)
+			{
+				this.Name = Name;
+				this.AverageWaitTime = AverageWaitTime;
+				this.AverageInitTime = AverageInitTime;
+				this.AverageDuration = AverageDuration;
+			}
+		}
+
+		/// <summary>
+		/// Timing information for a particular job
+		/// </summary>
+		class JobTimingDocument : IJobTiming
+		{
+			public ObjectId Id { get; set; }
+			public List<JobStepTimingDocument> Steps { get; set; } = new List<JobStepTimingDocument>();
+			public int UpdateIndex { get; set; }
+
+			[BsonIgnore]
+			public Dictionary<string, IJobStepTiming>? NameToStep;
+
+			public bool TryGetStepTiming(string Name, [NotNullWhen(true)] out IJobStepTiming? Timing)
+			{
+				if (NameToStep == null)
+				{
+					NameToStep = new Dictionary<string, IJobStepTiming>();
+					foreach (JobStepTimingDocument Step in Steps)
+					{
+						if (NameToStep.ContainsKey(Step.Name))
+						{
+							Serilog.Log.Logger.Warning("Step {Name} appears twice in job timing document {Id}", Step.Name, Id);
+						}
+						NameToStep[Step.Name] = Step;
+					}
+				}
+				return NameToStep.TryGetValue(Name, out Timing);
+			}
+		}
+
+		/// <summary>
+		/// Collection of timing documents
+		/// </summary>
+		IMongoCollection<JobTimingDocument> Collection;
+
+		/// <summary>
+		/// Constructor
+		/// </summary>
+		/// <param name="DatabaseService">The database service singleton</param>
+		public JobTimingCollection(DatabaseService DatabaseService)
+		{
+			Collection = DatabaseService.GetCollection<JobTimingDocument>("JobTiming");
+		}
+
+		/// <inheritdoc/>
+		public async Task<IJobTiming?> TryAddAsync(ObjectId JobId, List<JobStepTimingData> Steps)
+		{
+			JobTimingDocument JobTiming = new JobTimingDocument();
+			JobTiming.Id = JobId;
+			JobTiming.Steps.AddRange(Steps.Select(x => new JobStepTimingDocument(x.Name, x.AverageWaitTime, x.AverageInitTime, x.AverageDuration)));
+			JobTiming.UpdateIndex = 1;
+
+			try
+			{
+				await Collection.InsertOneAsync(JobTiming);
+				return JobTiming;
+			}
+			catch (MongoWriteException Ex)
+			{
+				if (Ex.WriteError.Category == ServerErrorCategory.DuplicateKey)
+				{
+					return null;
+				}
+				else
+				{
+					throw;
+				}
+			}
+		}
+
+		/// <inheritdoc/>
+		public async Task<IJobTiming?> TryGetAsync(ObjectId JobId)
+		{
+			return await Collection.Find(x => x.Id == JobId).FirstOrDefaultAsync();
+		}
+
+		/// <inheritdoc/>
+		public async Task<IJobTiming?> TryAddStepsAsync(IJobTiming JobTiming, List<JobStepTimingData> Steps)
+		{
+			JobTimingDocument JobTimingDocument = (JobTimingDocument)JobTiming;
+			List<JobStepTimingDocument> StepDocuments = Steps.ConvertAll(x => new JobStepTimingDocument(x.Name, x.AverageWaitTime, x.AverageInitTime, x.AverageDuration));
+
+			FilterDefinition<JobTimingDocument> Filter = Builders<JobTimingDocument>.Filter.Expr(x => x.Id == JobTimingDocument.Id && x.UpdateIndex == JobTimingDocument.UpdateIndex);
+			UpdateDefinition<JobTimingDocument> Update = Builders<JobTimingDocument>.Update.Set(x => x.UpdateIndex, JobTimingDocument.UpdateIndex + 1).PushEach(x => x.Steps, StepDocuments);
+
+			UpdateResult Result = await Collection.UpdateOneAsync(Filter, Update);
+			if (Result.ModifiedCount > 0)
+			{
+				JobTimingDocument.NameToStep = null;
+				JobTimingDocument.Steps.AddRange(StepDocuments);
+				JobTimingDocument.UpdateIndex++;
+				return JobTimingDocument;
+			}
+			return null;
+		}
+	}
+}
diff --git a/Engine/Source/Programs/Horde/HordeServer/Collections/Impl/LeaseCollection.cs b/Engine/Source/Programs/Horde/HordeServer/Collections/Impl/LeaseCollection.cs
new file mode 100644
--- /dev/null
+++ b/Engine/Source/Programs/Horde/HordeServer/Collections/Impl/LeaseCollection.cs
@@ -0,0 +1,152 @@
+using HordeCommon;
+using HordeServer.Models;
+using HordeServer.Services;
+using HordeServer.Utilities;
+using MongoDB.Bson;
+using MongoDB.Bson.Serialization.Attributes;
+using MongoDB.Driver;
+using System;
+using System.Collections.Generic;
+using System.Linq;
+using System.Threading.Tasks;
+
+using PoolId = HordeServer.Utilities.StringId<HordeServer.Models.IPool>;
+using StreamId = HordeServer.Utilities.StringId<HordeServer.Models.IStream>;
+
+namespace HordeServer.Collections.Impl
+{
+	/// <summary>
+	/// Collection of lease documents
+	/// </summary>
+	public class LeaseCollection : ILeaseCollection
+	{
+		/// <summary>
+		/// Concrete implementation of a lease document
+		/// </summary>
+		class LeaseDocument : ILease
+		{
+			public ObjectId Id { get; set; }
+			public string Name { get; set; }
+			public AgentId AgentId { get; set; }
+			public ObjectId SessionId { get; set; }
+			public StreamId? StreamId { get; set; }
+			public PoolId? PoolId { get; set; }
+			public ObjectId? LogId { get; set; }
+			public DateTime StartTime { get; set; }
+			public DateTime? FinishTime { get; set; }
+			public byte[] Payload { get; set; }
+			public LeaseOutcome Outcome { get; set; }
+
+			ReadOnlyMemory<byte> ILease.Payload => Payload;
+
+			[BsonConstructor]
+			private LeaseDocument()
+			{
+				Name = String.Empty;
+				Payload = null!;
+			}
+
+			public LeaseDocument(ObjectId Id, string Name, AgentId AgentId, ObjectId SessionId, StreamId? StreamId, PoolId? PoolId, ObjectId? LogId, DateTime StartTime, byte[] Payload)
+			{
+				this.Id = Id;
+				this.Name = Name;
+				this.AgentId = AgentId;
+				this.SessionId = SessionId;
+				this.StreamId = StreamId;
+				this.PoolId = PoolId;
+				this.LogId = LogId;
+				this.StartTime = StartTime;
+				this.Payload = Payload;
+			}
+		}
+
+		/// <summary>
+		/// Collection of lease documents
+		/// </summary>
+		readonly IMongoCollection<LeaseDocument> Leases;
+
+		/// <summary>
+		/// Constructor
+		/// </summary>
+		/// <param name="DatabaseService">The database service instance</param>
+		public LeaseCollection(DatabaseService DatabaseService)
+		{
+			Leases = DatabaseService.GetCollection<LeaseDocument>("Leases");
+
+			if (!DatabaseService.ReadOnlyMode)
+			{
+				Leases.Indexes.CreateOne(new CreateIndexModel<LeaseDocument>(Builders<LeaseDocument>.IndexKeys.Ascending(x => x.AgentId)));
+				Leases.Indexes.CreateOne(new CreateIndexModel<LeaseDocument>(Builders<LeaseDocument>.IndexKeys.Ascending(x => x.SessionId)));
+				Leases.Indexes.CreateOne(new CreateIndexModel<LeaseDocument>(Builders<LeaseDocument>.IndexKeys.Ascending(x => x.StartTime)));
+				Leases.Indexes.CreateOne(new CreateIndexModel<LeaseDocument>(Builders<LeaseDocument>.IndexKeys.Ascending(x => x.FinishTime)));
+			}
+		}
+
+		/// <inheritdoc/>
+		public async Task<ILease> AddAsync(ObjectId Id, string Name, AgentId AgentId, ObjectId SessionId, StreamId? StreamId, PoolId? PoolId, ObjectId? LogId, DateTime StartTime, byte[] Payload)
+		{
+			LeaseDocument Lease = new LeaseDocument(Id, Name, AgentId, SessionId, StreamId, PoolId, LogId, StartTime, Payload);
+			await Leases.InsertOneAsync(Lease);
+			return Lease;
+		}
+
+		/// <inheritdoc/>
+		public async Task DeleteAsync(ObjectId LeaseId)
+		{
+			await Leases.DeleteOneAsync(x => x.Id == LeaseId);
+		}
+
+		/// <inheritdoc/>
+		public async Task<ILease?> GetAsync(ObjectId LeaseId)
+		{
+			return await Leases.Find(x => x.Id == LeaseId).FirstOrDefaultAsync();
+		}
+
+		/// <inheritdoc/>
+		public async Task<List<ILease>> FindLeasesAsync(AgentId? AgentId, ObjectId? SessionId, DateTime? MinTime, DateTime? MaxTime, int? Index, int? Count)
+		{
+			FilterDefinitionBuilder<LeaseDocument> FilterBuilder = Builders<LeaseDocument>.Filter;
+
+			FilterDefinition<LeaseDocument> Filter = FilterDefinition<LeaseDocument>.Empty;
+			if (AgentId != null)
+			{
+				Filter &= FilterBuilder.Eq(x => x.AgentId, AgentId.Value);
+			}
+			if (SessionId != null)
+			{
+				Filter &= FilterBuilder.Eq(x => x.SessionId, SessionId.Value);
+			}
+			if (MinTime != null)
+			{
+				Filter &= FilterBuilder.Or(FilterBuilder.Eq(x => x.FinishTime!, null), FilterBuilder.Gt(x => x.FinishTime!, MinTime.Value));
+			}
+			if (MaxTime != null)
+			{
+				Filter &= FilterBuilder.Lt(x => x.StartTime, MaxTime.Value);
+			}
+
+			List<LeaseDocument> Results = await Leases.Find(Filter).SortByDescending(x => x.StartTime).Range(Index, Count).ToListAsync();
+			return Results.ConvertAll<ILease>(x => x);
+		}
+
+		/// <inheritdoc/>
+		public async Task<List<ILease>> FindActiveLeasesAsync(int? Index = null, int? Count = null)
+		{
+			List<LeaseDocument> Results = await Leases.Find(x => x.FinishTime == null).Range(Index, Count).ToListAsync();
+			return Results.ConvertAll<ILease>(x => x);
+		}
+
+		/// <inheritdoc/>
+		public async Task<bool> TrySetOutcomeAsync(ObjectId LeaseId, DateTime FinishTime, LeaseOutcome Outcome)
+		{
+			FilterDefinitionBuilder<LeaseDocument> FilterBuilder = Builders<LeaseDocument>.Filter;
+			FilterDefinition<LeaseDocument> Filter = FilterBuilder.Eq(x => x.Id, LeaseId) & FilterBuilder.Eq(x => x.FinishTime, null);
+
+			UpdateDefinitionBuilder<LeaseDocument> UpdateBuilder = Builders<LeaseDocument>.Update;
+			UpdateDefinition<LeaseDocument> Update = UpdateBuilder.Set(x => x.FinishTime, FinishTime).Set(x => x.Outcome, Outcome);
+
+			UpdateResult Result = await Leases.UpdateOneAsync(Filter, Update);
+			return Result.ModifiedCount > 0;
+		}
+	}
+}
diff --git a/Engine/Source/Programs/Horde/HordeServer/Collections/Impl/PoolCollection.cs b/Engine/Source/Programs/Horde/HordeServer/Collections/Impl/PoolCollection.cs
new file mode 100644
--- /dev/null
+++ b/Engine/Source/Programs/Horde/HordeServer/Collections/Impl/PoolCollection.cs
@@ -0,0 +1,259 @@
+using HordeServer.Models;
+using HordeServer.Services;
+using HordeServer.Utilities;
+using MongoDB.Bson;
+using MongoDB.Bson.Serialization;
+using MongoDB.Bson.Serialization.Attributes;
+using MongoDB.Driver;
+using System;
+using System.Collections.Generic;
+using System.Linq;
+using System.Linq.Expressions;
+using System.Threading.Tasks;
+
+using PoolId = HordeServer.Utilities.StringId<HordeServer.Models.IPool>;
+
+namespace HordeServer.Collections.Impl
+{
+	/// <summary>
+	/// Collection of pool documents
+	/// </summary>
+	public class PoolCollection : IPoolCollection
+	{
+		class PoolDocument : IPool
+		{
+			// Properties
+			[BsonId]
+			public PoolId Id { get; set; }
+			[BsonRequired]
+			public string Name { get; set; } = null!;
+			public List<AgentWorkspace> Workspaces { get; set; } = new List<AgentWorkspace>();
+			public AgentRequirements? Requirements { get; set; }
+			public Dictionary<string, string> Properties { get; set; } = new Dictionary<string, string>();
+			[BsonIgnoreIfDefault(true)]
+			public bool EnableAutoscaling { get; set; } = true;
+			public int? MinAgents { get; set; }
+			public int? NumReserveAgents { get; set; }
+			public DateTime? LastScaleUpTime { get; set; }
+			public DateTime? LastScaleDownTime { get; set; }
+			public int UpdateIndex { get; set; }
+
+			// Read-only wrappers
+			IReadOnlyList<AgentWorkspace> IPool.Workspaces => Workspaces;
+			IReadOnlyDictionary<string, string> IPool.Properties => Properties;
+
+			public PoolDocument()
+			{
+			}
+
+			public PoolDocument(IPool Other)
+			{
+				Id = Other.Id;
+				Name = Other.Name;
+				Workspaces.AddRange(Other.Workspaces);
+				if(Other.Requirements != null)
+				{
+					Requirements = new AgentRequirements(Other.Requirements);
+				}
+				Properties = new Dictionary<string, string>(Other.Properties);
+				EnableAutoscaling = Other.EnableAutoscaling;
+				MinAgents = Other.MinAgents;
+				NumReserveAgents = Other.NumReserveAgents;
+				LastScaleUpTime = Other.LastScaleUpTime;
+				LastScaleDownTime = Other.LastScaleDownTime;
+				UpdateIndex = Other.UpdateIndex;
+			}
+		}
+
+		/// <summary>
+		/// Collection of pool documents
+		/// </summary>
+		IMongoCollection<PoolDocument> Pools;
+
+		/// <summary>
+		/// Constructor
+		/// </summary>
+		/// <param name="DatabaseService">The database service instance</param>
+		public PoolCollection(DatabaseService DatabaseService)
+		{
+			Pools = DatabaseService.GetCollection<PoolDocument>("Pools");
+
+			if (!DatabaseService.ReadOnlyMode)
+			{
+				// Drop any legacy indexes
+				List<BsonDocument> IndexList = Pools.Indexes.List().ToList();
+				if (IndexList.Any(x => { BsonValue Name; return x.TryGetValue("name", out Name) && Name.AsString == "Deleted_1_Name_1"; }))
+				{
+					Pools.Indexes.DropOne("Deleted_1_Name_1");
+				}
+				
+				// Remove any pools with an objectid for unique id
+				Pools.DeleteMany("{ \"_id\": { \"$type\": 7 } }");
+			}
+		}
+
+		/// <inheritdoc/>
+		public async Task<IPool> AddAsync(PoolId Id, string Name, AgentRequirements? Requirements, bool? EnableAutoscaling, int? MinAgents, int? NumReserveAgents, IEnumerable<KeyValuePair<string, string>>? Properties)
+		{
+			PoolDocument Pool = new PoolDocument();
+			Pool.Id = Id;
+			Pool.Name = Name;
+			if(Requirements != null)
+			{
+				Pool.Requirements = new AgentRequirements(Requirements);
+			}
+			if (EnableAutoscaling != null)
+			{
+				Pool.EnableAutoscaling = EnableAutoscaling.Value;
+			}
+			Pool.MinAgents = MinAgents;
+			Pool.NumReserveAgents = NumReserveAgents;
+			if (Properties != null)
+			{
+				Pool.Properties = new Dictionary<string, string>(Properties);
+			}
+			await Pools.InsertOneAsync(Pool);
+			return Pool;
+		}
+
+		/// <inheritdoc/>
+		public async Task<List<IPool>> GetAsync()
+		{
+			List<PoolDocument> Results = await Pools.Find(FilterDefinition<PoolDocument>.Empty).ToListAsync();
+			return Results.ConvertAll<IPool>(x => x);
+		}
+
+		/// <inheritdoc/>
+		public async Task<IPool?> GetAsync(PoolId Id)
+		{
+			return await Pools.Find(x => x.Id == Id).FirstOrDefaultAsync();
+		}
+
+		/// <inheritdoc/>
+		public async Task<List<PoolId>> GetPoolIdsAsync()
+		{
+			ProjectionDefinition<PoolDocument, BsonDocument> Projection = Builders<PoolDocument>.Projection.Include(x => x.Id);
+			List<BsonDocument> Results = await Pools.Find(FilterDefinition<PoolDocument>.Empty).Project(Projection).ToListAsync();
+			return Results.ConvertAll(x => new PoolId(x.GetElement("_id").Value.AsString));
+		}
+
+		/// <inheritdoc/>
+		public async Task<bool> DeleteAsync(PoolId Id)
+		{
+			FilterDefinition<PoolDocument> Filter = Builders<PoolDocument>.Filter.Eq(x => x.Id, Id);
+			DeleteResult Result = await Pools.DeleteOneAsync(Filter);
+			return Result.DeletedCount > 0;
+		}
+
+		/// <summary>
+		/// Attempts to update a pool, upgrading it to the latest document schema if necessary
+		/// </summary>
+		/// <param name="PoolToUpdate">Interface for the document to update</param>
+		/// <param name="Transaction">The transaction to execute</param>
+		/// <returns>New pool definition, or null on failure.</returns>
+		async Task<IPool?> TryUpdateAsync(IPool PoolToUpdate, TransactionBuilder<PoolDocument> Transaction)
+		{
+			if (Transaction.IsEmpty)
+			{
+				return PoolToUpdate;
+			}
+
+			Transaction.Set(x => x.UpdateIndex, PoolToUpdate.UpdateIndex + 1);
+
+			PoolDocument? MutablePool = PoolToUpdate as PoolDocument;
+			if (MutablePool != null)
+			{
+				UpdateResult Result = await Pools.UpdateOneAsync(x => x.Id == PoolToUpdate.Id && x.UpdateIndex == PoolToUpdate.UpdateIndex, Transaction.ToUpdateDefinition());
+				if (Result.ModifiedCount == 0)
+				{
+					return null;
+				}
+
+				Transaction.ApplyTo(MutablePool);
+			}
+			else
+			{
+				MutablePool = new PoolDocument(PoolToUpdate);
+				Transaction.ApplyTo(MutablePool);
+
+				ReplaceOneResult Result = await Pools.ReplaceOneAsync<PoolDocument>(x => x.Id == PoolToUpdate.Id && x.UpdateIndex == PoolToUpdate.UpdateIndex, MutablePool);
+				if (Result.ModifiedCount == 0)
+				{
+					return null;
+				}
+			}
+			return MutablePool;
+		}
+
+		/// <inheritdoc/>
+		public Task<IPool?> TryUpdateAsync(IPool Pool, string? NewName, AgentRequirements? NewRequirements, bool? NewEnableAutoscaling, int? NewMinAgents, int? NewNumReserveAgents, List<AgentWorkspace>? NewWorkspaces, Dictionary<string, string?>? NewProperties, DateTime? LastScaleUpTime, DateTime? LastScaleDownTime)
+		{
+			TransactionBuilder<PoolDocument> Transaction = new TransactionBuilder<PoolDocument>();
+			if (NewName != null)
+			{
+				Transaction.Set(x => x.Name, NewName);
+			}
+			if (NewRequirements != null)
+			{
+				if (NewRequirements.IsEmpty())
+				{
+					Transaction.Unset(x => x.Requirements!);
+				}
+				else
+				{
+					Transaction.Set(x => x.Requirements, NewRequirements);
+				}
+			}
+			if (NewEnableAutoscaling != null)
+			{
+				if (NewEnableAutoscaling.Value)
+				{
+					Transaction.Unset(x => x.EnableAutoscaling);
+				}
+				else
+				{
+					Transaction.Set(x => x.EnableAutoscaling, NewEnableAutoscaling.Value);
+				}
+			}
+			if (NewMinAgents != null)
+			{
+				if (NewMinAgents.Value < 0)
+				{
+					Transaction.Unset(x => x.MinAgents!);
+				}
+				else
+				{
+					Transaction.Set(x => x.MinAgents, NewMinAgents.Value);
+				}
+			}
+			if (NewNumReserveAgents != null)
+			{
+				if (NewNumReserveAgents.Value < 0)
+				{
+					Transaction.Unset(x => x.NumReserveAgents!);
+				}
+				else
+				{
+					Transaction.Set(x => x.NumReserveAgents, NewNumReserveAgents.Value);
+				}
+			}
+			if (NewWorkspaces != null)
+			{
+				Transaction.Set(x => x.Workspaces, NewWorkspaces);
+			}
+			if (NewProperties != null)
+			{
+				Transaction.UpdateDictionary(x => x.Properties, NewProperties);
+			}
+			if (LastScaleUpTime != null)
+			{
+				Transaction.Set(x => x.LastScaleUpTime, LastScaleUpTime);
+			}
+			if (LastScaleDownTime != null)
+			{
+				Transaction.Set(x => x.LastScaleDownTime, LastScaleDownTime);
+			}
+			return TryUpdateAsync(Pool, Transaction);
+		}
+	}
+}
diff --git a/Engine/Source/Programs/Horde/HordeServer/Collections/Impl/ProjectCollection.cs b/Engine/Source/Programs/Horde/HordeServer/Collections/Impl/ProjectCollection.cs
new file mode 100644
--- /dev/null
+++ b/Engine/Source/Programs/Horde/HordeServer/Collections/Impl/ProjectCollection.cs
@@ -0,0 +1,194 @@
+// Copyright Epic Games, Inc. All Rights Reserved.
+
+using HordeServer.Api;
+using HordeServer.Models;
+using HordeServer.Services;
+using HordeServer.Utilities;
+using Microsoft.Extensions.Options;
+using MongoDB.Bson;
+using MongoDB.Bson.Serialization.Attributes;
+using MongoDB.Driver;
+using System;
+using System.Collections.Generic;
+using System.Diagnostics.CodeAnalysis;
+using System.Linq;
+using System.Security.Claims;
+using System.Threading.Tasks;
+
+using ProjectId = HordeServer.Utilities.StringId<HordeServer.Models.IProject>;
+
+namespace HordeServer.Collections.Impl
+{
+	/// <summary>
+	/// Wraps functionality for manipulating projects
+	/// </summary>
+	public class ProjectCollection : IProjectCollection
+	{
+		/// <summary>
+		/// Represents a project
+		/// </summary>
+		class ProjectDocument : IProject
+		{
+			[BsonRequired, BsonId]
+			public ProjectId Id { get; set; }
+
+			[BsonRequired]
+			public string Name { get; set; }
+
+			public int Order { get; set; } = 128;
+			public Acl? Acl { get; set; }
+			public List<StreamCategory> Categories { get; set; } = new List<StreamCategory>();
+			public Dictionary<string, string> Properties { get; set; } = new Dictionary<string, string>(StringComparer.Ordinal);
+
+			IReadOnlyList<StreamCategory> IProject.Categories => Categories;
+			IReadOnlyDictionary<string, string> IProject.Properties => Properties;
+
+			[BsonConstructor]
+			private ProjectDocument()
+			{
+				Name = null!;
+			}
+
+			public ProjectDocument(ProjectId Id, string Name)
+			{
+				this.Id = Id;
+				this.Name = Name;
+			}
+		}
+
+		/// <summary>
+		/// Projection of a project definition to just include permissions info
+		/// </summary>
+		[SuppressMessage("Design", "CA1812: Class is never instantiated")]
+		class ProjectPermissions : IProjectPermissions
+		{
+			/// <summary>
+			/// ACL for the project
+			/// </summary>
+			public Acl? Acl { get; set; }
+
+			/// <summary>
+			/// Projection to extract the permissions info from the project
+			/// </summary>
+			public static readonly ProjectionDefinition<ProjectDocument> Projection = Builders<ProjectDocument>.Projection.Include(x => x.Acl);
+		}
+
+		/// <summary>
+		/// Collection of project documents
+		/// </summary>
+		IMongoCollection<ProjectDocument> Projects;
+
+		/// <summary>
+		/// Constructor
+		/// </summary>
+		/// <param name="DatabaseService">The database service instance</param>
+		public ProjectCollection(DatabaseService DatabaseService)
+		{
+			Projects = DatabaseService.GetCollection<ProjectDocument>("Projects");
+		}
+
+		/// <inheritdoc/>
+		public async Task<IProject?> TryAddAsync(ProjectId Id, string Name, int? Order, List<StreamCategory>? Categories, Dictionary<string, string>? Properties)
+		{
+			ProjectDocument NewProject = new ProjectDocument(Id, Name);
+			if (Order != null)
+			{
+				NewProject.Order = Order.Value;
+			}
+			if (Categories != null)
+			{
+				NewProject.Categories = Categories;
+			}
+			if (Properties != null)
+			{
+				NewProject.Properties = new Dictionary<string, string>(Properties, NewProject.Properties.Comparer);
+			}
+
+			// Attempt to insert the stream, but fail gracefully if it already exists
+			try
+			{
+				await Projects.InsertOneAsync(NewProject);
+				return NewProject;
+			}
+			catch (MongoWriteException Ex)
+			{
+				if (Ex.WriteError.Category == ServerErrorCategory.DuplicateKey)
+				{
+					return null;
+				}
+				else
+				{
+					throw;
+				}
+			}
+		}
+
+		/// <inheritdoc/>
+		public async Task UpdateAsync(ProjectId ProjectId, string? NewName, int? NewOrder, List<StreamCategory>? NewCategories, Dictionary<string, string>? NewProperties, Acl? NewAcl)
+		{
+			UpdateDefinitionBuilder<ProjectDocument> UpdateBuilder = Builders<ProjectDocument>.Update;
+
+			List<UpdateDefinition<ProjectDocument>> Updates = new List<UpdateDefinition<ProjectDocument>>();
+			if (NewName != null)
+			{
+				Updates.Add(UpdateBuilder.Set(x => x.Name, NewName));
+			}
+			if (NewOrder != null)
+			{
+				Updates.Add(UpdateBuilder.Set(x => x.Order, NewOrder.Value));
+			}
+			if (NewCategories != null)
+			{
+				Updates.Add(UpdateBuilder.Set(x => x.Categories, NewCategories));
+			}
+			if (NewProperties != null)
+			{
+				foreach (KeyValuePair<string, string> Pair in NewProperties)
+				{
+					if (Pair.Value == null)
+					{
+						Updates.Add(UpdateBuilder.Unset(x => x.Properties[Pair.Key]));
+					}
+					else
+					{
+						Updates.Add(UpdateBuilder.Set(x => x.Properties[Pair.Key], Pair.Value));
+					}
+				}
+			}
+			if (NewAcl != null)
+			{
+				Updates.Add(Acl.CreateUpdate<ProjectDocument>(x => x.Acl!, NewAcl));
+			}
+
+			if (Updates.Count > 0)
+			{
+				await Projects.FindOneAndUpdateAsync<ProjectDocument>(x => x.Id == ProjectId, UpdateBuilder.Combine(Updates));
+			}
+		}
+
+		/// <inheritdoc/>
+		public async Task<List<IProject>> FindAllAsync()
+		{
+			List<ProjectDocument> Results = await Projects.Find(FilterDefinition<ProjectDocument>.Empty).ToListAsync();
+			return Results.OrderBy(x => x.Order).ThenBy(x => x.Name).Select<ProjectDocument, IProject>(x => x).ToList();
+		}
+
+		/// <inheritdoc/>
+		public async Task<IProject?> GetAsync(ProjectId ProjectId)
+		{
+			return await Projects.Find<ProjectDocument>(x => x.Id == ProjectId).FirstOrDefaultAsync();
+		}
+
+		/// <inheritdoc/>
+		public async Task<IProjectPermissions?> GetPermissionsAsync(ProjectId ProjectId)
+		{
+			return await Projects.Find<ProjectDocument>(x => x.Id == ProjectId).Project<ProjectPermissions>(ProjectPermissions.Projection).FirstOrDefaultAsync();
+		}
+
+		/// <inheritdoc/>
+		public async Task DeleteAsync(ProjectId ProjectId)
+		{
+			await Projects.DeleteOneAsync<ProjectDocument>(x => x.Id == ProjectId);
+		}
+	}
+}
diff --git a/Engine/Source/Programs/Horde/HordeServer/Collections/Impl/ServiceAccountCollection.cs b/Engine/Source/Programs/Horde/HordeServer/Collections/Impl/ServiceAccountCollection.cs
new file mode 100644
--- /dev/null
+++ b/Engine/Source/Programs/Horde/HordeServer/Collections/Impl/ServiceAccountCollection.cs
@@ -0,0 +1,129 @@
+using System;
+using HordeServer.Models;
+using HordeServer.Services;
+using MongoDB.Bson;
+using MongoDB.Bson.Serialization.Attributes;
+using MongoDB.Driver;
+using System.Collections.Generic;
+using System.Linq;
+using System.Threading.Tasks;
+
+namespace HordeServer.Collections.Impl
+{
+	/// <summary>
+	/// Collection of service account documents
+	/// </summary>
+	public class ServiceAccountCollection : IServiceAccountCollection
+	{
+		/// <summary>
+		/// Concrete implementation of IServiceAccount
+		/// </summary>
+		class ServiceAccountDocument : IServiceAccount
+		{
+			[BsonRequired, BsonId]
+			public ObjectId Id { get; set; }
+
+			[BsonRequired]
+			public string SecretToken { get; set; } = "empty";
+
+			[BsonRequired]
+			public IReadOnlyDictionary<string, string> Claims { get; set; } = new Dictionary<string, string>();
+			
+			public bool Enabled { get; set; }
+			public string Description { get; set; } = "empty";
+
+			[BsonConstructor]
+			private ServiceAccountDocument()
+			{
+			}
+
+			public ServiceAccountDocument(ObjectId Id, string SecretToken, IReadOnlyDictionary<string, string> Claims, bool Enabled, string Description)
+			{
+				this.Id = Id;
+				this.SecretToken = SecretToken;
+				this.Claims = Claims;
+				this.Enabled = Enabled;
+				this.Description = Description;
+			}
+
+			protected bool Equals(ServiceAccountDocument other)
+			{
+				bool AreClaimsEqual = !Claims.Except(other.Claims).Any();
+
+                return Id.Equals(other.Id) && SecretToken == other.SecretToken && AreClaimsEqual && Enabled == other.Enabled && Description == other.Description;
+			}
+
+			public override bool Equals(object? obj)
+			{
+				if (ReferenceEquals(null, obj)) return false;
+				if (ReferenceEquals(this, obj)) return true;
+				if (obj.GetType() != this.GetType()) return false;
+				return Equals((ServiceAccountDocument) obj);
+			}
+
+			public override int GetHashCode()
+			{
+				return HashCode.Combine(Id, SecretToken, Claims, Enabled, Description);
+			}
+		}
+
+		/// <summary>
+		/// Collection of session documents
+		/// </summary>
+		private readonly IMongoCollection<ServiceAccountDocument> ServiceAccounts;
+
+		/// <summary>
+		/// Constructor
+		/// </summary>
+		/// <param name="DatabaseService">The database service</param>
+		public ServiceAccountCollection(DatabaseService DatabaseService)
+		{
+			ServiceAccounts = DatabaseService.GetCollection<ServiceAccountDocument>("ServiceAccounts");
+
+			if (!DatabaseService.ReadOnlyMode)
+			{
+				ServiceAccounts.Indexes.CreateOne(new CreateIndexModel<ServiceAccountDocument>(Builders<ServiceAccountDocument>.IndexKeys.Ascending(x => x.SecretToken)));
+			}
+		}
+
+		/// <inheritdoc/>
+		public async Task<IServiceAccount> AddAsync(string SecretToken, Dictionary<string, string> Claims, string Description)
+		{
+			ServiceAccountDocument NewSession = new ServiceAccountDocument(ObjectId.GenerateNewId(), SecretToken, Claims, true, Description);
+			await ServiceAccounts.InsertOneAsync(NewSession);
+			return NewSession;
+		}
+
+		/// <inheritdoc/>
+		public async Task<IServiceAccount?> GetAsync(ObjectId Id)
+		{
+			return await ServiceAccounts.Find(x => x.Id == Id).FirstOrDefaultAsync();
+		}
+
+		/// <inheritdoc/>
+		public async Task<IServiceAccount?> GetBySecretTokenAsync(string SecretToken)
+		{
+			return await ServiceAccounts.Find(x => x.SecretToken == SecretToken).FirstOrDefaultAsync();
+		}
+
+		/// <inheritdoc/>
+		public Task UpdateAsync(ObjectId Id, string? SecretToken, Dictionary<string, string>? Claims, bool? Enabled, string? Description)
+		{
+			UpdateDefinitionBuilder<ServiceAccountDocument> Update = Builders<ServiceAccountDocument>.Update;
+			List<UpdateDefinition<ServiceAccountDocument>> Updates = new List<UpdateDefinition<ServiceAccountDocument>>();
+			
+			if (SecretToken != null) Updates.Add(Update.Set(x => x.SecretToken, SecretToken));
+			if (Claims != null) Updates.Add(Update.Set(x => x.Claims, Claims));
+			if (Enabled != null) Updates.Add(Update.Set(x => x.Enabled, Enabled));
+			if (Description != null) Updates.Add(Update.Set(x => x.Description, Description));
+
+			return ServiceAccounts.FindOneAndUpdateAsync(x => x.Id == Id, Update.Combine(Updates));
+		}
+
+		/// <inheritdoc/>
+		public Task DeleteAsync(ObjectId SessionId)
+		{
+			return ServiceAccounts.DeleteOneAsync(x => x.Id == SessionId);
+		}
+	}
+}
diff --git a/Engine/Source/Programs/Horde/HordeServer/Collections/Impl/SessionCollection.cs b/Engine/Source/Programs/Horde/HordeServer/Collections/Impl/SessionCollection.cs
new file mode 100644
--- /dev/null
+++ b/Engine/Source/Programs/Horde/HordeServer/Collections/Impl/SessionCollection.cs
@@ -0,0 +1,130 @@
+using HordeServer.Models;
+using HordeServer.Services;
+using HordeServer.Utilities;
+using Microsoft.Extensions.DependencyInjection;
+using MongoDB.Bson;
+using MongoDB.Bson.Serialization.Attributes;
+using MongoDB.Driver;
+using System;
+using System.Collections.Generic;
+using System.Linq;
+using System.Threading.Tasks;
+
+using AgentSoftwareVersion = HordeServer.Utilities.StringId<HordeServer.Collections.IAgentSoftwareCollection>;
+
+namespace HordeServer.Collections.Impl
+{
+	/// <summary>
+	/// Collection of session documents
+	/// </summary>
+	public class SessionCollection : ISessionCollection
+	{
+		/// <summary>
+		/// Concrete implementation of ISession
+		/// </summary>
+		class SessionDocument : ISession
+		{
+			[BsonRequired, BsonId]
+			public ObjectId Id { get; set; }
+
+			[BsonRequired]
+			public AgentId AgentId { get; set; }
+
+			public DateTime StartTime { get; set; }
+			public DateTime? FinishTime { get; set; }
+			public AgentCapabilities? Capabilities { get; set; }
+			public AgentSoftwareVersion? Version { get; set; }
+
+			[BsonConstructor]
+			private SessionDocument()
+			{
+			}
+
+			public SessionDocument(ObjectId Id, AgentId AgentId, DateTime StartTime, AgentCapabilities? Capabilities, AgentSoftwareVersion? Version)
+			{
+				this.Id = Id;
+				this.AgentId = AgentId;
+				this.StartTime = StartTime;
+				this.Capabilities = Capabilities;
+				this.Version = Version;
+			}
+		}
+
+		/// <summary>
+		/// Collection of session documents
+		/// </summary>
+		readonly IMongoCollection<SessionDocument> Sessions;
+
+		/// <summary>
+		/// Constructor
+		/// </summary>
+		/// <param name="DatabaseService">The database service</param>
+		public SessionCollection(DatabaseService DatabaseService)
+		{
+			Sessions = DatabaseService.GetCollection<SessionDocument>("Sessions");
+
+			if (!DatabaseService.ReadOnlyMode)
+			{
+				Sessions.Indexes.CreateOne(new CreateIndexModel<SessionDocument>(Builders<SessionDocument>.IndexKeys.Ascending(x => x.AgentId)));
+				Sessions.Indexes.CreateOne(new CreateIndexModel<SessionDocument>(Builders<SessionDocument>.IndexKeys.Ascending(x => x.StartTime)));
+				Sessions.Indexes.CreateOne(new CreateIndexModel<SessionDocument>(Builders<SessionDocument>.IndexKeys.Ascending(x => x.FinishTime)));
+			}
+		}
+
+		/// <inheritdoc/>
+		public async Task<ISession> AddAsync(ObjectId Id, AgentId AgentId, DateTime StartTime, AgentCapabilities? Capabilities, AgentSoftwareVersion? Version)
+		{
+			SessionDocument NewSession = new SessionDocument(Id, AgentId, StartTime, Capabilities, Version);
+			await Sessions.InsertOneAsync(NewSession);
+			return NewSession;
+		}
+
+		/// <inheritdoc/>
+		public async Task<ISession?> GetAsync(ObjectId SessionId)
+		{
+			return await Sessions.Find(x => x.Id == SessionId).FirstOrDefaultAsync();
+		}
+
+		/// <inheritdoc/>
+		public async Task<List<ISession>> FindAsync(AgentId AgentId, DateTime? StartTime, DateTime? FinishTime, int Index, int Count)
+		{
+			FilterDefinitionBuilder<SessionDocument> FilterBuilder = Builders<SessionDocument>.Filter;
+
+			FilterDefinition<SessionDocument> Filter = FilterBuilder.Eq(x => x.AgentId, AgentId);
+			if (StartTime != null)
+			{
+				Filter &= FilterBuilder.Gte(x => x.StartTime, StartTime.Value);
+			}
+			if (FinishTime != null)
+			{
+				Filter &= FilterBuilder.Or(FilterBuilder.Eq(x => x.FinishTime, null), FilterBuilder.Lte(x => x.FinishTime, FinishTime.Value));
+			}
+
+			List<SessionDocument> Results = await Sessions.Find(Filter).SortByDescending(x => x.StartTime).Skip(Index).Limit(Count).ToListAsync();
+			return Results.ConvertAll<ISession>(x => x);
+		}
+
+		/// <inheritdoc/>
+		public async Task<List<ISession>> FindActiveSessionsAsync(int? Index, int? Count)
+		{
+			List<SessionDocument> Results = await Sessions.Find(x => x.FinishTime == null).Range(Index, Count).ToListAsync();
+			return Results.ConvertAll<ISession>(x => x);
+		}
+
+		/// <inheritdoc/>
+		public Task UpdateAsync(ObjectId SessionId, DateTime FinishTime, AgentCapabilities Capabilities)
+		{
+			UpdateDefinition<SessionDocument> Update = Builders<SessionDocument>.Update
+				.Set(x => x.FinishTime, FinishTime)
+				.Set(x => x.Capabilities, Capabilities);
+
+			return Sessions.FindOneAndUpdateAsync(x => x.Id == SessionId, Update);
+		}
+
+		/// <inheritdoc/>
+		public Task DeleteAsync(ObjectId SessionId)
+		{
+			return Sessions.DeleteOneAsync(x => x.Id == SessionId);
+		}
+	}
+}
