Initial drop of Horde Server/Agent for UE5/Main. Still missing some Epic-specific functionality required to build/run correctly.

diff --git a/Engine/Source/Programs/Horde/HordeServer/Collections/Impl/SoftwareCollection.cs b/Engine/Source/Programs/Horde/HordeServer/Collections/Impl/SoftwareCollection.cs
new file mode 100644
--- /dev/null
+++ b/Engine/Source/Programs/Horde/HordeServer/Collections/Impl/SoftwareCollection.cs
@@ -0,0 +1,221 @@
+using EpicGames.Core;
+using HordeCommon;
+using HordeServer.Models;
+using HordeServer.Services;
+using MongoDB.Bson;
+using MongoDB.Bson.Serialization.Attributes;
+using MongoDB.Driver;
+using System;
+using System.Collections.Generic;
+using System.IO;
+using System.IO.Compression;
+using System.Linq;
+using System.Text.Json;
+using System.Threading.Tasks;
+
+namespace HordeServer.Collections.Impl
+{
+	/// <summary>
+	/// Collection of template documents
+	/// </summary>
+	public class SoftwareCollection : ISoftwareCollection
+	{
+		/// <summary>
+		/// Information about agent software
+		/// </summary>
+		class SoftwareDocument : ISoftware
+		{
+			[BsonRequired, BsonId]
+			public ObjectId Id { get; set; }
+
+			public byte[]? Data { get; set; }
+			public string? UploadedByUser { get; set; }
+			public DateTime UploadedAtTime { get; set; }
+			public string? MadeDefaultByUser { get; set; }
+			public DateTime? MadeDefaultAtTime { get; set; }
+
+			[BsonConstructor]
+			private SoftwareDocument()
+			{
+			}
+
+			public SoftwareDocument(ObjectId Id, byte[] Data, string? UserName, bool bDefault)
+			{
+				this.Id = Id;
+				this.Data = Data;
+				this.UploadedAtTime = DateTime.UtcNow;
+				this.UploadedByUser = UserName;
+
+				if (bDefault)
+				{
+					this.MadeDefaultAtTime = DateTime.UtcNow;
+				}
+			}
+
+		}
+
+		/// <summary>
+		/// Template documents
+		/// </summary>
+		IMongoCollection<SoftwareDocument> Collection;
+
+		/// <summary>
+		/// Constructor
+		/// </summary>
+		/// <param name="DatabaseService">The database service singleton</param>
+		public SoftwareCollection(DatabaseService DatabaseService)
+		{
+			Collection = DatabaseService.GetCollection<SoftwareDocument>("Software");
+			if (!DatabaseService.ReadOnlyMode)
+			{
+				Collection.Indexes.CreateOne(new CreateIndexModel<SoftwareDocument>(Builders<SoftwareDocument>.IndexKeys.Descending(x => x.MadeDefaultAtTime!)));
+			}
+		}
+
+		/// <inheritdoc/>
+		public async Task<ISoftware> AddAsync(byte[] Data, string? User, bool Default)
+		{
+			ObjectId Id = ObjectId.GenerateNewId();
+			byte[] NewData = UpdateVersion(Data, Id.ToString());
+
+			SoftwareDocument Document = new SoftwareDocument(Id, NewData, User, Default);
+			await Collection.InsertOneAsync(Document);
+			return Document;
+		}
+
+		/// <summary>
+		/// Updates the client id within the archive in the data array
+		/// </summary>
+		/// <param name="Data">Data for the zip archive</param>
+		/// <param name="Version">The new version identifier</param>
+		/// <returns>New agent app data</returns>
+		public static byte[] UpdateVersion(byte[] Data, string Version)
+		{
+			bool bWrittenClientId = false;
+
+			MemoryStream OutputStream = new MemoryStream();
+			using (ZipArchive OutputArchive = new ZipArchive(OutputStream, ZipArchiveMode.Create, true))
+			{
+				MemoryStream InputStream = new MemoryStream(Data);
+				using (ZipArchive InputArchive = new ZipArchive(InputStream, ZipArchiveMode.Read, true))
+				{
+					foreach (ZipArchiveEntry InputEntry in InputArchive.Entries)
+					{
+						ZipArchiveEntry OutputEntry = OutputArchive.CreateEntry(InputEntry.FullName);
+
+						using System.IO.Stream InputEntryStream = InputEntry.Open();
+						using System.IO.Stream OutputEntryStream = OutputEntry.Open();
+
+						if (InputEntry.FullName.Equals("appsettings.json", StringComparison.OrdinalIgnoreCase))
+						{
+							using MemoryStream MemoryStream = new MemoryStream();
+							InputEntryStream.CopyTo(MemoryStream);
+
+							Dictionary<string, Dictionary<string, object>> Document = JsonSerializer.Deserialize<Dictionary<string, Dictionary<string, object>>>(MemoryStream.ToArray());
+							Document["Horde"]["Version"] = Version;
+
+							using Utf8JsonWriter Writer = new Utf8JsonWriter(OutputEntryStream, new JsonWriterOptions { Indented = true });
+							JsonSerializer.Serialize<Dictionary<string, Dictionary<string, object>>>(Writer, Document, new JsonSerializerOptions { WriteIndented = true });
+
+							bWrittenClientId = true;
+						}
+						else
+						{
+							InputEntryStream.CopyTo(OutputEntryStream);
+						}
+					}
+				}
+			}
+
+			if (!bWrittenClientId)
+			{
+				throw new InvalidDataException("Missing appsettings.json file from zip archive");
+			}
+
+			return OutputStream.ToArray();
+		}
+
+		/// <inheritdoc/>
+		public async Task<bool> UpdateAsync(ObjectId Id, string? User, bool Default)
+		{
+			UpdateDefinition<SoftwareDocument> Update;
+			if (Default)
+			{
+				Update = Builders<SoftwareDocument>.Update.Set(x => x.MadeDefaultAtTime, DateTime.UtcNow).Set(x => x.MadeDefaultByUser, User);
+			}
+			else
+			{
+				Update = Builders<SoftwareDocument>.Update.Set(x => x.MadeDefaultAtTime, null);
+			}
+
+			UpdateResult Result = await Collection.UpdateOneAsync(x => x.Id == Id, Update);
+			return Result.ModifiedCount > 0;
+		}
+
+		/// <inheritdoc/>
+		public async Task<List<ISoftware>> FindAsync(string? UploadedByUser, string? MadeDefaultByUser, bool? MadeDefault, int Offset, int Count)
+		{
+			FilterDefinitionBuilder<SoftwareDocument> FilterBuilder = Builders<SoftwareDocument>.Filter;
+
+			FilterDefinition<SoftwareDocument> Filter = FilterDefinition<SoftwareDocument>.Empty;
+			if (UploadedByUser != null)
+			{
+				Filter &= FilterBuilder.Eq(x => x.UploadedByUser, UploadedByUser);
+			}
+			if (MadeDefaultByUser != null)
+			{
+				Filter &= FilterBuilder.Eq(x => x.MadeDefaultByUser, MadeDefaultByUser);
+			}
+			if (MadeDefault.HasValue)
+			{
+				if (MadeDefault.Value)
+				{
+					Filter &= FilterBuilder.Ne(x => x.MadeDefaultAtTime, null);
+				}
+				else
+				{
+					Filter &= FilterBuilder.Eq(x => x.MadeDefaultAtTime, null);
+				}
+			}
+
+			List<SoftwareDocument> Results = await Collection.Find(Filter).SortByDescending(x => x.MadeDefaultAtTime).Skip(Offset).Limit(Count).ToListAsync();
+			return Results.ConvertAll<ISoftware>(x => x);
+		}
+
+		/// <inheritdoc/>
+		public async Task<ISoftware?> GetAsync(ObjectId Id, bool bIncludeData)
+		{
+			IFindFluent<SoftwareDocument, SoftwareDocument> Matches = Collection.Find(x => x.Id == Id);
+			if (!bIncludeData)
+			{
+				Matches = Matches.Project<SoftwareDocument>(Builders<SoftwareDocument>.Projection.Exclude(x => x.Data));
+			}
+			return await Matches.FirstOrDefaultAsync();
+		}
+
+		/// <inheritdoc/>
+		public async Task<ISoftware?> GetDefaultAsync(bool bIncludeData)
+		{
+			IFindFluent<SoftwareDocument, SoftwareDocument> Matches = Collection.Find(FilterDefinition<SoftwareDocument>.Empty).SortByDescending(x => x.MadeDefaultAtTime);
+			if (!bIncludeData)
+			{
+				Matches = Matches.Project<SoftwareDocument>(Builders<SoftwareDocument>.Projection.Exclude(x => x.Data));
+			}
+
+			SoftwareDocument? DefaultSoftware = await Matches.FirstOrDefaultAsync();
+			if (DefaultSoftware != null && !DefaultSoftware.MadeDefaultAtTime.HasValue)
+			{
+				DefaultSoftware = null;
+			}
+
+			return DefaultSoftware;
+		}
+
+		/// <inheritdoc/>
+		public async Task<bool> DeleteAsync(ObjectId Id)
+		{
+			DeleteResult Result = await Collection.DeleteOneAsync(x => x.Id == Id);
+			return Result.DeletedCount > 0;
+		}
+	}
+}
diff --git a/Engine/Source/Programs/Horde/HordeServer/Collections/Impl/StreamCollection.cs b/Engine/Source/Programs/Horde/HordeServer/Collections/Impl/StreamCollection.cs
new file mode 100644
--- /dev/null
+++ b/Engine/Source/Programs/Horde/HordeServer/Collections/Impl/StreamCollection.cs
@@ -0,0 +1,485 @@
+using HordeServer.Api;
+using HordeServer.Models;
+using HordeServer.Services;
+using HordeServer.Utilities;
+using MongoDB.Bson;
+using MongoDB.Bson.Serialization.Attributes;
+using MongoDB.Driver;
+using System;
+using System.Collections.Generic;
+using System.Diagnostics.CodeAnalysis;
+using System.Linq;
+using System.Threading.Tasks;
+
+using ProjectId = HordeServer.Utilities.StringId<HordeServer.Models.IProject>;
+using StreamId = HordeServer.Utilities.StringId<HordeServer.Models.IStream>;
+using TemplateRefId = HordeServer.Utilities.StringId<HordeServer.Models.TemplateRef>;
+
+namespace HordeServer.Collections.Impl
+{
+	/// <summary>
+	/// Collection of stream documents
+	/// </summary>
+	public class StreamCollection : IStreamCollection
+	{
+		/// <summary>
+		/// Information about a stream
+		/// </summary>
+		class StreamDocument : IStream
+		{
+			public const int DefaultOrder = 128;
+
+			[BsonRequired, BsonId]
+			public StreamId Id { get; set; }
+
+			[BsonRequired]
+			public ProjectId ProjectId { get; set; }
+
+			[BsonRequired]
+			public string Name { get; set; }
+
+			public string? ConfigPath { get; set; }
+			public int? ConfigChange { get; set; }
+
+			public int Order { get; set; } = DefaultOrder;
+			public string? NotificationChannel { get; set; }
+			public string? NotificationChannelFilter { get; set; }
+			public string? TriageChannel { get; set; }
+			public DefaultPreflight? DefaultPreflight { get; set; }
+			public List<StreamTab> Tabs { get; set; } = new List<StreamTab>();
+			public Dictionary<string, AgentType> AgentTypes { get; set; } = new Dictionary<string, AgentType>(StringComparer.Ordinal);
+			public Dictionary<string, WorkspaceType> WorkspaceTypes { get; set; } = new Dictionary<string, WorkspaceType>(StringComparer.Ordinal);
+			public Dictionary<TemplateRefId, TemplateRef> Templates { get; set; } = new Dictionary<TemplateRefId, TemplateRef>();
+			public Dictionary<string, string> Properties { get; set; } = new Dictionary<string, string>(StringComparer.Ordinal);
+			public DateTime? LastCommitTime { get; set; }
+			public DateTime? PausedUntil { get; set; }
+			public string? PauseComment { get; set; }
+			public Acl? Acl { get; set; }
+			public int UpdateIndex { get; set; }
+
+			IReadOnlyList<StreamTab> IStream.Tabs => Tabs;
+			IReadOnlyDictionary<string, AgentType> IStream.AgentTypes => AgentTypes;
+			IReadOnlyDictionary<string, WorkspaceType> IStream.WorkspaceTypes => WorkspaceTypes;
+			IReadOnlyDictionary<TemplateRefId, TemplateRef> IStream.Templates => Templates;
+			IReadOnlyDictionary<string, string> IStream.Properties => Properties;
+
+			[BsonConstructor]
+			private StreamDocument()
+			{
+				Name = null!;
+			}
+
+			public StreamDocument(StreamId Id, string Name, ProjectId ProjectId, string? ConfigPath)
+			{
+				this.Id = Id;
+				this.Name = Name;
+				this.ProjectId = ProjectId;
+				this.ConfigPath = ConfigPath;
+			}
+		}
+
+		/// <summary>
+		/// Projection of a stream definition to just include permissions info
+		/// </summary>
+		[SuppressMessage("Design", "CA1812: Class is never instantiated")]
+		private class StreamPermissions : IStreamPermissions
+		{
+			public Acl? Acl { get; set; }
+			public ProjectId ProjectId { get; set; }
+
+			public static readonly ProjectionDefinition<StreamDocument> Projection = Builders<StreamDocument>.Projection.Include(x => x.Acl).Include(x => x.ProjectId);
+		}
+
+		/// <summary>
+		/// The stream collection
+		/// </summary>
+		IMongoCollection<StreamDocument> Streams;
+
+		/// <summary>
+		/// Constructor
+		/// </summary>
+		/// <param name="DatabaseService">The database service instance</param>
+		public StreamCollection(DatabaseService DatabaseService)
+		{
+			this.Streams = DatabaseService.GetCollection<StreamDocument>("Streams");
+		}
+
+		/// <inheritdoc/>
+		public async Task<IStream?> TryCreateAsync(StreamId Id, string Name, ProjectId ProjectId, string? ConfigPath, int? Order, string? NotificationChannel, string? NotificationChannelFilter, string? TriageChannel, DefaultPreflight? DefaultPreflight, List<StreamTab>? Tabs, Dictionary<string, AgentType>? AgentTypes, Dictionary<string, WorkspaceType>? WorkspaceTypes, Dictionary<TemplateRefId, TemplateRef>? TemplateRefs, Dictionary<string, string>? Properties, Acl? Acl)
+		{
+			StreamDocument NewStream = new StreamDocument(Id, Name, ProjectId, ConfigPath);
+			if (Order != null)
+			{
+				NewStream.Order = Order.Value;
+			}
+			if(NotificationChannel != null)
+			{
+				NewStream.NotificationChannel = NotificationChannel;
+			}
+			if (NotificationChannelFilter != null)
+			{
+				NewStream.NotificationChannelFilter = NotificationChannelFilter;
+			}
+			if (TriageChannel != null)
+			{
+				NewStream.TriageChannel = TriageChannel;
+			}
+			if (DefaultPreflight != null)
+			{
+				NewStream.DefaultPreflight = DefaultPreflight;
+			}
+			if (Tabs != null)
+			{
+				NewStream.Tabs = Tabs;
+			}
+			if (AgentTypes != null)
+			{
+				NewStream.AgentTypes = new Dictionary<string, AgentType>(AgentTypes, NewStream.AgentTypes.Comparer);
+			}
+			if (WorkspaceTypes != null)
+			{
+				NewStream.WorkspaceTypes = new Dictionary<string, WorkspaceType>(WorkspaceTypes, NewStream.WorkspaceTypes.Comparer);
+			}
+			if (TemplateRefs != null)
+			{
+				NewStream.Templates = TemplateRefs;
+			}
+			if (Properties != null)
+			{
+				NewStream.Properties = new Dictionary<string, string>(Properties, NewStream.Properties.Comparer);
+			}
+			if (Acl != null)
+			{
+				NewStream.Acl = Acl;
+			}
+			NewStream.Validate();
+
+			try
+			{
+				await Streams.InsertOneAsync(NewStream);
+				return NewStream;
+			}
+			catch (MongoWriteException Ex)
+			{
+				if (Ex.WriteError.Category == ServerErrorCategory.DuplicateKey)
+				{
+					return null;
+				}
+				else
+				{
+					throw;
+				}
+			}
+		}
+
+
+		/// <inheritdoc/>
+		public async Task<IStream?> TryReplaceAsync(IStream StreamInterface, string Name, string? ConfigPath, int? ConfigChange, int? Order, string? NotificationChannel, string? NotificationChannelFilter, string? TriageChannel, DefaultPreflight? DefaultPreflight, List<StreamTab>? Tabs, Dictionary<string, AgentType>? AgentTypes, Dictionary<string, WorkspaceType>? WorkspaceTypes, Dictionary<TemplateRefId, TemplateRef>? TemplateRefs, Dictionary<string, string>? Properties, Acl? Acl, DateTime? PausedUntil = null, string? PauseComment = null)
+		{
+			StreamDocument Stream = (StreamDocument)StreamInterface;
+
+			UpdateDefinitionBuilder<StreamDocument> UpdateBuilder = Builders<StreamDocument>.Update;
+
+			List<UpdateDefinition<StreamDocument>> Updates = new List<UpdateDefinition<StreamDocument>>();
+
+			Stream.Name = Name;
+			Updates.Add(UpdateBuilder.Set(x => x.Name, Stream.Name));
+
+			if (ConfigPath != null)
+			{
+				if (ConfigPath.Length == 0)
+				{
+					Stream.ConfigPath = null;
+					Updates.Add(UpdateBuilder.Unset(x => x.ConfigPath));
+				}
+				else
+				{
+					Stream.ConfigPath = ConfigPath;
+					Updates.Add(UpdateBuilder.Set(x => x.ConfigPath, ConfigPath));
+				}
+
+				ConfigChange = ConfigChange ?? 0;
+			}
+
+			if (ConfigChange != null)
+			{
+				if (ConfigChange == 0)
+				{
+					Stream.ConfigChange = null;
+					Updates.Add(UpdateBuilder.Unset(x => x.ConfigChange));
+				}
+				else
+				{
+					Stream.ConfigChange = ConfigChange;
+					Updates.Add(UpdateBuilder.Set(x => x.ConfigChange, ConfigChange.Value));
+				}
+			}
+
+			Stream.Order = Order ?? StreamDocument.DefaultOrder;
+			Updates.Add(UpdateBuilder.Set(x => x.Order, Stream.Order));
+
+			Stream.NotificationChannel = NotificationChannel;
+			Updates.Add(UpdateBuilder.Set(x => x.NotificationChannel, Stream.NotificationChannel));
+
+			Stream.NotificationChannelFilter = NotificationChannelFilter;
+			Updates.Add(UpdateBuilder.Set(x => x.NotificationChannelFilter, Stream.NotificationChannelFilter));
+
+			Stream.TriageChannel = TriageChannel;
+			Updates.Add(UpdateBuilder.Set(x => x.TriageChannel, Stream.TriageChannel));
+
+			Stream.DefaultPreflight = DefaultPreflight;
+			Updates.Add(UpdateBuilder.Set(x => x.DefaultPreflight, Stream.DefaultPreflight));
+			
+			Stream.Tabs = Tabs ?? new List<StreamTab>();
+			Updates.Add(UpdateBuilder.Set(x => x.Tabs, Stream.Tabs));
+
+			Stream.AgentTypes = AgentTypes ?? new Dictionary<string, AgentType>();
+			Updates.Add(UpdateBuilder.Set(x => x.AgentTypes, Stream.AgentTypes));
+
+			Stream.WorkspaceTypes = WorkspaceTypes ?? new Dictionary<string, WorkspaceType>();
+			Updates.Add(UpdateBuilder.Set(x => x.WorkspaceTypes, Stream.WorkspaceTypes));
+
+			Stream.Templates = TemplateRefs ?? new Dictionary<TemplateRefId, TemplateRef>();
+			Updates.Add(UpdateBuilder.Set(x => x.Templates, Stream.Templates));
+
+			Stream.Properties = Properties ?? new Dictionary<string, string>();
+			Updates.Add(UpdateBuilder.Set(x => x.Properties, Stream.Properties));
+
+			Stream.Acl = Acl;
+			Updates.Add(UpdateBuilder.SetOrUnsetNullRef(x => x.Acl, Acl));
+			
+			Stream.PausedUntil = PausedUntil;
+			Updates.Add(UpdateBuilder.Set(x => x.PausedUntil, Stream.PausedUntil));
+			
+			Stream.PauseComment = PauseComment;
+			Updates.Add(UpdateBuilder.Set(x => x.PauseComment, Stream.PauseComment));
+
+			Stream.Validate();
+			if(await TryUpdateStreamAsync(Stream, UpdateBuilder.Combine(Updates)))
+			{
+				return Stream;
+			}
+			else
+			{
+				return null;
+			}
+		}
+
+		/// <inheritdoc/>
+		public async Task<IStream?> GetAsync(StreamId StreamId)
+		{
+			return await Streams.Find<StreamDocument>(x => x.Id == StreamId).FirstOrDefaultAsync();
+		}
+
+		/// <inheritdoc/>
+		public async Task<IStreamPermissions?> GetPermissionsAsync(StreamId StreamId)
+		{
+			return await Streams.Find<StreamDocument>(x => x.Id == StreamId).Project<StreamPermissions>(StreamPermissions.Projection).FirstOrDefaultAsync();
+		}
+
+		/// <inheritdoc/>
+		public async Task<List<IStream>> FindAllAsync()
+		{
+			List<StreamDocument> Results = await Streams.Find(FilterDefinition<StreamDocument>.Empty).ToListAsync();
+			return Results.ConvertAll<IStream>(x => x);
+		}
+
+		/// <inheritdoc/>
+		public async Task<List<IStream>> FindForProjectsAsync(ProjectId[] ProjectIds)
+		{
+			FilterDefinition<StreamDocument> Filter = Builders<StreamDocument>.Filter.In(x => x.ProjectId, ProjectIds);
+			List<StreamDocument> Results = await Streams.Find(Filter).ToListAsync();
+			return Results.ConvertAll<IStream>(x => x);
+		}
+
+		/// <inheritdoc/>
+		public Task<bool> TryUpdatePropertiesAsync(IStream StreamInterface, string? NewName, int? NewOrder, string? NewNotificationChannel, string? NewNotificationChannelFilter, string? NewTriageChannel, List<StreamTab>? NewTabs, Dictionary<string, AgentType?>? NewAgentTypes, Dictionary<string, WorkspaceType?>? NewWorkspaceTypes, Dictionary<TemplateRefId, TemplateRef>? NewTemplateRefs, Dictionary<string, string?>? NewProperties, Acl? NewAcl, bool? UpdatePauseFields, DateTime? NewPausedUntil, string? NewPauseComment)
+		{
+			StreamDocument Stream = (StreamDocument)StreamInterface;
+
+			UpdateDefinitionBuilder<StreamDocument> UpdateBuilder = Builders<StreamDocument>.Update;
+
+			List<UpdateDefinition<StreamDocument>> Updates = new List<UpdateDefinition<StreamDocument>>();
+			if (NewName != null)
+			{
+				Stream.Name = NewName;
+				Updates.Add(UpdateBuilder.Set(x => x.Name, NewName));
+			}
+			if (NewOrder != null)
+			{
+				Stream.Order = NewOrder.Value;
+				Updates.Add(UpdateBuilder.Set(x => x.Order, NewOrder.Value));
+			}
+			if (NewNotificationChannel != null)
+			{
+				Stream.NotificationChannel = NewNotificationChannel;
+				Updates.Add(UpdateBuilder.Set(x => x.NotificationChannel, NewNotificationChannel));
+			}
+			if (NewNotificationChannelFilter != null)
+			{
+				Stream.NotificationChannelFilter = NewNotificationChannelFilter;
+				Updates.Add(UpdateBuilder.Set(x => x.NotificationChannelFilter, NewNotificationChannelFilter));
+			}
+			if (NewTriageChannel != null)
+			{
+				Stream.TriageChannel = NewTriageChannel;
+				Updates.Add(UpdateBuilder.Set(x => x.TriageChannel, NewTriageChannel));
+			}
+			if (NewTabs != null)
+			{
+				Stream.Tabs = NewTabs;
+				Updates.Add(UpdateBuilder.Set(x => x.Tabs, NewTabs));
+			}
+			if (NewAgentTypes != null)
+			{
+				foreach (KeyValuePair<string, AgentType?> Pair in NewAgentTypes)
+				{
+					if (Pair.Value == null)
+					{
+						Stream.AgentTypes.Remove(Pair.Key);
+						Updates.Add(UpdateBuilder.Unset(x => x.AgentTypes[Pair.Key]));
+					}
+					else
+					{
+						Stream.AgentTypes[Pair.Key] = Pair.Value;
+						Updates.Add(UpdateBuilder.Set(x => x.AgentTypes[Pair.Key], Pair.Value));
+					}
+				}
+			}
+			if (NewWorkspaceTypes != null)
+			{
+				foreach (KeyValuePair<string, WorkspaceType?> Pair in NewWorkspaceTypes)
+				{
+					if (Pair.Value == null)
+					{
+						Stream.WorkspaceTypes.Remove(Pair.Key);
+						Updates.Add(UpdateBuilder.Unset(x => x.WorkspaceTypes[Pair.Key]));
+					}
+					else
+					{
+						Stream.WorkspaceTypes[Pair.Key] = Pair.Value;
+						Updates.Add(UpdateBuilder.Set(x => x.WorkspaceTypes[Pair.Key], Pair.Value));
+					}
+				}
+			}
+			if (NewTemplateRefs != null)
+			{
+				Stream.Templates = NewTemplateRefs;
+				Updates.Add(UpdateBuilder.Set(x => x.Templates, Stream.Templates));
+			}
+			if (NewProperties != null)
+			{
+				foreach (KeyValuePair<string, string?> Pair in NewProperties)
+				{
+					if (Pair.Value == null)
+					{
+						Stream.Properties.Remove(Pair.Key);
+						Updates.Add(UpdateBuilder.Unset(x => x.Properties[Pair.Key]));
+					}
+					else
+					{
+						Stream.Properties[Pair.Key] = Pair.Value;
+						Updates.Add(UpdateBuilder.Set(x => x.Properties[Pair.Key], Pair.Value));
+					}
+				}
+			}
+			if (NewAcl != null)
+			{
+				Stream.Acl = NewAcl;
+				Updates.Add(Acl.CreateUpdate<StreamDocument>(x => x.Acl!, NewAcl));
+			}
+
+			if (UpdatePauseFields != null && UpdatePauseFields == true)
+			{
+				Stream.PausedUntil = NewPausedUntil;
+				Stream.PauseComment = NewPauseComment;
+				Updates.Add(UpdateBuilder.Set(x => x.PausedUntil, NewPausedUntil));
+				Updates.Add(UpdateBuilder.Set(x => x.PauseComment, NewPauseComment));
+			}
+
+			if (Updates.Count > 0)
+			{
+				return TryUpdateStreamAsync(Stream, UpdateBuilder.Combine(Updates));
+			}
+			else
+			{
+				return Task.FromResult(true);
+			}
+		}
+
+		/// <inheritdoc/>
+		public async Task<bool> TryUpdateCommitTimeAsync(IStream StreamInterface, DateTime LastCommitTime)
+		{
+			StreamDocument Stream = (StreamDocument)StreamInterface;
+
+			// Get the update to apply
+			UpdateDefinition<StreamDocument> Update = Builders<StreamDocument>.Update
+				.Set(x => x.LastCommitTime, Stream.LastCommitTime)
+				.Set(x => x.UpdateIndex, Stream.UpdateIndex + 1);
+
+			// Apply the update
+			UpdateResult Result = await Streams.UpdateOneAsync(x => x.Id == Stream.Id && x.UpdateIndex == Stream.UpdateIndex, Update);
+			if (Result.ModifiedCount == 0)
+			{
+				return false;
+			}
+
+			// Update the input document
+			Stream.LastCommitTime = LastCommitTime;
+			Stream.UpdateIndex++;
+			return true;
+		}
+
+		/// <inheritdoc/>
+		public async Task<bool> TryUpdateScheduleTriggerAsync(IStream StreamInterface, TemplateRefId TemplateRefId, DateTimeOffset? LastTriggerTime, int? LastTriggerChange, List<ObjectId> NewActiveJobs)
+		{
+			StreamDocument Stream = (StreamDocument)StreamInterface;
+			Schedule Schedule = Stream.Templates[TemplateRefId].Schedule!;
+
+			// Build the updates. MongoDB driver cannot parse TemplateRefId in expression tree; need to specify field name explicitly
+			List<UpdateDefinition<StreamDocument>> Updates = new List<UpdateDefinition<StreamDocument>>();
+			if (LastTriggerTime.HasValue && LastTriggerTime.Value != Schedule.LastTriggerTime)
+			{
+				FieldDefinition<StreamDocument, DateTimeOffset> LastTriggerTimeField = $"{nameof(Stream.Templates)}.{TemplateRefId}.{nameof(Schedule)}.{nameof(Schedule.LastTriggerTime)}";
+				Updates.Add(Builders<StreamDocument>.Update.Set(LastTriggerTimeField, LastTriggerTime.Value));
+				Schedule.LastTriggerTime = LastTriggerTime.Value;
+			}
+			if (LastTriggerChange.HasValue && LastTriggerChange.Value != Schedule.LastTriggerChange)
+			{
+				FieldDefinition<StreamDocument, int> LastTriggerChangeField = $"{nameof(Stream.Templates)}.{TemplateRefId}.{nameof(Schedule)}.{nameof(Schedule.LastTriggerChange)}";
+				Updates.Add(Builders<StreamDocument>.Update.Set(LastTriggerChangeField, LastTriggerChange.Value));
+				Schedule.LastTriggerChange = LastTriggerChange.Value;
+			}
+			if (NewActiveJobs != null)
+			{
+				FieldDefinition<StreamDocument, List<ObjectId>> Field = $"{nameof(Stream.Templates)}.{TemplateRefId}.{nameof(Schedule)}.{nameof(Schedule.ActiveJobs)}";
+				Updates.Add(Builders<StreamDocument>.Update.Set(Field, NewActiveJobs));
+				Schedule.ActiveJobs = NewActiveJobs;
+			}
+
+			return Updates.Count == 0 || await TryUpdateStreamAsync(Stream, Builders<StreamDocument>.Update.Combine(Updates));
+		}
+
+		/// <summary>
+		/// Update a stream
+		/// </summary>
+		/// <param name="Stream">The stream to update</param>
+		/// <param name="Update">The update definition</param>
+		/// <returns>True if the stream was updated, false otherwise</returns>
+		private async Task<bool> TryUpdateStreamAsync(StreamDocument Stream, UpdateDefinition<StreamDocument> Update)
+		{
+			Stream.Validate();
+
+			int InitialUpdateIndex = Stream.UpdateIndex++;
+
+			UpdateResult Result = await Streams.UpdateOneAsync(x => x.Id == Stream.Id && x.UpdateIndex == InitialUpdateIndex, Update.Set(x => x.UpdateIndex, Stream.UpdateIndex));
+			return Result.ModifiedCount > 0;
+		}
+
+		/// <inheritdoc/>
+		public async Task DeleteAsync(StreamId StreamId)
+		{
+			await Streams.DeleteOneAsync<StreamDocument>(x => x.Id == StreamId);
+		}
+	}
+}
diff --git a/Engine/Source/Programs/Horde/HordeServer/Collections/Impl/TelemetryCollection.cs b/Engine/Source/Programs/Horde/HordeServer/Collections/Impl/TelemetryCollection.cs
new file mode 100644
--- /dev/null
+++ b/Engine/Source/Programs/Horde/HordeServer/Collections/Impl/TelemetryCollection.cs
@@ -0,0 +1,134 @@
+using EpicGames.Core;
+using HordeServer.Models;
+using HordeServer.Services;
+using HordeServer.Utilities;
+using Microsoft.Extensions.Hosting;
+using MongoDB.Bson;
+using MongoDB.Bson.Serialization.Attributes;
+using MongoDB.Driver;
+using System;
+using System.Collections.Generic;
+using System.Linq;
+using System.Threading;
+using System.Threading.Tasks;
+
+using PoolId = HordeServer.Utilities.StringId<HordeServer.Models.IPool>;
+using StreamId = HordeServer.Utilities.StringId<HordeServer.Models.IStream>;
+
+namespace HordeServer.Collections.Impl
+{
+	/// <summary>
+	/// Collection of utilization data
+	/// </summary>
+	public class TelemetryCollection : ITelemetryCollection
+	{
+		class UtilizationDocument : IUtilizationTelemetry
+		{
+			public DateTime StartTime { get; set; }
+			public DateTime FinishTime { get; set; }
+
+			public int NumAgents { get; set; }
+			public List<PoolUtilizationDocument> Pools { get; set; }
+			IReadOnlyList<IPoolUtilizationTelemetry> IUtilizationTelemetry.Pools => Pools;
+
+			public double AdminTime { get; set; }
+			public int UpdateIndex { get; set; }
+
+			[BsonConstructor]
+			private UtilizationDocument()
+			{
+				Pools = new List<PoolUtilizationDocument>();
+			}
+
+			public UtilizationDocument(IUtilizationTelemetry Other)
+			{
+				this.StartTime = Other.StartTime;
+				this.FinishTime = Other.FinishTime;
+				this.NumAgents = Other.NumAgents;
+				this.Pools = Other.Pools.ConvertAll(x => new PoolUtilizationDocument(x));
+				this.AdminTime = Other.AdminTime;
+			}
+		}
+
+		class PoolUtilizationDocument : IPoolUtilizationTelemetry
+		{
+			public PoolId PoolId { get; set; }
+			public int NumAgents { get; set; }
+
+			public List<StreamUtilizationDocument> Streams { get; set; } 
+			IReadOnlyList<IStreamUtilizationTelemetry> IPoolUtilizationTelemetry.Streams => Streams;
+
+			public double AdminTime { get; set; }
+			public double OtherTime { get; set; }
+
+			[BsonConstructor]
+			private PoolUtilizationDocument()
+			{
+				Streams = new List<StreamUtilizationDocument>();
+			}
+
+			public PoolUtilizationDocument(IPoolUtilizationTelemetry Other)
+			{
+				this.PoolId = Other.PoolId;
+				this.NumAgents = Other.NumAgents;
+				this.Streams = Other.Streams.ConvertAll(x => new StreamUtilizationDocument(x));
+				this.AdminTime = Other.AdminTime;
+				this.OtherTime = Other.OtherTime;
+			}
+		}
+
+		class StreamUtilizationDocument : IStreamUtilizationTelemetry
+		{
+			public StreamId StreamId { get; set; }
+			public double Time { get; set; }
+
+			[BsonConstructor]
+			private StreamUtilizationDocument()
+			{
+			}
+
+			public StreamUtilizationDocument(IStreamUtilizationTelemetry Other)
+			{
+				this.StreamId = Other.StreamId;
+				this.Time = Other.Time;
+			}
+		}
+
+		IMongoCollection<UtilizationDocument> Utilization;
+
+		/// <summary>
+		/// Constructor
+		/// </summary>
+		/// <param name="Database"></param>
+		public TelemetryCollection(DatabaseService Database)
+		{
+			Utilization = Database.GetCollection<UtilizationDocument>("Utilization");
+
+			if (!Database.ReadOnlyMode)
+			{
+				Utilization.Indexes.CreateOne(new CreateIndexModel<UtilizationDocument>(Builders<UtilizationDocument>.IndexKeys.Ascending(x => x.FinishTime).Ascending(x => x.StartTime)));
+			}
+		}
+
+		/// <inheritdoc/>
+		public async Task AddUtilizationTelemetryAsync(IUtilizationTelemetry NewTelemetry)
+		{
+			await Utilization.InsertOneAsync(new UtilizationDocument(NewTelemetry));
+		}
+
+		/// <inheritdoc/>
+		public async Task<List<IUtilizationTelemetry>> GetUtilizationTelemetryAsync(DateTime StartTimeUtc, DateTime FinishTimeUtc)
+		{
+			FilterDefinition<UtilizationDocument> Filter = Builders<UtilizationDocument>.Filter.Gte(x => x.FinishTime, StartTimeUtc) & Builders<UtilizationDocument>.Filter.Lte(x => x.StartTime, FinishTimeUtc);
+
+			List<UtilizationDocument> Documents = await Utilization.Find(Filter).ToListAsync();
+			return Documents.ConvertAll<IUtilizationTelemetry>(x => x);
+		}
+
+		/// <inheritdoc/>
+		public async Task<IUtilizationTelemetry?> GetLatestUtilizationTelemetryAsync()
+		{
+			return await Utilization.Find(FilterDefinition<UtilizationDocument>.Empty).SortByDescending(x => x.FinishTime).FirstOrDefaultAsync();
+		}
+	}
+}
diff --git a/Engine/Source/Programs/Horde/HordeServer/Collections/Impl/UserCollection.cs b/Engine/Source/Programs/Horde/HordeServer/Collections/Impl/UserCollection.cs
new file mode 100644
--- /dev/null
+++ b/Engine/Source/Programs/Horde/HordeServer/Collections/Impl/UserCollection.cs
@@ -0,0 +1,201 @@
+using HordeServer.Models;
+using HordeServer.Services;
+using HordeServer.Utilities;
+using MongoDB.Bson;
+using MongoDB.Bson.Serialization;
+using MongoDB.Bson.Serialization.Attributes;
+using MongoDB.Driver;
+using System;
+using System.Collections.Generic;
+using System.Linq;
+using System.Security.Claims;
+using System.Threading.Tasks;
+
+namespace HordeServer.Collections.Impl
+{
+	/// <summary>
+	/// Manages user documents
+	/// </summary>
+	class UserCollectionV1 : IUserCollection
+	{
+		class UserDocument : IUser, IUserClaims, IUserSettings
+		{
+			public ObjectId Id { get; set; }
+
+			public ClaimDocument PrimaryClaim { get; set; } = null!;
+			public List<ClaimDocument> Claims { get; set; } = new List<ClaimDocument>();
+
+			[BsonDefaultValue(false), BsonIgnoreIfDefault]
+			public bool EnableExperimentalFeatures { get; set; }
+
+			[BsonDefaultValue(false), BsonIgnoreIfDefault]
+			public bool EnableIssueNotifications { get; set; }
+
+			public BsonValue DashboardSettings { get; set; } = BsonNull.Value;
+			public List<ObjectId> PinnedJobIds { get; set; } = new List<ObjectId>();
+
+			string IUser.Name => Claims.FirstOrDefault(x => String.Equals(x.Type, "name", StringComparison.Ordinal))?.Value ?? PrimaryClaim.Value;
+			string IUser.Login => Claims.FirstOrDefault(x => String.Equals(x.Type, ClaimTypes.Name, StringComparison.Ordinal))?.Value ?? PrimaryClaim.Value;
+			string? IUser.Email => Claims.FirstOrDefault(x => String.Equals(x.Type, ClaimTypes.Email, StringComparison.Ordinal))?.Value;
+
+			ObjectId IUserClaims.UserId => Id;
+			IReadOnlyList<IUserClaim> IUserClaims.Claims => Claims;
+
+			ObjectId IUserSettings.UserId => Id;
+			IReadOnlyList<ObjectId> IUserSettings.PinnedJobIds => PinnedJobIds;
+		}
+
+		class ClaimDocument : IUserClaim
+		{
+			public string Type { get; set; }
+			public string Value { get; set; }
+
+			private ClaimDocument()
+			{
+				this.Type = null!;
+				this.Value = null!;
+			}
+
+			public ClaimDocument(string Type, string Value)
+			{
+				this.Type = Type;
+				this.Value = Value;
+			}
+
+			public ClaimDocument(IUserClaim Other)
+			{
+				this.Type = Other.Type;
+				this.Value = Other.Value;
+			}
+		}
+
+		IMongoCollection<UserDocument> Users;
+
+		/// <summary>
+		/// Static constructor
+		/// </summary>
+		static UserCollectionV1()
+		{
+			BsonSerializer.RegisterDiscriminatorConvention(typeof(UserDocument), NullDiscriminatorConvention.Instance);
+			BsonSerializer.RegisterDiscriminatorConvention(typeof(ClaimDocument), NullDiscriminatorConvention.Instance);
+		}
+
+		/// <summary>
+		/// Constructor
+		/// </summary>
+		/// <param name="DatabaseService"></param>
+		public UserCollectionV1(DatabaseService DatabaseService)
+		{
+			Users = DatabaseService.GetCollection<UserDocument>("Users");
+
+			if (!DatabaseService.ReadOnlyMode)
+			{
+				Users.Indexes.CreateOne(new CreateIndexModel<UserDocument>(Builders<UserDocument>.IndexKeys.Ascending(x => x.PrimaryClaim), new CreateIndexOptions { Unique = true }));
+			}
+		}
+
+		/// <inheritdoc/>
+		public async Task<IUser?> GetUserAsync(ObjectId Id)
+		{
+			return await Users.Find(x => x.Id == Id).FirstOrDefaultAsync();
+		}
+
+		/// <inheritdoc/>
+		public async Task<List<IUser>> FindUsersAsync(IEnumerable<ObjectId>? Ids, int? Index, int? Count)
+		{
+			FilterDefinition<UserDocument> Filter = Builders<UserDocument>.Filter.In(x => x.Id, Ids);
+			return await Users.Find(Filter).Range(Index, Count).ToListAsync<UserDocument, IUser>();
+		}
+
+		/// <inheritdoc/>
+		public async Task<IUser?> FindUserByLoginAsync(string Login)
+		{
+			ClaimDocument PrimaryClaim = new ClaimDocument(ClaimTypes.Name, Login);
+			return await Users.Find(x => x.PrimaryClaim == PrimaryClaim).FirstOrDefaultAsync();
+		}
+
+		/// <inheritdoc/>
+		public async Task<IUser> FindOrAddUserByLoginAsync(string Login, string? Name, string? Email)
+		{
+			ClaimDocument NewPrimaryClaim = new ClaimDocument(ClaimTypes.Name, Login);
+			UpdateDefinition<UserDocument> Update = Builders<UserDocument>.Update.SetOnInsert(x => x.Id, ObjectId.GenerateNewId());
+			return await Users.FindOneAndUpdateAsync<UserDocument>(x => x.PrimaryClaim == NewPrimaryClaim, Update, new FindOneAndUpdateOptions<UserDocument> { IsUpsert = true, ReturnDocument = ReturnDocument.After });
+		}
+
+		/// <inheritdoc/>
+		public async Task<IUserClaims> GetClaimsAsync(ObjectId UserId)
+		{
+			return await Users.Find(x => x.Id == UserId).FirstOrDefaultAsync() ?? new UserDocument { Id = UserId };
+		}
+
+		/// <inheritdoc/>
+		public async Task UpdateClaimsAsync(ObjectId UserId, IEnumerable<IUserClaim> Claims)
+		{
+			List<ClaimDocument> NewClaims = Claims.Select(x => new ClaimDocument(x)).ToList();
+			await Users.FindOneAndUpdateAsync(x => x.Id == UserId, Builders<UserDocument>.Update.Set(x => x.Claims, NewClaims));
+		}
+
+		/// <inheritdoc/>
+		public async Task<IUserSettings> GetSettingsAsync(ObjectId UserId)
+		{
+			return await Users.Find(x => x.Id == UserId).FirstOrDefaultAsync() ?? new UserDocument { Id = UserId };
+		}
+
+		/// <inheritdoc/>
+		public async Task UpdateSettingsAsync(ObjectId UserId, bool? EnableExperimentalFeatures, bool? EnableIssueNotifications, BsonValue? DashboardSettings = null, IEnumerable<ObjectId>? AddPinnedJobIds = null, IEnumerable<ObjectId>? RemovePinnedJobIds = null)
+		{
+			if (AddPinnedJobIds != null)
+			{
+				foreach (ObjectId PinnedJobId in AddPinnedJobIds)
+				{
+					FilterDefinition<UserDocument> Filter = Builders<UserDocument>.Filter.Eq(x => x.Id, UserId) & Builders<UserDocument>.Filter.AnyNin<ObjectId>(x => x.PinnedJobIds, new[] { PinnedJobId });
+					UpdateDefinition<UserDocument> Update = Builders<UserDocument>.Update.PushEach(x => x.PinnedJobIds, new[] { PinnedJobId }, -50);
+					await Users.UpdateOneAsync(Filter, Update);
+				}
+			}
+
+			List<UpdateDefinition<UserDocument>> Updates = new List<UpdateDefinition<UserDocument>>();
+			if (EnableExperimentalFeatures != null)
+			{
+				Updates.Add(Builders<UserDocument>.Update.SetOrUnsetNull(x => x.EnableExperimentalFeatures, EnableExperimentalFeatures));
+			}
+			if (EnableIssueNotifications != null)
+			{
+				Updates.Add(Builders<UserDocument>.Update.SetOrUnsetNull(x => x.EnableIssueNotifications, EnableIssueNotifications));
+			}
+			if (DashboardSettings != null)
+			{
+				Updates.Add(Builders<UserDocument>.Update.Set(x => x.DashboardSettings, DashboardSettings));
+			}
+			if (RemovePinnedJobIds != null && RemovePinnedJobIds.Any())
+			{
+				Updates.Add(Builders<UserDocument>.Update.PullAll(x => x.PinnedJobIds, RemovePinnedJobIds));
+			}
+			if (Updates.Count > 0)
+			{
+				await Users.UpdateOneAsync<UserDocument>(x => x.Id == UserId, Builders<UserDocument>.Update.Combine(Updates));
+			}
+		}
+
+		/// <summary>
+		/// Enumerate all the documents in this collection
+		/// </summary>
+		/// <returns></returns>
+		public async IAsyncEnumerable<(IUser, IUserClaims, IUserSettings)> EnumerateDocumentsAsync()
+		{
+			using (IAsyncCursor<UserDocument> Cursor = await Users.Find(FilterDefinition<UserDocument>.Empty).ToCursorAsync())
+			{
+				while (await Cursor.MoveNextAsync())
+				{
+					foreach (UserDocument Document in Cursor.Current)
+					{
+						if (Document.Claims.Count > 0)
+						{
+							yield return (Document, Document, Document);
+						}
+					}
+				}
+			}
+		}
+	}
+}
diff --git a/Engine/Source/Programs/Horde/HordeServer/Collections/Impl/UserCollectionV2.cs b/Engine/Source/Programs/Horde/HordeServer/Collections/Impl/UserCollectionV2.cs
new file mode 100644
--- /dev/null
+++ b/Engine/Source/Programs/Horde/HordeServer/Collections/Impl/UserCollectionV2.cs
@@ -0,0 +1,321 @@
+using HordeServer.Models;
+using HordeServer.Services;
+using HordeServer.Utilities;
+using Microsoft.Extensions.Logging;
+using MongoDB.Bson;
+using MongoDB.Bson.Serialization;
+using MongoDB.Bson.Serialization.Attributes;
+using MongoDB.Driver;
+using System;
+using System.Collections.Generic;
+using System.Linq;
+using System.Security.Claims;
+using System.Threading.Tasks;
+
+namespace HordeServer.Collections.Impl
+{
+	/// <summary>
+	/// Manages user documents
+	/// </summary>
+	class UserCollectionV2 : IUserCollection
+	{
+		class UserDocument : IUser
+		{
+			public ObjectId Id { get; set; }
+
+			public string Name { get; set; }
+			public string Login { get; set; }
+			public string LoginUpper { get; set; }
+			public string? Email { get; set; }
+			public string? EmailUpper { get; set; }
+
+			[BsonConstructor]
+			private UserDocument()
+			{
+				Name = null!;
+				Login = null!;
+				LoginUpper = null!;
+			}
+
+			public UserDocument(IUser Other)
+				: this(Other.Id, Other.Name, Other.Login, Other.Email)
+			{
+			}
+
+			public UserDocument(ObjectId Id, string Name, string Login, string? Email)
+			{
+				this.Id = Id;
+				this.Name = Name;
+				this.Login = Login;
+				this.LoginUpper = Login.ToUpperInvariant();
+				this.Email = Email;
+				this.EmailUpper = Email?.ToUpperInvariant();
+			}
+		}
+
+		class UserClaimsDocument : IUserClaims
+		{
+			public ObjectId Id { get; set; }
+			public List<UserClaim> Claims { get; set; } = new List<UserClaim>();
+
+			ObjectId IUserClaims.UserId => Id;
+			IReadOnlyList<IUserClaim> IUserClaims.Claims => Claims;
+
+			[BsonConstructor]
+			private UserClaimsDocument()
+			{
+			}
+
+			public UserClaimsDocument(ObjectId Id)
+			{
+				this.Id = Id;
+			}
+
+			public UserClaimsDocument(IUserClaims Other)
+				: this(Other.UserId)
+			{
+				this.Claims.AddRange(Other.Claims.Select(x => new UserClaim(x)));
+			}
+		}
+
+		class UserSettingsDocument : IUserSettings
+		{
+			public ObjectId Id { get; set; }
+
+			[BsonDefaultValue(false), BsonIgnoreIfDefault]
+			public bool EnableExperimentalFeatures { get; set; }
+
+			[BsonDefaultValue(true), BsonIgnoreIfDefault]
+			public bool EnableIssueNotifications { get; set; } = true;
+
+			public BsonValue DashboardSettings { get; set; } = BsonNull.Value;
+			public List<ObjectId> PinnedJobIds { get; set; } = new List<ObjectId>();
+
+			ObjectId IUserSettings.UserId => Id;
+			IReadOnlyList<ObjectId> IUserSettings.PinnedJobIds => PinnedJobIds;
+
+			[BsonConstructor]
+			private UserSettingsDocument()
+			{
+			}
+
+			public UserSettingsDocument(ObjectId Id)
+			{
+				this.Id = Id;
+			}
+
+			public UserSettingsDocument(IUserSettings Other)
+				: this(Other.UserId)
+			{
+				this.EnableExperimentalFeatures = Other.EnableExperimentalFeatures;
+				this.EnableIssueNotifications = Other.EnableIssueNotifications;
+				this.DashboardSettings = Other.DashboardSettings;
+				this.PinnedJobIds = new List<ObjectId>(Other.PinnedJobIds);
+			}
+		}
+
+		class ClaimDocument : IUserClaim
+		{
+			public string Type { get; set; }
+			public string Value { get; set; }
+
+			private ClaimDocument()
+			{
+				this.Type = null!;
+				this.Value = null!;
+			}
+
+			public ClaimDocument(string Type, string Value)
+			{
+				this.Type = Type;
+				this.Value = Value;
+			}
+
+			public ClaimDocument(IUserClaim Other)
+			{
+				this.Type = Other.Type;
+				this.Value = Other.Value;
+			}
+		}
+
+		IMongoCollection<UserDocument> Users;
+		IMongoCollection<UserClaimsDocument> UserClaims;
+		IMongoCollection<UserSettingsDocument> UserSettings;
+
+		ILogger<UserCollectionV2> Logger;
+
+		/// <summary>
+		/// Static constructor
+		/// </summary>
+		static UserCollectionV2()
+		{
+			BsonSerializer.RegisterDiscriminatorConvention(typeof(UserDocument), NullDiscriminatorConvention.Instance);
+			BsonSerializer.RegisterDiscriminatorConvention(typeof(UserClaimsDocument), NullDiscriminatorConvention.Instance);
+			BsonSerializer.RegisterDiscriminatorConvention(typeof(UserSettingsDocument), NullDiscriminatorConvention.Instance);
+
+			BsonSerializer.RegisterDiscriminatorConvention(typeof(ClaimDocument), NullDiscriminatorConvention.Instance);
+		}
+
+		/// <summary>
+		/// Constructor
+		/// </summary>
+		/// <param name="DatabaseService"></param>
+		/// <param name="Logger"></param>
+		public UserCollectionV2(DatabaseService DatabaseService, ILogger<UserCollectionV2> Logger)
+		{
+			this.Logger = Logger;
+
+			Users = DatabaseService.GetCollection<UserDocument>("UsersV2");
+			UserClaims = DatabaseService.GetCollection<UserClaimsDocument>("UserClaimsV2");
+			UserSettings = DatabaseService.GetCollection<UserSettingsDocument>("UserSettingsV2");
+
+			if (!DatabaseService.ReadOnlyMode)
+			{
+				Users.Indexes.CreateOne(new CreateIndexModel<UserDocument>(Builders<UserDocument>.IndexKeys.Ascending(x => x.LoginUpper), new CreateIndexOptions { Unique = true }));
+				Users.Indexes.CreateOne(new CreateIndexModel<UserDocument>(Builders<UserDocument>.IndexKeys.Ascending(x => x.EmailUpper)));
+			}
+		}
+
+		/// <inheritdoc/>
+		public async Task<IUser?> GetUserAsync(ObjectId Id)
+		{
+			return await Users.Find(x => x.Id == Id).FirstOrDefaultAsync();
+		}
+
+		/// <inheritdoc/>
+		public async Task<List<IUser>> FindUsersAsync(IEnumerable<ObjectId>? Ids, int? Index = null, int? Count = null)
+		{
+			FilterDefinition<UserDocument> Filter = FilterDefinition<UserDocument>.Empty;
+			if (Ids != null)
+			{
+				Filter &= Builders<UserDocument>.Filter.In(x => x.Id, Ids);
+			}
+			return await Users.Find(Filter).Range(Index, Count ?? 100).ToListAsync<UserDocument, IUser>();
+		}
+
+		/// <inheritdoc/>
+		public async Task<IUser?> FindUserByLoginAsync(string Login)
+		{
+			string LoginUpper = Login.ToUpperInvariant();
+			return await Users.Find(x => x.LoginUpper == LoginUpper).FirstOrDefaultAsync();
+		}
+
+		/// <inheritdoc/>
+		public async Task<IUser> FindOrAddUserByLoginAsync(string Login, string? Name, string? Email)
+		{
+			UpdateDefinition<UserDocument> Update = Builders<UserDocument>.Update.SetOnInsert(x => x.Id, ObjectId.GenerateNewId()).SetOnInsert(x => x.Login, Login);
+
+			if (Name == null)
+			{
+				Update = Update.SetOnInsert(x => x.Name, Name ?? Login);
+			}
+			else
+			{
+				Update = Update.Set(x => x.Name, Name);
+			}
+
+			if (Email != null)
+			{
+				Update = Update.Set(x => x.Email, Email).Set(x => x.EmailUpper, Email.ToUpperInvariant());
+			}
+
+			string LoginUpper = Login.ToUpperInvariant();
+			return await Users.FindOneAndUpdateAsync<UserDocument>(x => x.LoginUpper == LoginUpper, Update, new FindOneAndUpdateOptions<UserDocument, UserDocument> { IsUpsert = true, ReturnDocument = ReturnDocument.After });
+		}
+
+		/// <inheritdoc/>
+		public async Task<IUserClaims> GetClaimsAsync(ObjectId UserId)
+		{
+			IUserClaims? Claims = await UserClaims.Find(x => x.Id == UserId).FirstOrDefaultAsync();
+			if(Claims == null)
+			{
+				Claims = new UserClaimsDocument(UserId);
+			}
+			return Claims;
+		}
+
+		/// <inheritdoc/>
+		public async Task UpdateClaimsAsync(ObjectId UserId, IEnumerable<IUserClaim> Claims)
+		{
+			UserClaimsDocument NewDocument = new UserClaimsDocument(UserId);
+			NewDocument.Claims.AddRange(Claims.Select(x => new UserClaim(x)));
+			await UserClaims.ReplaceOneAsync(x => x.Id == UserId, NewDocument, new ReplaceOptions { IsUpsert = true });
+		}
+
+		/// <inheritdoc/>
+		public async Task<IUserSettings> GetSettingsAsync(ObjectId UserId)
+		{
+			IUserSettings? Settings = await UserSettings.Find(x => x.Id == UserId).FirstOrDefaultAsync();
+			if (Settings == null)
+			{
+				Settings = new UserSettingsDocument(UserId);
+			}
+			return Settings;
+		}
+
+		/// <inheritdoc/>
+		public async Task UpdateSettingsAsync(ObjectId UserId, bool? EnableExperimentalFeatures = null, bool? EnableIssueNotifications = null, BsonValue? DashboardSettings = null, IEnumerable<ObjectId>? AddPinnedJobIds = null, IEnumerable<ObjectId>? RemovePinnedJobIds = null)
+		{
+			List<UpdateDefinition<UserSettingsDocument>> Updates = new List<UpdateDefinition<UserSettingsDocument>>();
+			if (EnableExperimentalFeatures != null)
+			{
+				Updates.Add(Builders<UserSettingsDocument>.Update.SetOrUnsetNull(x => x.EnableExperimentalFeatures, EnableExperimentalFeatures));
+			}
+			if (EnableIssueNotifications != null)
+			{
+				if (EnableIssueNotifications.Value)
+				{
+					Updates.Add(Builders<UserSettingsDocument>.Update.Unset(x => x.EnableIssueNotifications));
+				}
+				else
+				{
+					Updates.Add(Builders<UserSettingsDocument>.Update.Set(x => x.EnableIssueNotifications, false));
+				}
+			}
+			if (DashboardSettings != null)
+			{
+				Updates.Add(Builders<UserSettingsDocument>.Update.Set(x => x.DashboardSettings, DashboardSettings));
+			}
+			if (AddPinnedJobIds != null && AddPinnedJobIds.Any())
+			{
+				Updates.Add(Builders<UserSettingsDocument>.Update.AddToSetEach(x => x.PinnedJobIds, AddPinnedJobIds));
+			}
+			if (RemovePinnedJobIds != null && RemovePinnedJobIds.Any())
+			{
+				Updates.Add(Builders<UserSettingsDocument>.Update.PullAll(x => x.PinnedJobIds, RemovePinnedJobIds));
+			}
+			if (Updates.Count > 0)
+			{
+				await UserSettings.UpdateOneAsync<UserSettingsDocument>(x => x.Id == UserId, Builders<UserSettingsDocument>.Update.Combine(Updates), new UpdateOptions { IsUpsert = true });
+			}
+		}
+
+		/// <summary>
+		/// Upgrade from V1 collection
+		/// </summary>
+		/// <param name="UserCollectionV1"></param>
+		/// <returns></returns>
+		public async Task ResaveDocumentsAsync(UserCollectionV1 UserCollectionV1)
+		{
+			await foreach ((IUser User, IUserClaims Claims, IUserSettings Settings) in UserCollectionV1.EnumerateDocumentsAsync())
+			{
+				try
+				{
+					await Users.ReplaceOneAsync(x => x.Id == User.Id, new UserDocument(User), new ReplaceOptions { IsUpsert = true });
+					await UserClaims.ReplaceOneAsync(x => x.Id == User.Id, new UserClaimsDocument(Claims), new ReplaceOptions { IsUpsert = true });
+					await UserSettings.ReplaceOneAsync(x => x.Id == User.Id, new UserSettingsDocument(Settings), new ReplaceOptions { IsUpsert = true });
+					Logger.LogDebug("Updated user {UserId}", User.Id);
+				}
+				catch (MongoWriteException Ex)
+				{
+					Logger.LogWarning(Ex, "Unable to resave user {UserId}", User.Id);
+				}
+
+				if(Settings.PinnedJobIds.Count > 0)
+				{
+					await UpdateSettingsAsync(User.Id, AddPinnedJobIds: Settings.PinnedJobIds);
+				}
+			}
+		}
+	}
+}
diff --git a/Engine/Source/Programs/Horde/HordeServer/Controllers/AccountController.cs b/Engine/Source/Programs/Horde/HordeServer/Controllers/AccountController.cs
new file mode 100644
--- /dev/null
+++ b/Engine/Source/Programs/Horde/HordeServer/Controllers/AccountController.cs
@@ -0,0 +1,126 @@
+// Copyright Epic Games, Inc. All Rights Reserved.
+
+using HordeServer.Authentication;
+using HordeServer.Models;
+using HordeServer.Services;
+using Microsoft.AspNetCore.Authentication;
+using Microsoft.AspNetCore.Authentication.Cookies;
+using Microsoft.AspNetCore.Authentication.OpenIdConnect;
+using Microsoft.AspNetCore.Authorization;
+using Microsoft.AspNetCore.Mvc;
+using System;
+using System.Collections.Generic;
+using System.Linq;
+using System.Net;
+using System.Text;
+using System.Threading.Tasks;
+
+namespace HordeServer.Controllers
+{
+	/// <summary>
+	/// Controller managing account status
+	/// </summary>
+	[ApiController]
+	[Route("[controller]")]
+	public class AccountController : Controller
+	{
+		const string DefaultAuthenticationScheme = OktaDefaults.AuthenticationScheme;
+
+		/// <summary>
+		/// Style sheet for HTML responses
+		/// </summary>
+		const string StyleSheet =
+			"body { font-family: 'Segoe UI', 'Roboto', arial, sans-serif; } " +
+			"p { margin:20px; font-size:13px; } " +
+			"h1 { margin:20px; font-size:32px; font-weight:200; } " +
+			"table { margin:10px 20px; } " +
+			"td { margin:5px; font-size:13px; }";
+
+		/// <summary>
+		/// The ACL service singleton
+		/// </summary>
+		AclService AclService;
+
+		/// <summary>
+		/// Constructor
+		/// </summary>
+		/// <param name="AclService">ACL service instance</param>
+		public AccountController(AclService AclService)
+		{
+			this.AclService = AclService;
+		}
+
+		/// <summary>
+		/// Gets the current login status
+		/// </summary>
+		/// <returns>The current login state</returns>
+		[HttpGet]
+		[Route("/account")]
+		public async Task<ActionResult> State()
+		{
+			StringBuilder Content = new StringBuilder();
+			Content.Append($"<html><style>{StyleSheet}</style><h1>Horde Server</h1>");
+			if (User.Identity.IsAuthenticated)
+			{
+				Content.Append($"<p>User <b>{User.Identity.Name}</b> is logged in. <a href=\"/account/logout\">Log out</a></p>");
+				if (await AclService.AuthorizeAsync(AclAction.AdminWrite, User))
+				{
+					Content.Append("<p>");
+					Content.Append("<a href=\"/api/v1/admin/token\">Get bearer token</a><br/>");
+					Content.Append("<a href=\"/api/v1/admin/registrationtoken\">Get agent registration token</a><br/>");
+					Content.Append("<a href=\"/api/v1/admin/softwaretoken\">Get agent software upload token</a><br/>");
+					Content.Append("<a href=\"/api/v1/admin/softwaredownloadtoken\">Get agent software download token</a><br/>");
+					Content.Append("<a href=\"/api/v1/admin/configtoken\">Get configuration token</a><br/>");
+					Content.Append("<a href=\"/api/v1/admin/chainedjobtoken\">Get chained job token</a><br/>");
+					Content.Append("</p>");
+				}
+				Content.Append($"<p>Claims for {User.Identity.Name}:");
+				Content.Append("<table>");
+				foreach (System.Security.Claims.Claim Claim in User.Claims)
+				{
+					Content.Append($"<tr><td>{Claim.Type}</td><td>{Claim.Value}</td></tr>");
+				}
+				Content.Append("</table>");
+				Content.Append("</p>");
+			}
+			else
+			{
+				Content.Append("<p><a href=\"/account/login\"><b>Login with OAuth2</b></a></p>");
+			}
+			Content.Append("</html>");
+			return new ContentResult { ContentType = "text/html", StatusCode = (int)HttpStatusCode.OK, Content = Content.ToString() };
+		}
+
+		/// <summary>
+		/// Login to the server
+		/// </summary>
+		/// <returns>Http result</returns>
+		[HttpGet]
+		[Route("/account/login")]
+		public IActionResult Login()
+		{
+			return new ChallengeResult(DefaultAuthenticationScheme, new AuthenticationProperties { RedirectUri = "/account" });
+		}
+
+		/// <summary>
+		/// Logout of the current account
+		/// </summary>
+		/// <returns>Http result</returns>
+		[HttpGet]
+		[Route("/account/logout")]
+		public async Task<IActionResult> Logout()
+		{
+			await HttpContext.SignOutAsync(CookieAuthenticationDefaults.AuthenticationScheme);
+			try
+			{
+				await HttpContext.SignOutAsync(DefaultAuthenticationScheme);
+			}
+			catch
+			{
+			}
+
+			string Content = $"<html><style>{StyleSheet}</style><body onload=\"setTimeout(function(){{ window.location = '/account'; }}, 2000)\"><p>User has been logged out. Returning to login page.</p></body></html>";
+			return new ContentResult { ContentType = "text/html", StatusCode = (int)HttpStatusCode.OK, Content = Content };
+		}
+	}
+}
diff --git a/Engine/Source/Programs/Horde/HordeServer/Controllers/AdminController.cs b/Engine/Source/Programs/Horde/HordeServer/Controllers/AdminController.cs
new file mode 100644
--- /dev/null
+++ b/Engine/Source/Programs/Horde/HordeServer/Controllers/AdminController.cs
@@ -0,0 +1,435 @@
+// Copyright Epic Games, Inc. All Rights Reserved.
+
+using EpicGames.Core;
+using HordeServer.Api;
+using HordeServer.Collections;
+using HordeCommon;
+using HordeServer.Models;
+using HordeServer.Services;
+using HordeServer.Utilities;
+using Microsoft.AspNetCore.Authorization;
+using Microsoft.AspNetCore.Authorization.Infrastructure;
+using Microsoft.AspNetCore.Mvc;
+using Microsoft.Extensions.Options;
+using MongoDB.Bson;
+using System;
+using System.Collections.Generic;
+using System.IdentityModel.Tokens.Jwt;
+using System.IO;
+using System.Linq;
+using System.Runtime.InteropServices;
+using System.Security.Claims;
+using System.Threading.Tasks;
+
+namespace HordeServer.Controllers
+{
+	/// <summary>
+	/// Object containing settings for the server
+	/// </summary>
+	public class AdminSettings
+	{
+		/// <summary>
+		/// The default perforce server
+		/// </summary>
+		public string? DefaultServerAndPort { get; set; }
+
+		/// <summary>
+		/// The default perforce username
+		/// </summary>
+		public string? DefaultUserName { get; set; }
+
+		/// <summary>
+		/// The default perforce password
+		/// </summary>
+		public string? DefaultPassword { get; set; }
+	}
+
+	/// <summary>
+	/// The conform limit value
+	/// </summary>
+	public class ConformSettings
+	{
+		/// <summary>
+		/// Maximum number of conforms allowed at once
+		/// </summary>
+		public int MaxCount { get; set; }
+	}
+
+	/// <summary>
+	/// Controller managing account status
+	/// </summary>
+	[ApiController]
+	[Authorize]
+	[Route("[controller]")]
+	public class AdminController : ControllerBase
+	{
+		/// <summary>
+		/// The database service singleton
+		/// </summary>
+		DatabaseService DatabaseService;
+
+		/// <summary>
+		/// The acl service singleton
+		/// </summary>
+		AclService AclService;
+
+		/// <summary>
+		/// The upgrade service singleton
+		/// </summary>
+		UpgradeService UpgradeService;
+
+		/// <summary>
+		/// List of conforms running
+		/// </summary>
+		SingletonDocument<ConformList> ConformList;
+
+		/// <summary>
+		/// Settings for the server
+		/// </summary>
+		IOptionsMonitor<ServerSettings> Settings;
+
+		/// <summary>
+		/// Constructor
+		/// </summary>
+		/// <param name="DatabaseService">The database service singleton</param>
+		/// <param name="AclService">The ACL service singleton</param>
+		/// <param name="UpgradeService">The upgrade service singelton</param>
+		/// <param name="Settings">Server settings</param>
+		public AdminController(DatabaseService DatabaseService, AclService AclService, UpgradeService UpgradeService, IOptionsMonitor<ServerSettings> Settings)
+		{
+			this.DatabaseService = DatabaseService;
+			this.AclService = AclService;
+			this.UpgradeService = UpgradeService;
+			this.ConformList = new SingletonDocument<ConformList>(DatabaseService);
+			this.Settings = Settings;
+		}
+
+		/// <summary>
+		/// Force a reset on the database
+		/// </summary>
+		[HttpPost]
+		[Route("/api/v1/admin/reset")]
+		public async Task<ActionResult> ForceResetAsync([FromQuery] string? Instance = null)
+		{
+			if(!await AclService.AuthorizeAsync(AclAction.AdminWrite, User))
+			{
+				return Forbid();
+			}
+
+			for (; ; )
+			{
+				Globals Globals = await DatabaseService.GetGlobalsAsync();
+				if (Instance == null)
+				{
+					return NotFound($"Missing code query parameter. Set to {Globals.InstanceId} to reset.");
+				}
+				if (Globals.InstanceId != Instance.ToObjectId())
+				{
+					return NotFound($"Incorrect code query parameter. Should be {Globals.InstanceId}.");
+				}
+
+				Globals.ForceReset = true;
+
+				if (await DatabaseService.TryUpdateSingletonAsync(Globals))
+				{
+					return Ok("Database will be reinitialized on next restart");
+				}
+			}
+		}
+
+
+		/// <summary>
+		/// Gets the scheduled downtime 
+		/// </summary>
+		/// <returns>Async task</returns>
+		[HttpGet]
+		[Route("/api/v1/admin/downtime")]
+		public async Task<ActionResult<List<ScheduledDowntime>>> GetScheduledDowntimeAsync()
+		{
+			if (!await AclService.AuthorizeAsync(AclAction.AdminWrite, User))
+			{
+				return Forbid();
+			}
+
+			Globals Globals = await DatabaseService.GetGlobalsAsync();
+			return Globals.ScheduledDowntime;
+		}
+
+		/// <summary>
+		/// Sets the scheduled downtime
+		/// </summary>
+		/// <param name="Schedules">List of downtime</param>
+		/// <returns>Async task</returns>
+		[HttpPut]
+		[Route("/api/v1/admin/downtime")]
+		public async Task<ActionResult> SetScheduledDowntimeAsync(List<ScheduledDowntime> Schedules)
+		{
+			if (!await AclService.AuthorizeAsync(AclAction.AdminWrite, User))
+			{
+				return Forbid();
+			}
+
+			for (; ; )
+			{
+				Globals Globals = await DatabaseService.GetGlobalsAsync();
+				Globals.ScheduledDowntime = Schedules;
+
+				if (await DatabaseService.TryUpdateSingletonAsync(Globals))
+				{
+					break;
+				}
+			}
+
+			return Ok();
+		}
+
+		/// <summary>
+		/// Upgrade the database to the latest schema
+		/// </summary>
+		/// <param name="FromVersion">The schema version to upgrade from.</param>
+		[HttpPost]
+		[Route("/api/v1/admin/upgradeschema")]
+		public async Task<ActionResult> UpgradeSchemaAsync([FromQuery] int? FromVersion = null)
+		{
+			if (!await AclService.AuthorizeAsync(AclAction.AdminWrite, User))
+			{
+				return Forbid();
+			}
+
+			await UpgradeService.UpgradeSchemaAsync(FromVersion);
+			return Ok();
+		}
+
+		/// <summary>
+		/// Update the conform limit
+		/// </summary>
+		[HttpPut]
+		[Route("/api/v1/admin/conform")]
+		public async Task<ActionResult> SetConformSettingsAsync([FromBody] ConformSettings NewSettings)
+		{
+			if (!await AclService.AuthorizeAsync(AclAction.AdminWrite, User))
+			{
+				return Forbid();
+			}
+
+			for (; ; )
+			{
+				ConformList CurrentValue = await ConformList.GetAsync();
+				CurrentValue.MaxCount = NewSettings.MaxCount;
+
+				if (await ConformList.TryUpdateAsync(CurrentValue))
+				{
+					return Ok();
+				}
+			}
+		}
+
+		/// <summary>
+		/// Gets the current perforce settinsg
+		/// </summary>
+		/// <returns>List of Perforce clusters</returns>
+		[HttpGet]
+		[Route("/api/v1/admin/conform")]
+		public async Task<ActionResult<ConformSettings>> GetConformSettingsAsync()
+		{
+			if (!await AclService.AuthorizeAsync(AclAction.AdminRead, User))
+			{
+				return Forbid();
+			}
+
+			ConformList CurrentValue = await ConformList.GetAsync();
+			return new ConformSettings { MaxCount = CurrentValue.MaxCount };
+		}
+
+		/// <summary>
+		/// Update Perforce server settings
+		/// </summary>
+		[HttpPut]
+		[Route("/api/v1/admin/perforce")]
+		public async Task<ActionResult> UpdatePerforceSettingsAsync([FromBody] List<PerforceCluster> Clusters)
+		{
+			if (!await AclService.AuthorizeAsync(AclAction.AdminWrite, User))
+			{
+				return Forbid();
+			}
+
+			for (; ; )
+			{
+				Globals Globals = await DatabaseService.GetGlobalsAsync();
+				Globals.PerforceClusters = Clusters;
+
+				if (await DatabaseService.TryUpdateSingletonAsync(Globals))
+				{
+					return Ok("Settings have been updated");
+				}
+			}
+		}
+
+		/// <summary>
+		/// Gets the current perforce settinsg
+		/// </summary>
+		/// <returns>List of Perforce clusters</returns>
+		[HttpGet]
+		[Route("/api/v1/admin/perforce")]
+		public async Task<ActionResult<List<PerforceCluster>>> GetPerforceSettingsAsync()
+		{
+			if (!await AclService.AuthorizeAsync(AclAction.AdminRead, User))
+			{
+				return Forbid();
+			}
+
+			Globals Globals = await DatabaseService.GetGlobalsAsync();
+			return Globals.PerforceClusters;
+		}
+
+		/// <summary>
+		/// Issues a token for the given roles. Issues a token for the current user if not specified.
+		/// </summary>
+		/// <returns>Administrative settings for the server</returns>
+		[HttpGet]
+		[Route("/api/v1/admin/token")]
+		public async Task<ActionResult<string>> GetTokenAsync()
+		{
+			if (!await AclService.AuthorizeAsync(AclAction.IssueBearerToken, User))
+			{
+				return Forbid();
+			}
+
+			return AclService.IssueBearerToken(User.Claims, GetDefaultExpiryTime());
+		}
+
+		/// <summary>
+		/// Issues a token for the given roles. Issues a token for the current user if not specified.
+		/// </summary>
+		/// <param name="Roles">Roles for the new token</param>
+		/// <returns>Administrative settings for the server</returns>
+		[HttpGet]
+		[Route("/api/v1/admin/roletoken")]
+		public async Task<ActionResult<string>> GetRoleTokenAsync([FromQuery] string Roles)
+		{
+			if (!await AclService.AuthorizeAsync(AclAction.AdminWrite, User))
+			{
+				return Forbid();
+			}
+
+			List<Claim> Claims = new List<Claim>();
+			Claims.AddRange(Roles.Split('+').Select(x => new Claim(ClaimTypes.Role, x)));
+
+			return AclService.IssueBearerToken(Claims, GetDefaultExpiryTime());
+		}
+
+		/// <summary>
+		/// Issues a token for the given roles. Issues a token for the current user if not specified.
+		/// </summary>
+		/// <returns>Administrative settings for the server</returns>
+		[HttpGet]
+		[Route("/api/v1/admin/registrationtoken")]
+		public async Task<ActionResult<string>> GetRegistrationTokenAsync()
+		{
+			if (!await AclService.AuthorizeAsync(AclAction.AdminWrite, User))
+			{
+				return Forbid();
+			}
+
+			List<AclClaim> Claims = new List<AclClaim>();
+			Claims.Add(new AclClaim(ClaimTypes.Name, User.Identity.Name ?? "Unknown"));
+			Claims.Add(AclService.AgentRegistrationClaim);
+
+			return AclService.IssueBearerToken(Claims, null);
+		}
+
+		/// <summary>
+		/// Issues a token valid to upload new versions of the agent software.
+		/// </summary>
+		/// <returns>Administrative settings for the server</returns>
+		[HttpGet]
+		[Route("/api/v1/admin/softwaretoken")]
+		public async Task<ActionResult<string>> GetSoftwareTokenAsync()
+		{
+			if (!await AclService.AuthorizeAsync(AclAction.AdminWrite, User))
+			{
+				return Forbid();
+			}
+
+			List<AclClaim> Claims = new List<AclClaim>();
+			Claims.Add(new AclClaim(ClaimTypes.Name, User.Identity.Name ?? "Unknown"));
+			Claims.Add(AclService.UploadSoftwareClaim);
+
+			return AclService.IssueBearerToken(Claims, null);
+		}
+
+		/// <summary>
+		/// Issues a token valid to download new versions of the agent software.
+		/// </summary>
+		/// <returns>Administrative settings for the server</returns>
+		[HttpGet]
+		[Route("/api/v1/admin/softwaredownloadtoken")]
+		public async Task<ActionResult<string>> GetSoftwareDownloadTokenAsync()
+		{
+			if (!await AclService.AuthorizeAsync(AclAction.AdminRead, User))
+			{
+				return Forbid();
+			}
+
+			List<AclClaim> Claims = new List<AclClaim>();
+			Claims.Add(new AclClaim(ClaimTypes.Name, User.Identity.Name ?? "Unknown"));
+			Claims.Add(AclService.DownloadSoftwareClaim);
+
+			return AclService.IssueBearerToken(Claims, null);
+		}
+
+		/// <summary>
+		/// Issues a token valid to configure streams and projects
+		/// </summary>
+		/// <returns>Administrative settings for the server</returns>
+		[HttpGet]
+		[Route("/api/v1/admin/configtoken")]
+		public async Task<ActionResult<string>> GetConfigToken()
+		{
+			if (!await AclService.AuthorizeAsync(AclAction.AdminRead, User))
+			{
+				return Forbid();
+			}
+
+			List<AclClaim> Claims = new List<AclClaim>();
+			Claims.Add(new AclClaim(ClaimTypes.Name, User.Identity.Name ?? "Unknown"));
+			Claims.Add(AclService.ConfigureProjectsClaim);
+
+			return AclService.IssueBearerToken(Claims, null);
+		}
+
+		/// <summary>
+		/// Issues a token valid to start chained jobs
+		/// </summary>
+		/// <returns>Administrative settings for the server</returns>
+		[HttpGet]
+		[Route("/api/v1/admin/chainedjobtoken")]
+		public async Task<ActionResult<string>> GetChainedJobToken()
+		{
+			if (!await AclService.AuthorizeAsync(AclAction.AdminRead, User))
+			{
+				return Forbid();
+			}
+
+			List<AclClaim> Claims = new List<AclClaim>();
+			//Claims.Add(new AclClaim(ClaimTypes.Name, User.Identity.Name ?? "Unknown"));
+			Claims.Add(AclService.StartChainedJobClaim);
+
+			return AclService.IssueBearerToken(Claims, null);
+		}
+
+		/// <summary>
+		/// Gets the default expiry time for a token
+		/// </summary>
+		/// <returns></returns>
+		private TimeSpan? GetDefaultExpiryTime()
+		{
+			TimeSpan? ExpiryTime = null;
+			if (Settings.CurrentValue.JwtExpiryTimeHours != -1)
+			{
+				ExpiryTime = TimeSpan.FromHours(Settings.CurrentValue.JwtExpiryTimeHours);
+			}
+			return ExpiryTime;
+		}
+	}
+}
diff --git a/Engine/Source/Programs/Horde/HordeServer/Controllers/AgentSoftwareController.cs b/Engine/Source/Programs/Horde/HordeServer/Controllers/AgentSoftwareController.cs
new file mode 100644
--- /dev/null
+++ b/Engine/Source/Programs/Horde/HordeServer/Controllers/AgentSoftwareController.cs
@@ -0,0 +1,153 @@
+// Copyright Epic Games, Inc. All Rights Reserved.
+
+using HordeServer.Api;
+using HordeServer.Services;
+using Microsoft.AspNetCore.Authorization;
+using Microsoft.AspNetCore.Http;
+using Microsoft.AspNetCore.Mvc;
+using Microsoft.Net.Http.Headers;
+using System;
+using System.IO;
+using System.Collections.Generic;
+using System.Linq;
+using System.Threading.Tasks;
+using HordeServer.Models;
+using MongoDB.Bson;
+using HordeServer.Utilities;
+
+using AgentSoftwareChannelName = HordeServer.Utilities.StringId<HordeServer.Services.AgentSoftwareChannels>;
+
+namespace HordeServer.Controllers
+{
+	/// <summary>
+	/// Controller for the /api/v1/software endpoint
+	/// </summary>
+	[ApiController]
+	[Authorize]
+	[Route("[controller]")]
+	public class AgentSoftwareController : ControllerBase
+	{
+		/// <summary>
+		/// Singleton instance of the ACL service
+		/// </summary>
+		private readonly AclService AclService;
+
+		/// <summary>
+		/// Singleton instance of the client service
+		/// </summary>
+		private readonly AgentSoftwareService AgentSoftwareService;
+
+		/// <summary>
+		/// Constructor
+		/// </summary>
+		/// <param name="AclService">The ACL service</param>
+		/// <param name="AgentSoftwareService">The client service</param>
+		public AgentSoftwareController(AclService AclService, AgentSoftwareService AgentSoftwareService)
+		{
+			this.AclService = AclService;
+			this.AgentSoftwareService = AgentSoftwareService;
+		}
+
+		/// <summary>
+		/// Finds all uploaded software matching the given criteria
+		/// </summary>
+		/// <param name="Filter">Filter for the properties to return</param>
+		/// <returns>Http response</returns>
+		[HttpGet]
+		[Route("/api/v1/agentsoftware")]
+		[ProducesResponseType(typeof(List<GetAgentSoftwareChannelResponse>), 200)]
+		public async Task<ActionResult<List<object>>> FindSoftwareAsync([FromQuery] PropertyFilter? Filter = null)
+		{
+			if (!await AclService.AuthorizeAsync(AclAction.DownloadSoftware, User))
+			{
+				return Forbid();
+			}
+
+			List<IAgentSoftwareChannel> Results = await AgentSoftwareService.FindChannelsAsync();
+
+			List<object> Responses = new List<object>();
+			foreach (IAgentSoftwareChannel Result in Results)
+			{
+				Responses.Add(new GetAgentSoftwareChannelResponse(Result).ApplyFilter(Filter));
+			}
+			return Responses;
+		}
+
+		/// <summary>
+		/// Finds all uploaded software matching the given criteria
+		/// </summary>
+		/// <param name="Name">Name of the channel to get</param>
+		/// <param name="Filter">Filter for the properties to return</param>
+		/// <returns>Http response</returns>
+		[HttpGet]
+		[Route("/api/v1/agentsoftware/{Name}")]
+		[ProducesResponseType(typeof(GetAgentSoftwareChannelResponse), 200)]
+		public async Task<ActionResult<object>> FindSoftwareAsync(string Name, [FromQuery] PropertyFilter? Filter = null)
+		{
+			if (!await AclService.AuthorizeAsync(AclAction.DownloadSoftware, User))
+			{
+				return Forbid();
+			}
+
+			IAgentSoftwareChannel? Channel = await AgentSoftwareService.GetChannelAsync(new AgentSoftwareChannelName(Name));
+			if(Channel == null)
+			{
+				return NotFound();
+			}
+
+			return new GetAgentSoftwareChannelResponse(Channel).ApplyFilter(Filter);
+		}
+
+		/// <summary>
+		/// Uploads a new agent zip file
+		/// </summary>
+		/// <param name="Name">Name of the channel to post to</param>
+		/// <param name="File">Zip archive containing the new client software</param>
+		/// <returns>Http result code</returns>
+		[HttpPost]
+		[Route("/api/v1/agentsoftware/{Name}/zip")]
+		public async Task<ActionResult> SetArchiveAsync(string Name, [FromForm] IFormFile File)
+		{
+			if (!await AclService.AuthorizeAsync(AclAction.UploadSoftware, User))
+			{
+				return Forbid();
+			}
+
+			byte[] Data;
+			using (MemoryStream MemoryStream = new MemoryStream())
+			{
+				using (System.IO.Stream Stream = File.OpenReadStream())
+				{
+					await Stream.CopyToAsync(MemoryStream);
+				}
+				Data = MemoryStream.ToArray();
+			}
+
+			await AgentSoftwareService.SetArchiveAsync(new AgentSoftwareChannelName(Name), User.Identity.Name, Data);
+			return Ok();
+		}
+
+		/// <summary>
+		/// Gets the zip file for a specific channel
+		/// </summary>
+		/// <param name="Name">Name of the channel</param>
+		/// <returns>Http response</returns>
+		[HttpGet]
+		[Route("/api/v1/agentsoftware/{Name}/zip")]
+		public async Task<ActionResult<GetSoftwareResponse>> GetArchiveAsync(string Name)
+		{
+			if (!await AclService.AuthorizeAsync(AclAction.DownloadSoftware, User))
+			{
+				return Forbid();
+			}
+
+			byte[]? Data = await AgentSoftwareService.GetArchiveAsync(new AgentSoftwareChannelName(Name));
+			if (Data == null)
+			{
+				return NotFound();
+			}
+
+			return new FileStreamResult(new MemoryStream(Data), new MediaTypeHeaderValue("application/octet-stream"));
+		}
+	}
+}
diff --git a/Engine/Source/Programs/Horde/HordeServer/Controllers/AgentsController.cs b/Engine/Source/Programs/Horde/HordeServer/Controllers/AgentsController.cs
new file mode 100644
--- /dev/null
+++ b/Engine/Source/Programs/Horde/HordeServer/Controllers/AgentsController.cs
@@ -0,0 +1,316 @@
+// Copyright Epic Games, Inc. All Rights Reserved.
+
+using HordeServer.Api;
+using HordeServer.Models;
+using HordeServer.Services;
+using HordeServer.Utilities;
+using Microsoft.AspNetCore.Authentication;
+using Microsoft.AspNetCore.Authorization;
+using Microsoft.AspNetCore.Mvc;
+using Microsoft.AspNetCore.Mvc.ModelBinding;
+using Microsoft.Extensions.Options;
+using Microsoft.IdentityModel.Tokens;
+using MongoDB.Bson;
+using System;
+using System.Collections.Generic;
+using System.ComponentModel.DataAnnotations;
+using System.IdentityModel.Tokens.Jwt;
+using System.Linq;
+using System.Reflection;
+using System.Security.Claims;
+using System.Text.Json;
+using System.Threading.Tasks;
+
+using PoolId = HordeServer.Utilities.StringId<HordeServer.Models.IPool>;
+using AgentSoftwareChannelName = HordeServer.Utilities.StringId<HordeServer.Services.AgentSoftwareChannels>;
+
+namespace HordeServer.Controllers
+{
+	/// <summary>
+	/// Controller for the /api/v1/agents endpoint
+	/// </summary>
+	[ApiController]
+	[Authorize]
+	[Route("[controller]")]
+	public class AgentsController : ControllerBase
+	{
+		/// <summary>
+		/// Singleton instance of the ACL service
+		/// </summary>
+		AclService AclService;
+
+		/// <summary>
+		/// Singleton instance of the agent service
+		/// </summary>
+		AgentService AgentService;
+
+		/// <summary>
+		/// The pool service
+		/// </summary>
+		PoolService PoolService;
+
+		/// <summary>
+		/// Constructor
+		/// </summary>
+		/// <param name="AclService">The ACL service singleton</param>
+		/// <param name="AgentService">The agent service</param>
+		/// <param name="PoolService">The pool service</param>
+		public AgentsController(AclService AclService, AgentService AgentService, PoolService PoolService)
+		{
+			this.AclService = AclService;
+			this.AgentService = AgentService;
+			this.PoolService = PoolService;
+		}
+
+		/// <summary>
+		/// Register an agent to perform remote work.
+		/// </summary>
+		/// <param name="Request">Request parameters</param>
+		/// <returns>Information about the registered agent</returns>
+		[HttpPost]
+		[Route("/api/v1/agents")]
+		public async Task<ActionResult<CreateAgentResponse>> CreateAgentAsync([FromBody] CreateAgentRequest Request)
+		{
+			if(!await AclService.AuthorizeAsync(AclAction.CreateAgent, User))
+			{
+				return Forbid();
+			}
+
+			AgentSoftwareChannelName? Channel = String.IsNullOrEmpty(Request.Channel) ? (AgentSoftwareChannelName?)null : new AgentSoftwareChannelName(Request.Channel);
+			IAgent Agent = await AgentService.CreateAgentAsync(Request.Name, Request.Enabled, Request.Ephemeral, Channel, Request.Pools?.ConvertAll(x => new PoolId(x)));
+
+			return new CreateAgentResponse(Agent.Id.ToString());
+		}
+
+		/// <summary>
+		/// Finds the agents matching specified criteria.
+		/// </summary>
+		/// <param name="Pool">The pool containing the agent</param>
+		/// <param name="Index">First result to return</param>
+		/// <param name="Count">Number of results to return</param>
+		/// <param name="ModifiedAfter">If set, only returns agents modified after this time</param>
+		/// <param name="Filter">Filter for the properties to return</param>
+		/// <returns>List of matching agents</returns>
+		[HttpGet]
+		[Route("/api/v1/agents")]
+		[ProducesResponseType(typeof(List<GetAgentResponse>), 200)]
+		public async Task<ActionResult<List<object>>> FindAgentsAsync([FromQuery] string? Pool = null, [FromQuery] int? Index = null, [FromQuery] int? Count = null, [FromQuery] DateTimeOffset? ModifiedAfter = null, [FromQuery] PropertyFilter? Filter = null)
+		{
+			GlobalPermissionsCache PermissionsCache = new GlobalPermissionsCache();
+			List<IAgent> Agents = await AgentService.FindAgentsAsync(Pool?.ToObjectId(), ModifiedAfter?.UtcDateTime, Index, Count);
+			List<IPool> Pools = await PoolService.GetPoolsAsync();
+
+			List<object> Responses = new List<object>();
+			foreach (IAgent Agent in Agents)
+			{
+				bool bIncludeAcl = await AgentService.AuthorizeAsync(Agent, AclAction.ViewPermissions, User, PermissionsCache);
+				List<PoolId> PoolIds = Agent.GetPools(Pools).Select(x => x.Id).ToList();
+				Responses.Add(new GetAgentResponse(Agent, PoolIds, bIncludeAcl).ApplyFilter(Filter));
+			}
+
+			return Responses;
+		}
+
+		/// <summary>
+		/// Retrieve information about a specific agent
+		/// </summary>
+		/// <param name="AgentId">Id of the agent to get information about</param>
+		/// <param name="Filter">Filter for the properties to return</param>
+		/// <returns>Information about the requested agent</returns>
+		[HttpGet]
+		[Route("/api/v1/agents/{AgentId}")]
+		[ProducesResponseType(typeof(GetAgentResponse), 200)]
+		public async Task<ActionResult<object>> GetAgentAsync(string AgentId, [FromQuery] PropertyFilter? Filter = null)
+		{
+			IAgent? Agent = await AgentService.GetAgentAsync(new AgentId(AgentId));
+			if (Agent == null)
+			{
+				return NotFound();
+			}
+
+			List<IPool> Pools = await PoolService.GetPoolsAsync();
+
+			GlobalPermissionsCache Cache = new GlobalPermissionsCache();
+			bool bIncludeAcl = await AgentService.AuthorizeAsync(Agent, AclAction.ViewPermissions, User, Cache);
+			List<PoolId> PoolIds = Agent.GetPools(Pools).Select(x => x.Id).ToList();
+			return new GetAgentResponse(Agent, PoolIds, bIncludeAcl).ApplyFilter(Filter);
+		}
+
+		/// <summary>
+		/// Update an agent's properties.
+		/// </summary>
+		/// <param name="AgentId">Id of the agent to update.</param>
+		/// <param name="Update">Properties on the agent to update.</param>
+		/// <returns>Http result code</returns>
+		[HttpPut]
+		[Route("/api/v1/agents/{AgentId}")]
+		public async Task<ActionResult> UpdateAgentAsync(string AgentId, [FromBody] UpdateAgentRequest Update)
+		{
+			IAgent? Agent = await AgentService.GetAgentAsync(new AgentId(AgentId));
+			if (Agent == null)
+			{
+				return NotFound();
+			}
+
+			GlobalPermissionsCache Cache = new GlobalPermissionsCache();
+			if (!await AgentService.AuthorizeAsync(Agent, AclAction.UpdateAgent, User, Cache))
+			{
+				return Forbid();
+			}
+			if (Update.Acl != null && !await AgentService.AuthorizeAsync(Agent, AclAction.ChangePermissions, User, Cache))
+			{
+				return Forbid();
+			}
+
+			AgentSoftwareChannelName? Channel = String.IsNullOrEmpty(Update.Channel) ? (AgentSoftwareChannelName?)null : new AgentSoftwareChannelName(Update.Channel);
+			await AgentService.UpdateAgentAsync(Agent, Update.Enabled, Update.RequestConform, Update.RequestRestart, Update.RequestShutdown, Channel, Update.Pools?.ConvertAll(x => new PoolId(x)), Acl.Merge(Agent.Acl, Update.Acl), Update.Comment);
+			return Ok();
+		}
+
+		/// <summary>
+		/// Remove a registered agent.
+		/// </summary>
+		/// <param name="AgentId">Id of the agent to delete.</param>
+		/// <returns>Http result code</returns>
+		[HttpDelete]
+		[Route("/api/v1/agents/{AgentId}")]
+		public async Task<ActionResult> DeleteAgentAsync(string AgentId)
+		{
+			IAgent? Agent = await AgentService.GetAgentAsync(new AgentId(AgentId));
+			if (Agent == null)
+			{
+				return NotFound();
+			}
+			if (!await AgentService.AuthorizeAsync(Agent, AclAction.DeleteAgent, User, null))
+			{
+				return Forbid();
+			}
+
+			await AgentService.DeleteAgentAsync(Agent);
+			return new OkResult();
+		}
+
+		/// <summary>
+		/// Find all the sessions of a particular agent
+		/// </summary>
+		/// <param name="AgentId">Unique id of the agent to find</param>
+		/// <param name="StartTime">Start time to include in the search</param>
+		/// <param name="FinishTime">Finish time to include in the search</param>
+		/// <param name="Index">Index of the first result to return</param>
+		/// <param name="Count">Number of results to return</param>
+		/// <returns>Sessions </returns>
+		[HttpGet]
+		[Route("/api/v1/agents/{AgentId}/sessions")]
+		public async Task<ActionResult<List<GetAgentSessionResponse>>> FindSessionsAsync(string AgentId, [FromQuery] DateTimeOffset? StartTime, [FromQuery] DateTimeOffset? FinishTime, [FromQuery] int Index = 0, [FromQuery] int Count = 50)
+		{
+			AgentId AgentIdValue = new AgentId(AgentId);
+
+			IAgent? Agent = await AgentService.GetAgentAsync(AgentIdValue);
+			if (Agent == null)
+			{
+				return NotFound();
+			}
+			if (!await AgentService.AuthorizeAsync(Agent, AclAction.ViewSession, User, null))
+			{
+				return Forbid();
+			}
+
+			List<ISession> Sessions = await AgentService.FindSessionsAsync(AgentIdValue, StartTime?.UtcDateTime, FinishTime?.UtcDateTime, Index, Count);
+			return Sessions.ConvertAll(x => new GetAgentSessionResponse(x));
+		}
+
+		/// <summary>
+		/// Find all the sessions of a particular agent
+		/// </summary>
+		/// <param name="AgentId">Unique id of the agent to find</param>
+		/// <param name="SessionId">Unique id of the session</param>
+		/// <returns>Sessions </returns>
+		[HttpGet]
+		[Route("/api/v1/agents/{AgentId}/sessions/{SessionId}")]
+		public async Task<ActionResult<GetAgentSessionResponse>> GetSessionAsync(string AgentId, string SessionId)
+		{
+			AgentId AgentIdValue = new AgentId(AgentId);
+
+			IAgent? Agent = await AgentService.GetAgentAsync(AgentIdValue);
+			if (Agent == null)
+			{
+				return NotFound();
+			}
+			if (!await AgentService.AuthorizeAsync(Agent, AclAction.ViewSession, User, null))
+			{
+				return Forbid();
+			}
+
+			ISession? Session = await AgentService.GetSessionAsync(SessionId.ToObjectId());
+			if(Session == null || Session.AgentId != AgentIdValue)
+			{
+				return NotFound();
+			}
+
+			return new GetAgentSessionResponse(Session);
+		}
+
+		/// <summary>
+		/// Find all the leases for a particular agent
+		/// </summary>
+		/// <param name="AgentId">Unique id of the agent to find</param>
+		/// <param name="SessionId">The session to query</param>
+		/// <param name="StartTime">Start of the time window to consider</param>
+		/// <param name="FinishTime">End of the time window to consider</param>
+		/// <param name="Index">Index of the first result to return</param>
+		/// <param name="Count">Number of results to return</param>
+		/// <returns>Sessions </returns>
+		[HttpGet]
+		[Route("/api/v1/agents/{AgentId}/leases")]
+		public async Task<ActionResult<List<GetAgentLeaseResponse>>> FindLeasesAsync(string AgentId, [FromQuery] string? SessionId, [FromQuery] DateTimeOffset? StartTime, [FromQuery] DateTimeOffset? FinishTime, [FromQuery] int Index = 0, [FromQuery] int Count = 1000)
+		{
+			AgentId AgentIdValue = new AgentId(AgentId);
+
+			IAgent? Agent = await AgentService.GetAgentAsync(AgentIdValue);
+			if (Agent == null)
+			{
+				return NotFound();
+			}
+			if (!await AgentService.AuthorizeAsync(Agent, AclAction.ViewLeases, User, null))
+			{
+				return Forbid();
+			}
+
+			List<ILease> Leases = await AgentService.FindLeasesAsync(AgentIdValue, SessionId?.ToObjectId(), StartTime?.UtcDateTime, FinishTime?.UtcDateTime, Index, Count);
+			return Leases.ConvertAll(x => new GetAgentLeaseResponse(x));
+		}
+
+		/// <summary>
+		/// Get info about a particular lease
+		/// </summary>
+		/// <param name="AgentId">Unique id of the agent to find</param>
+		/// <param name="LeaseId">Unique id of the particular lease</param>
+		/// <returns>Lease matching the given id</returns>
+		[HttpGet]
+		[Route("/api/v1/agents/{AgentId}/leases/{LeaseId}")]
+		public async Task<ActionResult<GetAgentLeaseResponse>> GetLeaseAsync(string AgentId, string LeaseId)
+		{
+			AgentId AgentIdValue = new AgentId(AgentId);
+
+			IAgent? Agent = await AgentService.GetAgentAsync(AgentIdValue);
+			if(Agent == null)
+			{
+				return NotFound();
+			}
+			if (!await AclService.AuthorizeAsync(AclAction.ViewLeases, User))
+			{
+				return Forbid();
+			}
+
+			ObjectId LeaseIdValue = LeaseId.ToObjectId();
+
+			ILease? Lease = await AgentService.GetLeaseAsync(LeaseIdValue);
+			if (Lease == null || Lease.AgentId != AgentIdValue)
+			{
+				return NotFound();
+			}
+
+			return new GetAgentLeaseResponse(Lease);
+		}
+	}
+}
diff --git a/Engine/Source/Programs/Horde/HordeServer/Controllers/ArtifactsController.cs b/Engine/Source/Programs/Horde/HordeServer/Controllers/ArtifactsController.cs
new file mode 100644
--- /dev/null
+++ b/Engine/Source/Programs/Horde/HordeServer/Controllers/ArtifactsController.cs
@@ -0,0 +1,448 @@
+using Amazon.S3.Transfer;
+using HordeServer.Api;
+using HordeServer.Models;
+using HordeServer.Services;
+using HordeServer.Utilities;
+using Microsoft.AspNetCore.Authorization;
+using Microsoft.AspNetCore.Http;
+using Microsoft.AspNetCore.Mvc;
+using Microsoft.AspNetCore.StaticFiles;
+using Microsoft.IdentityModel.Tokens;
+using MongoDB.Bson;
+using System;
+using System.Collections.Generic;
+using System.Globalization;
+using System.IdentityModel.Tokens.Jwt;
+using System.IO;
+using System.IO.Compression;
+using System.Linq;
+using System.Net.Http.Headers;
+using System.Net.Mime;
+using System.Security.Claims;
+using System.Text;
+using System.Threading.Tasks;
+
+namespace HordeServer.Controllers
+{
+	/// <summary>
+	/// Controller for the /api/artifacts endpoint
+	/// </summary>
+	[ApiController]
+	[Route("[controller]")]
+	public class ArtifactsController : ControllerBase
+	{
+		/// <summary>
+		/// Instance of the database service
+		/// </summary>
+		private readonly DatabaseService DatabaseService;
+
+		/// <summary>
+		/// Instance of the Artifact service
+		/// </summary>
+		private readonly ArtifactService ArtifactService;
+
+		/// <summary>
+		/// Instance of the ACL service
+		/// </summary>
+		private readonly AclService AclService;
+
+		/// <summary>
+		/// Instance of the Job service
+		/// </summary>
+		private readonly JobService JobService;
+
+		/// <summary>
+		/// Constructor
+		/// </summary>
+		/// <param name="DatabaseSerivce">The database service</param>
+		/// <param name="ArtifactService">The Artifact service</param>
+		/// <param name="AclService">The ACL service</param>
+		/// <param name="JobService">The Job service</param>
+		public ArtifactsController(DatabaseService DatabaseSerivce, ArtifactService ArtifactService, AclService AclService, JobService JobService)
+		{
+			this.DatabaseService = DatabaseSerivce;
+			this.ArtifactService = ArtifactService;
+			this.AclService = AclService;
+			this.JobService = JobService;
+		}
+
+		/// <summary>
+		/// Creates an artifact
+		/// </summary>
+		/// <param name="JobId">BatchId</param>
+		/// <param name="StepId">StepId</param>
+		/// <param name="File">The file contents</param>
+		/// <returns>Http result code</returns>
+		[HttpPost]
+		[Authorize]
+		[Route("/api/v1/artifacts")]
+		public async Task<ActionResult<CreateArtifactResponse>> CreateArtifact([FromQuery]string JobId, [FromQuery]string? StepId, IFormFile File)
+		{
+			IJob? Job = await JobService.GetJobAsync(JobId.ToObjectId());
+			if(Job == null)
+			{
+				return NotFound();
+			}
+
+			if (!await JobService.AuthorizeAsync(Job, AclAction.UploadArtifact, User, null))
+			{
+				return Forbid();
+			}
+
+			IJobStep? Step = null;
+			if(StepId != null)
+			{
+				foreach(IJobStepBatch Batch in Job.Batches)
+				{
+					if(Batch.TryGetStep(StepId.ToSubResourceId(), out Step))
+					{
+						break;
+					}
+				}
+				if(Step == null)
+				{
+					// if the step doesn't exist in any of the batches, not found
+					return NotFound();
+				}
+			}
+		
+			Artifact NewArtifact = await ArtifactService.CreateArtifactAsync(Job.Id, Step?.Id, File.FileName, File.ContentType ?? "horde-mime/unknown", File.OpenReadStream());
+			return new CreateArtifactResponse(NewArtifact.Id.ToString());
+		}
+
+		/// <summary>
+		/// Updates an artifact
+		/// </summary>
+		/// <param name="ArtifactId">JobId</param>
+		/// <param name="File">The file contents</param>
+		/// <returns>Http result code</returns>
+		[HttpPut]
+		[Authorize]
+		[Route("/api/v1/artifacts/{ArtifactId}")]
+		public async Task<ActionResult<CreateArtifactResponse>> UpdateArtifact(string ArtifactId, IFormFile File)
+		{
+			Artifact? Artifact = await ArtifactService.GetArtifactAsync(ArtifactId.ToObjectId());
+			if (Artifact == null)
+			{
+				return NotFound();
+			}
+			if (!await JobService.AuthorizeAsync(Artifact.JobId, AclAction.UploadArtifact, User, null))
+			{
+				return Forbid();
+			}
+
+			await ArtifactService.UpdateArtifactAsync(Artifact, File.ContentType ?? "horde-mime/unknown", File.OpenReadStream());
+			return Ok();
+		}
+
+		/// <summary>
+		/// Query artifacts for a job step
+		/// </summary>
+		/// <param name="JobId">Optional JobId to filter by</param>
+		/// <param name="StepId">Optional StepId to filter by</param>
+		/// <param name="Code">Whether to generate a direct download code</param>
+		/// <param name="Filter">Filter for the properties to return</param>
+		/// <returns>Information about all the artifacts</returns>
+		[HttpGet]
+		[Authorize]
+		[Route("/api/v1/artifacts")]
+		[ProducesResponseType(typeof(List<GetArtifactResponse>), 200)]
+		public async Task<ActionResult<List<object>>> GetArtifacts([FromQuery] string JobId, [FromQuery] string? StepId = null, [FromQuery] bool Code = false, [FromQuery] PropertyFilter? Filter = null)
+		{
+			ObjectId JobIdValue = JobId.ToObjectId();
+			if (!await JobService.AuthorizeAsync(JobIdValue, AclAction.DownloadArtifact, User, null))
+			{
+				return Forbid();
+			}
+
+			string? DownloadCode = Code ? (string?)GetDirectDownloadCodeForJob(JobIdValue) : null;
+
+			List<Artifact> Artifacts = await ArtifactService.GetArtifactsAsync(JobIdValue, StepId?.ToSubResourceId(), null);
+			return Artifacts.ConvertAll(x => new GetArtifactResponse(x, DownloadCode).ApplyFilter(Filter));
+		}
+
+		/// <summary>
+		/// Gets the claim required to download artifacts for a particular job
+		/// </summary>
+		/// <param name="JobId">The job id</param>
+		/// <returns>The required claim</returns>
+		static Claim GetDirectDownloadClaim(ObjectId JobId)
+		{
+			return new Claim(HordeClaimTypes.JobArtifacts, JobId.ToString());
+		}
+
+		/// <summary>
+		/// Get a download code for the artifacts of a job
+		/// </summary>
+		/// <param name="JobId">The job id</param>
+		/// <returns>The download code</returns>
+		string GetDirectDownloadCodeForJob(ObjectId JobId)
+		{
+			Claim DownloadClaim = GetDirectDownloadClaim(JobId);
+			return AclService.IssueBearerToken(new[] { DownloadClaim }, TimeSpan.FromHours(4.0));
+		}
+
+		/// <summary>
+		/// Retrieve metadata about a specific artifact
+		/// </summary>
+		/// <param name="ArtifactId">Id of the artifact to get information about</param>
+		/// <param name="Code">Whether to generate a direct download code</param>
+		/// <param name="Filter">Filter for the properties to return</param>
+		/// <returns>Information about the requested project</returns>
+		[HttpGet]
+		[Authorize]
+		[Route("/api/v1/artifacts/{ArtifactId}")]
+		[ProducesResponseType(typeof(GetArtifactResponse), 200)]
+		public async Task<ActionResult<object>> GetArtifact(string ArtifactId, bool Code = false, [FromQuery] PropertyFilter? Filter = null)
+		{
+			Artifact? Artifact = await ArtifactService.GetArtifactAsync(ArtifactId.ToObjectId());
+			if (Artifact == null)
+			{
+				return NotFound();
+			}
+			if (!await JobService.AuthorizeAsync(Artifact.JobId, AclAction.DownloadArtifact, User, null))
+			{
+				return Forbid();
+			}
+
+			string? DownloadCode = Code? (string?)GetDirectDownloadCodeForJob(Artifact.JobId) : null;
+			return new GetArtifactResponse(Artifact, DownloadCode).ApplyFilter(Filter);
+		}
+
+		/// <summary>
+		/// Retrieve raw data for an artifact
+		/// </summary>
+		/// <param name="ArtifactId">Id of the artifact to get information about</param>
+		/// <returns>Raw artifact data</returns>
+		[HttpGet]
+		[Authorize]
+		[Route("/api/v1/artifacts/{ArtifactId}/data")]
+		public async Task<ActionResult> GetArtifactData(string ArtifactId)
+		{
+			Artifact? Artifact = await ArtifactService.GetArtifactAsync(ArtifactId.ToObjectId());
+			if (Artifact == null)
+			{
+				return NotFound();
+			}
+			if (!await JobService.AuthorizeAsync(Artifact.JobId, AclAction.DownloadArtifact, User, null))
+			{
+				return Forbid();
+			}
+
+			// Fun, filestream result automatically closes the stream!
+			return new FileStreamResult(ArtifactService.OpenArtifactReadStream(Artifact), Artifact.MimeType);
+		}
+		
+		/// <summary>
+		/// Retrieve raw data for an artifact by filename
+		/// </summary>
+		/// <param name="JobId">Unique id for the job</param>
+		/// <param name="StepId">Unique id for the step</param>
+		/// <param name="Filename">Filename of artifact from step</param>
+		/// <returns>Raw artifact data</returns>
+		[HttpGet]
+		[Route("/api/v1/jobs/{JobId}/steps/{StepId}/artifacts/{FileName}/data")]
+		public async Task<ActionResult<object>> GetArtifactDataByFilename(string JobId, string StepId, string Filename)
+		{
+			ObjectId JobIdValue = JobId.ToObjectId();
+			SubResourceId StepIdValue = StepId.ToSubResourceId();
+
+			if (!await JobService.AuthorizeAsync(JobIdValue, AclAction.DownloadArtifact, User, null))
+			{
+				return Forbid();
+			}
+			
+			List<Artifact> Artifacts = await ArtifactService.GetArtifactsAsync(JobIdValue, StepIdValue, Filename);
+			if (Artifacts.Count == 0)
+			{
+				return NotFound();
+			}
+
+			Artifact Artifact = Artifacts[0];
+			return new FileStreamResult(ArtifactService.OpenArtifactReadStream(Artifact), Artifact.MimeType);
+		}
+
+		/// <summary>
+		/// Class to return a file stream without the "content-disposition: attachment" header
+		/// </summary>
+		class InlineFileStreamResult : FileStreamResult
+		{
+			/// <summary>
+			/// The suggested download filename
+			/// </summary>
+			string FileName;
+
+			/// <summary>
+			/// Constructor
+			/// </summary>
+			/// <param name="Stream"></param>
+			/// <param name="MimeType"></param>
+			/// <param name="FileName"></param>
+			public InlineFileStreamResult(System.IO.Stream Stream, string MimeType, string FileName)
+				: base(Stream, MimeType)
+			{
+				this.FileName = FileName;
+			}
+
+			/// <inheritdoc/>
+			public override Task ExecuteResultAsync(ActionContext Context)
+			{
+				ContentDisposition ContentDisposition = new ContentDisposition();
+				ContentDisposition.Inline = true;
+				ContentDisposition.FileName = FileName;
+				Context.HttpContext.Response.Headers.Add("Content-Disposition", ContentDisposition.ToString());
+				
+				return base.ExecuteResultAsync(Context);
+			}
+		}
+
+		/// <summary>
+		/// Retrieve raw data for an artifact
+		/// </summary>
+		/// <param name="ArtifactId">Id of the artifact to get information about</param>
+		/// <param name="Code">The authorization code for this resource</param>
+		/// <returns>Raw artifact data</returns>
+		[HttpGet]
+		[AllowAnonymous]
+		[Route("/api/v1/artifacts/{ArtifactId}/download")]
+		public async Task<ActionResult> DownloadArtifact(string ArtifactId, [FromQuery] string Code)
+		{
+			TokenValidationParameters Parameters = new TokenValidationParameters();
+			Parameters.ValidateAudience = false;
+			Parameters.RequireExpirationTime = true;
+			Parameters.ValidateLifetime = true;
+			Parameters.ValidIssuer = DatabaseService.JwtIssuer;
+			Parameters.ValidateIssuer = true;
+			Parameters.ValidateIssuerSigningKey = true;
+			Parameters.IssuerSigningKey = DatabaseService.JwtSigningKey;
+
+			SecurityToken Token;
+			JwtSecurityTokenHandler Handler = new JwtSecurityTokenHandler();
+			ClaimsPrincipal Principal = Handler.ValidateToken(Code, Parameters, out Token);
+
+			Artifact? Artifact = await ArtifactService.GetArtifactAsync(ArtifactId.ToObjectId());
+			if (Artifact == null)
+			{
+				return NotFound();
+			}
+
+			Claim DirectDownloadClaim = GetDirectDownloadClaim(Artifact.JobId);
+			if (!Principal.HasClaim(DirectDownloadClaim.Type, DirectDownloadClaim.Value))
+			{
+				return Forbid();
+			}
+
+			return new InlineFileStreamResult(ArtifactService.OpenArtifactReadStream(Artifact), Artifact.MimeType, Path.GetFileName(Artifact.Name));
+		}
+
+		/// <summary>
+		/// Returns a zip archive of many artifacts
+		/// </summary>
+		/// <param name="ArtifactZipRequest">Artifact request params</param>
+		/// <returns>Zip of many artifacts</returns>
+		[HttpPost]
+		[Authorize]
+		[Route("/api/v1/artifacts/zip")]
+		public async Task<ActionResult> ZipArtifacts(GetArtifactZipRequest ArtifactZipRequest)
+		{
+			if (ArtifactZipRequest.JobId == null)
+			{
+				return BadRequest("Must specify a JobId");
+			}
+
+			IJob? Job = await JobService.GetJobAsync(ArtifactZipRequest.JobId!.ToObjectId());
+			if (Job == null)
+			{
+				return NotFound();
+			}
+			if (!await JobService.AuthorizeAsync(Job, AclAction.DownloadArtifact, User, null))
+			{
+				return Forbid();
+			}
+
+			List<Artifact> Artifacts = await ArtifactService.GetArtifactsAsync(Job.Id, ArtifactZipRequest.StepId?.ToSubResourceId(), null);
+
+			Dictionary<ObjectId, Artifact> IdToArtifact = Artifacts.ToDictionary(x => x.Id, x => x);
+
+			List<Artifact> ZipArtifacts;
+			if (ArtifactZipRequest.ArtifactIds == null)
+			{
+				ZipArtifacts = Artifacts;
+			}
+			else
+			{
+				ZipArtifacts = new List<Artifact>();
+				foreach (string ArtifactId in ArtifactZipRequest.ArtifactIds)
+				{
+					Artifact? Artifact;
+					if (IdToArtifact.TryGetValue(ArtifactId.ToObjectId(), out Artifact))
+					{
+						ZipArtifacts.Add(Artifact);
+					}
+					else
+					{
+						return NotFound();
+					}
+				}
+			}
+
+			IGraph Graph = await JobService.GetGraphAsync(Job);
+
+			return new CustomFileCallbackResult("Artifacts.zip", "application/octet-stream", false, async (OutputStream, Context) =>
+			{
+				// Make an unseekable MemoryStream for the ZipArchive. We have to do this because the ZipEntry stream falls back to a synchronous write to it's own stream wrappers.
+				using (CustomBufferStream ZipOutputStream = new CustomBufferStream())
+				{
+					// Keep the stream open after dispose so we can write the EOF bits.
+					using (ZipArchive ZipArchive = new ZipArchive(ZipOutputStream, ZipArchiveMode.Create, true))
+					{
+						foreach (Artifact Artifact in ZipArtifacts)
+						{
+							await using (System.IO.Stream ArtifactStream = ArtifactService.OpenArtifactReadStream(Artifact))
+							{
+								// tack on the step name into the directory if it exists
+								string StepName = string.Empty;
+								if (Artifact.StepId.HasValue)
+								{
+									foreach (IJobStepBatch Batch in Job.Batches)
+									{
+										IJobStep Step;
+										if (Batch.TryGetStep(Artifact.StepId.Value, out Step))
+										{
+											StepName = Graph.Groups[Batch.GroupIdx].Nodes[Step.NodeIdx].Name;
+											break;
+										}
+									}
+								}
+
+								ZipArchiveEntry ZipEntry = ZipArchive.CreateEntry(Artifact.Name);
+								using (System.IO.Stream EntryStream = ZipEntry.Open())
+								{
+									byte[] Buffer = new byte[4096];
+									int TotalBytesRead = 0;
+									while (TotalBytesRead < Artifact.Length)
+									{
+										int BytesRead = await ArtifactStream.ReadAsync(Buffer, 0, Buffer.Length);
+
+										// Write bytes to the entry stream.  Also advances the MemStream pos.
+										await EntryStream.WriteAsync(Buffer, 0, BytesRead);
+										// Dump what we have to the output stream
+										await OutputStream.WriteAsync(ZipOutputStream.GetBuffer(), 0, (int)ZipOutputStream.Position);
+
+										// Reset everything.
+										ZipOutputStream.Position = 0;
+										ZipOutputStream.SetLength(0);
+										TotalBytesRead += BytesRead;
+									}
+								}
+							}
+						}
+					}
+					// Write out the EOF stuff
+					ZipOutputStream.Position = 0;
+					await ZipOutputStream.CopyToAsync(OutputStream);
+				}
+			});
+		}
+	}
+}
