Initial drop of Horde Server/Agent for UE5/Main. Still missing some Epic-specific functionality required to build/run correctly.

diff --git a/Engine/Source/Programs/Horde/HordeServer/Controllers/JobsController.cs b/Engine/Source/Programs/Horde/HordeServer/Controllers/JobsController.cs
new file mode 100644
--- /dev/null
+++ b/Engine/Source/Programs/Horde/HordeServer/Controllers/JobsController.cs
@@ -0,0 +1,1399 @@
+// Copyright Epic Games, Inc. All Rights Reserved.
+
+using Datadog.Trace;
+using EpicGames.Core;
+using HordeServer.Api;
+using HordeServer.Collections;
+using HordeCommon;
+using HordeServer.Models;
+using HordeServer.Services;
+using HordeServer.Utilities;
+using Microsoft.AspNetCore.Authorization;
+using Microsoft.AspNetCore.Mvc;
+using MongoDB.Bson;
+using System;
+using System.Collections.Generic;
+using System.ComponentModel.DataAnnotations;
+using System.Globalization;
+using System.Linq;
+using System.Security.Claims;
+using System.Text.RegularExpressions;
+using System.Threading.Tasks;
+
+using StreamId = HordeServer.Utilities.StringId<HordeServer.Models.IStream>;
+using TemplateRefId = HordeServer.Utilities.StringId<HordeServer.Models.TemplateRef>;
+using Microsoft.Extensions.Logging;
+using HordeServer.Notifications;
+
+namespace HordeServer.Controllers
+{
+	/// <summary>
+	/// Controller for the /api/v1/jobs endpoing
+	/// </summary>
+	[ApiController]
+	[Authorize]
+	[Route("[controller]")]
+	public class JobsController : ControllerBase
+	{
+		/// <summary>
+		/// Instance of the ACL service
+		/// </summary>
+		private readonly AclService AclService;
+
+		/// <summary>
+		/// Collection of graphs
+		/// </summary>
+		private readonly IGraphCollection Graphs;
+
+		/// <summary>
+		/// The perforce service instance
+		/// </summary>
+		private readonly IPerforceService Perforce;
+
+		/// <summary>
+		/// Instance of the stream service
+		/// </summary>
+		private readonly StreamService StreamService;
+
+		/// <summary>
+		/// Instance of the JobService singleton
+		/// </summary>
+		private readonly JobService JobService;
+
+		/// <summary>
+		/// Instance of the TemplateService singleton
+		/// </summary>
+		private readonly TemplateService TemplateService;
+
+		/// <summary>
+		/// 
+		/// </summary>
+		private readonly ArtifactService ArtifactService;
+
+		/// <summary>
+		/// Instance of the notification service singleton
+		/// </summary>
+		private readonly INotificationService NotificationService;
+
+		/// <summary>
+		/// Logger instance
+		/// </summary>
+		private ILogger<JobsController> Logger;
+
+		/// <summary>
+		/// Constructor
+		/// </summary>
+		/// <param name="AclService">Instance of the ACL service</param>
+		/// <param name="Graphs">Collection of graph documents</param>
+		/// <param name="Perforce">The Perforce service instance</param>
+		/// <param name="StreamService">Instance of the stream service</param>
+		/// <param name="JobService">Instance of the JobService singleton</param>
+		/// <param name="TemplateService">Instance of the TemplateService singleton</param>
+		/// <param name="ArtifactService">Instance of the ArtifactService singleton</param>
+		/// <param name="NotificationService">Instance of the NotificationService singelton</param>
+		/// <param name="Logger">The logger instance</param>
+		public JobsController(AclService AclService, IGraphCollection Graphs, IPerforceService Perforce, StreamService StreamService, JobService JobService, TemplateService TemplateService, ArtifactService ArtifactService, INotificationService NotificationService, ILogger<JobsController> Logger)
+		{
+			this.AclService = AclService;
+			this.Graphs = Graphs;
+			this.Perforce = Perforce;
+			this.StreamService = StreamService;
+			this.JobService = JobService;
+			this.TemplateService = TemplateService;
+			this.ArtifactService = ArtifactService;
+			this.NotificationService = NotificationService;
+			this.Logger = Logger;
+		}
+
+		/// <summary>
+		/// Creates a new job
+		/// </summary>
+		/// <param name="Create">Properties of the new job</param>
+		/// <returns>Id of the new job</returns>
+		[HttpPost]
+		[Route("/api/v1/jobs")]
+		public async Task<ActionResult<CreateJobResponse>> CreateJobAsync([FromBody] CreateJobRequest Create)
+		{
+			IStream? Stream = await StreamService.GetStreamAsync(new StreamId(Create.StreamId));
+			if (Stream == null)
+			{
+				return BadRequest("Invalid StreamId parameter");
+			}
+			if (!await StreamService.AuthorizeAsync(Stream, AclAction.CreateJob, User, null))
+			{
+				return Forbid();
+			}
+
+			// Get the name of the template ref
+			TemplateRefId TemplateRefId = new TemplateRefId(Create.TemplateId);
+
+			// Augment the request with template properties
+			TemplateRef? TemplateRef;
+			if (!Stream.Templates.TryGetValue(TemplateRefId, out TemplateRef))
+			{
+				return BadRequest($"Invalid {Create.TemplateId} parameter");
+			}
+			if (!await StreamService.AuthorizeAsync(Stream, TemplateRef, AclAction.CreateJob, User, null))
+			{
+				return Forbid();
+			}
+
+			ITemplate? Template = await TemplateService.GetTemplateAsync(TemplateRef.Hash);
+			if (Template == null)
+			{
+				return BadRequest($"Missing template referenced by {Create.TemplateId}");
+			}
+			if (!Template.AllowPreflights && Create.PreflightChange > 0)
+			{
+				return BadRequest("Template does not allow preflights");
+			}
+
+			// Get the name of the new job
+			string Name = Create.Name ?? Template.Name;
+			if (Create.TemplateId.Equals("stage-to-marketplace", StringComparison.Ordinal) && Create.Arguments != null)
+			{
+				foreach (string Argument in Create.Arguments)
+				{
+					const string Prefix = "-set:UserContentItems=";
+					if (Argument.StartsWith(Prefix, StringComparison.Ordinal))
+					{
+						Name += $" - {Argument.Substring(Prefix.Length)}";
+						break;
+					}
+				}
+			}
+
+			// Get the priority of the new job
+			Priority Priority = Create.Priority ?? Template.Priority ?? Priority.Normal;
+
+			// New groups for the job
+			IGraph Graph = await Graphs.AddAsync(Template);
+			if (Create.Groups != null)
+			{
+				Graph = await Graphs.AppendAsync(Graph, Create.Groups, null, null);
+			}
+			if (Create.Aggregates != null)
+			{
+				Graph = await Graphs.AppendAsync(Graph, null, Create.Aggregates, null);
+			}
+			if (Create.Labels != null)
+			{
+				Graph = await Graphs.AppendAsync(Graph, null, null, Create.Labels);
+			}
+
+			// Get the change to build
+			int Change;
+			if (Create.Change.HasValue)
+			{
+				Change = Create.Change.Value;
+			}
+			else if (Create.ChangeQuery != null)
+			{
+				Change = await ExecuteChangeQueryAsync(Stream, Create.ChangeQuery);
+			}
+			else if (Create.PreflightChange == null && Template.SubmitNewChange != null)
+			{
+				Change = await Perforce.CreateNewChangeForTemplateAsync(Stream, Template);
+			}
+			else
+			{
+				Change = await Perforce.GetLatestChangeAsync(Stream.Name, null);
+			}
+
+			// And get the matching code changelist
+			int CodeChange = await Perforce.GetCodeChangeAsync(Stream.Name, Change);
+
+			// New properties for the job
+			List<string> Arguments = Create.Arguments ?? Template.GetDefaultArguments();
+
+			// Create the job
+			IJob Job = await JobService.CreateJobAsync(null, Stream.Id, TemplateRefId, Template.Id, Graph, Name, Change, CodeChange, Create.PreflightChange, null, User.GetUserId(), User.GetUserName(), Priority, Create.AutoSubmit, TemplateRef.ChainedJobs, TemplateRef.ShowUgsBadges, TemplateRef.ShowUgsAlerts, TemplateRef.NotificationChannel, TemplateRef.NotificationChannelFilter, null, Template.Counters, Arguments);
+			await UpdateNotificationsAsync(Job.Id.ToString(), new UpdateNotificationsRequest { Slack = true });
+			return new CreateJobResponse(Job.Id.ToString());
+		}
+
+		/// <summary>
+		/// Evaluate a change query to determine which CL to run a job at
+		/// </summary>
+		/// <param name="Stream"></param>
+		/// <param name="Request"></param>
+		/// <returns></returns>
+		async Task<int> ExecuteChangeQueryAsync(IStream Stream, ChangeQueryRequest Request)
+		{
+			IList<IJob> Jobs = await JobService.FindJobsAsync(StreamId: Stream.Id, Templates: new[] { new TemplateRefId(Request.TemplateId) }, Target: Request.Target, State: new[] { JobStepState.Completed }, Outcome: new[] { JobStepOutcome.Success }, Count: 1);
+			if (Jobs.Count == 0)
+			{
+				Logger.LogInformation("Unable to find successful build of {TemplateRefId} target {Target}. Using latest change instead", Request.TemplateId, Request.Target);
+				return await Perforce.GetLatestChangeAsync(Stream.Name, null);
+			}
+			else
+			{
+				Logger.LogInformation("Last successful build of {TemplateRefId} target {Target} was job {JobId} at change {Change}", Request.TemplateId, Request.Target, Jobs[0].Id, Jobs[0].Change);
+				return Jobs[0].Change;
+			}
+		}
+
+		/// <summary>
+		/// Deletes a specific job.
+		/// </summary>
+		/// <param name="JobId">Id of the job to delete</param>
+		/// <returns>Async task</returns>
+		[HttpDelete]
+		[Route("/api/v1/jobs/{JobId}")]
+		public async Task<ActionResult> DeleteJobAsync(string JobId)
+		{
+			IJob? Job = await JobService.GetJobAsync(JobId.ToObjectId());
+			if (Job == null)
+			{
+				return NotFound();
+			}
+			if (!await JobService.AuthorizeAsync(Job, AclAction.DeleteJob, User, null))
+			{
+				return Forbid();
+			}
+			if (!await JobService.DeleteJobAsync(Job))
+			{
+				return NotFound();
+			}
+			return Ok();
+		}
+
+		/// <summary>
+		/// Updates a specific job.
+		/// </summary>
+		/// <param name="JobId">Id of the job to find</param>
+		/// <param name="Request">Settings to update in the job</param>
+		/// <returns>Async task</returns>
+		[HttpPut]
+		[Route("/api/v1/jobs/{JobId}")]
+		public async Task<ActionResult> UpdateJobAsync(string JobId, [FromBody] UpdateJobRequest Request)
+		{
+			IJob? Job = await JobService.GetJobAsync(JobId.ToObjectId());
+			if (Job == null)
+			{
+				return NotFound();
+			}
+
+			StreamPermissionsCache PermissionsCache = new StreamPermissionsCache();
+			if (!await JobService.AuthorizeAsync(Job, AclAction.UpdateJob, User, PermissionsCache))
+			{
+				return Forbid();
+			}
+			if (Request.Acl != null && !await JobService.AuthorizeAsync(Job, AclAction.ChangePermissions, User, PermissionsCache))
+			{
+				return Forbid();
+			}
+
+			// Convert legacy behavior of clearing out the argument to setting the aborted flag
+			if (Request.Arguments != null && Request.Arguments.Count == 0)
+			{
+				Request.Aborted = true;
+				Request.Arguments = null;
+			}
+
+			string? AbortedByUser = null;
+			if (Request.Aborted ?? false)
+			{
+				AbortedByUser = User.Identity.Name;
+			}
+
+			if (!await JobService.UpdateJobAsync(Job, Name: Request.Name, Priority: Request.Priority, AutoSubmit: Request.AutoSubmit, AbortedByUser: AbortedByUser, Arguments: Request.Arguments))
+			{
+				return NotFound();
+			}
+			return Ok();
+		}
+
+		/// <summary>
+		/// Updates notifications for a specific job.
+		/// </summary>
+		/// <param name="JobId">Id of the job to find</param>
+		/// <param name="Request">The notification request</param>
+		/// <returns>Information about the requested job</returns>
+		[HttpPut]
+		[Route("/api/v1/jobs/{JobId}/notifications")]
+		public async Task<ActionResult> UpdateNotificationsAsync(string JobId, [FromBody] UpdateNotificationsRequest Request)
+		{
+			IJob? Job = await JobService.GetJobAsync(JobId.ToObjectId());
+			if (Job == null)
+			{
+				return NotFound();
+			}
+			if (!await JobService.AuthorizeAsync(Job, AclAction.CreateSubscription, User, null))
+			{
+				return Forbid();
+			}
+
+			ObjectId TriggerId = Job.NotificationTriggerId ?? ObjectId.GenerateNewId();
+			if (!await JobService.UpdateJobAsync(Job, null, null, null, null, TriggerId, null, null))
+			{
+				return NotFound();
+			}
+
+			await NotificationService.UpdateSubscriptionsAsync(TriggerId, User, Request.Email, Request.Slack);
+			return Ok();
+		}
+
+		/// <summary>
+		/// Gets information about a specific job.
+		/// </summary>
+		/// <param name="JobId">Id of the job to find</param>
+		/// <returns>Information about the requested job</returns>
+		[HttpGet]
+		[Route("/api/v1/jobs/{JobId}/notifications")]
+		public async Task<ActionResult<GetNotificationResponse>> GetNotificationsAsync(string JobId)
+		{
+			IJob? Job = await JobService.GetJobAsync(JobId.ToObjectId());
+			if (Job == null)
+			{
+				return NotFound();
+			}
+			if (!await JobService.AuthorizeAsync(Job, AclAction.CreateSubscription, User, null))
+			{
+				return Forbid();
+			}
+
+			INotificationSubscription? Subscription;
+			if (Job.NotificationTriggerId == null)
+			{
+				Subscription = null;
+			}
+			else
+			{
+				Subscription = await NotificationService.GetSubscriptionsAsync(Job.NotificationTriggerId.Value, User);
+			}
+			return new GetNotificationResponse(Subscription);
+		}
+
+		/// <summary>
+		/// Gets information about a specific job.
+		/// </summary>
+		/// <param name="JobId">Id of the job to find</param>
+		/// <param name="ModifiedAfter">If specified, returns an empty response unless the job's update time is equal to or less than the given value</param>
+		/// <param name="Filter">Filter for the fields to return</param>
+		/// <returns>Information about the requested job</returns>
+		[HttpGet]
+		[Route("/api/v1/jobs/{JobId}")]
+		[ProducesResponseType(typeof(GetJobResponse), 200)]
+		public async Task<ActionResult<object>> GetJobAsync(string JobId, [FromQuery] DateTimeOffset? ModifiedAfter = null, [FromQuery] PropertyFilter? Filter = null)
+		{
+			IJob? Job = await JobService.GetJobAsync(JobId.ToObjectId());
+			if (Job == null)
+			{
+				return NotFound();
+			}
+
+			StreamPermissionsCache Cache = new StreamPermissionsCache();
+			if (!await JobService.AuthorizeAsync(Job, AclAction.ViewJob, User, Cache))
+			{
+				return Forbid();
+			}
+			if (ModifiedAfter != null && Job.UpdateTimeUtc <= ModifiedAfter.Value)
+			{
+				return new Dictionary<string, object>();
+			}
+
+			IGraph Graph = await JobService.GetGraphAsync(Job);
+			bool bIncludeAcl = await JobService.AuthorizeAsync(Job, AclAction.ViewPermissions, User, Cache);
+			return Job.ToResponse(Graph, bIncludeAcl, Filter);
+		}
+
+		/// <summary>
+		/// Gets information about the graph for a specific job.
+		/// </summary>
+		/// <param name="JobId">Id of the job to find</param>
+		/// <param name="Filter">Filter for the fields to return</param>
+		/// <returns>Information about the requested job</returns>
+		[HttpGet]
+		[Route("/api/v1/jobs/{JobId}/graph")]
+		[ProducesResponseType(typeof(GetGraphResponse), 200)]
+		public async Task<ActionResult<object>> GetJobGraphAsync(string JobId, [FromQuery] PropertyFilter? Filter = null)
+		{
+			IJob? Job = await JobService.GetJobAsync(JobId.ToObjectId());
+			if (Job == null)
+			{
+				return NotFound();
+			}
+			if (!await JobService.AuthorizeAsync(Job, AclAction.ViewJob, User, null))
+			{
+				return Forbid();
+			}
+
+			IGraph Graph = await JobService.GetGraphAsync(Job);
+			return PropertyFilter.Apply(new GetGraphResponse(Graph), Filter);
+		}
+
+		/// <summary>
+		/// Gets timing information about the graph for a specific job.
+		/// </summary>
+		/// <param name="JobId">Id of the job to find</param>
+		/// <param name="Filter">Filter for the fields to return</param>
+		/// <returns>Information about the requested job</returns>
+		[HttpGet]
+		[Route("/api/v1/jobs/{JobId}/timing")]
+		[ProducesResponseType(typeof(GetJobTimingResponse), 200)]
+		public async Task<ActionResult<object>> GetJobTimingAsync(string JobId, [FromQuery] PropertyFilter? Filter = null)
+		{
+			IJob? Job = await JobService.GetJobAsync(JobId.ToObjectId());
+			if (Job == null)
+			{
+				return NotFound();
+			}
+			if (!await JobService.AuthorizeAsync(Job, AclAction.ViewJob, User, null))
+			{
+				return Forbid();
+			}
+
+			IJobTiming JobTiming = await JobService.GetJobTimingAsync(Job);
+
+			IGraph Graph = await JobService.GetGraphAsync(Job);
+			Dictionary<INode, TimingInfo> NodeToTimingInfo = Job.GetTimingInfo(Graph, JobTiming);
+
+			Dictionary<string, GetStepTimingInfoResponse> Steps = new Dictionary<string, GetStepTimingInfoResponse>();
+			foreach (IJobStepBatch Batch in Job.Batches)
+			{
+				foreach (IJobStep Step in Batch.Steps)
+				{
+					INode Node = Graph.Groups[Batch.GroupIdx].Nodes[Step.NodeIdx];
+					Steps[Step.Id.ToString()] = new GetStepTimingInfoResponse(Node.Name, NodeToTimingInfo[Node]);
+				}
+			}
+
+			List<GetLabelTimingInfoResponse> Labels = new List<GetLabelTimingInfoResponse>();
+			foreach (ILabel Label in Graph.Labels)
+			{
+				TimingInfo TimingInfo = TimingInfo.Max(Label.GetDependencies(Graph.Groups).Select(x => NodeToTimingInfo[x]));
+				Labels.Add(new GetLabelTimingInfoResponse(Label, TimingInfo));
+			}
+
+			return PropertyFilter.Apply(new GetJobTimingResponse(Steps, Labels), Filter);
+		}
+
+		/// <summary>
+		/// Gets information about the template for a specific job.
+		/// </summary>
+		/// <param name="JobId">Id of the job to find</param>
+		/// <param name="Filter">Filter for the fields to return</param>
+		/// <returns>Information about the requested job</returns>
+		[HttpGet]
+		[Route("/api/v1/jobs/{JobId}/template")]
+		[ProducesResponseType(typeof(GetTemplateResponse), 200)]
+		public async Task<ActionResult<object>> GetJobTemplateAsync(string JobId, [FromQuery] PropertyFilter? Filter = null)
+		{
+			IJob? Job = await JobService.GetJobAsync(JobId.ToObjectId());
+			if (Job == null || Job.TemplateHash == null)
+			{
+				return NotFound();
+			}
+			if (!await JobService.AuthorizeAsync(Job, AclAction.ViewJob, User, null))
+			{
+				return Forbid();
+			}
+
+			ITemplate? Template = await TemplateService.GetTemplateAsync(Job.TemplateHash);
+			if(Template == null)
+			{
+				return NotFound();
+			}
+
+			return new GetTemplateResponse(Template).ApplyFilter(Filter);
+		}
+
+		/// <summary>
+		/// Find jobs matching a criteria
+		/// </summary>
+		/// <param name="Ids">The job ids to return</param>
+		/// <param name="Name">Name of the job to find</param>
+		/// <param name="Templates">List of templates to find</param>
+		/// <param name="StreamId">The stream to search for</param>
+		/// <param name="MinChange">The minimum changelist number</param>
+		/// <param name="MaxChange">The maximum changelist number</param>
+		/// <param name="IncludePreflight">Whether to include preflight jobs</param>
+		/// <param name="PreflightChange">The preflighted changelist</param>
+		/// <param name="PreflightStartedByUser">User for which to include preflight jobs</param>
+		/// <param name="MinCreateTime">Minimum creation time</param>
+		/// <param name="MaxCreateTime">Maximum creation time</param>
+		/// <param name="ModifiedBefore">If specified, only jobs updated before the give time will be returned</param>
+		/// <param name="ModifiedAfter">If specified, only jobs updated after the give time will be returned</param>
+		/// <param name="Target">Target to filter the returned jobs by</param>
+		/// <param name="State">Filter state of the returned jobs</param>
+		/// <param name="Outcome">Filter outcome of the returned jobs</param>
+		/// <param name="Filter">Filter for properties to return</param>
+		/// <param name="Index">Index of the first result to be returned</param>
+		/// <param name="Count">Number of results to return</param>
+		/// <returns>List of jobs</returns>
+		[HttpGet]
+		[Route("/api/v1/jobs")]
+		[ProducesResponseType(typeof(List<GetJobResponse>), 200)]
+		public async Task<ActionResult<List<object>>> FindJobsAsync(
+			[FromQuery(Name = "Id")] string[]? Ids = null,
+			[FromQuery] string? StreamId = null,
+			[FromQuery] string? Name = null,
+			[FromQuery(Name = "template")] string[]? Templates = null,
+			[FromQuery] int? MinChange = null,
+			[FromQuery] int? MaxChange = null,
+			[FromQuery] bool IncludePreflight = true,
+			[FromQuery] int? PreflightChange = null,
+			[FromQuery] string? PreflightStartedByUser = null,
+			[FromQuery] DateTimeOffset? MinCreateTime = null,
+			[FromQuery] DateTimeOffset? MaxCreateTime = null,
+			[FromQuery] DateTimeOffset? ModifiedBefore = null,
+			[FromQuery] DateTimeOffset? ModifiedAfter = null,
+			[FromQuery] string? Target = null,
+			[FromQuery] JobStepState[]? State = null,
+			[FromQuery] JobStepOutcome[]? Outcome = null, 
+			[FromQuery] PropertyFilter? Filter = null,
+			[FromQuery] int Index = 0,
+			[FromQuery] int Count = 100)
+		{
+			ObjectId[]? JobIdValues = (Ids == null) ? (ObjectId[]?)null : Array.ConvertAll(Ids, x => x.ToObjectId());
+			StreamId? StreamIdValue = (StreamId == null)? (StreamId?)null : new StreamId(StreamId);
+
+			TemplateRefId[]? TemplateRefIds = (Templates != null && Templates.Length > 0) ? Templates.Select(x => new TemplateRefId(x)).ToArray() : null;
+
+			if (IncludePreflight == false)
+			{
+				PreflightChange = 0;
+			}
+
+			List<IJob> Jobs;
+			using (Scope _ = Tracer.Instance.StartActive("FindJobs"))
+			{
+				Jobs = await JobService.FindJobsAsync(JobIdValues, StreamIdValue, Name, TemplateRefIds, MinChange,
+					MaxChange, PreflightChange, PreflightStartedByUser, MinCreateTime?.UtcDateTime, MaxCreateTime?.UtcDateTime, Target, State, Outcome,
+					ModifiedBefore, ModifiedAfter, Index, Count);
+			}
+
+			StreamPermissionsCache PermissionsCache = new StreamPermissionsCache();
+
+			List<object> Responses = new List<object>();
+			foreach (IJob Job in Jobs)
+			{
+				using Scope JobScope = Tracer.Instance.StartActive("JobIteration");
+				JobScope.Span.SetTag("jobId", Job.Id.ToString());
+				
+				bool ViewJobAuthorized;
+				using (Scope _ = Tracer.Instance.StartActive("AuthorizeViewJob"))
+				{
+					ViewJobAuthorized = await JobService.AuthorizeAsync(Job, AclAction.ViewJob, User, PermissionsCache);
+				}
+				
+				if (ViewJobAuthorized)
+				{
+					IGraph Graph;
+					using (Scope _ = Tracer.Instance.StartActive("GetGraph"))
+					{
+						Graph = await JobService.GetGraphAsync(Job);
+					}
+
+					bool bIncludeAcl;
+					using (Scope _ = Tracer.Instance.StartActive("AuthorizeViewPermissions"))
+					{
+						bIncludeAcl = await JobService.AuthorizeAsync(Job, AclAction.ViewPermissions, User, PermissionsCache);
+					}
+
+					using (Scope _ = Tracer.Instance.StartActive("CreateResponse"))
+					{
+						Responses.Add(Job.ToResponse(Graph, bIncludeAcl, Filter));
+					}
+				}
+			}
+			return Responses;
+		}
+
+		/// <summary>
+		/// Adds an array of nodes to be executed for a job
+		/// </summary>
+		/// <param name="JobId">Unique id for the job</param>
+		/// <param name="Requests">Properties of the new nodes</param>
+		/// <returns>Id of the new job</returns>
+		[HttpPost]
+		[Route("/api/v1/jobs/{JobId}/groups")]
+		public async Task<ActionResult> CreateGroupsAsync(string JobId, [FromBody] List<CreateGroupRequest> Requests)
+		{
+			Dictionary<string, int> ExpectedDurationCache = new Dictionary<string, int>();
+
+			ObjectId JobIdValue = JobId.ToObjectId();
+			for (; ; )
+			{
+				IJob? Job = await JobService.GetJobAsync(JobIdValue);
+				if (Job == null)
+				{
+					return NotFound();
+				}
+				if (!await JobService.AuthorizeAsync(Job, AclAction.ExecuteJob, User, null))
+				{
+					return Forbid();
+				}
+
+				IGraph Graph = await JobService.GetGraphAsync(Job);
+				Graph = await Graphs.AppendAsync(Graph, Requests, null, null);
+
+				if (await JobService.TryUpdateGraphAsync(Job, Graph))
+				{
+					return Ok();
+				}
+			}
+		}
+
+		/// <summary>
+		/// Gets the nodes to be executed for a job
+		/// </summary>
+		/// <param name="JobId">Unique id for the job</param>
+		/// <param name="Filter">Filter for the properties to return</param>
+		/// <returns>List of nodes to be executed</returns>
+		[HttpGet]
+		[Route("/api/v1/jobs/{JobId}/groups")]
+		[ProducesResponseType(typeof(List<GetGroupResponse>), 200)]
+		public async Task<ActionResult<List<object>>> GetGroupsAsync(string JobId, [FromQuery] PropertyFilter? Filter = null)
+		{
+			IJob? Job = await JobService.GetJobAsync(JobId.ToObjectId());
+			if (Job == null)
+			{
+				return NotFound();
+			}
+			if (!await JobService.AuthorizeAsync(Job, AclAction.ViewJob, User, null))
+			{
+				return Forbid();
+			}
+
+			IGraph Graph = await JobService.GetGraphAsync(Job);
+			return Graph.Groups.ConvertAll(x => new GetGroupResponse(x, Graph.Groups).ApplyFilter(Filter));
+		}
+
+		/// <summary>
+		/// Gets the nodes in a group to be executed for a job
+		/// </summary>
+		/// <param name="JobId">Unique id for the job</param>
+		/// <param name="GroupIdx">The group index</param>
+		/// <param name="Filter">Filter for the properties to return</param>
+		/// <returns>List of nodes to be executed</returns>
+		[HttpGet]
+		[Route("/api/v1/jobs/{JobId}/groups/{GroupIdx}")]
+		[ProducesResponseType(typeof(GetGroupResponse), 200)]
+		public async Task<ActionResult<object>> GetGroupAsync(string JobId, int GroupIdx, [FromQuery] PropertyFilter? Filter = null)
+		{
+			IJob? Job = await JobService.GetJobAsync(JobId.ToObjectId());
+			if (Job == null)
+			{
+				return NotFound();
+			}
+			if (!await JobService.AuthorizeAsync(Job, AclAction.ViewJob, User, null))
+			{
+				return Forbid();
+			}
+
+			IGraph Graph = await JobService.GetGraphAsync(Job);
+			if (GroupIdx < 0 || GroupIdx >= Graph.Groups.Count)
+			{
+				return NotFound();
+			}
+
+			return new GetGroupResponse(Graph.Groups[GroupIdx], Graph.Groups).ApplyFilter(Filter);
+		}
+
+		/// <summary>
+		/// Gets the nodes for a particular group
+		/// </summary>
+		/// <param name="JobId">Unique id for the job</param>
+		/// <param name="GroupIdx">Index of the group containing the node to update</param>
+		/// <param name="Filter">Filter for the properties to return</param>
+		[HttpGet]
+		[Route("/api/v1/jobs/{JobId}/groups/{GroupIdx}/nodes")]
+		[ProducesResponseType(typeof(List<GetNodeResponse>), 200)]
+		public async Task<ActionResult<List<object>>> GetNodesAsync(string JobId, int GroupIdx, [FromQuery] PropertyFilter? Filter = null)
+		{
+			IJob? Job = await JobService.GetJobAsync(JobId.ToObjectId());
+			if (Job == null)
+			{
+				return NotFound();
+			}
+			if (!await JobService.AuthorizeAsync(Job, AclAction.ViewJob, User, null))
+			{
+				return Forbid();
+			}
+
+			IGraph Graph = await JobService.GetGraphAsync(Job);
+			if (GroupIdx < 0 || GroupIdx >= Graph.Groups.Count)
+			{
+				return NotFound();
+			}
+
+			return Graph.Groups[GroupIdx].Nodes.ConvertAll(x => new GetNodeResponse(x, Graph.Groups).ApplyFilter(Filter));
+		}
+
+		/// <summary>
+		/// Gets a particular node definition
+		/// </summary>
+		/// <param name="JobId">Unique id for the job</param>
+		/// <param name="GroupIdx">Index of the group containing the node to update</param>
+		/// <param name="NodeIdx">Index of the node to update</param>
+		/// <param name="Filter">Filter for the properties to return</param>
+		[HttpGet]
+		[Route("/api/v1/jobs/{JobId}/groups/{GroupIdx}/nodes/{NodeIdx}")]
+		[ProducesResponseType(typeof(GetNodeResponse), 200)]
+		public async Task<ActionResult<object>> GetNodeAsync(string JobId, int GroupIdx, int NodeIdx, [FromQuery] PropertyFilter? Filter = null)
+		{
+			IJob? Job = await JobService.GetJobAsync(JobId.ToObjectId());
+			if (Job == null)
+			{
+				return NotFound();
+			}
+			if (!await JobService.AuthorizeAsync(Job, AclAction.ViewJob, User, null))
+			{
+				return Forbid();
+			}
+
+			IGraph Graph = await JobService.GetGraphAsync(Job);
+			if (GroupIdx < 0 || GroupIdx >= Graph.Groups.Count || NodeIdx < 0 || NodeIdx >= Graph.Groups[GroupIdx].Nodes.Count)
+			{
+				return NotFound();
+			}
+
+			return new GetNodeResponse(Graph.Groups[GroupIdx].Nodes[NodeIdx], Graph.Groups).ApplyFilter(Filter);
+		}
+
+		/// <summary>
+		/// Gets the steps currently scheduled to be executed for a job
+		/// </summary>
+		/// <param name="JobId">Unique id for the job</param>
+		/// <param name="Filter">Filter for the properties to return</param>
+		/// <returns>List of nodes to be executed</returns>
+		[HttpGet]
+		[Route("/api/v1/jobs/{JobId}/batches")]
+		[ProducesResponseType(typeof(List<GetBatchResponse>), 200)]
+		public async Task<ActionResult<List<object>>> GetBatchesAsync(string JobId, [FromQuery] PropertyFilter? Filter = null)
+		{
+			ObjectId JobIdValue = JobId.ToObjectId();
+
+			IJob? Job = await JobService.GetJobAsync(JobIdValue);
+			if (Job == null)
+			{
+				return NotFound();
+			}
+			if (!await JobService.AuthorizeAsync(Job, AclAction.ViewJob, User, null))
+			{
+				return Forbid();
+			}
+
+			IGraph Graph = await JobService.GetGraphAsync(Job);
+
+			List<object> Responses = new List<object>();
+			foreach (IJobStepBatch Batch in Job.Batches)
+			{
+				Responses.Add(new GetBatchResponse(Batch).ApplyFilter(Filter));
+			}
+			return Responses;
+		}
+
+		/// <summary>
+		/// Updates the state of a jobstep
+		/// </summary>
+		/// <param name="JobId">Unique id for the job</param>
+		/// <param name="BatchId">Unique id for the step</param>
+		/// <param name="Request">Updates to apply to the node</param>
+		[HttpPut]
+		[Route("/api/v1/jobs/{JobId}/batches/{BatchId}")]
+		public async Task<ActionResult> UpdateBatchAsync(string JobId, string BatchId, [FromBody] UpdateBatchRequest Request)
+		{
+			ObjectId JobIdValue = JobId.ToObjectId();
+			SubResourceId BatchIdValue = BatchId.ToSubResourceId();
+
+			IJob? Job = await JobService.GetJobAsync(JobIdValue);
+			if (Job == null)
+			{
+				return NotFound();
+			}
+
+			IJobStepBatch Batch = Job.Batches.FirstOrDefault(x => x.Id == BatchIdValue);
+			if (Batch == null)
+			{
+				return NotFound();
+			}
+			if (Batch.SessionId == null || !User.HasSessionClaim(Batch.SessionId.Value))
+			{
+				return Forbid();
+			}
+
+			if (!await JobService.UpdateBatchAsync(Job, BatchIdValue, Request.LogId?.ToObjectId(), Request.State))
+			{
+				return NotFound();
+			}
+			return Ok();
+		}
+
+		/// <summary>
+		/// Gets a particular step currently scheduled to be executed for a job
+		/// </summary>
+		/// <param name="JobId">Unique id for the job</param>
+		/// <param name="BatchId">Unique id for the step</param>
+		/// <param name="Filter">Filter for the properties to return</param>
+		/// <returns>List of nodes to be executed</returns>
+		[HttpGet]
+		[Route("/api/v1/jobs/{JobId}/batches/{BatchId}")]
+		[ProducesResponseType(typeof(GetBatchResponse), 200)]
+		public async Task<ActionResult<object>> GetBatchAsync(string JobId, string BatchId, [FromQuery] PropertyFilter? Filter = null)
+		{
+			ObjectId JobIdValue = JobId.ToObjectId();
+			SubResourceId BatchIdValue = BatchId.ToSubResourceId();
+
+			IJob? Job = await JobService.GetJobAsync(JobIdValue);
+			if (Job == null)
+			{
+				return NotFound();
+			}
+			if (!await JobService.AuthorizeAsync(Job, AclAction.ViewJob, User, null))
+			{
+				return Forbid();
+			}
+
+			IGraph Graph = await JobService.GetGraphAsync(Job);
+			foreach (IJobStepBatch Batch in Job.Batches)
+			{
+				if (Batch.Id == BatchIdValue)
+				{
+					return new GetBatchResponse(Batch).ApplyFilter(Filter);
+				}
+			}
+
+			return NotFound();
+		}
+
+		/// <summary>
+		/// Gets the steps currently scheduled to be executed for a job
+		/// </summary>
+		/// <param name="JobId">Unique id for the job</param>
+		/// <param name="BatchId">Unique id for the batch</param>
+		/// <param name="Filter">Filter for the properties to return</param>
+		/// <returns>List of nodes to be executed</returns>
+		[HttpGet]
+		[Route("/api/v1/jobs/{JobId}/batches/{BatchId}/steps")]
+		[ProducesResponseType(typeof(List<GetStepResponse>), 200)]
+		public async Task<ActionResult<List<object>>> GetStepsAsync(string JobId, string BatchId, [FromQuery] PropertyFilter? Filter = null)
+		{
+			ObjectId JobIdValue = JobId.ToObjectId();
+			SubResourceId BatchIdValue = BatchId.ToSubResourceId();
+
+			IJob? Job = await JobService.GetJobAsync(JobIdValue);
+			if (Job == null)
+			{
+				return NotFound();
+			}
+			if (!await JobService.AuthorizeAsync(Job, AclAction.ViewJob, User, null))
+			{
+				return Forbid();
+			}
+
+			IGraph Graph = await JobService.GetGraphAsync(Job);
+			foreach (IJobStepBatch Batch in Job.Batches)
+			{
+				if (Batch.Id == BatchIdValue)
+				{
+					INodeGroup Group = Graph.Groups[Batch.GroupIdx];
+
+					List<object> Responses = new List<object>();
+					foreach (IJobStep Step in Batch.Steps)
+					{
+						Responses.Add(new GetStepResponse(Step).ApplyFilter(Filter));
+					}
+					return Responses;
+				}
+			}
+
+			return NotFound();
+		}
+
+		/// <summary>
+		/// Updates the state of a jobstep
+		/// </summary>
+		/// <param name="JobId">Unique id for the job</param>
+		/// <param name="BatchId">Unique id for the batch</param>
+		/// <param name="StepId">Unique id for the step</param>
+		/// <param name="Request">Updates to apply to the node</param>
+		[HttpPut]
+		[Route("/api/v1/jobs/{JobId}/batches/{BatchId}/steps/{StepId}")]
+		public async Task<ActionResult<UpdateStepResponse>> UpdateStepAsync(string JobId, string BatchId, string StepId, [FromBody] UpdateStepRequest Request)
+		{
+			ObjectId JobIdValue = JobId.ToObjectId();
+			SubResourceId BatchIdValue = BatchId.ToSubResourceId();
+			SubResourceId StepIdValue = StepId.ToSubResourceId();
+
+			IJob? Job = await JobService.GetJobAsync(JobIdValue);
+			if (Job == null)
+			{
+				return NotFound();
+			}
+
+			// Check permissions for updating this step. Only the agent executing the step can modify the state of it.
+			if (Request.State != JobStepState.Unspecified || Request.Outcome != JobStepOutcome.Unspecified)
+			{
+				IJobStepBatch Batch = Job.Batches.FirstOrDefault(x => x.Id == BatchIdValue);
+				if (Batch == null)
+				{
+					return NotFound();
+				}
+				if (!Batch.SessionId.HasValue || !User.HasSessionClaim(Batch.SessionId.Value))
+				{
+					return Forbid();
+				}
+			}
+			if (Request.Retry != null || Request.Priority != null)
+			{
+				if (!await JobService.AuthorizeAsync(Job, AclAction.RetryJobStep, User, null))
+				{
+					return Forbid();
+				}
+			}
+			if (Request.Properties != null)
+			{
+				if (!await JobService.AuthorizeAsync(Job, AclAction.UpdateJob, User, null))
+				{
+					return Forbid();
+				}
+			}
+
+			string? RetryByUser = (Request.Retry.HasValue && Request.Retry.Value) ? (User.Identity.Name ?? "Anonymous") : null;
+			string? AbortByUser = (Request.AbortRequested.HasValue && Request.AbortRequested.Value) ? (User.Identity.Name ?? "Anonymous") : null;
+
+			try
+			{
+				IJob? NewJob = await JobService.UpdateStepAsync(Job, BatchIdValue, StepIdValue, Request.State, Request.Outcome, Request.AbortRequested, AbortByUser, Request.LogId?.ToObjectId(), null, RetryByUser, Request.Priority, null, Request.Properties);
+				if (NewJob == null)
+				{
+					return NotFound();
+				}
+
+				UpdateStepResponse Response = new UpdateStepResponse();
+				if (Request.Retry ?? false)
+				{
+					JobStepRefId? RetriedStepId = FindRetriedStep(Job, BatchIdValue, StepIdValue);
+					if (RetriedStepId != null)
+					{
+						Response.BatchId = RetriedStepId.Value.BatchId.ToString();
+						Response.StepId = RetriedStepId.Value.StepId.ToString();
+					}
+				}
+				return Response;
+			}
+			catch (RetryNotAllowedException Ex)
+			{
+				return BadRequest(Ex.Message);
+			}
+		}
+
+		/// <summary>
+		/// Find the first retried step after the given step
+		/// </summary>
+		/// <param name="Job">The job being run</param>
+		/// <param name="BatchId">Batch id of the last step instance</param>
+		/// <param name="StepId">Step id of the last instance</param>
+		/// <returns>The retried step information</returns>
+		static JobStepRefId? FindRetriedStep(IJob Job, SubResourceId BatchId, SubResourceId StepId)
+		{
+			NodeRef? LastNodeRef = null;
+			foreach (IJobStepBatch Batch in Job.Batches)
+			{
+				if ((LastNodeRef == null && Batch.Id == BatchId) || (LastNodeRef != null && Batch.GroupIdx == LastNodeRef.GroupIdx))
+				{
+					foreach (IJobStep Step in Batch.Steps)
+					{
+						if (LastNodeRef == null && Step.Id == StepId)
+						{
+							LastNodeRef = new NodeRef(Batch.GroupIdx, Step.NodeIdx);
+						}
+						else if (LastNodeRef != null && Step.NodeIdx == LastNodeRef.NodeIdx)
+						{
+							return new JobStepRefId(Job.Id, Batch.Id, Step.Id);
+						}
+					}
+				}
+			}
+			return null;
+		}
+
+		/// <summary>
+		/// Gets a particular step currently scheduled to be executed for a job
+		/// </summary>
+		/// <param name="JobId">Unique id for the job</param>
+		/// <param name="BatchId">Unique id for the batch</param>
+		/// <param name="StepId">Unique id for the step</param>
+		/// <param name="Filter">Filter for the properties to return</param>
+		/// <returns>List of nodes to be executed</returns>
+		[HttpGet]
+		[Route("/api/v1/jobs/{JobId}/batches/{BatchId}/steps/{StepId}")]
+		[ProducesResponseType(typeof(GetStepResponse), 200)]
+		public async Task<ActionResult<object>> GetStepAsync(string JobId, string BatchId, string StepId, [FromQuery] PropertyFilter? Filter = null)
+		{
+			ObjectId JobIdValue = JobId.ToObjectId();
+			SubResourceId BatchIdValue = BatchId.ToSubResourceId();
+			SubResourceId StepIdValue = StepId.ToSubResourceId();
+
+			IJob? Job = await JobService.GetJobAsync(JobIdValue);
+			if (Job == null)
+			{
+				return NotFound();
+			}
+			if (!await JobService.AuthorizeAsync(Job, AclAction.ViewJob, User, null))
+			{
+				return Forbid();
+			}
+
+			IGraph Graph = await JobService.GetGraphAsync(Job);
+			foreach (IJobStepBatch Batch in Job.Batches)
+			{
+				if (Batch.Id == BatchIdValue)
+				{
+					foreach (IJobStep Step in Batch.Steps)
+					{
+						if (Step.Id == StepIdValue)
+						{
+							return new GetStepResponse(Step).ApplyFilter(Filter);
+						}
+					}
+					break;
+				}
+			}
+
+			return NotFound();
+		}
+
+		/// <summary>
+		/// Updates notifications for a specific job.
+		/// </summary>
+		/// <param name="JobId">Unique id for the job</param>
+		/// <param name="BatchId">Unique id for the batch</param>
+		/// <param name="StepId">Unique id for the step</param>
+		/// <param name="Request">The notification request</param>
+		/// <returns>Information about the requested job</returns>
+		[HttpPut]
+		[Route("/api/v1/jobs/{JobId}/batches/{BatchId}/steps/{StepId}/notifications")]
+		public async Task<ActionResult> UpdateStepNotificationsAsync(string JobId, string BatchId, string StepId, [FromBody] UpdateNotificationsRequest Request)
+		{
+			ObjectId JobIdValue = JobId.ToObjectId();
+			SubResourceId BatchIdValue = BatchId.ToSubResourceId();
+			SubResourceId StepIdValue = StepId.ToSubResourceId();
+
+			IJob? Job = await JobService.GetJobAsync(JobIdValue);
+			if (Job == null)
+			{
+				return NotFound();
+			}
+			if (!await JobService.AuthorizeAsync(Job, AclAction.CreateSubscription, User, null))
+			{
+				return Forbid();
+			}
+
+			if (!Job.TryGetBatch(BatchIdValue, out IJobStepBatch? Batch))
+			{
+				return NotFound();
+			}
+
+			if (!Batch.TryGetStep(StepIdValue, out IJobStep? Step))
+			{
+				return NotFound();
+			}
+
+			ObjectId? TriggerId = Step.NotificationTriggerId;
+			if (TriggerId == null)
+			{
+				TriggerId = ObjectId.GenerateNewId();
+				if (await JobService.UpdateStepAsync(Job, BatchIdValue, StepIdValue, JobStepState.Unspecified, JobStepOutcome.Unspecified, null, null, null, TriggerId, null, null, null) == null)
+				{
+					return NotFound();
+				}
+			}
+
+			await NotificationService.UpdateSubscriptionsAsync(TriggerId.Value, User, Request.Email, Request.Slack);
+			return Ok();
+		}
+
+		/// <summary>
+		/// Gets information about a specific job.
+		/// </summary>
+		/// <param name="JobId">Id of the job to find</param>
+		/// <param name="BatchId">Unique id for the batch</param>
+		/// <param name="StepId">Unique id for the step</param>
+		/// <returns>Information about the requested job</returns>
+		[HttpGet]
+		[Route("/api/v1/jobs/{JobId}/batches/{BatchId}/steps/{StepId}/notifications")]
+		public async Task<ActionResult<GetNotificationResponse>> GetStepNotificationsAsync(string JobId, string BatchId, string StepId)
+		{
+			IJob? Job = await JobService.GetJobAsync(JobId.ToObjectId());
+			if (Job == null)
+			{
+				return NotFound();
+			}
+
+			IJobStep? Step;
+			if(!Job.TryGetStep(BatchId.ToSubResourceId(), StepId.ToSubResourceId(), out Step))
+			{
+				return NotFound();
+			}
+			if (!await JobService.AuthorizeAsync(Job, AclAction.CreateSubscription, User, null))
+			{
+				return Forbid();
+			}
+
+			INotificationSubscription? Subscription;
+			if (Step.NotificationTriggerId == null)
+			{
+				Subscription = null;
+			}
+			else
+			{
+				Subscription = await NotificationService.GetSubscriptionsAsync(Step.NotificationTriggerId.Value, User);
+			}
+			return new GetNotificationResponse(Subscription);
+		}
+
+		/// <summary>
+		/// Gets a particular step currently scheduled to be executed for a job
+		/// </summary>
+		/// <param name="JobId">Unique id for the job</param>
+		/// <param name="BatchId">Unique id for the batch</param>
+		/// <param name="StepId">Unique id for the step</param>
+		/// <param name="Name"></param>
+		/// <returns>List of nodes to be executed</returns>
+		[HttpGet]
+		[Route("/api/v1/jobs/{JobId}/batches/{BatchId}/steps/{StepId}/artifacts/{*Name}")]
+		public async Task<ActionResult> GetArtifactAsync(string JobId, string BatchId, string StepId, string Name)
+		{
+			ObjectId JobIdValue = JobId.ToObjectId();
+			SubResourceId BatchIdValue = BatchId.ToSubResourceId();
+			SubResourceId StepIdValue = StepId.ToSubResourceId();
+
+			IJob? Job = await JobService.GetJobAsync(JobIdValue);
+			if (Job == null)
+			{
+				return NotFound();
+			}
+			if (!await JobService.AuthorizeAsync(Job, AclAction.ViewJob, User, null))
+			{
+				return Forbid();
+			}
+			if (!Job.TryGetStep(BatchId.ToSubResourceId(), StepId.ToSubResourceId(), out _))
+			{
+				return NotFound();
+			}
+
+			List<Artifact> Artifacts = await ArtifactService.GetArtifactsAsync(JobIdValue, StepIdValue, Name);
+			if (Artifacts.Count == 0)
+			{
+				return NotFound();
+			}
+
+			Artifact Artifact = Artifacts[0];
+			return new FileStreamResult(ArtifactService.OpenArtifactReadStream(Artifact), Artifact.MimeType);
+		}
+
+		/// <summary>
+		/// Gets a particular step currently scheduled to be executed for a job
+		/// </summary>
+		/// <param name="JobId">Unique id for the job</param>
+		/// <param name="BatchId">Unique id for the batch</param>
+		/// <param name="StepId">Unique id for the step</param>
+		/// <returns>List of nodes to be executed</returns>
+		[HttpGet]
+		[Route("/api/v1/jobs/{JobId}/batches/{BatchId}/steps/{StepId}/trace")]
+		public async Task<ActionResult> GetStepTraceAsync(string JobId, string BatchId, string StepId)
+		{
+			ObjectId JobIdValue = JobId.ToObjectId();
+			SubResourceId BatchIdValue = BatchId.ToSubResourceId();
+			SubResourceId StepIdValue = StepId.ToSubResourceId();
+
+			IJob? Job = await JobService.GetJobAsync(JobIdValue);
+			if (Job == null)
+			{
+				return NotFound();
+			}
+			if (!await JobService.AuthorizeAsync(Job, AclAction.ViewJob, User, null))
+			{
+				return Forbid();
+			}
+			if (!Job.TryGetStep(BatchId.ToSubResourceId(), StepId.ToSubResourceId(), out _))
+			{
+				return NotFound();
+			}
+
+			List<Artifact> Artifacts = await ArtifactService.GetArtifactsAsync(JobIdValue, StepIdValue, null);
+			foreach (Artifact Artifact in Artifacts)
+			{
+				if (Artifact.Name.Equals("trace.json", StringComparison.OrdinalIgnoreCase))
+				{
+					return new FileStreamResult(ArtifactService.OpenArtifactReadStream(Artifact), "text/json");
+				}
+			}
+			return NotFound();
+		}
+
+		/// <summary>
+		/// Adds an array of aggregates to a job
+		/// </summary>
+		/// <param name="JobId">Unique id for the job</param>
+		/// <param name="Requests">Properties of the new aggregates</param>
+		/// <returns>Response object with the first index of the new aggregates</returns>
+		[HttpPost]
+		[Route("/api/v1/jobs/{JobId}/aggregates")]
+		public async Task<ActionResult> CreateAggregatesAsync(string JobId, [FromBody] List<CreateAggregateRequest> Requests)
+		{
+			ObjectId JobIdValue = JobId.ToObjectId();
+			for (; ; )
+			{
+				IJob? Job = await JobService.GetJobAsync(JobIdValue);
+				if (Job == null)
+				{
+					return NotFound();
+				}
+				if (!await JobService.AuthorizeAsync(Job, AclAction.ExecuteJob, User, null))
+				{
+					return Forbid();
+				}
+
+				IGraph Graph = await JobService.GetGraphAsync(Job);
+				Graph = await Graphs.AppendAsync(Graph, null, Requests, null);
+
+				if (await JobService.TryUpdateGraphAsync(Job, Graph))
+				{
+					return Ok();
+				}
+			}
+		}
+
+		/// <summary>
+		/// Gets the aggregates for a job
+		/// </summary>
+		/// <param name="JobId">Unique id for the job</param>
+		/// <param name="Filter">Filter for the properties to return</param>
+		/// <returns>List of aggregates for the job</returns>
+		[HttpGet]
+		[Route("/api/v1/jobs/{JobId}/aggregates")]
+		[ProducesResponseType(typeof(List<GetAggregateResponse>), 200)]
+		public async Task<ActionResult<List<object>>> GetAggregatesAsync(string JobId, [FromQuery] PropertyFilter? Filter = null)
+		{
+			IJob? Job = await JobService.GetJobAsync(JobId.ToObjectId());
+			if (Job == null)
+			{
+				return NotFound();
+			}
+			if (!await JobService.AuthorizeAsync(Job, AclAction.ViewJob, User, null))
+			{
+				return Forbid();
+			}
+
+			IGraph Graph = await JobService.GetGraphAsync(Job);
+
+			List<GetLabelStateResponse> Aggregates = new List<GetLabelStateResponse>();
+			Job.GetLabelStateResponses(Graph, Aggregates);
+			return Aggregates.ConvertAll(x => PropertyFilter.Apply(x, Filter));
+		}
+
+		/// <summary>
+		/// Gets an aggregate in a job
+		/// </summary>
+		/// <param name="JobId">Unique id for the job</param>
+		/// <param name="AggregateIdx">The aggregate index</param>
+		/// <param name="Filter">Filter for the properties to return</param>
+		/// <returns>Aggregate information</returns>
+		[HttpGet]
+		[Route("/api/v1/jobs/{JobId}/aggregates/{AggregateIdx}")]
+		[ProducesResponseType(typeof(GetLabelStateResponse), 200)]
+		public async Task<ActionResult<object>> GetAggregateAsync(string JobId, int AggregateIdx, [FromQuery] PropertyFilter? Filter = null)
+		{
+			IJob? Job = await JobService.GetJobAsync(JobId.ToObjectId());
+			if (Job == null)
+			{
+				return NotFound();
+			}
+			if (!await JobService.AuthorizeAsync(Job, AclAction.ViewJob, User, null))
+			{
+				return Forbid();
+			}
+
+			IGraph Graph = await JobService.GetGraphAsync(Job);
+
+			List<GetLabelStateResponse> Responses = new List<GetLabelStateResponse>();
+			Job.GetLabelStateResponses(Graph, Responses);
+
+			if (AggregateIdx < 0 || AggregateIdx >= Responses.Count)
+			{
+				return NotFound();
+			}
+			else
+			{
+				return PropertyFilter.Apply(Responses[AggregateIdx], Filter);
+			}
+		}
+
+		/// <summary>
+		/// Updates notifications for a specific label.
+		/// </summary>
+		/// <param name="JobId">Unique id for the job</param>
+		/// <param name="LabelIndex">Index for the label</param>
+		/// <param name="Request">The notification request</param>
+		[HttpPut]
+		[Route("/api/v1/jobs/{JobId}/labels/{LabelIndex}/notifications")]
+		public async Task<ActionResult> UpdateLabelNotificationsAsync(string JobId, int LabelIndex, [FromBody] UpdateNotificationsRequest Request)
+		{
+			ObjectId JobIdValue = JobId.ToObjectId();
+
+			ObjectId TriggerId;
+			for (; ; )
+			{
+				IJob? Job = await JobService.GetJobAsync(JobIdValue);
+				if (Job == null)
+				{
+					return NotFound();
+				}
+				if (!await JobService.AuthorizeAsync(Job, AclAction.CreateSubscription, User, null))
+				{
+					return Forbid();
+				}
+
+				ObjectId NewTriggerId;
+				if (Job.LabelIdxToTriggerId.TryGetValue(LabelIndex, out NewTriggerId))
+				{
+					TriggerId = NewTriggerId;
+					break;
+				}
+
+				NewTriggerId = ObjectId.GenerateNewId();
+				if (await JobService.UpdateJobAsync(Job, LabelIdxToTriggerId: new KeyValuePair<int, ObjectId>(LabelIndex, NewTriggerId)))
+				{
+					TriggerId = NewTriggerId;
+					break;
+				}
+			}
+
+			await NotificationService.UpdateSubscriptionsAsync(TriggerId, User, Request.Email, Request.Slack);
+			return Ok();
+		}
+
+		/// <summary>
+		/// Gets notification info about a specific label in a job.
+		/// </summary>
+		/// <param name="JobId">Id of the job to find</param>
+		/// <param name="LabelIndex">Index for the label</param>
+		/// <returns>Notification info for the requested label in the job</returns>
+		[HttpGet]
+		[Route("/api/v1/jobs/{JobId}/labels/{LabelIndex}/notifications")]
+		public async Task<ActionResult<GetNotificationResponse>> GetLabelNotificationsAsync(string JobId, int LabelIndex)
+		{
+			IJob? Job = await JobService.GetJobAsync(JobId.ToObjectId());
+			if (Job == null)
+			{
+				return NotFound();
+			}
+
+			INotificationSubscription? Subscription;
+			if (!Job.LabelIdxToTriggerId.ContainsKey(LabelIndex) || Job.LabelIdxToTriggerId[LabelIndex] == null)
+			{
+				Subscription = null;
+			}
+			else
+			{
+				Subscription = await NotificationService.GetSubscriptionsAsync(Job.LabelIdxToTriggerId[LabelIndex], User);
+			}
+			return new GetNotificationResponse(Subscription);
+		}
+	}
+}
diff --git a/Engine/Source/Programs/Horde/HordeServer/Controllers/LeasesController.cs b/Engine/Source/Programs/Horde/HordeServer/Controllers/LeasesController.cs
new file mode 100644
--- /dev/null
+++ b/Engine/Source/Programs/Horde/HordeServer/Controllers/LeasesController.cs
@@ -0,0 +1,129 @@
+// Copyright Epic Games, Inc. All Rights Reserved.
+
+using HordeServer.Api;
+using HordeServer.Models;
+using HordeServer.Services;
+using HordeServer.Utilities;
+using Microsoft.AspNetCore.Authentication;
+using Microsoft.AspNetCore.Authorization;
+using Microsoft.AspNetCore.Mvc;
+using Microsoft.AspNetCore.Mvc.ModelBinding;
+using Microsoft.Extensions.Options;
+using Microsoft.IdentityModel.Tokens;
+using MongoDB.Bson;
+using System;
+using System.Collections.Generic;
+using System.ComponentModel.DataAnnotations;
+using System.IdentityModel.Tokens.Jwt;
+using System.Linq;
+using System.Reflection;
+using System.Security.Claims;
+using System.Text.Json;
+using System.Threading.Tasks;
+
+using PoolId = HordeServer.Utilities.StringId<HordeServer.Models.IPool>;
+using AgentSoftwareChannelName = HordeServer.Utilities.StringId<HordeServer.Services.AgentSoftwareChannels>;
+
+namespace HordeServer.Controllers
+{
+	/// <summary>
+	/// Controller for the /api/v1/leases endpoint
+	/// </summary>
+	[ApiController]
+	[Authorize]
+	[Route("[controller]")]
+	public class LeasesController : ControllerBase
+	{
+		/// <summary>
+		/// Singleton instance of the ACL service
+		/// </summary>
+		AclService AclService;
+
+		/// <summary>
+		/// Singleton instance of the agent service
+		/// </summary>
+		AgentService AgentService;
+
+		/// <summary>
+		/// Constructor
+		/// </summary>
+		/// <param name="AclService">The ACL service singleton</param>
+		/// <param name="AgentService">The agent service</param>
+		public LeasesController(AclService AclService, AgentService AgentService)
+		{
+			this.AclService = AclService;
+			this.AgentService = AgentService;
+		}
+
+		/// <summary>
+		/// Find all the leases for a particular agent
+		/// </summary>
+		/// <param name="AgentId">Unique id of the agent to find</param>
+		/// <param name="SessionId">The session to query</param>
+		/// <param name="StartTime">Start of the time window to consider</param>
+		/// <param name="FinishTime">End of the time window to consider</param>
+		/// <param name="Index">Index of the first result to return</param>
+		/// <param name="Count">Number of results to return</param>
+		/// <returns>Sessions </returns>
+		[HttpGet]
+		[Route("/api/v1/leases")]
+		public async Task<ActionResult<List<GetAgentLeaseResponse>>> FindLeasesAsync([FromQuery] string? AgentId, [FromQuery] string? SessionId, [FromQuery] DateTimeOffset? StartTime, [FromQuery] DateTimeOffset? FinishTime, [FromQuery] int Index = 0, [FromQuery] int Count = 1000)
+		{
+			AgentId? AgentIdValue = null;
+			if (AgentId == null)
+			{
+				if (!await AclService.AuthorizeAsync(AclAction.ViewLeases, User))
+				{
+					return Forbid();
+				}
+			}
+			else
+			{
+				AgentIdValue = new AgentId(AgentId);
+
+				IAgent? Agent = await AgentService.GetAgentAsync(AgentIdValue.Value);
+				if (Agent == null)
+				{
+					return NotFound();
+				}
+				if (!await AgentService.AuthorizeAsync(Agent, AclAction.ViewLeases, User, null))
+				{
+					return Forbid();
+				}
+			}
+
+			List<ILease> Leases = await AgentService.FindLeasesAsync(AgentIdValue, SessionId?.ToObjectId(), StartTime?.UtcDateTime, FinishTime?.UtcDateTime, Index, Count);
+			return Leases.ConvertAll(x => new GetAgentLeaseResponse(x));
+		}
+
+		/// <summary>
+		/// Get info about a particular lease
+		/// </summary>
+		/// <param name="LeaseId">Unique id of the particular lease</param>
+		/// <returns>Lease matching the given id</returns>
+		[HttpGet]
+		[Route("/api/v1/leases/{LeaseId}")]
+		public async Task<ActionResult<GetAgentLeaseResponse>> GetLeaseAsync(string LeaseId)
+		{
+			ObjectId LeaseIdValue = LeaseId.ToObjectId();
+
+			ILease? Lease = await AgentService.GetLeaseAsync(LeaseIdValue);
+			if (Lease == null)
+			{
+				return NotFound();
+			}
+
+			IAgent? Agent = await AgentService.GetAgentAsync(Lease.AgentId);
+			if (Agent == null)
+			{
+				return NotFound();
+			}
+			if (!await AgentService.AuthorizeAsync(Agent, AclAction.ViewLeases, User, null))
+			{
+				return Forbid();
+			}
+
+			return new GetAgentLeaseResponse(Lease);
+		}
+	}
+}
diff --git a/Engine/Source/Programs/Horde/HordeServer/Controllers/LifetimeController.cs b/Engine/Source/Programs/Horde/HordeServer/Controllers/LifetimeController.cs
new file mode 100644
--- /dev/null
+++ b/Engine/Source/Programs/Horde/HordeServer/Controllers/LifetimeController.cs
@@ -0,0 +1,65 @@
+// Copyright Epic Games, Inc. All Rights Reserved.
+
+using HordeServer.Services;
+using Microsoft.AspNetCore.Mvc;
+using System.Threading.Tasks;
+
+namespace HordeServer.Controllers
+{
+	/// <summary>
+	/// Controller for app lifetime related routes
+	/// </summary>
+	[ApiController]
+	[Route("[controller]")]
+	public class LifetimeController : ControllerBase
+	{
+		/// <summary>
+		/// Singleton instance of the lifetime service
+		/// </summary>
+		LifetimeService LifetimeService;
+
+		/// <summary>
+		/// Constructor
+		/// </summary>
+		/// <param name="LifetimeService">The lifetime service singleton</param>
+		public LifetimeController(LifetimeService LifetimeService)
+		{
+			this.LifetimeService = LifetimeService;
+		}
+
+		/// <summary>
+		/// Readiness check for Kubernetes
+		/// If this return a non-successful HTTP response, Kubernetes will remove it from the load balancer,
+		/// preventing it from serving traffic.
+		/// </summary>
+		/// <returns>Ok if app is not stopping and all databases can be reached</returns>
+		[HttpGet]
+		[Route("/health/ready")]
+		public Task<ActionResult> K8sReadinessProbe()
+		{
+			// Disabled for now
+			bool IsRunning = true;//!LifetimeService.IsStopping;
+			bool IsMongoDBHealthy = true; //await LifetimeService.IsMongoDbConnectionHealthy();
+			bool IsRedisHealthy = true; //await LifetimeService.IsRedisConnectionHealthy();
+			int StatusCode = IsRunning && IsMongoDBHealthy && IsRedisHealthy ? 200 : 500;
+
+			string Content = $"IsRunning={IsRunning}\n";
+			Content += $"IsMongoDBHealthy={IsMongoDBHealthy}\n";
+			Content += $"IsRedisHealthy={IsRedisHealthy}\n";
+
+			return Task.FromResult<ActionResult>(new ContentResult { ContentType = "text/plain", StatusCode = StatusCode, Content = Content });
+		}
+		
+		/// <summary>
+		/// Liveness check for Kubernetes
+		/// If this return a non-successful HTTP response, Kubernetes will kill the pod and restart it
+		/// </summary>
+		/// <returns>Ok if app is not stopping and all databases can be reached</returns>
+		[HttpGet]
+		[Route("/health/live")]
+		public ActionResult K8sLivenessProbe()
+		{
+			return new ContentResult { ContentType = "text/plain", StatusCode = 200, Content = "ok" };
+		}
+	}
+}
diff --git a/Engine/Source/Programs/Horde/HordeServer/Controllers/LogsController.cs b/Engine/Source/Programs/Horde/HordeServer/Controllers/LogsController.cs
new file mode 100644
--- /dev/null
+++ b/Engine/Source/Programs/Horde/HordeServer/Controllers/LogsController.cs
@@ -0,0 +1,418 @@
+// Copyright Epic Games, Inc. All Rights Reserved.
+
+using System;
+using System.Buffers;
+using System.Collections.Generic;
+using System.ComponentModel.DataAnnotations;
+using System.IO;
+using System.Linq;
+using System.Net.Http;
+using System.Security.Claims;
+using System.Text;
+using System.Text.Json;
+using System.Threading.Tasks;
+using HordeServer.Api;
+using HordeServer.Collections;
+using HordeServer.Logs;
+using HordeServer.Models;
+using HordeServer.Services;
+using HordeServer.Utilities;
+using Microsoft.AspNetCore.Authorization;
+using Microsoft.AspNetCore.Http;
+using Microsoft.AspNetCore.Mvc;
+using MongoDB.Bson;
+
+namespace HordeServer.Controllers
+{
+	/// <summary>
+	/// Format for the returned data
+	/// </summary>
+	public enum LogOutputFormat
+	{
+		/// <summary>
+		/// Plain text
+		/// </summary>
+		Text,
+
+		/// <summary>
+		/// Raw output (text/json)
+		/// </summary>
+		Raw,
+	}
+
+	/// <summary>
+	/// Controller for the /api/logs endpoint
+	/// </summary>
+	[ApiController]
+	[Authorize]
+	[Route("[controller]")]
+	public class LogsController : ControllerBase
+	{
+		/// <summary>
+		/// Instance of the LogFile service
+		/// </summary>
+		private readonly ILogFileService LogFileService;
+
+		/// <summary>
+		/// Instance of the issue collection
+		/// </summary>
+		private readonly IIssueCollection IssueCollection;
+
+		/// <summary>
+		/// Instance of the ACL service
+		/// </summary>
+		private readonly AclService AclService;
+
+		/// <summary>
+		/// Instance of the Job service
+		/// </summary>
+		private readonly JobService JobService;
+
+		/// <summary>
+		/// Constructor
+		/// </summary>
+		/// <param name="LogFileService">The Logfile service</param>
+		/// <param name="IssueCollection">The issue collection</param>
+		/// <param name="AclService">The ACL service</param>
+		/// <param name="JobService">The Job service</param>
+		public LogsController(ILogFileService LogFileService, IIssueCollection IssueCollection, AclService AclService, JobService JobService)
+		{
+			this.LogFileService = LogFileService;
+			this.IssueCollection = IssueCollection;
+			this.AclService = AclService;
+			this.JobService = JobService;
+		}
+
+		/// <summary>
+		/// Creates a new logfile
+		/// </summary>
+		/// <param name="Create">Parameters for the new LogFile</param>
+		/// <returns>Http result code</returns>
+		[HttpPost]
+		[Route("/api/v1/logs")]
+		public async Task<ActionResult<CreateLogFileResponse>> CreateLogFile([FromBody] CreateLogFileRequest Create)
+		{
+			ObjectId JobId = Create.JobId.ToObjectId();
+			if (!await JobService.AuthorizeAsync(JobId, AclAction.CreateLog, User, null))
+			{
+				return Forbid();
+			}
+
+			ILogFile NewLogFile = await LogFileService.CreateLogFileAsync(JobId, null, Create.Type);
+			return new CreateLogFileResponse(NewLogFile.Id.ToString());
+		}
+
+		/// <summary>
+		/// Retrieve metadata about a specific log file
+		/// </summary>
+		/// <param name="LogFileId">Id of the log file to get information about</param>
+		/// <param name="Filter">Filter for the properties to return</param>
+		/// <returns>Information about the requested project</returns>
+		[HttpGet]
+		[Route("/api/v1/logs/{LogFileId}")]
+		[ProducesResponseType(typeof(GetLogFileResponse), 200)]
+		public async Task<ActionResult<object>> GetLog(string LogFileId, [FromQuery] PropertyFilter? Filter = null)
+		{
+			ILogFile? LogFile = await LogFileService.GetLogFileAsync(LogFileId.ToObjectId());
+			if (LogFile == null)
+			{
+				return NotFound();
+			}
+			if (!await AuthorizeAsync(LogFile, AclAction.ViewLog, User, null))
+			{
+				return Forbid();
+			}
+
+			LogMetadata Metadata = await LogFileService.GetMetadataAsync(LogFile);
+			return new GetLogFileResponse(LogFile, Metadata).ApplyFilter(Filter);       
+		}
+
+		/// <summary>
+		/// Retrieve raw data for a log file
+		/// </summary>
+		/// <param name="LogFileId">Id of the log file to get information about</param>
+		/// <param name="Format">Format for the returned data</param>
+		/// <param name="Offset">The log offset in bytes</param>
+		/// <param name="Length">Number of bytes to return</param>
+		/// <param name="FileName">Name of the default filename to download</param>
+		/// <param name="Download">Whether to download the file rather than display in the browser</param>
+		/// <returns>Raw log data for the requested range</returns>
+		[HttpGet]
+		[Route("/api/v1/logs/{LogFileId}/data")]
+		public async Task<ActionResult> GetLogData(string LogFileId, [FromQuery] LogOutputFormat Format = LogOutputFormat.Raw, [FromQuery] long Offset = 0, [FromQuery] long Length = long.MaxValue, [FromQuery] string? FileName = null, [FromQuery] bool Download = false)
+		{
+			ILogFile? LogFile = await LogFileService.GetLogFileAsync(LogFileId.ToObjectId());
+			if (LogFile == null)
+			{
+				return NotFound();
+			}
+			if (!await AuthorizeAsync(LogFile, AclAction.ViewLog, User, null))
+			{
+				return Forbid();
+			}
+
+			Func<Stream, ActionContext, Task> CopyTask;
+			if (Format == LogOutputFormat.Text && LogFile.Type == LogType.Json)
+			{
+				CopyTask = (OutputStream, Context) => LogFileService.CopyPlainTextStreamAsync(LogFile, Offset, Length, OutputStream);
+			}
+			else
+			{
+				CopyTask = (OutputStream, Context) => LogFileService.CopyRawStreamAsync(LogFile, Offset, Length, OutputStream);
+			}
+
+			return new CustomFileCallbackResult(FileName ?? $"log-{LogFileId}.txt", "text/plain", !Download, CopyTask);
+		}
+
+		/// <summary>
+		/// Retrieve line data for a logfile
+		/// </summary>
+		/// <param name="LogFileId">Id of the log file to get information about</param>
+		/// <param name="Index">Index of the first line to retrieve</param>
+		/// <param name="Count">Number of lines to retrieve</param>
+		/// <returns>Information about the requested project</returns>
+		[HttpGet]
+		[Route("/api/v1/logs/{LogFileId}/lines")]
+		public async Task<ActionResult> GetLogLines(string LogFileId, [FromQuery] int Index = 0, [FromQuery] int Count = int.MaxValue)
+		{
+			ILogFile? LogFile = await LogFileService.GetLogFileAsync(LogFileId.ToObjectId());
+			if (LogFile == null)
+			{
+				return NotFound();
+			}
+			if (!await AuthorizeAsync(LogFile, AclAction.ViewLog, User, null))
+			{
+				return Forbid();
+			}
+
+			LogMetadata Metadata = await LogFileService.GetMetadataAsync(LogFile);
+
+			(int MinIndex, long MinOffset) = await LogFileService.GetLineOffsetAsync(LogFile, Index);
+			(int MaxIndex, long MaxOffset) = await LogFileService.GetLineOffsetAsync(LogFile, Index + Math.Min(Count, int.MaxValue - Index));
+			Index = MinIndex;
+			Count = MaxIndex - MinIndex;
+
+			byte[] Result;
+			using (System.IO.Stream Stream = await LogFileService.OpenRawStreamAsync(LogFile, MinOffset, MaxOffset - MinOffset))
+			{
+				Result = new byte[Stream.Length];
+				await Stream.ReadFixedSizeDataAsync(Result, 0, Result.Length);
+			}
+
+			using (MemoryStream Stream = new MemoryStream(Result.Length + (Count * 20)))
+			{
+				Stream.WriteByte((byte)'{');
+
+				Stream.Write(Encoding.UTF8.GetBytes($"\"index\":{Index},"));
+				Stream.Write(Encoding.UTF8.GetBytes($"\"count\":{Count},"));
+				Stream.Write(Encoding.UTF8.GetBytes($"\"maxLineIndex\":{Metadata.MaxLineIndex},"));
+				Stream.Write(Encoding.UTF8.GetBytes($"\"format\":{ (LogFile.Type == LogType.Json ? "\"JSON\"" : "\"TEXT\"")},"));
+
+//				Stream.Write(Encoding.UTF8.GetBytes($"\"minIndex\":{MinIndex},"));
+//				Stream.Write(Encoding.UTF8.GetBytes($"\"minOffset\":{MinOffset},"));
+//				Stream.Write(Encoding.UTF8.GetBytes($"\"maxIndex\":{MaxIndex},"));
+//				Stream.Write(Encoding.UTF8.GetBytes($"\"maxOffset\":{MaxOffset},"));
+//				Stream.Write(Encoding.UTF8.GetBytes($"\"length\":{Result.Length},"));
+
+				Stream.Write(Encoding.UTF8.GetBytes($"\"lines\":["));
+				Stream.WriteByte((byte)'\n');
+
+				int Offset = 0;
+				for (int Line = Index; Line < Index + Count; Line++)
+				{
+					Stream.WriteByte((byte)' ');
+					Stream.WriteByte((byte)' ');
+
+					if (LogFile.Type == LogType.Json)
+					{
+						// Find the end of the line and output it as an opaque blob
+						int StartOffset = Offset;
+						for (; ; Offset++)
+						{
+							if (Offset == Result.Length)
+							{
+								Stream.WriteByte((byte)'{');
+								Stream.WriteByte((byte)'}');
+								break;
+							}
+							else if (Result[Offset] == (byte)'\n')
+							{
+								Stream.Write(Result, StartOffset, Offset - StartOffset);
+								Offset++;
+								break;
+							}
+						}
+					}
+					else
+					{
+						Stream.WriteByte((byte)'\"');
+						for (; Offset < Result.Length; Offset++)
+						{
+							if (Result[Offset] == '\\' || Result[Offset] == '\"')
+							{
+								Stream.WriteByte((byte)'\\');
+								Stream.WriteByte(Result[Offset]);
+							}
+							else if (Result[Offset] == (byte)'\n')
+							{
+								Offset++;
+								break;
+							}
+							else if (Result[Offset] >= 32 && Result[Offset] <= 126)
+							{
+								Stream.WriteByte(Result[Offset]);
+							}
+							else
+							{
+								Stream.Write(Encoding.UTF8.GetBytes($"\\x{Result[Offset]:x2}"));
+							}
+						}
+						Stream.WriteByte((byte)'\"');
+					}
+
+					if (Line + 1 < Index + Count)
+					{
+						Stream.WriteByte((byte)',');
+					}
+
+					Stream.WriteByte((byte)'\n');
+				}
+
+				if (LogFile.Type == LogType.Json)
+				{
+					Stream.Write(Encoding.UTF8.GetBytes($"]"));
+				}
+
+				Stream.WriteByte((byte)'}');
+
+				Response.ContentType = "application/json";
+				Response.Headers.ContentLength = Stream.Length;
+				Stream.Position = 0;
+				await Stream.CopyToAsync(Response.Body);
+			}
+			return new EmptyResult();
+		}
+
+		/// <summary>
+		/// Search log data
+		/// </summary>
+		/// <param name="LogFileId">Id of the log file to get information about</param>
+		/// <param name="Text">Text to search for</param>
+		/// <param name="FirstLine">First line to search from</param>
+		/// <param name="Count">Number of results to return</param>
+		/// <returns>Raw log data for the requested range</returns>
+		[HttpGet]
+		[Route("/api/v1/logs/{LogFileId}/search")]
+		public async Task<ActionResult<SearchLogFileResponse>> SearchLogFileAsync(string LogFileId, [FromQuery] string Text, [FromQuery] int FirstLine = 0, [FromQuery] int Count = 5)
+		{
+			ILogFile? LogFile = await LogFileService.GetLogFileAsync(LogFileId.ToObjectId());
+			if (LogFile == null)
+			{
+				return NotFound();
+			}
+			if (!await AuthorizeAsync(LogFile, AclAction.ViewLog, User, null))
+			{
+				return Forbid();
+			}
+
+			SearchLogFileResponse Response = new SearchLogFileResponse();
+			Response.Stats = new LogSearchStats();
+			Response.Lines = await LogFileService.SearchLogDataAsync(LogFile, Text, FirstLine, Count, Response.Stats);
+			return Response;
+		}
+
+		/// <summary>
+		/// Retrieve events for a logfile
+		/// </summary>
+		/// <param name="LogFileId">Id of the log file to get information about</param>
+		/// <param name="Index">Index of the first line to retrieve</param>
+		/// <param name="Count">Number of lines to retrieve</param>
+		/// <returns>Information about the requested project</returns>
+		[HttpGet]
+		[Route("/api/v1/logs/{LogFileId}/events")]
+		[ProducesResponseType(typeof(List<GetLogEventResponse>), 200)]
+		public async Task<ActionResult<List<GetLogEventResponse>>> GetEventsAsync(string LogFileId, [FromQuery] int? Index = null, [FromQuery] int? Count = null)
+		{
+			ILogFile? LogFile = await LogFileService.GetLogFileAsync(LogFileId.ToObjectId());
+			if (LogFile == null)
+			{
+				return NotFound();
+			}
+			if (!await AuthorizeAsync(LogFile, AclAction.ViewLog, User, null))
+			{
+				return Forbid();
+			}
+
+			List<ILogEvent> LogEvents = await LogFileService.FindLogEventsAsync(LogFile, Index, Count);
+
+			Dictionary<ObjectId, int?> SpanIdToIssueId = new Dictionary<ObjectId, int?>();
+
+			List<GetLogEventResponse> Responses = new List<GetLogEventResponse>();
+			foreach (ILogEvent LogEvent in LogEvents)
+			{
+				ILogEventData LogEventData = await LogFileService.GetEventDataAsync(LogFile, LogEvent.LineIndex, LogEvent.LineCount);
+
+				int? IssueId = null;
+				if (LogEvent.SpanId != null && !SpanIdToIssueId.TryGetValue(LogEvent.SpanId.Value, out IssueId))
+				{
+					IIssueSpan? Span = await IssueCollection.GetSpanAsync(LogEvent.SpanId.Value);
+					IssueId = Span?.IssueId;
+					SpanIdToIssueId[LogEvent.SpanId.Value] = IssueId;
+				}
+
+				Responses.Add(new GetLogEventResponse(LogEvent, LogEventData, IssueId));
+			}
+			return Responses;
+		}
+
+		/// <summary>
+		/// Appends data to a log file
+		/// </summary>
+		/// <param name="LogFileId">The logfile id</param>
+		/// <param name="Offset">Offset within the log file</param>
+		/// <param name="LineIndex">The line index</param>
+		/// <returns>Http result code</returns>
+		[HttpPost]
+		[Route("/api/v1/logs/{LogFileId}")]
+		public async Task<ActionResult> WriteData(string LogFileId, [FromQuery] long Offset, [FromQuery] int LineIndex)
+		{
+			ILogFile? LogFile = await LogFileService.GetLogFileAsync(LogFileId.ToObjectId());
+			if (LogFile == null)
+			{
+				return NotFound();
+			}
+			if (!await AuthorizeAsync(LogFile, AclAction.WriteLogData, User, null))
+			{
+				return Forbid();
+			}
+
+			using (MemoryStream BodyStream = new MemoryStream())
+			{
+				await Request.Body.CopyToAsync(BodyStream);
+				await LogFileService.WriteLogDataAsync(LogFile, Offset, LineIndex, BodyStream.ToArray(), false);
+			}
+			return Ok();
+		}
+
+		/// <summary>
+		/// Determines if the user is authorized to perform an action on a particular template
+		/// </summary>
+		/// <param name="LogFile">The template to check</param>
+		/// <param name="Action">The action being performed</param>
+		/// <param name="User">The principal to authorize</param>
+		/// <param name="PermissionsCache">Permissions cache</param>
+		/// <returns>True if the action is authorized</returns>
+		async Task<bool> AuthorizeAsync(ILogFile LogFile, AclAction Action, ClaimsPrincipal User, JobPermissionsCache? PermissionsCache)
+		{
+			if (LogFile.JobId != ObjectId.Empty && await JobService.AuthorizeAsync(LogFile.JobId, Action, User, PermissionsCache))
+			{
+				return true;
+			}
+			if (LogFile.SessionId != null && await AclService.AuthorizeAsync(AclAction.ViewSession, User, PermissionsCache))
+			{
+				return true;
+			}
+			return false;
+		}
+	}
+}
\ No newline at end of file
diff --git a/Engine/Source/Programs/Horde/HordeServer/Controllers/NoticesController.cs b/Engine/Source/Programs/Horde/HordeServer/Controllers/NoticesController.cs
new file mode 100644
--- /dev/null
+++ b/Engine/Source/Programs/Horde/HordeServer/Controllers/NoticesController.cs
@@ -0,0 +1,119 @@
+// Copyright Epic Games, Inc. All Rights Reserved.
+
+using HordeServer.Models;
+using HordeServer.Services;
+using HordeServer.Utilities;
+using Microsoft.AspNetCore.Authorization;
+using Microsoft.AspNetCore.Mvc;
+using Microsoft.Extensions.Options;
+using MongoDB.Bson;
+using System;
+using System.Collections.Generic;
+using System.Globalization;
+using System.Linq;
+using System.Threading.Tasks;
+using TimeZoneConverter;
+
+namespace HordeServer.Controllers
+{
+	/// <summary>
+	/// Controller for the /api/v1/status endpoint
+	/// </summary>
+	[ApiController]
+	[Authorize]
+	[Route("[controller]")]
+	public class NoticesController : ControllerBase
+	{
+		/// <summary>
+		/// Timezone to use for display notices
+		/// </summary>
+		TimeZoneInfo TimeZone;
+
+		/// <summary>
+		/// The globals singleton
+		/// </summary>
+		ISingletonDocument<Globals> Globals;
+
+		/// <summary>
+		/// Constructor
+		/// </summary>
+		/// <param name="Settings">The server settings</param>
+		/// <param name="Globals">The globals singleton</param>
+		public NoticesController(IOptions<ServerSettings> Settings, ISingletonDocument<Globals> Globals)
+		{
+			string? ScheduleTimeZone = Settings.Value.ScheduleTimeZone;
+			this.TimeZone = (ScheduleTimeZone == null) ? TimeZoneInfo.Local : TZConvert.GetTimeZoneInfo(ScheduleTimeZone);
+
+			this.Globals = Globals;
+		}
+
+		/// <summary>
+		/// Update a status message
+		/// </summary>
+		/// <param name="Status"></param>
+		/// <returns></returns>
+		[HttpPost("/api/v1/notices")]
+		public async Task<ActionResult> AddNoticeAsync(Notice Status)
+		{
+			Status.Id = ObjectId.GenerateNewId().ToString();
+			await Globals.UpdateAsync(x => x.Notices.Add(Status));
+			return Ok();
+		}
+
+		/// <summary>
+		/// Remove a manually added status message
+		/// </summary>
+		/// <returns></returns>
+		[HttpDelete("/api/v1/notices/{Id}")]
+		public async Task<ActionResult> DeleteNoticeAsync(string Id)
+		{
+			await Globals.UpdateAsync(x => x.Notices.RemoveAll(x => x.Id == Id));
+			return Ok();
+		}
+
+		/// <summary>
+		/// Gets the current status messages
+		/// </summary>
+		/// <returns>The status messages</returns>
+		[HttpGet("/api/v1/notices")]
+		public async Task<List<Notice>> GetNoticesAsync()
+		{
+			Globals Value = await Globals.GetAsync();
+
+			List<Notice> Messages = new List<Notice>(Value.Notices);
+
+			if(Value.ScheduledDowntime.Count > 0)
+			{
+				DateTimeOffset Now = DateTimeOffset.Now;
+
+				(DateTimeOffset StartTime, DateTimeOffset FinishTime) = Value.ScheduledDowntime[0].GetNext(Now);
+				for (int Idx = 1; Idx < Value.ScheduledDowntime.Count; Idx++)
+				{
+					(DateTimeOffset NextStartTime, DateTimeOffset NextFinishTime) = Value.ScheduledDowntime[Idx].GetNext(Now);
+					if (NextFinishTime < StartTime)
+					{
+						StartTime = NextStartTime;
+					}
+					if (NextStartTime < FinishTime && NextFinishTime > FinishTime)
+					{
+						FinishTime = NextFinishTime;
+					}
+				}
+
+				if (StartTime < Now)
+				{
+					string FinishTimeString = String.Format(CultureInfo.CurrentCulture, "{0:t}", TimeZoneInfo.ConvertTime(FinishTime, TimeZone));
+					Messages.Add(new Notice { StartTime = StartTime.UtcDateTime, FinishTime = FinishTime.UtcDateTime, Message = $"Horde is currently in scheduled downtime. Jobs will resume execution at {FinishTimeString} {TimeZone.Id}." });
+				}
+				else if(StartTime < Now + TimeSpan.FromHours(12.0))
+				{
+					string StartTimeString = String.Format(CultureInfo.CurrentCulture, "{0:t}", StartTime);
+					string FinishTimeString = String.Format(CultureInfo.CurrentCulture, "{0:t}", FinishTime);
+					Messages.Add(new Notice { StartTime = StartTime.UtcDateTime, FinishTime = FinishTime.UtcDateTime, Message = $"Downtime is scheduled from {StartTimeString} to {FinishTimeString} {TimeZone.Id}. No jobs will run during this time." });
+				}
+			}
+
+			return Messages;
+		}
+	}
+}
diff --git a/Engine/Source/Programs/Horde/HordeServer/Controllers/PerforceController.cs b/Engine/Source/Programs/Horde/HordeServer/Controllers/PerforceController.cs
new file mode 100644
--- /dev/null
+++ b/Engine/Source/Programs/Horde/HordeServer/Controllers/PerforceController.cs
@@ -0,0 +1,141 @@
+// Copyright Epic Games, Inc. All Rights Reserved.
+
+using Microsoft.AspNetCore.Mvc;
+using System.Net;
+using Microsoft.Extensions.Logging;
+using System.Threading.Tasks;
+using System.Collections.Generic;
+using HordeServer.Services;
+using Microsoft.AspNetCore.Authorization;
+using HordeServer.Models;
+
+namespace HordeServer.Controllers
+{
+	/// <summary>
+	/// Controller for Perforce triggers and callbacks
+	/// </summary>
+	[Authorize]
+	[ApiController]
+	[Route("[controller]")]
+	public class PerforceController : ControllerBase
+	{
+		/// <summary>
+		/// The database service instance
+		/// </summary>
+		private DatabaseService DatabaseService;
+
+		/// <summary>
+		/// The ACL service instance
+		/// </summary>
+		private AclService AclService;
+
+		/// <summary>
+		/// Load balancer instance
+		/// </summary>
+		private PerforceLoadBalancer PerforceLoadBalancer;
+
+		/// <summary>
+		/// Constructor
+		/// </summary>
+		public PerforceController(DatabaseService DatabaseService, AclService AclService, PerforceLoadBalancer PerforceLoadBalancer)
+		{
+			this.DatabaseService = DatabaseService;
+			this.AclService = AclService;
+			this.PerforceLoadBalancer = PerforceLoadBalancer;
+		}
+
+		/// <summary>
+		/// Get the current server status
+		/// </summary>
+		/// <returns></returns>
+		[HttpGet]
+		[Route("/api/v1/perforce/status")]
+		public async Task<ActionResult<List<object>>> GetStatusAsync()
+		{
+			List<IPerforceServer> Servers = await PerforceLoadBalancer.GetServersAsync();
+
+			List<object> Responses = new List<object>();
+			foreach (IPerforceServer Server in Servers)
+			{
+				Responses.Add(new { Server.ServerAndPort, Server.BaseServerAndPort, Server.Cluster, Server.NumLeases, Server.Status, Server.Detail });
+			}
+			return Responses;
+		}
+
+		/// <summary>
+		/// Update Perforce server settings
+		/// </summary>
+		[HttpPut]
+		[Route("/api/v1/perforce/settings")]
+		public async Task<ActionResult> UpdatePerforceSettingsAsync([FromBody] List<PerforceCluster> Clusters)
+		{
+			if (!await AclService.AuthorizeAsync(AclAction.AdminWrite, User))
+			{
+				return Forbid();
+			}
+
+			for (; ; )
+			{
+				Globals Globals = await DatabaseService.GetGlobalsAsync();
+				Globals.PerforceClusters = Clusters;
+
+				if (await DatabaseService.TryUpdateSingletonAsync(Globals))
+				{
+					return Ok("Settings have been updated");
+				}
+			}
+		}
+
+		/// <summary>
+		/// Gets the current perforce settinsg
+		/// </summary>
+		/// <returns>List of Perforce clusters</returns>
+		[HttpGet]
+		[Route("/api/v1/perforce/settings")]
+		public async Task<ActionResult<List<PerforceCluster>>> GetPerforceSettingsAsync()
+		{
+			if (!await AclService.AuthorizeAsync(AclAction.AdminRead, User))
+			{
+				return Forbid();
+			}
+
+			Globals Globals = await DatabaseService.GetGlobalsAsync();
+			return Globals.PerforceClusters;
+		}
+	}
+
+	/// <summary>
+	/// Controller for Perforce triggers and callbacks
+	/// </summary>
+	[ApiController]
+	[Route("[controller]")]
+	public class PublicPerforceController : ControllerBase
+	{
+		/// <summary>
+		/// Logger instance
+		/// </summary>
+		private readonly ILogger<PerforceController> Logger;
+
+		/// <summary>
+		/// Constructor
+		/// </summary>
+		public PublicPerforceController(ILogger<PerforceController> Logger)
+		{
+			this.Logger = Logger;
+		}
+
+		/// <summary>
+		/// Endpoint which trigger scripts in Perforce will call
+		/// </summary>
+		/// <returns>200 OK on success</returns>
+		[HttpPost]
+		[Route("/api/v1/perforce/trigger/{Type}")]
+		public ActionResult TriggerCallback(string Type, [FromQuery] long? Changelist = null, [FromQuery(Name = "user")] string? PerforceUser = null)
+		{
+			// Currently just a placeholder until correct triggers are in place.
+			Logger.LogDebug("Received Perforce trigger callback. Type={Type} Changelist={Changelist} User={User}", Type, Changelist, PerforceUser);
+			string Content = "{\"message\": \"Trigger received\"}";
+			return new ContentResult { ContentType = "text/plain", StatusCode = (int)HttpStatusCode.OK, Content = Content };
+		}
+	}
+}
diff --git a/Engine/Source/Programs/Horde/HordeServer/Controllers/PermissionsController.cs b/Engine/Source/Programs/Horde/HordeServer/Controllers/PermissionsController.cs
new file mode 100644
--- /dev/null
+++ b/Engine/Source/Programs/Horde/HordeServer/Controllers/PermissionsController.cs
@@ -0,0 +1,85 @@
+// Copyright Epic Games, Inc. All Rights Reserved.
+
+using HordeServer.Api;
+using HordeServer.Models;
+using HordeServer.Services;
+using HordeServer.Utilities;
+using Microsoft.AspNetCore.Authorization;
+using Microsoft.AspNetCore.Mvc;
+using MongoDB.Bson;
+using System;
+using System.Collections.Generic;
+using System.ComponentModel.DataAnnotations;
+using System.Linq;
+using System.Linq.Expressions;
+using System.Threading.Tasks;
+
+namespace HordeServer.Controllers
+{
+	/// <summary>
+	/// Controller for the /api/v1/permissions endpoint
+	/// </summary>
+	[ApiController]
+	[Authorize]
+	[Route("[controller]")]
+	public class PermissionsController : ControllerBase
+	{
+		/// <summary>
+		/// Singleton instance of the ACL service
+		/// </summary>
+		private readonly AclService AclService;
+
+		/// <summary>
+		/// Constructor
+		/// </summary>
+		/// <param name="AclService">The ACL service</param>
+		public PermissionsController(AclService AclService)
+		{
+			this.AclService = AclService;
+		}
+
+		/// <summary>
+		/// Gets a scope ACL
+		/// </summary>
+		/// <returns>The ACL information</returns>
+		[HttpGet]
+		[Route("/api/v1/permissions")]
+		public async Task<ActionResult<GetAclResponse>> GetPermissions()
+		{
+			GlobalPermissionsCache PermissionsCache = new GlobalPermissionsCache();
+			if (!await AclService.AuthorizeAsync(AclAction.ViewPermissions, User, PermissionsCache))
+			{
+				return Forbid();
+			}
+
+			GlobalPermissions GlobalPermissions = PermissionsCache.GlobalPermissions ?? await AclService.GetGlobalPermissionsAsync();
+			return new GetAclResponse(GlobalPermissions.Acl);
+		}
+
+		/// <summary>
+		/// Updates a scope ACL
+		/// </summary>
+		/// <param name="Update">The update request</param>
+		[HttpPut]
+		[Route("/api/v1/permissions")]
+		public async Task<ActionResult> UpdateScopeAcl(UpdateAclRequest Update)
+		{
+			for (; ; )
+			{
+				GlobalPermissionsCache PermissionsCache = new GlobalPermissionsCache();
+				if (!await AclService.AuthorizeAsync(AclAction.ChangePermissions, User, PermissionsCache))
+				{
+					return Forbid();
+				}
+
+				GlobalPermissions GlobalPermissions = PermissionsCache.GlobalPermissions ?? await AclService.GetGlobalPermissionsAsync();
+				GlobalPermissions.Acl = Acl.Merge(GlobalPermissions.Acl, Update) ?? new Acl();
+
+				if (await AclService.TryUpdateGlobalPermissionsAsync(GlobalPermissions))
+				{
+					return Ok();
+				}
+			}
+		}
+	}
+}
diff --git a/Engine/Source/Programs/Horde/HordeServer/Controllers/PoolsController.cs b/Engine/Source/Programs/Horde/HordeServer/Controllers/PoolsController.cs
new file mode 100644
--- /dev/null
+++ b/Engine/Source/Programs/Horde/HordeServer/Controllers/PoolsController.cs
@@ -0,0 +1,197 @@
+// Copyright Epic Games, Inc. All Rights Reserved.
+
+using HordeServer.Api;
+using HordeServer.Models;
+using HordeServer.Services;
+using HordeServer.Utilities;
+using Microsoft.AspNetCore.Authorization;
+using Microsoft.AspNetCore.Mvc;
+using MongoDB.Bson;
+using System;
+using System.Collections.Generic;
+using System.ComponentModel.DataAnnotations;
+using System.Linq;
+using System.Threading.Tasks;
+
+using PoolId = HordeServer.Utilities.StringId<HordeServer.Models.IPool>;
+
+namespace HordeServer.Controllers
+{
+	/// <summary>
+	/// Controller for the /api/v1/pools endpoint
+	/// </summary>
+	[ApiController]
+	[Authorize]
+	[Route("[controller]")]
+	public class PoolsController : ControllerBase
+	{
+		/// <summary>
+		/// Singleton instance of the ACL service
+		/// </summary>
+		private readonly AclService AclService;
+
+		/// <summary>
+		/// Singleton instance of the pool service
+		/// </summary>
+		private readonly PoolService PoolService;
+
+		/// <summary>
+		/// Constructor
+		/// </summary>
+		/// <param name="AclService">The ACL service</param>
+		/// <param name="PoolService">The pool service</param>
+		public PoolsController(AclService AclService, PoolService PoolService)
+		{
+			this.AclService = AclService;
+			this.PoolService = PoolService;
+		}
+
+		/// <summary>
+		/// Creates a new pool
+		/// </summary>
+		/// <param name="Create">Parameters for the new pool.</param>
+		/// <returns>Http result code</returns>
+		[HttpPost]
+		[Route("/api/v1/pools")]
+		public async Task<ActionResult<CreatePoolResponse>> CreatePoolAsync([FromBody] CreatePoolRequest Create)
+		{
+			if(!await AclService.AuthorizeAsync(AclAction.CreatePool, User))
+			{
+				return Forbid();
+			}
+
+			IPool NewPool = await PoolService.CreatePoolAsync(Create.Name, Create.Requirements, Create.EnableAutoscaling, Create.MinAgents, Create.NumReserveAgents, Create.Properties);
+			return new CreatePoolResponse(NewPool.Id.ToString());
+		}
+
+		/// <summary>
+		/// Query all the pools
+		/// </summary>
+		/// <param name="Filter">Filter for the properties to return</param>
+		/// <returns>Information about all the pools</returns>
+		[HttpGet]
+		[Route("/api/v1/pools")]
+		[ProducesResponseType(typeof(List<GetPoolResponse>), 200)]
+		public async Task<ActionResult<List<object>>> GetPoolsAsync([FromQuery] PropertyFilter? Filter = null)
+		{
+			if (!await AclService.AuthorizeAsync(AclAction.ListPools, User))
+			{
+				return Forbid();
+			}
+
+			List<IPool> Pools = await PoolService.GetPoolsAsync();
+
+			List<object> Responses = new List<object>();
+			foreach (IPool Pool in Pools)
+			{
+				Responses.Add(new GetPoolResponse(Pool).ApplyFilter(Filter));
+			}
+			return Responses;
+		}
+
+		/// <summary>
+		/// Retrieve information about a specific pool
+		/// </summary>
+		/// <param name="PoolId">Id of the pool to get information about</param>
+		/// <param name="Filter">Filter to apply to the returned properties</param>
+		/// <returns>Information about the requested pool</returns>
+		[HttpGet]
+		[Route("/api/v1/pools/{PoolId}")]
+		[ProducesResponseType(typeof(GetPoolResponse), 200)]
+		public async Task<ActionResult<object>> GetPoolAsync(string PoolId, [FromQuery] PropertyFilter? Filter = null)
+		{
+			if (!await AclService.AuthorizeAsync(AclAction.ViewPool, User))
+			{
+				return Forbid();
+			}
+
+			PoolId PoolIdValue = new PoolId(PoolId);
+
+			IPool? Pool = await PoolService.GetPoolAsync(PoolIdValue);
+			if (Pool == null)
+			{
+				return NotFound();
+			}
+
+			return new GetPoolResponse(Pool).ApplyFilter(Filter);
+		}
+
+		/// <summary>
+		/// Update a pool's properties.
+		/// </summary>
+		/// <param name="PoolId">Id of the pool to update</param>
+		/// <param name="Update">Items on the pool to update</param>
+		/// <returns>Http result code</returns>
+		[HttpPut]
+		[Route("/api/v1/pools/{PoolId}")]
+		public async Task<ActionResult> UpdatePoolAsync(string PoolId, [FromBody] UpdatePoolRequest Update)
+		{
+			if (!await AclService.AuthorizeAsync(AclAction.UpdatePool, User))
+			{
+				return Forbid();
+			}
+
+			PoolId PoolIdValue = new PoolId(PoolId);
+			
+			IPool? Pool = await PoolService.GetPoolAsync(PoolIdValue);
+			if(Pool == null)
+			{
+				return NotFound();
+			}
+
+			await PoolService.UpdatePoolAsync(Pool, Update.Name, Update.Requirements, Update.EnableAutoscaling, Update.MinAgents, Update.NumReserveAgents, Update.Properties);
+			return new OkResult();
+		}
+
+		/// <summary>
+		/// Delete a pool
+		/// </summary>
+		/// <param name="PoolId">Id of the pool to delete</param>
+		/// <returns>Http result code</returns>
+		[HttpDelete]
+		[Route("/api/v1/pools/{PoolId}")]
+		public async Task<ActionResult> DeletePoolAsync(string PoolId)
+		{
+			if (!await AclService.AuthorizeAsync(AclAction.DeletePool, User))
+			{
+				return Forbid();
+			}
+
+			PoolId PoolIdValue = new PoolId(PoolId);
+			if(!await PoolService.DeletePoolAsync(PoolIdValue))
+			{
+				return NotFound();
+			}
+			return new OkResult();
+		}
+
+		/// <summary>
+		/// Batch update pool properties
+		/// </summary>
+		/// <param name="BatchUpdates">List of pools to update</param>
+		/// <returns>Http result code</returns>
+		[HttpPut]
+		[Route("/api/v1/pools")]
+		public async Task<ActionResult> UpdatePoolAsync([FromBody] List<BatchUpdatePoolRequest> BatchUpdates)
+		{
+			if (!await AclService.AuthorizeAsync(AclAction.UpdatePool, User))
+			{
+				return Forbid();
+			}
+
+			foreach (BatchUpdatePoolRequest Update in BatchUpdates)
+			{
+				PoolId PoolIdValue = new PoolId(Update.Id);
+
+				IPool? Pool = await PoolService.GetPoolAsync(PoolIdValue);
+				if (Pool == null)
+				{
+					return NotFound();
+				}
+
+				await PoolService.UpdatePoolAsync(Pool, Update.Name, NewProperties: Update.Properties);
+			}
+			return new OkResult();
+		}
+	}
+}
diff --git a/Engine/Source/Programs/Horde/HordeServer/Controllers/ProjectsController.cs b/Engine/Source/Programs/Horde/HordeServer/Controllers/ProjectsController.cs
new file mode 100644
--- /dev/null
+++ b/Engine/Source/Programs/Horde/HordeServer/Controllers/ProjectsController.cs
@@ -0,0 +1,250 @@
+// Copyright Epic Games, Inc. All Rights Reserved.
+
+using HordeServer.Api;
+using HordeServer.Models;
+using HordeServer.Services;
+using HordeServer.Utilities;
+using Microsoft.AspNetCore.Authorization;
+using Microsoft.AspNetCore.Mvc;
+using MongoDB.Bson;
+using System;
+using System.Collections.Generic;
+using System.ComponentModel.DataAnnotations;
+using System.Linq;
+using System.Threading.Tasks;
+
+using ProjectId = HordeServer.Utilities.StringId<HordeServer.Models.IProject>;
+
+namespace HordeServer.Controllers
+{
+	/// <summary>
+	/// Controller for the /api/v1/projects endpoint
+	/// </summary>
+	[ApiController]
+	[Authorize]
+	[Route("[controller]")]
+	public class ProjectsController : ControllerBase
+	{
+		/// <summary>
+		/// Singleton instance of the ACL service
+		/// </summary>
+		private readonly AclService AclService;
+
+		/// <summary>
+		/// Singleton instance of the project service
+		/// </summary>
+		private readonly ProjectService ProjectService;
+
+		/// <summary>
+		/// Singleton instance of the stream service
+		/// </summary>
+		private readonly StreamService StreamService;
+
+		/// <summary>
+		/// Singleton instance of the job service
+		/// </summary>
+		private readonly JobService JobService;
+
+		/// <summary>
+		/// Singleton instance of the schedule service
+		/// </summary>
+		private readonly ScheduleService ScheduleService;
+
+		/// <summary>
+		/// Constructor
+		/// </summary>
+		/// <param name="AclService">The ACL service</param>
+		/// <param name="ProjectService">The project service</param>
+		/// <param name="StreamService">The stream service</param>
+		/// <param name="JobService">The job service</param>
+		/// <param name="ScheduleService">The schedule service</param>
+		public ProjectsController(AclService AclService, ProjectService ProjectService, StreamService StreamService, JobService JobService, ScheduleService ScheduleService)
+		{
+			this.AclService = AclService;
+			this.ProjectService = ProjectService;
+			this.StreamService = StreamService;
+			this.JobService = JobService;
+			this.ScheduleService = ScheduleService;
+		}
+
+		/// <summary>
+		/// Creates a new project
+		/// </summary>
+		/// <param name="Create">Parameters for the new project.</param>
+		/// <returns>Http result code</returns>
+		[HttpPost]
+		[Route("/api/v1/projects")]
+		public async Task<ActionResult<CreateProjectResponse>> CreateProjectAsync([FromBody] CreateProjectRequest Create)
+		{
+			if(!await AclService.AuthorizeAsync(AclAction.CreateProject, User, null))
+			{
+				return Forbid();
+			}
+
+			IProject? NewProject = await ProjectService.TryCreateProjectAsync(ProjectId.Sanitize(Create.Name), Create.Name, Create.Order, Create.Categories?.ConvertAll(x => new StreamCategory(x)), Create.Properties);
+			if (NewProject == null)
+			{
+				return BadRequest("Project already exists");
+			}
+			else
+			{
+				return new CreateProjectResponse(NewProject.Id.ToString());
+			}
+		}
+
+		/// <summary>
+		/// Query all the projects
+		/// </summary>
+		/// <param name="IncludeStreams">Whether to include streams in the response</param>
+		/// <param name="IncludeCategories">Whether to include categories in the response</param>
+		/// <param name="Filter">Filter for the properties to return</param>
+		/// <returns>Information about all the projects</returns>
+		[HttpGet]
+		[Route("/api/v1/projects")]
+		[ProducesResponseType(typeof(List<GetProjectResponse>), 200)]
+		public async Task<ActionResult<List<object>>> GetProjectsAsync([FromQuery(Name = "Streams")] bool IncludeStreams = false, [FromQuery(Name = "Categories")] bool IncludeCategories = false, [FromQuery] PropertyFilter? Filter = null)
+		{
+			List<IProject> Projects = await ProjectService.GetProjectsAsync();
+			ProjectPermissionsCache PermissionsCache = new ProjectPermissionsCache();
+
+			List<IStream>? Streams = null;
+			if (IncludeStreams || IncludeCategories)
+			{
+				Streams = await StreamService.GetStreamsAsync();
+			}
+
+			List<object> Responses = new List<object>();
+			foreach (IProject Project in Projects)
+			{
+				if (await ProjectService.AuthorizeAsync(Project, AclAction.ViewProject, User, PermissionsCache))
+				{
+					bool bIncludeAcl = await ProjectService.AuthorizeAsync(Project, AclAction.ViewPermissions, User, PermissionsCache);
+					Responses.Add(Project.ToResponse(IncludeStreams, IncludeCategories, Streams, bIncludeAcl).ApplyFilter(Filter));
+				}
+			}
+			return Responses;
+		}
+
+		/// <summary>
+		/// Retrieve information about a specific project
+		/// </summary>
+		/// <param name="ProjectId">Id of the project to get information about</param>
+		/// <param name="Filter">Filter for the properties to return</param>
+		/// <returns>Information about the requested project</returns>
+		[HttpGet]
+		[Route("/api/v1/projects/{ProjectId}")]
+		[ProducesResponseType(typeof(List<GetProjectResponse>), 200)]
+		public async Task<ActionResult<object>> GetProjectAsync(string ProjectId, [FromQuery] PropertyFilter? Filter = null)
+		{
+			ProjectId ProjectIdValue = new ProjectId(ProjectId);
+
+			IProject? Project = await ProjectService.GetProjectAsync(ProjectIdValue);
+			if (Project == null)
+			{
+				return NotFound();
+			}
+
+			ProjectPermissionsCache Cache = new ProjectPermissionsCache();
+			if (!await ProjectService.AuthorizeAsync(Project, AclAction.ViewProject, User, Cache))
+			{
+				return Forbid();
+			}
+
+			bool bIncludeStreams = PropertyFilter.Includes(Filter, nameof(GetProjectResponse.Streams));
+			bool bIncludeCategories = PropertyFilter.Includes(Filter, nameof(GetProjectResponse.Categories));
+
+			List<IStream>? VisibleStreams = null;
+			if (bIncludeStreams || bIncludeCategories)
+			{
+				VisibleStreams = new List<IStream>();
+
+				List<IStream> Streams = await StreamService.GetStreamsAsync(Project.Id);
+				foreach (IStream Stream in Streams)
+				{
+					if (await StreamService.AuthorizeAsync(Stream, AclAction.ViewStream, User, Cache))
+					{
+						VisibleStreams.Add(Stream);
+					}
+				}
+			}
+
+			bool bIncludeAcl = await ProjectService.AuthorizeAsync(Project, AclAction.ViewPermissions, User, Cache);
+			return Project.ToResponse(bIncludeStreams, bIncludeCategories, VisibleStreams, bIncludeAcl).ApplyFilter(Filter);
+		}
+
+		/// <summary>
+		/// Update a project's properties.
+		/// </summary>
+		/// <param name="ProjectId">Id of the project to update</param>
+		/// <param name="Update">Items on the project to update</param>
+		/// <returns>Http result code</returns>
+		[HttpPut]
+		[Route("/api/v1/projects/{ProjectId}")]
+		public async Task<ActionResult> UpdateProjectAsync(string ProjectId, [FromBody] UpdateProjectRequest Update)
+		{
+			ProjectId ProjectIdValue = new ProjectId(ProjectId);
+			for (; ; )
+			{
+				IProject? Project = await ProjectService.GetProjectAsync(ProjectIdValue);
+				if (Project == null)
+				{
+					if (!await AclService.AuthorizeAsync(AclAction.CreateProject, User, null))
+					{
+						return Forbid();
+					}
+					if (Update.Name == null)
+					{
+						return BadRequest("Project does not exist; 'Name' must be specified on update request to create it");
+					}
+
+					IProject? NewProject = await ProjectService.TryCreateProjectAsync(ProjectIdValue, Update.Name, Update.Order, Update.Categories?.ConvertAll(x => new StreamCategory(x)), Update.Properties);
+					if (NewProject != null)
+					{
+						return Ok();
+					}
+				}
+				else
+				{
+					ProjectPermissionsCache PermissionsCache = new ProjectPermissionsCache();
+					if (!await ProjectService.AuthorizeAsync(Project, AclAction.UpdateProject, User, PermissionsCache))
+					{
+						return Forbid();
+					}
+					if (Update.Acl != null && !await ProjectService.AuthorizeAsync(Project, AclAction.ChangePermissions, User, PermissionsCache))
+					{
+						return Forbid();
+					}
+
+					List<StreamCategory>? NewCategories = Update.Categories?.ConvertAll(x => new StreamCategory(x));
+					await ProjectService.UpdateProjectAsync(ProjectIdValue, Update.Name, Update.Order, NewCategories, Update.Properties, Acl.Merge(Project.Acl, Update.Acl));
+					return new OkResult();
+				}
+			}
+		}
+
+		/// <summary>
+		/// Delete a project
+		/// </summary>
+		/// <param name="ProjectId">Id of the project to update</param>
+		/// <returns>Http result code</returns>
+		[HttpDelete]
+		[Route("/api/v1/projects/{ProjectId}")]
+		public async Task<ActionResult> DeleteProjectAsync(string ProjectId)
+		{
+			ProjectId ProjectIdValue = new ProjectId(ProjectId);
+
+			IProject? Project = await ProjectService.GetProjectAsync(ProjectIdValue);
+			if (Project == null)
+			{
+				return NotFound();
+			}
+			if (!await ProjectService.AuthorizeAsync(Project, AclAction.DeleteProject, User, null))
+			{
+				return Forbid();
+			}
+
+			await ProjectService.DeleteProjectAsync(ProjectIdValue, StreamService, JobService);
+			return new OkResult();
+		}
+	}
+}
