Initial drop of Horde Server/Agent for UE5/Main. Still missing some Epic-specific functionality required to build/run correctly.

diff --git a/Engine/Source/Programs/Horde/HordeServer/Controllers/ReportsController.cs b/Engine/Source/Programs/Horde/HordeServer/Controllers/ReportsController.cs
new file mode 100644
--- /dev/null
+++ b/Engine/Source/Programs/Horde/HordeServer/Controllers/ReportsController.cs
@@ -0,0 +1,57 @@
+// Copyright Epic Games, Inc. All Rights Reserved.
+
+using System;
+using System.Collections.Generic;
+using System.Linq;
+using System.Threading.Tasks;
+using HordeServer.Api;
+using HordeServer.Collections;
+using HordeServer.Models;
+using HordeServer.Services;
+using Microsoft.AspNetCore.Authorization;
+using Microsoft.AspNetCore.Mvc;
+using EpicGames.Core;
+
+namespace HordeServer.Controllers
+{
+	/// <summary>
+	/// Controller for the /api/v1/reports endpoint, used for the reports pages
+	/// </summary>
+	[ApiController]
+	[Route("[controller]")]
+	public sealed class ReportsController : ControllerBase
+	{
+		/// <summary>
+		/// the Telemetry collection singleton
+		/// </summary>
+		ITelemetryCollection TelemetryCollection;
+
+		/// <summary>
+		/// Constructor
+		/// </summary>
+		/// <param name="TelemetryCollection">The telemetry collection</param>
+		public ReportsController(ITelemetryCollection TelemetryCollection)
+		{
+			this.TelemetryCollection = TelemetryCollection;
+		}
+
+		/// <summary>
+		/// Gets a collection of utilization data including an endpoint and a range
+		/// </summary>
+		[HttpGet]
+		[Route("/api/v1/reports/utilization/{EndDate}")]
+		public async Task<ActionResult<List<UtilizationTelemetryResponse>>> GetStreamUtilizationData(DateTime EndDate, [FromQuery(Name = "Range")] int Range, [FromQuery(Name = "TzOffset")] int? TzOffset)
+		{
+			// Logic here is a bit messy. The client is always passing in a date at midnight
+			// If user passes in 12/1/2020 into the date with range of 1, the range should be 12/1/2020:00:00:00 to 12/1/2020:23:59:59
+			// Range 2 should be 11/30/2020:00:00:00 to 12/1/2020:23:59:59
+			int Offset = TzOffset ?? 0;
+			DateTimeOffset EndDateOffset = new DateTimeOffset(EndDate, TimeSpan.FromHours(Offset)).Add(new TimeSpan(23, 59, 59));
+			DateTimeOffset StartDateOffset = EndDate.Subtract(new TimeSpan(Range - 1, 0, 0, 0));
+
+			List<IUtilizationTelemetry> Telemetry = await TelemetryCollection.GetUtilizationTelemetryAsync(StartDateOffset.UtcDateTime, EndDateOffset.UtcDateTime);
+
+			return Telemetry.ConvertAll(Telemetry => new UtilizationTelemetryResponse(Telemetry));
+		}
+	}
+}
diff --git a/Engine/Source/Programs/Horde/HordeServer/Controllers/SchedulesController.cs b/Engine/Source/Programs/Horde/HordeServer/Controllers/SchedulesController.cs
new file mode 100644
--- /dev/null
+++ b/Engine/Source/Programs/Horde/HordeServer/Controllers/SchedulesController.cs
@@ -0,0 +1,42 @@
+// Copyright Epic Games, Inc. All Rights Reserved.
+
+using EpicGames.Core;
+using HordeServer.Api;
+using HordeServer.Models;
+using HordeServer.Services;
+using HordeServer.Utilities;
+using Microsoft.AspNetCore.Authorization;
+using Microsoft.AspNetCore.Mvc;
+using MongoDB.Bson;
+using System;
+using System.Collections.Generic;
+using System.Linq;
+using System.Threading.Tasks;
+
+namespace HordeServer.Controllers
+{
+	/// <summary>
+	/// Controller for the /api/v1/projects endpoint
+	/// </summary>
+	[ApiController]
+	[Authorize]
+	[Route("[controller]")]
+	public class SchedulesController : ControllerBase
+	{
+		/// <summary>
+		/// Query all the schedules
+		/// </summary>
+		/// <param name="StreamId">Id of the stream to search for schedules for</param>
+		/// <param name="Filter">Filter for the properties to return</param>
+		/// <returns>Information about all the schedules</returns>
+		[HttpGet]
+		[Route("/api/v1/schedules")]
+		[ProducesResponseType(typeof(List<GetScheduleResponse>), 200)]
+#pragma warning disable CA1801 // Remove unused parameter
+		public Task<ActionResult<List<object>>> GetSchedulesAsync([FromQuery] string? StreamId = null, [FromQuery] PropertyFilter? Filter = null)
+		{
+			return Task.FromResult<ActionResult<List<object>>>(new List<object>());
+		}
+#pragma warning restore CA1801 // Remove unused parameter
+	}
+}
diff --git a/Engine/Source/Programs/Horde/HordeServer/Controllers/SoftwareController.cs b/Engine/Source/Programs/Horde/HordeServer/Controllers/SoftwareController.cs
new file mode 100644
--- /dev/null
+++ b/Engine/Source/Programs/Horde/HordeServer/Controllers/SoftwareController.cs
@@ -0,0 +1,185 @@
+// Copyright Epic Games, Inc. All Rights Reserved.
+
+using HordeServer.Api;
+using HordeServer.Services;
+using Microsoft.AspNetCore.Authorization;
+using Microsoft.AspNetCore.Http;
+using Microsoft.AspNetCore.Mvc;
+using Microsoft.Net.Http.Headers;
+using System;
+using System.IO;
+using System.Collections.Generic;
+using System.Linq;
+using System.Threading.Tasks;
+using HordeServer.Models;
+using MongoDB.Bson;
+using HordeServer.Utilities;
+
+namespace HordeServer.Controllers
+{
+	/// <summary>
+	/// Controller for the /api/v1/software endpoint
+	/// </summary>
+	[ApiController]
+	[Authorize]
+	[Route("[controller]")]
+	public class SoftwareController : ControllerBase
+	{
+		/// <summary>
+		/// Singleton instance of the ACL service
+		/// </summary>
+		private readonly AclService AclService;
+
+		/// <summary>
+		/// Singleton instance of the client service
+		/// </summary>
+		private readonly SoftwareService SoftwareService;
+
+		/// <summary>
+		/// Constructor
+		/// </summary>
+		/// <param name="AclService">The ACL service</param>
+		/// <param name="SoftwareService">The client service</param>
+		public SoftwareController(AclService AclService, SoftwareService SoftwareService)
+		{
+			this.AclService = AclService;
+			this.SoftwareService = SoftwareService;
+		}
+
+		/// <summary>
+		/// Uploads a new agent zip file
+		/// </summary>
+		/// <param name="File">Zip archive containing the new client software</param>
+		/// <param name="Default">Whether the client should immediately be made the default</param>
+		/// <returns>Http result code</returns>
+		[HttpPost]
+		[Route("/api/v1/software")]
+		public async Task<ActionResult<CreateSoftwareResponse>> CreateSoftwareAsync([FromForm] IFormFile File, [FromForm] bool Default = false)
+		{
+			if (!await AclService.AuthorizeAsync(AclAction.UploadSoftware, User))
+			{
+				return Forbid();
+			}
+
+			ObjectId Id;
+			using (System.IO.Stream Stream = File.OpenReadStream())
+			{
+				Id = await SoftwareService.CreateSoftwareAsync(Stream, User.Identity.Name, Default);
+			}
+			return new CreateSoftwareResponse(Id.ToString());
+		}
+
+		/// <summary>
+		/// Finds all uploaded software matching the given criteria
+		/// </summary>
+		/// <param name="CreatedByUser">The user that created the software</param>
+		/// <param name="MadeDefaultByUser">The user that made it the default</param>
+		/// <param name="MadeDefault">Whether the software was made the default</param>
+		/// <param name="Offset">Index of the first result to return</param>
+		/// <param name="Count">Number of results to return</param>
+		/// <param name="Filter">Filter for the properties to return</param>
+		/// <returns>Http response</returns>
+		[HttpGet]
+		[Route("/api/v1/software")]
+		[ProducesResponseType(typeof(List<GetSoftwareResponse>), 200)]
+		public async Task<ActionResult<List<object>>> FindSoftwareAsync([FromQuery] string? CreatedByUser = null, [FromQuery] string? MadeDefaultByUser = null, [FromQuery] bool? MadeDefault = null, [FromQuery] PropertyFilter? Filter = null, [FromQuery] int Offset = 0, [FromQuery] int Count = 200)
+		{
+			if (!await AclService.AuthorizeAsync(AclAction.DownloadSoftware, User))
+			{
+				return Forbid();
+			}
+
+			List<ISoftware> Results = await SoftwareService.FindSoftwareAsync(CreatedByUser, MadeDefaultByUser, MadeDefault, Offset, Count);
+
+			List<object> Responses = new List<object>();
+			foreach (ISoftware Result in Results)
+			{
+				Responses.Add(new GetSoftwareResponse(Result).ApplyFilter(Filter));
+			}
+			return Responses;
+		}
+
+		/// <summary>
+		/// Gets information about a specific client revision
+		/// </summary>
+		/// <param name="Id">Unique id of the client software</param>
+		/// <param name="Filter">Filter for the properties to return</param>
+		/// <returns>Http response</returns>
+		[HttpGet]
+		[Route("/api/v1/software/{Id}")]
+		[ProducesResponseType(typeof(GetSoftwareResponse), 200)]
+		public async Task<ActionResult<object>> GetSoftwareAsync(string Id, [FromQuery] PropertyFilter? Filter = null)
+		{
+			if (!await AclService.AuthorizeAsync(AclAction.DownloadSoftware, User))
+			{
+				return Forbid();
+			}
+
+			ISoftware? Client = await GetSoftwareInternalAsync(Id, false);
+			if (Client == null)
+			{
+				return NotFound();
+			}
+
+			return new GetSoftwareResponse(Client).ApplyFilter(Filter);
+		}
+
+		/// <summary>
+		/// Gets the zip file for a specific client revision
+		/// </summary>
+		/// <param name="Id">Unique id of the client</param>
+		/// <returns>Http response</returns>
+		[HttpGet]
+		[Route("/api/v1/software/{Id}/download")]
+		public async Task<ActionResult<GetSoftwareResponse>> DownloadSoftwareAsync(string Id)
+		{
+			if (!await AclService.AuthorizeAsync(AclAction.DownloadSoftware, User))
+			{
+				return Forbid();
+			}
+
+			ISoftware? Software = await GetSoftwareInternalAsync(Id, true);
+			if (Software == null)
+			{
+				return NotFound();
+			}
+
+			return new FileStreamResult(new MemoryStream(Software.Data!), new MediaTypeHeaderValue("application/octet-stream"));
+		}
+
+		/// <summary>
+		/// Find the software matching the given identifier
+		/// </summary>
+		/// <param name="Id">The identifier to match. May be an object id or "latest"</param>
+		/// <param name="bIncludeData">Whether to include data in the response</param>
+		/// <returns>Software interface</returns>
+		async Task<ISoftware?> GetSoftwareInternalAsync(string Id, bool bIncludeData)
+		{
+			if (Id.Equals("latest", StringComparison.OrdinalIgnoreCase))
+			{
+				return await SoftwareService.GetDefaultSoftwareAsync(bIncludeData);
+			}
+			else
+			{
+				return await SoftwareService.GetSoftwareAsync(Id.ToObjectId(), bIncludeData);
+			}
+		}
+
+		/// <summary>
+		/// Deletes a version of the client
+		/// </summary>
+		/// <param name="Id">Unique id of the client to delete</param>
+		/// <returns>Http response</returns>
+		[HttpDelete]
+		[Route("/api/v1/software/{Id}")]
+		public async Task<ActionResult> DeleteSoftwareAsync(string Id)
+		{
+			if (!await AclService.AuthorizeAsync(AclAction.DeleteSoftware, User))
+			{
+				return Forbid();
+			}
+			await SoftwareService.DeleteSoftwareAsync(Id.ToObjectId());
+			return new OkResult();
+		}
+	}
+}
diff --git a/Engine/Source/Programs/Horde/HordeServer/Controllers/StreamsController.cs b/Engine/Source/Programs/Horde/HordeServer/Controllers/StreamsController.cs
new file mode 100644
--- /dev/null
+++ b/Engine/Source/Programs/Horde/HordeServer/Controllers/StreamsController.cs
@@ -0,0 +1,552 @@
+// Copyright Epic Games, Inc. All Rights Reserved.
+
+using EpicGames.Core;
+using HordeServer.Api;
+using HordeServer.Models;
+using HordeServer.Services;
+using HordeServer.Utilities;
+using Microsoft.AspNetCore.Authorization;
+using Microsoft.AspNetCore.Mvc;
+using MongoDB.Bson;
+using System;
+using System.Collections.Generic;
+using System.ComponentModel.DataAnnotations;
+using System.Diagnostics.CodeAnalysis;
+using System.Linq;
+using System.Threading.Tasks;
+
+using ProjectId = HordeServer.Utilities.StringId<HordeServer.Models.IProject>;
+using StreamId = HordeServer.Utilities.StringId<HordeServer.Models.IStream>;
+using TemplateRefId = HordeServer.Utilities.StringId<HordeServer.Models.TemplateRef>;
+using HordeServer.Collections;
+using System.Security.Claims;
+using System.IO;
+
+namespace HordeServer.Controllers
+{
+	/// <summary>
+	/// Controller for the /api/v1/streams endpoint
+	/// </summary>
+	[ApiController]
+	[Authorize]
+	[Route("[controller]")]
+	public class StreamsController : ControllerBase
+	{
+		/// <summary>
+		/// Singleton instance of the ACL service
+		/// </summary>
+		private readonly AclService AclService;
+
+		/// <summary>
+		/// Singleton instance of the stream service
+		/// </summary>
+		private readonly ProjectService ProjectService;
+
+		/// <summary>
+		/// Singleton instance of the stream service
+		/// </summary>
+		private readonly StreamService StreamService;
+
+		/// <summary>
+		/// Singleton instance of the template service
+		/// </summary>
+		private readonly TemplateService TemplateService;
+
+		/// <summary>
+		/// Singleton instance of the job service
+		/// </summary>
+		private readonly JobService JobService;
+
+		/// <summary>
+		/// Collection of jobstep refs
+		/// </summary>
+		private readonly IJobStepRefCollection JobStepRefCollection;
+
+		/// <summary>
+		/// Singleton instance of the perforce service
+		/// </summary>
+		private readonly IPerforceService PerforceService;
+
+		/// <summary>
+		/// Constructor
+		/// </summary>
+		/// <param name="AclService">The ACL service</param>
+		/// <param name="ProjectService">The project service</param>
+		/// <param name="StreamService">The stream service</param>
+		/// <param name="TemplateService">The template service</param>
+		/// <param name="JobService">The job service</param>
+		/// <param name="JobStepRefCollection">The jobstep ref collection</param>
+		/// <param name="PerforceService">The perforce service</param>
+		public StreamsController(AclService AclService, ProjectService ProjectService, StreamService StreamService, TemplateService TemplateService, JobService JobService, IJobStepRefCollection JobStepRefCollection, IPerforceService PerforceService)
+		{
+			this.AclService = AclService;
+			this.ProjectService = ProjectService;
+			this.StreamService = StreamService;
+			this.TemplateService = TemplateService;
+			this.JobService = JobService;
+			this.JobStepRefCollection = JobStepRefCollection;
+			this.PerforceService = PerforceService;
+		}
+
+		/// <summary>
+		/// Query all the streams for a particular project.
+		/// </summary>
+		/// <param name="ProjectIds">Unique id of the project to query</param>
+		/// <param name="Filter">Filter for the properties to return</param>
+		/// <returns>Information about all the projects</returns>
+		[HttpGet]
+		[Route("/api/v1/streams")]
+		[ProducesResponseType(typeof(List<GetStreamResponse>), 200)]
+		public async Task<ActionResult<List<object>>> GetStreamsAsync([FromQuery(Name = "ProjectId")] string[] ProjectIds, [FromQuery] PropertyFilter? Filter = null)
+		{
+			ProjectId[] ProjectIdValues = Array.ConvertAll(ProjectIds, x => new ProjectId(x));
+
+			List<IStream> Streams = await StreamService.GetStreamsAsync(ProjectIdValues);
+			ProjectPermissionsCache PermissionsCache = new ProjectPermissionsCache();
+
+			List<GetStreamResponse> Responses = new List<GetStreamResponse>();
+			foreach (IStream Stream in Streams)
+			{
+				if (await StreamService.AuthorizeAsync(Stream, AclAction.ViewStream, User, PermissionsCache))
+				{
+					GetStreamResponse Response = await CreateGetStreamResponse(Stream, PermissionsCache);
+					Responses.Add(Response);
+				}
+			}
+			return Responses.OrderBy(x => x.Order).ThenBy(x => x.Name).Select(x => PropertyFilter.Apply(x, Filter)).ToList();
+		}
+
+		/// <summary>
+		/// Retrieve information about a specific stream.
+		/// </summary>
+		/// <param name="StreamId">Id of the stream to get information about</param>
+		/// <param name="Filter">Filter for the properties to return</param>
+		/// <returns>Information about the requested project</returns>
+		[HttpGet]
+		[Route("/api/v1/streams/{StreamId}")]
+		[ProducesResponseType(typeof(GetStreamResponse), 200)]
+		public async Task<ActionResult<object>> GetStreamAsync(string StreamId, [FromQuery] PropertyFilter? Filter = null)
+		{
+			StreamId StreamIdValue = new StreamId(StreamId);
+
+			IStream? Stream = await StreamService.GetStreamAsync(StreamIdValue);
+			if (Stream == null)
+			{
+				return NotFound();
+			}
+
+			ProjectPermissionsCache PermissionsCache = new ProjectPermissionsCache();
+			if (!await StreamService.AuthorizeAsync(Stream, AclAction.ViewStream, User, PermissionsCache))
+			{
+				return Forbid();
+			}
+
+			bool bIncludeAcl = await StreamService.AuthorizeAsync(Stream, AclAction.ViewPermissions, User, PermissionsCache);
+			return PropertyFilter.Apply(await CreateGetStreamResponse(Stream, PermissionsCache), Filter);
+		}
+
+		/// <summary>
+		/// Create a stream response object, including all the templates
+		/// </summary>
+		/// <param name="Stream">Stream to create response for</param>
+		/// <param name="Cache">Permissions cache</param>
+		/// <returns>Response object</returns>
+		async Task<GetStreamResponse> CreateGetStreamResponse(IStream Stream, ProjectPermissionsCache Cache)
+		{
+			bool bIncludeAcl = Stream.Acl != null && await StreamService.AuthorizeAsync(Stream, AclAction.ViewPermissions, User, Cache);
+
+			List<GetTemplateRefResponse> ApiTemplateRefs = new List<GetTemplateRefResponse>();
+			foreach (KeyValuePair<TemplateRefId, TemplateRef> Pair in Stream.Templates)
+			{
+				if (await StreamService.AuthorizeAsync(Stream, Pair.Value, AclAction.ViewTemplate, User, Cache))
+				{
+					ITemplate? Template = await TemplateService.GetTemplateAsync(Pair.Value.Hash);
+					if (Template != null)
+					{
+						bool bIncludeTemplateAcl = Pair.Value.Acl != null && await StreamService.AuthorizeAsync(Stream, Pair.Value, AclAction.ViewPermissions, User, Cache);
+						ApiTemplateRefs.Add(new GetTemplateRefResponse(Pair.Key, Pair.Value, Template, bIncludeTemplateAcl));
+					}
+				}
+			}
+
+			return Stream.ToApiResponse(bIncludeAcl, ApiTemplateRefs);
+		}
+
+		/// <summary>
+		/// Gets a list of changes for a stream
+		/// </summary>
+		/// <param name="StreamId">The stream id</param>
+		/// <param name="Min">The starting changelist number</param>
+		/// <param name="Max">The ending changelist number</param>
+		/// <param name="Results">Number of results to return</param>
+		/// <param name="Filter">The filter to apply to the results</param>
+		/// <returns>Http result code</returns>
+		[HttpGet]
+		[Route("/api/v1/streams/{StreamId}/changes")]
+		[ProducesResponseType(typeof(List<GetChangeSummaryResponse>), 200)]
+		public async Task<ActionResult<List<object>>> GetChangesAsync(string StreamId, [FromQuery] int? Min = null, [FromQuery] int? Max = null, [FromQuery] int Results = 50, PropertyFilter? Filter = null)
+		{
+			StreamId StreamIdValue = new StreamId(StreamId);
+
+			IStream? Stream = await StreamService.GetStreamAsync(StreamIdValue);
+			if (Stream == null)
+			{
+				return NotFound();
+			}
+			if (!await StreamService.AuthorizeAsync(Stream, AclAction.ViewChanges, User, null))
+			{
+				return Forbid();
+			}
+
+			string? PerforceUser = User.GetPerforceUser();
+			if(PerforceUser == null)
+			{
+				return Forbid();
+			}
+
+			List<ChangeSummary> Commits = await PerforceService.GetChangesAsync(Stream.Name, Min, Max, Results, PerforceUser);
+			return Commits.ConvertAll(x => PropertyFilter.Apply(new GetChangeSummaryResponse(x), Filter));
+		}
+
+		/// <summary>
+		/// Gets a list of changes for a stream
+		/// </summary>
+		/// <param name="StreamId">The stream id</param>
+		/// <param name="Number">The changelist number</param>
+		/// <param name="Filter">The filter to apply to the results</param>
+		/// <returns>Http result code</returns>
+		[HttpGet]
+		[Route("/api/v1/streams/{StreamId}/changes/{Number}")]
+		[ProducesResponseType(typeof(GetChangeDetailsResponse), 200)]
+		public async Task<ActionResult<object>> GetChangeDetailsAsync(string StreamId, int Number, PropertyFilter? Filter = null)
+		{
+			StreamId StreamIdValue = new StreamId(StreamId);
+
+			IStream? Stream = await StreamService.GetStreamAsync(StreamIdValue);
+			if (Stream == null)
+			{
+				return NotFound();
+			}
+			if (!await StreamService.AuthorizeAsync(Stream, AclAction.ViewChanges, User, null))
+			{
+				return Forbid();
+			}
+
+			string? PerforceUser = User.GetPerforceUser();
+			if(PerforceUser == null)
+			{
+				return Forbid();
+			}
+
+			ChangeDetails? ChangeDetails = await PerforceService.GetChangeDetailsAsync(Stream.Name, Number, PerforceUser);
+			if(ChangeDetails == null)
+			{
+				return NotFound();
+			}
+
+			return PropertyFilter.Apply(new GetChangeDetailsResponse(ChangeDetails), Filter);
+		}
+
+		/// <summary>
+		/// Gets the history of a step in the stream
+		/// </summary>
+		/// <param name="StreamId">The stream id</param>
+		/// <param name="TemplateId"></param>
+		/// <param name="Step">Name of the step to search for</param>
+		/// <param name="Change">Maximum changelist number to return</param>
+		/// <param name="Count">Number of results to return</param>
+		/// <param name="Filter">The filter to apply to the results</param>
+		/// <returns>Http result code</returns>
+		[HttpGet]
+		[Route("/api/v1/streams/{StreamId}/history")]
+		[ProducesResponseType(typeof(List<GetJobStepRefResponse>), 200)]
+		public async Task<ActionResult<List<object>>> GetStepHistoryAsync(string StreamId, [FromQuery] string TemplateId, [FromQuery] string Step, [FromQuery] int? Change = null, [FromQuery] int Count = 10, [FromQuery] PropertyFilter? Filter = null)
+		{
+			StreamId StreamIdValue = new StreamId(StreamId);
+
+			IStream? Stream = await StreamService.GetStreamAsync(StreamIdValue);
+			if (Stream == null)
+			{
+				return NotFound();
+			}
+			if (!await StreamService.AuthorizeAsync(Stream, AclAction.ViewJob, User, null))
+			{
+				return Forbid();
+			}
+
+			TemplateRefId TemplateIdValue = new TemplateRefId(TemplateId);
+
+			List<IJobStepRef> Steps = await JobStepRefCollection.GetStepsForNodeAsync(StreamIdValue, TemplateIdValue, Step, Change, true, Count);
+			return Steps.ConvertAll(x => PropertyFilter.Apply(new GetJobStepRefResponse(x), Filter));
+		}
+
+		/// <summary>
+		/// Deletes a stream
+		/// </summary>
+		/// <param name="StreamId">Id of the stream to update.</param>
+		/// <returns>Http result code</returns>
+		[HttpDelete]
+		[Route("/api/v1/streams/{StreamId}")]
+		public async Task<ActionResult> DeleteStreamAsync(string StreamId)
+		{
+			StreamId StreamIdValue = new StreamId(StreamId);
+
+			IStream? Stream = await StreamService.GetStreamAsync(StreamIdValue);
+			if (Stream == null)
+			{
+				return NotFound();
+			}
+			if (!await StreamService.AuthorizeAsync(Stream, AclAction.DeleteStream, User, null))
+			{
+				return Forbid();
+			}
+
+			await StreamService.DeleteStreamAsync(StreamIdValue, JobService);
+			return new OkResult();
+		}
+
+		/// <summary>
+		/// Update a stream's properties.
+		/// </summary>
+		/// <param name="StreamId">Id of the stream to update.</param>
+		/// <param name="Request">Properties on the stream to update</param>
+		/// <returns>Http result code</returns>
+		[HttpPut]
+		[Route("/api/v1/streams/{StreamId}")]
+		public async Task<ActionResult> SetStreamAsync(string StreamId, [FromBody] SetStreamRequest Request)
+		{
+			ActionResult? ClaimError;
+			if (!TryAuthorizeTemplateClaims(Request.Templates, out ClaimError))
+			{
+				return ClaimError;
+			}
+
+			StreamId StreamIdValue = new StreamId(StreamId);
+			ProjectId ProjectIdValue = new ProjectId(Request.ProjectId);
+			for (; ;)
+			{
+				IStream? Stream = await StreamService.GetStreamAsync(StreamIdValue);
+				if (Stream == null)
+				{
+					IProject? Project = await ProjectService.GetProjectAsync(ProjectIdValue);
+					if (Project == null)
+					{
+						return BadRequest("Invalid project id");
+					}
+					if (!await ProjectService.AuthorizeAsync(Project, AclAction.CreateStream, User, null))
+					{
+						return Forbid();
+					}
+				}
+				else
+				{
+					if (ProjectIdValue != Stream.ProjectId)
+					{
+						return BadRequest("Cannot modify project id after a stream has been created");
+					}
+
+					ProjectPermissionsCache PermissionsCache = new ProjectPermissionsCache();
+					if (!await StreamService.AuthorizeAsync(Stream, AclAction.UpdateStream, User, PermissionsCache))
+					{
+						return Forbid();
+					}
+					if (Request.Acl != null && !await StreamService.AuthorizeAsync(Stream, AclAction.ChangePermissions, User, PermissionsCache))
+					{
+						return Forbid();
+					}
+				}
+
+				try
+				{
+					Stream = await SetStreamAsync(ProjectIdValue, StreamIdValue, Stream, null, Request, StreamService, TemplateService);
+				}
+				catch(InvalidStreamException Ex)
+				{
+					return BadRequest(Ex.Message);
+				}
+
+				if(Stream != null)
+				{
+					return Ok();
+				}
+			}
+		}
+
+		internal static async Task<IStream?> SetStreamAsync(ProjectId ProjectId, StreamId StreamId, IStream? Stream, int? ConfigChange, SetStreamRequest Request, StreamService StreamService, TemplateService TemplateService)
+		{
+			List<StreamTab>? Tabs = null;
+			if (Request.Tabs != null)
+			{
+				Tabs = Request.Tabs.ConvertAll(x => StreamTab.FromRequest(x));
+			}
+
+			Dictionary<string, AgentType>? AgentTypes = null;
+			if (Request.AgentTypes != null)
+			{
+				AgentTypes = Request.AgentTypes.Where(x => x.Value != null).ToDictionary(x => x.Key, x => new AgentType(x.Value!));
+			}
+
+			Dictionary<string, WorkspaceType>? WorkspaceTypes = null;
+			if (Request.WorkspaceTypes != null)
+			{
+				WorkspaceTypes = Request.WorkspaceTypes.Where(x => x.Value != null).ToDictionary(x => x.Key, x => new WorkspaceType(x.Value!));
+			}
+
+			Dictionary<string, string>? Properties = null;
+			if (Request.Properties != null)
+			{
+				Properties = Request.Properties.Where(x => x.Value != null).ToDictionary(x => x.Key, x => x.Value!);
+			}
+
+			Dictionary<TemplateRefId, TemplateRef>? TemplateRefs = null;
+			if (Request.Templates != null)
+			{
+				TemplateRefs = await CreateTemplateRefs(Request.Templates, Stream, TemplateService);
+			}
+
+			Acl? Acl = null;
+			if (Request.Acl != null)
+			{
+				Acl = Acl.Merge(null, Request.Acl);
+			}
+
+			DefaultPreflight? DefaultPreflight = Request.DefaultPreflight?.ToModel();
+			if (DefaultPreflight == null && Request.DefaultPreflightTemplate != null)
+			{
+				DefaultPreflight = new DefaultPreflight(new TemplateRefId(Request.DefaultPreflightTemplate), null);
+			}
+
+			if (Stream == null)
+			{
+				return await StreamService.TryCreateStreamAsync(StreamId, Request.Name, ProjectId, Request.ConfigPath, Request.Order, Request.NotificationChannel, Request.NotificationChannelFilter, Request.TriageChannel, DefaultPreflight, Tabs, AgentTypes, WorkspaceTypes, TemplateRefs, Properties, Acl);
+			}
+			else
+			{
+				return await StreamService.TryReplaceStreamAsync(Stream, Request.Name, Request.ConfigPath, ConfigChange, Request.Order, Request.NotificationChannel, Request.NotificationChannelFilter, Request.TriageChannel, DefaultPreflight, Tabs, AgentTypes, WorkspaceTypes, TemplateRefs, Properties, Acl, Request.PausedUntil, Request.PauseComment);
+			}
+		}
+
+		/// <summary>
+		/// Validates that the user is allowed to specify the given template refs
+		/// </summary>
+		/// <param name="Requests">List of template refs</param>
+		/// <param name="OutResult">On failure, receives information about the missing claim</param>
+		/// <returns>True if the refs are valid, false otherwise</returns>
+		bool TryAuthorizeTemplateClaims(List<CreateTemplateRefRequest>? Requests, [NotNullWhen(false)] out ActionResult? OutResult)
+		{
+			if (Requests != null)
+			{
+				foreach (CreateTemplateRefRequest Request in Requests)
+				{
+					if (Request.Schedule != null && Request.Schedule.Claims != null)
+					{
+						foreach (CreateAclClaimRequest Claim in Request.Schedule.Claims)
+						{
+							if (!User.HasClaim(Claim.Type, Claim.Value))
+							{
+								OutResult = BadRequest("User does not have the {Claim.Type}: {Claim.Value} claim");
+								return false;
+							}
+						}
+					}
+				}
+			}
+
+			OutResult = null;
+			return true;
+		}
+
+		/// <summary>
+		/// Creates a list of template refs from a set of request objects
+		/// </summary>
+		/// <param name="Requests">Request objects</param>
+		/// <param name="Stream">The current stream state</param>
+		/// <param name="TemplateService">The template service</param>
+		/// <returns>List of new template references</returns>
+		static async Task<Dictionary<TemplateRefId, TemplateRef>> CreateTemplateRefs(List<CreateTemplateRefRequest> Requests, IStream? Stream, TemplateService TemplateService)
+		{
+			Dictionary<TemplateRefId, TemplateRef> NewTemplateRefs = new Dictionary<TemplateRefId, TemplateRef>();
+			foreach (CreateTemplateRefRequest Request in Requests)
+			{
+				// Create the template
+				ITemplate NewTemplate = await TemplateService.CreateTemplateAsync(Request.Name, Request.Priority, Request.AllowPreflights, Request.InitialAgentType, Request.SubmitNewChange, Request.Counters, Request.Arguments, Request.Parameters.ConvertAll(x => x.ToModel()));
+
+				// Get an identifier for the new template ref
+				TemplateRefId NewTemplateRefId;
+				if (Request.Id != null)
+				{
+					NewTemplateRefId = new TemplateRefId(Request.Id);
+				}
+				else
+				{
+					NewTemplateRefId = TemplateRefId.Sanitize(Request.Name);
+				}
+
+				// Add it to the list
+				TemplateRef NewTemplateRef = new TemplateRef(NewTemplate, Request.ShowUgsBadges, Request.ShowUgsAlerts, Request.NotificationChannel, Request.NotificationChannelFilter, Request.TriageChannel, Request.Schedule?.ToModel(), Request.ChainedJobs?.ConvertAll(x => new ChainedJobTemplate(x)), Acl.Merge(null, Request.Acl));
+				if (Stream != null && Stream.Templates.TryGetValue(NewTemplateRefId, out TemplateRef? OldTemplateRef))
+				{
+					if (OldTemplateRef.Schedule != null && NewTemplateRef.Schedule != null)
+					{
+						NewTemplateRef.Schedule.CopyState(OldTemplateRef.Schedule);
+					}
+				}
+				NewTemplateRefs.Add(NewTemplateRefId, NewTemplateRef);
+			}
+			foreach (TemplateRef TemplateRef in NewTemplateRefs.Values)
+			{
+				if (TemplateRef.ChainedJobs != null)
+				{
+					foreach (ChainedJobTemplate ChainedJob in TemplateRef.ChainedJobs)
+					{
+						if (!NewTemplateRefs.ContainsKey(ChainedJob.TemplateRefId))
+						{
+							throw new InvalidDataException($"Invalid template ref id '{ChainedJob.TemplateRefId}");
+						}
+					}
+				}
+			}
+			return NewTemplateRefs;
+		}
+
+		/// <summary>
+		/// Gets all the templates for a stream
+		/// </summary>
+		/// <param name="StreamId">Unique id of the stream to query</param>
+		/// <param name="TemplateRefId">Unique id of the template to query</param>
+		/// <param name="Filter">Filter for properties to return</param>
+		/// <returns>Information about all the templates</returns>
+		[HttpGet]
+		[Route("/api/v1/streams/{StreamId}/templates/{TemplateRefId}")]
+		[ProducesResponseType(typeof(List<GetTemplateResponse>), 200)]
+		public async Task<ActionResult<object>> GetTemplateAsync(string StreamId, string TemplateRefId, [FromQuery] PropertyFilter? Filter = null)
+		{
+			StreamId StreamIdValue = new StreamId(StreamId);
+			TemplateRefId TemplateRefIdValue = new TemplateRefId(TemplateRefId);
+
+			IStream? Stream = await StreamService.GetStreamAsync(StreamIdValue);
+			if (Stream == null)
+			{
+				return NotFound();
+			}
+
+			TemplateRef? TemplateRef;
+			if (!Stream.Templates.TryGetValue(TemplateRefIdValue, out TemplateRef))
+			{
+				return NotFound();
+			}
+			if (!await StreamService.AuthorizeAsync(Stream, TemplateRef, AclAction.ViewTemplate, User, null))
+			{
+				return Forbid();
+			}
+
+			ITemplate? Template = await TemplateService.GetTemplateAsync(TemplateRef.Hash);
+			if(Template == null)
+			{
+				return NotFound();
+			}
+
+			return new GetTemplateResponse(Template).ApplyFilter(Filter);
+		}
+	}
+}
diff --git a/Engine/Source/Programs/Horde/HordeServer/Controllers/SubscriptionsController.cs b/Engine/Source/Programs/Horde/HordeServer/Controllers/SubscriptionsController.cs
new file mode 100644
--- /dev/null
+++ b/Engine/Source/Programs/Horde/HordeServer/Controllers/SubscriptionsController.cs
@@ -0,0 +1,188 @@
+using HordeServer.Api;
+using HordeServer.Collections;
+using HordeServer.Models;
+using HordeServer.Services;
+using HordeServer.Utilities;
+using Microsoft.AspNetCore.Authorization;
+using Microsoft.AspNetCore.Mvc;
+using MongoDB.Bson;
+using MongoDB.Driver;
+using System;
+using System.Collections.Generic;
+using System.Linq;
+using System.Threading.Tasks;
+
+namespace HordeServer.Controllers
+{
+	/// <summary>
+	/// Controller for the /api/v1/agents endpoint
+	/// </summary>
+	[ApiController]
+	[Authorize]
+	[Route("[controller]")]
+	public class SubscriptionsController : ControllerBase
+	{
+		/// <summary>
+		/// The ACL service singleton
+		/// </summary>
+		AclService AclService;
+
+		/// <summary>
+		/// Collection of subscription documents
+		/// </summary>
+		ISubscriptionCollection SubscriptionCollection;
+
+		/// <summary>
+		/// Constructor
+		/// </summary>
+		/// <param name="AclService">The acl service singleton</param>
+		/// <param name="SubscriptionCollection">The collection of subscription documents</param>
+		public SubscriptionsController(AclService AclService, ISubscriptionCollection SubscriptionCollection)
+		{
+			this.AclService = AclService;
+			this.SubscriptionCollection = SubscriptionCollection;
+		}
+
+		/// <summary>
+		/// Find subscriptions matching a criteria
+		/// </summary>
+		/// <param name="UserId">Name of the user</param>
+		/// <param name="Filter">Filter for properties to return</param>
+		/// <returns>List of subscriptions</returns>
+		[HttpGet]
+		[Route("/api/v1/subscriptions")]
+		[ProducesResponseType(typeof(List<GetSubscriptionResponse>), 200)]
+		public async Task<ActionResult<List<object>>> GetSubscriptionsAsync([FromQuery] string UserId, [FromQuery] PropertyFilter? Filter = null)
+		{
+			ObjectId UserIdValue;
+			if (!TryParseUserId(UserId, out UserIdValue))
+			{
+				return BadRequest("Invalid user id");
+			}
+			if (!await AclService.AuthorizeAsUserAsync(User, UserIdValue))
+			{
+				return Forbid();
+			}
+
+			List<ISubscription> Results = await SubscriptionCollection.FindSubscriptionsAsync(UserIdValue);
+			return Results.ConvertAll(x => PropertyFilter.Apply(new GetSubscriptionResponse(x), Filter));
+		}
+
+		/// <summary>
+		/// Find subscriptions matching a criteria
+		/// </summary>
+		/// <param name="SubscriptionId">The subscription id</param>
+		/// <param name="Filter">Filter for properties to return</param>
+		/// <returns>List of subscriptions</returns>
+		[HttpGet]
+		[Route("/api/v1/subscriptions/{SubscriptionId}")]
+		[ProducesResponseType(typeof(GetSubscriptionResponse), 200)]
+		public async Task<ActionResult<object>> GetSubscriptionAsync(string SubscriptionId, [FromQuery] PropertyFilter? Filter = null)
+		{
+			ISubscription? Subscription = await SubscriptionCollection.GetAsync(SubscriptionId);
+			if (Subscription == null)
+			{
+				return NotFound();
+			}
+			if (!await AclService.AuthorizeAsUserAsync(User, Subscription.UserId))
+			{
+				return Forbid();
+			}
+
+			return PropertyFilter.Apply(new GetSubscriptionResponse(Subscription), Filter);
+		}
+
+
+		/// <summary>
+		/// Remove a subscription
+		/// </summary>
+		/// <param name="SubscriptionId">The subscription id</param>
+		/// <returns>Async task</returns>
+		[HttpDelete]
+		[Route("/api/v1/subscriptions/{SubscriptionId}")]
+		[ProducesResponseType(typeof(List<GetSubscriptionResponse>), 200)]
+		public async Task<ActionResult> DeleteSubscriptionAsync(string SubscriptionId)
+		{
+			ISubscription? Subscription = await SubscriptionCollection.GetAsync(SubscriptionId);
+			if (Subscription == null)
+			{
+				return NotFound();
+			}
+			if (!await AclService.AuthorizeAsUserAsync(User, Subscription.UserId))
+			{
+				return Forbid();
+			}
+
+			await SubscriptionCollection.RemoveAsync(new[] { Subscription });
+			return Ok();
+		}
+
+		/// <summary>
+		/// Find subscriptions matching a criteria
+		/// </summary>
+		/// <param name="Subscriptions">The new subscriptions to create</param>
+		/// <returns>List of subscriptions</returns>
+		[HttpPost]
+		[Route("/api/v1/subscriptions")]
+		public async Task<ActionResult<List<CreateSubscriptionResponse>>> CreateSubscriptionsAsync(List<CreateSubscriptionRequest> Subscriptions)
+		{
+			HashSet<ObjectId> AuthorizedUsers = new HashSet<ObjectId>();
+
+			ObjectId? CurrentUserId = User.GetUserId();
+			if(CurrentUserId != null)
+			{
+				AuthorizedUsers.Add(CurrentUserId.Value);
+			}
+
+			GlobalPermissionsCache Cache = new GlobalPermissionsCache();
+
+			List<NewSubscription> NewSubscriptions = new List<NewSubscription>();
+			foreach (CreateSubscriptionRequest Subscription in Subscriptions)
+			{
+				ObjectId NewUserId;
+				if (!TryParseUserId(Subscription.UserId, out NewUserId))
+				{
+					return BadRequest($"Invalid user id: '{Subscription.UserId}'.");
+				}
+				if (AuthorizedUsers.Add(NewUserId) && !await AclService.AuthorizeAsync(AclAction.Impersonate, User, Cache))
+				{
+					return Forbid();
+				}
+				NewSubscriptions.Add(new NewSubscription(Subscription.Event, NewUserId, Subscription.NotificationType));
+			}
+
+			List<ISubscription> Results = await SubscriptionCollection.AddAsync(NewSubscriptions);
+			return Results.ConvertAll(x => new CreateSubscriptionResponse(x));
+		}
+
+		/// <summary>
+		/// Parse a user id from a string. Allows passing the user's name as well as their objectid value.
+		/// </summary>
+		/// <param name="UserName"></param>
+		/// <param name="ObjectId"></param>
+		/// <returns></returns>
+		bool TryParseUserId(string UserName, out ObjectId ObjectId)
+		{
+			ObjectId NewObjectId;
+			if (ObjectId.TryParse(UserName, out NewObjectId))
+			{
+				ObjectId = NewObjectId;
+				return true;
+			}
+
+			string? CurrentUserName = User.GetUserName();
+			if (CurrentUserName != null && String.Equals(UserName, CurrentUserName, StringComparison.OrdinalIgnoreCase))
+			{
+				ObjectId? CurrentUserId = User.GetUserId();
+				if (CurrentUserId != null)
+				{
+					ObjectId = CurrentUserId.Value;
+					return true;
+				}
+			}
+
+			ObjectId = default;
+			return false;
+		}
+	}
+}
\ No newline at end of file
diff --git a/Engine/Source/Programs/Horde/HordeServer/Controllers/SwarmController.cs b/Engine/Source/Programs/Horde/HordeServer/Controllers/SwarmController.cs
new file mode 100644
--- /dev/null
+++ b/Engine/Source/Programs/Horde/HordeServer/Controllers/SwarmController.cs
@@ -0,0 +1,43 @@
+using Amazon.Runtime.Internal.Util;
+using Microsoft.AspNetCore.Mvc;
+using Microsoft.Extensions.Logging;
+using System;
+using System.Collections.Generic;
+using System.Linq;
+using System.Threading.Tasks;
+
+namespace HordeServer.Controllers
+{
+	/// <summary>
+	/// Controller for the /api/v1/swarm endpoint
+	/// </summary>
+	[ApiController]
+	[Route("[controller]")]
+	public sealed class SwarmController : ControllerBase
+	{
+		/// <summary>
+		/// Logger instance
+		/// </summary>
+		ILogger<SwarmController> Logger;
+
+		/// <summary>
+		/// Constructor
+		/// </summary>
+		public SwarmController(ILogger<SwarmController> Logger)
+		{
+			this.Logger = Logger;
+		}
+
+		/// <summary>
+		/// Gets the latest version info
+		/// </summary>
+		/// <returns>Result code</returns>
+		[HttpPost]
+		[Route("/api/v1/swarm/reviews")]
+		public ActionResult AddReview(object Request)
+		{
+			Logger.LogInformation("Added Swarm review: {Request}", Request.ToString());
+			return Ok();
+		}
+	}
+}
diff --git a/Engine/Source/Programs/Horde/HordeServer/Controllers/TestDataController.cs b/Engine/Source/Programs/Horde/HordeServer/Controllers/TestDataController.cs
new file mode 100644
--- /dev/null
+++ b/Engine/Source/Programs/Horde/HordeServer/Controllers/TestDataController.cs
@@ -0,0 +1,146 @@
+// Copyright Epic Games, Inc. All Rights Reserved.
+
+using HordeServer.Api;
+using HordeServer.Collections;
+using HordeCommon;
+using HordeServer.Models;
+using HordeServer.Services;
+using HordeServer.Utilities;
+using Microsoft.AspNetCore.Authorization;
+using Microsoft.AspNetCore.Mvc;
+using Microsoft.AspNetCore.Server.Kestrel.Core.Features;
+using MongoDB.Bson;
+using MongoDB.Bson.Serialization;
+using System;
+using System.Collections.Generic;
+using System.ComponentModel.DataAnnotations;
+using System.IO;
+using System.Linq;
+using System.Threading.Tasks;
+
+using StreamId = HordeServer.Utilities.StringId<HordeServer.Models.IStream>;
+
+namespace HordeServer.Controllers
+{
+	/// <summary>
+	/// Controller for the /api/v1/testdata endpoint
+	/// </summary>
+	[ApiController]
+	[Authorize]
+	[Route("[controller]")]
+	public class TestDataController : ControllerBase
+	{
+		/// <summary>
+		/// Collection of job documents
+		/// </summary>
+		private JobService JobService;
+
+		/// <summary>
+		/// Collection of test data documents
+		/// </summary>
+		private readonly ITestDataCollection TestDataCollection;
+
+		/// <summary>
+		/// Constructor
+		/// </summary>
+		/// <param name="JobService">The job service singleton</param>
+		/// <param name="TestDataCollection">Collection of test data documents</param>
+		public TestDataController(JobService JobService, ITestDataCollection TestDataCollection)
+		{
+			this.JobService = JobService;
+			this.TestDataCollection = TestDataCollection;
+		}
+
+		/// <summary>
+		/// Creates a new TestData document
+		/// </summary>
+		/// <returns>The stream document</returns>
+		[HttpPost]
+		[Route("/api/v1/testdata")]
+		public async Task<ActionResult<CreateTestDataResponse>> CreateAsync(CreateTestDataRequest Request)
+		{
+			IJob? Job = await JobService.GetJobAsync(Request.JobId.ToObjectId());
+			if (Job == null)
+			{
+				return NotFound();
+			}
+			if (!await JobService.AuthorizeAsync(Job, AclAction.UpdateJob, User, null))
+			{
+				return Forbid();
+			}
+
+			IJobStep? JobStep;
+			if (!Job.TryGetStep(Request.StepId.ToSubResourceId(), out JobStep))
+			{
+				return NotFound();
+			}
+
+			ITestData TestData = await TestDataCollection.AddAsync(Job, JobStep, Request.Key, new BsonDocument(Request.Data));
+			return new CreateTestDataResponse(TestData.Id.ToString());
+		}
+
+		/// <summary>
+		/// Searches for test data that matches a set of criteria
+		/// </summary>
+		/// <param name="StreamId">The stream id</param>
+		/// <param name="MinChange">The minimum changelist number to return (inclusive)</param>
+		/// <param name="MaxChange">The maximum changelist number to return (inclusive)</param>
+		/// <param name="JobId">The job id</param>
+		/// <param name="JobStepId">The unique step id</param>
+		/// <param name="Key">Key identifying the result to return</param>
+		/// <param name="Index">Offset within the results to return</param>
+		/// <param name="Count">Number of results to return</param>
+		/// <param name="Filter">Filter for properties to return</param>
+		/// <returns>The stream document</returns>
+		[HttpGet]
+		[Route("/api/v1/testdata")]
+		[ProducesResponseType(typeof(List<GetTestDataResponse>), 200)]
+		public async Task<ActionResult<List<object>>> FindTestDataAsync([FromQuery] string? StreamId = null, [FromQuery] int? MinChange = null, [FromQuery] int? MaxChange = null, string? JobId = null, string? JobStepId = null, string? Key = null, int Index = 0, int Count = 10, PropertyFilter? Filter = null)
+		{
+			StreamId? StreamIdValue = null;
+			if(StreamId != null)
+			{
+				StreamIdValue = new StreamId(StreamId);
+			}
+
+			JobPermissionsCache Cache = new JobPermissionsCache();
+
+			List<object> Results = new List<object>();
+
+			List<ITestData> Documents = await TestDataCollection.FindAsync(StreamIdValue, MinChange, MaxChange, JobId?.ToObjectId(), JobStepId?.ToSubResourceId(), Key, Index, Count);
+			foreach (ITestData Document in Documents)
+			{
+				if (await JobService.AuthorizeAsync(Document.JobId, AclAction.ViewJob, User, Cache))
+				{
+					Results.Add(PropertyFilter.Apply(new GetTestDataResponse(Document), Filter));
+				}
+			}
+
+			return Results;
+		}
+
+		/// <summary>
+		/// Retrieve information about a specific issue
+		/// </summary>
+		/// <param name="TestDataId">Id of the document to get information about</param>
+		/// <param name="Filter">Filter for the properties to return</param>
+		/// <returns>List of matching agents</returns>
+		[HttpGet]
+		[Route("/api/v1/testdata/{TestDataId}")]
+		[ProducesResponseType(typeof(GetTestDataResponse), 200)]
+		public async Task<ActionResult<object>> GetTestDataAsync(string TestDataId, [FromQuery] PropertyFilter? Filter = null)
+		{
+			ITestData? TestData = await TestDataCollection.GetAsync(TestDataId.ToObjectId());
+			if (TestData == null)
+			{
+				return NotFound();
+			}
+			if (!await JobService.AuthorizeAsync(TestData.JobId, AclAction.ViewJob, User, null))
+			{
+				return Forbid();
+			}
+
+			return PropertyFilter.Apply(new GetTestDataResponse(TestData), Filter);
+		}
+	}
+}
diff --git a/Engine/Source/Programs/Horde/HordeServer/Controllers/UgsController.cs b/Engine/Source/Programs/Horde/HordeServer/Controllers/UgsController.cs
new file mode 100644
--- /dev/null
+++ b/Engine/Source/Programs/Horde/HordeServer/Controllers/UgsController.cs
@@ -0,0 +1,389 @@
+// Copyright Epic Games, Inc. All Rights Reserved.
+
+using HordeServer.Api;
+using HordeServer.Collections;
+using HordeCommon;
+using HordeServer.Models;
+using HordeServer.Services;
+using HordeServer.Utilities;
+using Microsoft.AspNetCore.Authorization;
+using Microsoft.AspNetCore.Mvc;
+using Microsoft.AspNetCore.Server.Kestrel.Core.Features;
+using Microsoft.Extensions.Caching.Memory;
+using Microsoft.Extensions.Options;
+using MongoDB.Bson;
+using Serilog.Events;
+using System;
+using System.Collections.Generic;
+using System.ComponentModel.DataAnnotations;
+using System.IO;
+using System.Linq;
+using System.Threading.Tasks;
+
+using StreamId = HordeServer.Utilities.StringId<HordeServer.Models.IStream>;
+using Microsoft.Extensions.Logging;
+
+namespace HordeServer.Controllers
+{
+	/// <summary>
+	/// Controller for the /api/v1/issues endpoint
+	/// </summary>
+	[ApiController]
+	[Route("[controller]")]
+	public sealed class UgsController : ControllerBase
+	{
+		/// <summary>
+		/// Singleton instance of the issue service
+		/// </summary>
+		private readonly IIssueService IssueService;
+
+		/// <summary>
+		/// Singleton instance of the job service
+		/// </summary>
+		private readonly JobService JobService;
+
+		/// <summary>
+		/// Collection of stream documents
+		/// </summary>
+		private readonly StreamService StreamService;
+
+		/// <summary>
+		/// Collection of metadata documents
+		/// </summary>
+		private readonly IUgsMetadataCollection UgsMetadataCollection;
+
+		/// <summary>
+		/// Collection of log events
+		/// </summary>
+		private readonly ILogEventCollection LogEventCollection;
+
+		/// <summary>
+		/// The log file service
+		/// </summary>
+		private readonly ILogFileService LogFileService;
+
+		/// <summary>
+		/// Server settings
+		/// </summary>
+		private readonly ServerSettings Settings;
+
+		/// <summary>
+		/// Logger 
+		/// </summary>
+		ILogger Logger;
+
+		/// <summary>
+		/// Constructor
+		/// </summary>
+		/// <param name="IssueService">The issue service</param>
+		/// <param name="JobService">The job service</param>
+		/// <param name="StreamService">Collection of stream documents</param>
+		/// <param name="UgsMetadataCollection">Collection of UGS metadata documents</param>
+		/// <param name="LogEventCollection">Collection of log event documents</param>
+		/// <param name="LogFileService">The log file service</param>
+		/// <param name="OptionsMonitor">The server settings</param>
+		/// <param name="Logger">Logger</param>
+		public UgsController(IIssueService IssueService, JobService JobService, StreamService StreamService, IUgsMetadataCollection UgsMetadataCollection, ILogEventCollection LogEventCollection, ILogFileService LogFileService, IOptionsMonitor<ServerSettings> OptionsMonitor, ILogger<UgsController> Logger)
+		{
+			this.IssueService = IssueService;
+			this.JobService = JobService;
+			this.StreamService = StreamService;
+			this.UgsMetadataCollection = UgsMetadataCollection;
+			this.LogEventCollection = LogEventCollection;
+			this.LogFileService = LogFileService;
+			this.Settings = OptionsMonitor.CurrentValue;
+			this.Logger = Logger;
+		}
+
+		/// <summary>
+		/// Gets the latest version info
+		/// </summary>
+		/// <returns>Result code</returns>
+		[HttpGet]
+		[Route("/ugs/api/latest")]
+		public ActionResult<object> GetLatest()
+		{
+			return new { Version = 2 };
+		}
+
+		/// <summary>
+		/// Adds new metadata to the database
+		/// </summary>
+		/// <param name="Request">Request object</param>
+		/// <returns>Result code</returns>
+		[HttpPost]
+		[Route("/ugs/api/metadata")]
+		public async Task<ActionResult> AddMetadataAsync(AddUgsMetadataRequest Request)
+		{
+			IUgsMetadata Metadata = await UgsMetadataCollection.FindOrAddAsync(Request.Stream, Request.Change, Request.Project);
+			if (Request.Synced != null || Request.Vote != null || Request.Investigating != null || Request.Starred != null || Request.Comment != null)
+			{
+				if (Request.UserName == null)
+				{
+					return BadRequest("Missing UserName field on request body");
+				}
+				Metadata = await UgsMetadataCollection.UpdateUserAsync(Metadata, Request.UserName, Request.Synced, Request.Vote, Request.Investigating, Request.Starred, Request.Comment);
+			}
+			if (Request.Badges != null)
+			{
+				foreach (AddUgsBadgeRequest Badge in Request.Badges)
+				{
+					Metadata = await UgsMetadataCollection.UpdateBadgeAsync(Metadata, Badge.Name, Badge.Url, Badge.State);
+				}
+			}
+			return Ok();
+		}
+
+		/// <summary>
+		/// Searches for metadata updates
+		/// </summary>
+		/// <param name="Stream">THe stream to search for</param>
+		/// <param name="MinChange">Minimum changelist number</param>
+		/// <param name="MaxChange">Maximum changelist number</param>
+		/// <param name="Project">The project identifiers to search for</param>
+		/// <param name="Sequence">Last sequence number</param>
+		/// <returns>List of metadata updates</returns>
+		[HttpGet]
+		[Route("/ugs/api/metadata")]
+		public async Task<GetUgsMetadataListResponse> FindMetadataAsync([FromQuery] string Stream, [FromQuery] int MinChange, [FromQuery] int? MaxChange = null, [FromQuery] string? Project = null, [FromQuery] long? Sequence = null)
+		{
+			List<IUgsMetadata> MetadataList = await UgsMetadataCollection.FindAsync(Stream, MinChange, MaxChange, Sequence);
+		
+			GetUgsMetadataListResponse Response = new GetUgsMetadataListResponse();
+			if(Sequence != null)
+			{
+				Response.SequenceNumber = Sequence.Value;
+			}
+
+			foreach (IUgsMetadata Metadata in MetadataList)
+			{
+				if (Metadata.UpdateTicks > Response.SequenceNumber)
+				{
+					Response.SequenceNumber = Metadata.UpdateTicks;
+				}
+				if (String.IsNullOrEmpty(Metadata.Project) || Metadata.Project.Equals(Project, StringComparison.OrdinalIgnoreCase))
+				{
+					Response.Items.Add(new GetUgsMetadataResponse(Metadata));
+				}
+			}
+			return Response;
+		}
+
+		/// <summary>
+		/// Retrieve information about open issues
+		/// </summary>
+		/// <param name="User"></param>
+		/// <param name="IncludeResolved">Whether to include resolved issues</param>
+		/// <param name="MaxResults">Maximum number of results to return</param>
+		/// <returns>List of matching agents</returns>
+		[HttpGet]
+		[Route("/ugs/api/issues")]
+		[ProducesResponseType(typeof(GetUgsIssueResponse), 200)]
+		[System.Diagnostics.CodeAnalysis.SuppressMessage("Design", "CA1801:Remove unused parameter", Justification = "<Pending>")]
+		[System.Diagnostics.CodeAnalysis.SuppressMessage("Style", "IDE0060:Remove unused parameter", Justification = "<Pending>")]
+		public async Task<ActionResult<List<GetUgsIssueResponse>>> GetIssuesAsync([FromQuery] string? User = null, [FromQuery] bool IncludeResolved = false, [FromQuery] int MaxResults = 100)
+		{
+			List<GetUgsIssueResponse> Responses = new List<GetUgsIssueResponse>();
+			if (IncludeResolved)
+			{
+				List<IIssue> Issues = await IssueService.FindIssuesAsync(null, Resolved: null, Count: MaxResults);
+				foreach(IIssue Issue in Issues)
+				{
+					IIssueDetails Details = await IssueService.GetIssueDetailsAsync(Issue);
+					bool bNotify = User != null && Issue.Suspects.Any(x => x.Author.Equals(User, StringComparison.OrdinalIgnoreCase));
+					Responses.Add(CreateIssueResponse(Details, bNotify));
+				}
+			}
+			else
+			{
+				foreach (IIssueDetails CachedOpenIssue in IssueService.CachedOpenIssues)
+				{
+					if (Responses.Count >= MaxResults)
+					{
+						break;
+					}
+
+					if (CachedOpenIssue.ShowNotifications())
+					{
+						bool bNotify = User != null && CachedOpenIssue.IncludeForUser(User);
+						Responses.Add(CreateIssueResponse(CachedOpenIssue, bNotify));
+					}
+				}
+			}
+			return Responses;
+		}
+
+		/// <summary>
+		/// Retrieve information about a specific issue
+		/// </summary>
+		/// <param name="IssueId">Id of the issue to get information about</param>
+		/// <param name="Filter">Filter for the properties to return</param>
+		/// <returns>List of matching agents</returns>
+		[HttpGet]
+		[Route("/ugs/api/issues/{IssueId}")]
+		[ProducesResponseType(typeof(GetUgsIssueBuildResponse), 200)]
+		public async Task<ActionResult<object>> GetIssueAsync(int IssueId, [FromQuery] PropertyFilter? Filter = null)
+		{
+			IIssueDetails? Issue = await IssueService.GetIssueDetailsAsync(IssueId);
+			if (Issue == null)
+			{
+				return NotFound();
+			}
+
+			return PropertyFilter.Apply(CreateIssueResponse(Issue, false), Filter);
+		}
+
+		/// <summary>
+		/// Retrieve information about builds for a specific issue
+		/// </summary>
+		/// <param name="IssueId">Id of the issue to get information about</param>
+		/// <param name="Filter">Filter for the properties to return</param>
+		/// <returns>List of matching agents</returns>
+		[HttpGet]
+		[Route("/ugs/api/issues/{IssueId}/builds")]
+		[ProducesResponseType(typeof(List<GetUgsIssueBuildResponse>), 200)]
+		public async Task<ActionResult<List<object>>> GetIssueBuildsAsync(int IssueId, [FromQuery] PropertyFilter? Filter = null)
+		{
+			IIssueDetails? Issue = await IssueService.GetCachedIssueDetailsAsync(IssueId);
+			if (Issue == null)
+			{
+				return NotFound();
+			}
+
+			List<GetUgsIssueBuildResponse> Responses = new List<GetUgsIssueBuildResponse>();
+			foreach (IIssueSpan Span in Issue.Spans)
+			{
+				if (Span.LastSuccess != null)
+				{
+					Responses.Add(CreateBuildResponse(Span, Span.LastSuccess, IssueBuildOutcome.Success));
+				}
+				foreach (IIssueStep Step in Issue.Steps)
+				{
+					Responses.Add(CreateBuildResponse(Span, Step, IssueBuildOutcome.Error));
+				}
+				if (Span.NextSuccess != null)
+				{
+					Responses.Add(CreateBuildResponse(Span, Span.NextSuccess, IssueBuildOutcome.Success));
+				}
+			}
+
+			return Responses.ConvertAll(x => PropertyFilter.Apply(x, Filter));
+		}
+
+		/// <summary>
+		/// Retrieve information about builds for a specific issue
+		/// </summary>
+		/// <param name="IssueId">Id of the issue to get information about</param>
+		/// <returns>List of matching agents</returns>
+		[HttpGet]
+		[Route("/ugs/api/issues/{IssueId}/diagnostics")]
+		[ProducesResponseType(typeof(List<GetUgsIssueDiagnosticResponse>), 200)]
+		public async Task<ActionResult<List<GetUgsIssueDiagnosticResponse>>> GetIssueDiagnosticsAsync(int IssueId)
+		{
+			List<GetUgsIssueDiagnosticResponse> Diagnostics = new List<GetUgsIssueDiagnosticResponse>();
+
+			Dictionary<ObjectId, ILogFile?> LogFiles = new Dictionary<ObjectId, ILogFile?>();
+
+			JobPermissionsCache PermissionsCache = new JobPermissionsCache();
+
+			List<ILogEvent> Events = await IssueService.FindEventsForIssueAsync(IssueId, Count: 10);
+			foreach (ILogEvent Event in Events)
+			{
+				ILogFile? LogFile;
+				if(!LogFiles.TryGetValue(Event.LogId, out LogFile))
+				{
+					LogFile = await LogFileService.GetLogFileAsync(Event.LogId);
+					LogFiles.Add(Event.LogId, LogFile);
+				}
+				if (LogFile != null)
+				{
+					ILogEventData EventData = await LogFileService.GetEventDataAsync(LogFile, Event.LineIndex, Event.LineCount);
+					long BuildId = Event.LogId.GetHashCode();
+					Uri Url = new Uri(Settings.DashboardUrl, $"log/{Event.LogId}?lineindex={Event.LineIndex}");
+
+					GetUgsIssueDiagnosticResponse Diagnostic = new GetUgsIssueDiagnosticResponse(BuildId, EventData.Message, Url);
+					Diagnostics.Add(Diagnostic);
+				}
+			}
+
+			return Diagnostics;
+		}
+
+		/// <summary>
+		/// Gets the URL for a failing step in the
+		/// </summary>
+		/// <param name="Issue">The issue to get a URL for</param>
+		/// <param name="bNotify">Whether to show notifications for this issue</param>
+		/// <returns>The issue response</returns>
+		GetUgsIssueResponse CreateIssueResponse(IIssueDetails Issue, bool bNotify)
+		{
+			Uri? BuildUrl = GetIssueBuildUrl(Issue);
+			return new GetUgsIssueResponse(Issue, bNotify, BuildUrl);
+		}
+
+		/// <summary>
+		/// Gets the URL for a failing step in the given issue
+		/// </summary>
+		/// <param name="Issue">The issue details</param>
+		/// <returns>The build URL</returns>
+		Uri? GetIssueBuildUrl(IIssueDetails Issue)
+		{
+			HashSet<ObjectId> UnresolvedSpans = new HashSet<ObjectId>(Issue.Spans.Where(x => x.NextSuccess == null).Select(x => x.Id));
+
+			IIssueStep? Step = Issue.Steps.OrderByDescending(x => UnresolvedSpans.Contains(x.SpanId)).ThenByDescending(x => x.Change).FirstOrDefault();
+			if (Step == null)
+			{
+				return null;
+			}
+
+			return new Uri(Settings.DashboardUrl, $"job/{Step.JobId}?step={Step.StepId}");
+		}
+
+		/// <summary>
+		/// Creates the response for a particular build
+		/// </summary>
+		/// <param name="Span">Span containing the step</param>
+		/// <param name="Step">The step to describe</param>
+		/// <param name="Outcome">Outcome of this step</param>
+		/// <returns>Response object</returns>
+		GetUgsIssueBuildResponse CreateBuildResponse(IIssueSpan Span, IIssueStep Step, IssueBuildOutcome Outcome)
+		{
+			GetUgsIssueBuildResponse Response = new GetUgsIssueBuildResponse(Span.StreamName, Step.Change, Outcome);
+			Response.Id = Step.LogId.GetHashCode();
+			Response.JobName = $"{Step.JobName}: {Span.NodeName}";
+			Response.JobUrl = new Uri(Settings.DashboardUrl, $"job/{Step.JobId}");
+			Response.JobStepName = Span.NodeName;
+			Response.JobStepUrl = new Uri(Settings.DashboardUrl, $"job/{Step.JobId}?step={Step.StepId}");
+			Response.ErrorUrl = Response.JobStepUrl;
+			return Response;
+		}
+
+		/// <summary>
+		/// Update an issue
+		/// </summary>
+		/// <param name="IssueId">Id of the issue to get information about</param>
+		/// <param name="Request">The update information</param>
+		/// <returns>List of matching agents</returns>
+		[HttpPut]
+		[Route("/ugs/api/issues/{IssueId}")]
+		public async Task<ActionResult> UpdateIssueAsync(int IssueId, [FromBody] UpdateUgsIssueRequest Request)
+		{
+			if (!await IssueService.UpdateIssueAsync(IssueId, null, Request.Owner, Request.NominatedBy, Request.Acknowledged, Request.DeclinedBy, Request.FixChange, Request.Resolved))
+			{
+				return NotFound();
+			}
+			return Ok();
+		}
+
+		/// <summary>
+		/// Post information about net core installation. 
+		/// </summary>
+		[HttpPost]
+		[Route("/ugs/api/netcore")]
+		public ActionResult<object> PostNetCoreInfo(string? User = null, string? Machine = null, bool NetCore = false)
+		{
+			Logger.LogInformation("NetCore: User={User}, Machine={Machine}, NetCore={NetCore}", User, Machine, NetCore);
+			return new object();
+		}
+	}
+}
diff --git a/Engine/Source/Programs/Horde/HordeServer/Controllers/UserController.cs b/Engine/Source/Programs/Horde/HordeServer/Controllers/UserController.cs
new file mode 100644
--- /dev/null
+++ b/Engine/Source/Programs/Horde/HordeServer/Controllers/UserController.cs
@@ -0,0 +1,74 @@
+using HordeServer.Api;
+using HordeServer.Collections;
+using HordeServer.Models;
+using HordeServer.Utilities;
+using Microsoft.AspNetCore.Authorization;
+using Microsoft.AspNetCore.Mvc;
+using MongoDB.Bson;
+using System;
+using System.Collections.Generic;
+using System.Linq;
+using System.Threading.Tasks;
+
+namespace HordeServer.Controllers
+{
+	/// <summary>
+	/// Controller for the /api/v1/user endpoint
+	/// </summary>
+	[ApiController]
+	[Authorize]
+	[Route("[controller]")]
+	public class UserController : ControllerBase
+	{
+		/// <summary>
+		/// The user collection instance
+		/// </summary>
+		IUserCollection UserCollection { get; set; }
+
+		/// <summary>
+		/// Constructor
+		/// </summary>
+		/// <param name="UserCollection"></param>
+		public UserController(IUserCollection UserCollection)
+		{
+			this.UserCollection = UserCollection;
+		}
+
+		/// <summary>
+		/// Gets information about the logged in user
+		/// </summary>
+		/// <returns>Http result code</returns>
+		[HttpGet]
+		[Route("/api/v1/user")]
+		[ProducesResponseType(typeof(List<GetUserResponse>), 200)]
+		public async Task<ActionResult<object>> GetUserAsync([FromQuery] PropertyFilter? Filter = null)
+		{
+			IUser? InternalUser = await UserCollection.GetUserAsync(User);
+			if (InternalUser == null)
+			{
+				return NotFound();
+			}
+
+			IUserClaims? Claims = await UserCollection.GetClaimsAsync(InternalUser.Id);
+			IUserSettings? Settings = await UserCollection.GetSettingsAsync(InternalUser.Id);
+			return PropertyFilter.Apply(new GetUserResponse(InternalUser, Claims, Settings), Filter);
+		}
+
+		/// <summary>
+		/// Updates the logged in user
+		/// </summary>
+		/// <returns>Http result code</returns>
+		[HttpPut]
+		[Route("/api/v1/user")]
+		public async Task<ActionResult> UpdateUserAsync(UpdateUserRequest Request)
+		{
+			ObjectId? UserId = User.GetUserId();
+			if(UserId == null)
+			{
+				return BadRequest("Current user does not have a registered profile");
+			}
+			await UserCollection.UpdateSettingsAsync(UserId.Value, Request.EnableExperimentalFeatures, Request.EnableIssueNotifications, Request.DashboardSettings?.ToBsonValue(), Request.AddPinnedJobIds?.Select(x => new ObjectId(x)), Request.RemovePinnedJobIds?.Select(x => new ObjectId(x)));
+			return Ok();
+		}
+	}
+}
diff --git a/Engine/Source/Programs/Horde/HordeServer/GlobalSuppressions.cs b/Engine/Source/Programs/Horde/HordeServer/GlobalSuppressions.cs
new file mode 100644
--- /dev/null
+++ b/Engine/Source/Programs/Horde/HordeServer/GlobalSuppressions.cs
@@ -0,0 +1,40 @@
+// This file is used by Code Analysis to maintain SuppressMessage
+// attributes that are applied to this project.
+// Project-level suppressions either have no target or are given
+// a specific target and scoped to a namespace, type, member, etc.
+
+using System.Diagnostics.CodeAnalysis;
+
+[assembly: System.Diagnostics.CodeAnalysis.SuppressMessage("Reliability", "CA2007:Consider calling ConfigureAwait on the awaited task", Justification = "Not necessary in NET Core; no SynchronizationContext.", Scope = "module")]
+[assembly: System.Diagnostics.CodeAnalysis.SuppressMessage("Usage", "CA2227:Collection properties should be read only", Justification = "<Pending>", Scope = "module")]
+[assembly: System.Diagnostics.CodeAnalysis.SuppressMessage("Design", "CA1062:Validate arguments of public methods", Justification = "<Pending>", Scope = "module")]
+[assembly: System.Diagnostics.CodeAnalysis.SuppressMessage("Globalization", "CA1303:Do not pass literals as localized parameters", Justification = "<Pending>", Scope = "module")]
+[assembly: System.Diagnostics.CodeAnalysis.SuppressMessage("Performance", "CA1819:Properties should not return arrays", Justification = "<Pending>", Scope = "module")]
+[assembly: System.Diagnostics.CodeAnalysis.SuppressMessage("Naming", "CA1716:Identifiers should not match keywords", Justification = "<Pending>", Scope = "type", Target = "~T:HordeServer.Models.Event")]
+[assembly: System.Diagnostics.CodeAnalysis.SuppressMessage("Design", "CA1032:Implement standard exception constructors", Justification = "<Pending>", Scope = "type", Target = "~T:HordeServer.Services.RetryNotAllowedException")]
+[assembly: System.Diagnostics.CodeAnalysis.SuppressMessage("Naming", "CA1714:Flags enums should have plural names", Justification = "<Pending>", Scope = "type", Target = "~T:HordeServer.Api.ScheduleDaysOfWeek")]
+[assembly: System.Diagnostics.CodeAnalysis.SuppressMessage("Naming", "CA1720:Identifier contains type name", Justification = "<Pending>", Scope = "member", Target = "~F:HordeServer.Api.TemplateParameterType.String")]
+
+[assembly: System.Diagnostics.CodeAnalysis.SuppressMessage("Build", "CA1822:Member Login does not access instance data and can be marked as static (Shared in VisualBasic)", Justification = "<Pending>", Scope = "member", Target = "~M:HordeServer.Controllers.AccountController.Login~Microsoft.AspNetCore.Mvc.IActionResult")]
+[assembly: System.Diagnostics.CodeAnalysis.SuppressMessage("Build", "CA1031:Modify 'Logout' to catch a more specific allowed exception type, or rethrow the exception.", Justification = "<Pending>", Scope = "member", Target = "~M:HordeServer.Controllers.AccountController.Logout~System.Threading.Tasks.Task{Microsoft.AspNetCore.Mvc.IActionResult}")]
+[assembly: System.Diagnostics.CodeAnalysis.SuppressMessage("Build", "CA1305:The behavior of 'string.ToString()' could vary based on the current user's locale settings. Replace this call in 'AgentLease.ToResponse()' with a call to 'string.ToString(IFormatProvider)'.", Justification = "<Pending>", Scope = "member", Target = "~M:HordeServer.Models.AgentLease.ToResponse~HordeServer.Api.GetAgentLeaseResponse")]
+[assembly: System.Diagnostics.CodeAnalysis.SuppressMessage("Build", "CA5380:Adding certificates to the operating system's trusted root certificates increases the risk of incorrectly authenticating an illegitimate certificate", Justification = "<Pending>", Scope = "member", Target = "~T:HordeServer.Services.DatabaseService")]
+[assembly: System.Diagnostics.CodeAnalysis.SuppressMessage("Build", "CA1724:The type name Template conflicts in whole or in part with the namespace name 'Microsoft.AspNetCore.Routing.Template'. Change either name to eliminate the conflict.", Justification = "<Pending>", Scope = "type", Target = "~T:HordeServer.Models.Template")]
+[assembly: System.Diagnostics.CodeAnalysis.SuppressMessage("Build", "CA1724:The type name Session conflicts in whole or in part with the namespace name 'Microsoft.AspNetCore.Session'. Change either name to eliminate the conflict.", Justification = "<Pending>", Scope = "type", Target = "~T:HordeServer.Models.Session")]
+[assembly: System.Diagnostics.CodeAnalysis.SuppressMessage("Build", "CA1812:BotsService is an internal class that is apparently never instantiated. If so, remove the code from the assembly. If this class is intended to contain only static members, make it static (Shared in Visual Basic).", Justification = "<Pending>", Scope = "type", Target = "~T:HordeServer.Services.BotsService")]
+[assembly: System.Diagnostics.CodeAnalysis.SuppressMessage("Build", "CA1812:RpcService is an internal class that is apparently never instantiated. If so, remove the code from the assembly. If this class is intended to contain only static members, make it static (Shared in Visual Basic).", Justification = "<Pending>", Scope = "type", Target = "~T:HordeServer.Services.RpcService")]
+[assembly: System.Diagnostics.CodeAnalysis.SuppressMessage("Build", "CA1812:Startup is an internal class that is apparently never instantiated. If so, remove the code from the assembly. If this class is intended to contain only static members, make it static (Shared in Visual Basic).", Justification = "<Pending>", Scope = "type", Target = "~T:HordeServer.Startup")]
+[assembly: System.Diagnostics.CodeAnalysis.SuppressMessage("Build", "CA1812:Startup.AgentJwtBearerHandler is an internal class that is apparently never instantiated. If so, remove the code from the assembly. If this class is intended to contain only static members, make it static (Shared in Visual Basic).", Justification = "<Pending>", Scope = "type", Target = "~T:HordeServer.Startup.AgentJwtBearerHandler")]
+[assembly: System.Diagnostics.CodeAnalysis.SuppressMessage("Build", "CA1812:Startup.AnonymousAuthenticationHandler is an internal class that is apparently never instantiated. If so, remove the code from the assembly. If this class is intended to contain only static members, make it static (Shared in Visual Basic).", Justification = "<Pending>", Scope = "type", Target = "~T:HordeServer.Startup.AnonymousAuthenticationHandler")]
+[assembly: System.Diagnostics.CodeAnalysis.SuppressMessage("Performance", "CA1822:Mark members as static", Justification = "<Pending>", Scope = "member", Target = "~M:HordeServer.Controllers.DebugController.Environment~Microsoft.AspNetCore.Mvc.IActionResult")]
+[assembly: System.Diagnostics.CodeAnalysis.SuppressMessage("Naming", "CA1716:Identifiers should not match keywords", Justification = "<Pending>", Scope = "module")]
+[assembly: SuppressMessage("Design", "CA1031:Do not catch general exception types", Justification = "<Pending>", Scope = "member", Target = "~M:HordeServer.Services.AgentService.ExecuteAsync(System.Threading.CancellationToken)~System.Threading.Tasks.Task")]
+[assembly: SuppressMessage("Design", "CA1031:Do not catch general exception types", Justification = "<Pending>", Scope = "member", Target = "~M:HordeServer.Utilities.PeriodicBackgroundService.ExecuteAsync(System.Threading.CancellationToken)~System.Threading.Tasks.Task")]
+[assembly: SuppressMessage("Performance", "CA1815:Override equals and operator equals on value types", Justification = "<Pending>", Scope = "type", Target = "~T:HordeServer.Models.JobStepRefId")]
+[assembly: SuppressMessage("Style", "IDE0054:Use compound assignment", Justification = "<Pending>", Scope = "member", Target = "~M:HordeServer.Models.Job.GetTimingInfo~System.Collections.Generic.Dictionary{HordeServer.Models.Node,HordeServer.Models.TimingInfo}")]
+[assembly: SuppressMessage("Design", "CA1031:Do not catch general exception types", Justification = "<Pending>", Scope = "member", Target = "~M:HordeServer.Startup.DashboardBearerHandler.HandleAuthenticateAsync~System.Threading.Tasks.Task{Microsoft.AspNetCore.Authentication.AuthenticateResult}")]
+[assembly: SuppressMessage("Design", "CA1031:Do not catch general exception types", Justification = "<Pending>", Scope = "member", Target = "~M:HordeServer.Services.ConformService.ExecuteAsync(System.Threading.CancellationToken)~System.Threading.Tasks.Task")]
+[assembly: SuppressMessage("Design", "CA1031:Do not catch general exception types", Justification = "<Pending>", Scope = "member", Target = "~M:HordeServer.Services.ScheduleService.ExecuteAsync(System.Threading.CancellationToken)~System.Threading.Tasks.Task")]
+[assembly: SuppressMessage("Design", "CA1031:Do not catch general exception types", Justification = "<Pending>", Scope = "member", Target = "~M:HordeServer.Services.ScheduleService.UpdateSchedulesAsync~System.Threading.Tasks.Task{System.DateTimeOffset}")]
+[assembly: SuppressMessage("Design", "CA1030:Use events where appropriate", Justification = "<Pending>", Scope = "module")]
+[assembly: SuppressMessage("Design", "CA1031:Do not catch general exception types", Justification = "<Pending>", Scope = "module")]
diff --git a/Engine/Source/Programs/Horde/HordeServer/HordeServer.csproj b/Engine/Source/Programs/Horde/HordeServer/HordeServer.csproj
new file mode 100644
--- /dev/null
+++ b/Engine/Source/Programs/Horde/HordeServer/HordeServer.csproj
@@ -0,0 +1,78 @@
+<Project Sdk="Microsoft.NET.Sdk.Web">
+
+  <PropertyGroup>
+    <TargetFramework>netcoreapp3.1</TargetFramework>
+    <Nullable>enable</Nullable>
+    <Deterministic>true</Deterministic>
+    <AssemblyName>HordeServer</AssemblyName>
+    <DocumentationFile>HordeServer.xml</DocumentationFile>
+  </PropertyGroup>
+
+  <ItemGroup>
+    <PackageReference Include="AWSSDK.Core" Version="3.5.1.49" />
+    <PackageReference Include="AWSSDK.EC2" Version="3.5.20" />
+    <PackageReference Include="AWSSDK.Extensions.NETCore.Setup" Version="3.3.101" />
+    <PackageReference Include="AWSSDK.S3" Version="3.5.6.5" />
+    <PackageReference Include="AWSSDK.SecurityToken" Version="3.5.1.29" />
+    <PackageReference Include="Datadog.Trace" Version="1.18.0" />
+    <PackageReference Include="DogStatsD-CSharp-Client" Version="6.0.0" />
+    <PackageReference Include="Google.Protobuf" Version="3.12.3" />
+    <PackageReference Include="Grpc.AspNetCore" Version="2.29.0" />
+    <PackageReference Include="Grpc.Net.Client" Version="2.29.0" />
+    <PackageReference Include="Grpc.Tools" Version="2.30.0">
+      <PrivateAssets>all</PrivateAssets>
+      <IncludeAssets>runtime; build; native; contentfiles; analyzers; buildtransitive</IncludeAssets>
+    </PackageReference>
+    <PackageReference Include="Microsoft.AspNetCore.Authentication.Cookies" Version="2.2.0" />
+    <PackageReference Include="Microsoft.AspNetCore.Authentication.JwtBearer" Version="3.1.5" />
+    <PackageReference Include="Microsoft.AspNetCore.Authentication.OpenIdConnect" Version="3.1.5" />
+    <PackageReference Include="Microsoft.AspNetCore.DataProtection.StackExchangeRedis" Version="3.1.7" />
+    <PackageReference Include="Microsoft.CodeAnalysis.FxCopAnalyzers" Version="3.0.0">
+      <PrivateAssets>all</PrivateAssets>
+      <IncludeAssets>runtime; build; native; contentfiles; analyzers; buildtransitive</IncludeAssets>
+    </PackageReference>
+    <PackageReference Include="Microsoft.Extensions.Caching.Memory" Version="3.1.6" />
+    <PackageReference Include="MongoDB.Driver" Version="2.10.4" />
+    <PackageReference Include="Polly" Version="7.2.1" />
+    <PackageReference Include="Polly.Extensions.Http" Version="3.0.0" />
+    <PackageReference Include="Serilog.AspNetCore" Version="3.2.0" />
+    <PackageReference Include="Serilog.Extensions.Hosting" Version="3.1.0" />
+    <PackageReference Include="Serilog.Sinks.Console" Version="3.1.1" />
+    <PackageReference Include="Serilog.Sinks.File" Version="4.1.0" />
+    <PackageReference Include="SharpZipLib" Version="1.2.0" />
+    <PackageReference Include="StackExchange.Redis" Version="2.2.4" />
+    <PackageReference Include="Swashbuckle.AspNetCore" Version="5.5.0" />
+    <PackageReference Include="TimeZoneConverter" Version="3.2.0" />
+  </ItemGroup>
+
+  <ItemGroup>
+    <Content Include="..\HordeCommon\Protos\**\*.*">
+      <Link>Protos\%(RecursiveDir)%(FileName)%(Extension)</Link>
+    </Content>
+  </ItemGroup>
+
+  <ItemGroup>
+    <None Remove="HordeServer.xml" />
+  </ItemGroup>
+
+  <ItemGroup>
+    <Protobuf Include="..\HordeCommon\Protos\build\bazel\remote\execution\v2\remote_execution.proto" ProtoRoot="..\HordeCommon\Protos" GrpcServices="Server" />
+    <Protobuf Include="..\HordeCommon\Protos\build\bazel\semver\semver.proto" ProtoRoot="..\HordeCommon\Protos" GrpcServices="Server" />
+    <Protobuf Include="..\HordeCommon\Protos\google\api\annotations.proto" ProtoRoot="..\HordeCommon\Protos" GrpcServices="Server" />
+    <Protobuf Include="..\HordeCommon\Protos\google\api\client.proto" ProtoRoot="..\HordeCommon\Protos" GrpcServices="Server" />
+    <Protobuf Include="..\HordeCommon\Protos\google\api\http.proto" ProtoRoot="..\HordeCommon\Protos" GrpcServices="Server" />
+    <Protobuf Include="..\HordeCommon\Protos\google\bytestream\bytestream.proto" ProtoRoot="..\HordeCommon\Protos" GrpcServices="Server" />
+    <Protobuf Include="..\HordeCommon\Protos\google\longrunning\operations.proto" ProtoRoot="..\HordeCommon\Protos" GrpcServices="Server" />
+    <Protobuf Include="..\HordeCommon\Protos\google\rpc\code.proto" ProtoRoot="..\HordeCommon\Protos" GrpcServices="Server" />
+    <Protobuf Include="..\HordeCommon\Protos\horde\action_rpc.proto" ProtoRoot="..\HordeCommon\Protos" GrpcServices="Server" />
+    <Protobuf Include="..\HordeCommon\Protos\horde\horde_rpc.proto" ProtoRoot="..\HordeCommon\Protos" GrpcServices="Server" />
+    <Protobuf Include="..\HordeCommon\Protos\horde\tasks\action_task.proto" ProtoRoot="..\HordeCommon\Protos" GrpcServices="Server" />
+  </ItemGroup>
+
+  <ItemGroup>
+    <ProjectReference Include="..\..\Shared\EpicGames.Core\EpicGames.Core.csproj" />
+    <ProjectReference Include="..\..\Shared\EpicGames.Perforce\EpicGames.Perforce.csproj" />
+    <ProjectReference Include="..\HordeCommon\HordeCommon.csproj" />
+  </ItemGroup>
+
+</Project>
diff --git a/Engine/Source/Programs/Horde/HordeServer/IssueHandlers/IIssueHandler.cs b/Engine/Source/Programs/Horde/HordeServer/IssueHandlers/IIssueHandler.cs
new file mode 100644
--- /dev/null
+++ b/Engine/Source/Programs/Horde/HordeServer/IssueHandlers/IIssueHandler.cs
@@ -0,0 +1,155 @@
+using EpicGames.Core;
+using HordeServer.Models;
+using HordeServer.Collections;
+using System;
+using System.Collections.Generic;
+using System.Linq;
+using System.Threading.Tasks;
+using System.Diagnostics.CodeAnalysis;
+
+namespace HordeServer.IssueHandlers
+{
+	/// <summary>
+	/// Information about a changelist and a value ranking the likelihood that it caused an issue
+	/// </summary>
+	class SuspectChange
+	{
+		/// <summary>
+		/// Set of extensions to treat as code
+		/// </summary>
+		static readonly HashSet<StringView> CodeExtensions = new HashSet<StringView>(StringViewComparer.OrdinalIgnoreCase)
+		{
+			".c",
+			".cc",
+			".cpp",
+			".m",
+			".mm",
+			".rc",
+			".cs",
+			".h",
+			".hpp",
+			".inl",
+			".uproject",
+			".uplugin"
+		};
+
+		/// <summary>
+		/// Set of file extensions to treat as content
+		/// </summary>
+		static readonly HashSet<StringView> ContentExtensions = new HashSet<StringView>(StringViewComparer.OrdinalIgnoreCase)
+		{
+			".uasset",
+			".umap",
+			".ini"
+		};
+
+		/// <summary>
+		/// The change detials
+		/// </summary>
+		public ChangeDetails Details { get; set; }
+
+		/// <summary>
+		/// Rank for the change. A value less than or equal to zero indicates a lack of culpability, positive values indicate
+		/// the possibility of a change being the culprit.
+		/// </summary>
+		public int Rank { get; set; }
+
+		/// <summary>
+		/// Whether the change contains code
+		/// </summary>
+		public bool ContainsCode { get; }
+
+		/// <summary>
+		/// Whether the change modifies content
+		/// </summary>
+		public bool ContainsContent { get; }
+
+		/// <summary>
+		/// Constructor
+		/// </summary>
+		/// <param name="Details">The changelist details</param>
+		public SuspectChange(ChangeDetails Details)
+		{
+			this.Details = Details;
+
+			foreach (string Path in Details.Files)
+			{
+				int Idx = Path.LastIndexOf('.');
+				if (Idx != -1)
+				{
+					StringView Extension = new StringView(Path, Idx);
+					if (CodeExtensions.Contains(Extension))
+					{
+						ContainsCode = true;
+					}
+					if (ContentExtensions.Contains(Extension))
+					{
+						ContainsContent = true;
+					}
+					if (ContainsCode && ContainsContent)
+					{
+						break;
+					}
+				}
+			}
+		}
+
+		/// <summary>
+		/// Determines whether this change modifies the given file
+		/// </summary>
+		/// <param name="FileToCheck">The file to look for</param>
+		/// <returns>True if the change modifies the given file</returns>
+		public bool ModifiesFile(string FileToCheck)
+		{
+			foreach (string File in Details.Files)
+			{
+				if (File.EndsWith(FileToCheck, StringComparison.OrdinalIgnoreCase) && (File.Length == FileToCheck.Length || File[File.Length - FileToCheck.Length - 1] == '/'))
+				{
+					return true;
+				}
+			}
+			return false;
+		}
+	}
+
+	/// <summary>
+	/// Interface for issue matchers
+	/// </summary>
+	interface IIssueHandler
+	{
+		/// <summary>
+		/// Identifier for the type of issue
+		/// </summary>
+		string Type { get; }
+
+		/// <summary>
+		/// Priority of this matcher
+		/// </summary>
+		int Priority { get; }
+
+		/// <summary>
+		/// Match the given event and produce a fingerprint
+		/// </summary>
+		/// <param name="Job">The job that spawned the event</param>
+		/// <param name="Node">Node that was executed</param>
+		/// <param name="EventData">The event data</param>
+		/// <param name="Fingerprint">Receives the fingerprint on success</param>
+		/// <returns>True if the match is successful</returns>
+		bool TryGetFingerprint(IJob Job, INode Node, ILogEventData EventData, [NotNullWhen(true)] out NewIssueFingerprint? Fingerprint);
+
+		/// <summary>
+		/// Rank all the suspect changes for a given fingerprint
+		/// </summary>
+		/// <param name="Fingerprint">The issue fingerprint</param>
+		/// <param name="Suspects">Potential suspects</param>
+		void RankSuspects(IIssueFingerprint Fingerprint, List<SuspectChange> Suspects);
+
+		/// <summary>
+		/// Gets the summary text for an issue
+		/// </summary>
+		/// <param name="Fingerprint">The fingerprint</param>
+		/// <param name="Severity">Severity of the issue</param>
+		/// <returns>The summary text</returns>
+		string GetSummary(IIssueFingerprint Fingerprint, IssueSeverity Severity);
+	}
+}
diff --git a/Engine/Source/Programs/Horde/HordeServer/IssueHandlers/Impl/CompileIssueHandler.cs b/Engine/Source/Programs/Horde/HordeServer/IssueHandlers/Impl/CompileIssueHandler.cs
new file mode 100644
--- /dev/null
+++ b/Engine/Source/Programs/Horde/HordeServer/IssueHandlers/Impl/CompileIssueHandler.cs
@@ -0,0 +1,130 @@
+using HordeCommon;
+using HordeServer.Collections;
+using HordeCommon.Rpc;
+using HordeServer.Services;
+using HordeServer.Utilities;
+using Microsoft.Extensions.Logging;
+using MongoDB.Bson;
+using MongoDB.Bson.Serialization;
+using MongoDB.Driver;
+using System;
+using System.Collections.Generic;
+using System.Diagnostics.CodeAnalysis;
+using System.IO;
+using System.Linq;
+using System.Security.Policy;
+using System.Text;
+using System.Threading.Tasks;
+using System.Text.Json;
+using HordeServer.Models;
+using EpicGames.Core;
+
+namespace HordeServer.IssueHandlers.Impl
+{
+	/// <summary>
+	/// Instance of a particular compile error
+	/// </summary>
+	class CompileIssueHandler : IIssueHandler
+	{
+		/// <inheritdoc/>
+		public string Type => "Compile";
+
+		/// <inheritdoc/>
+		public int Priority => 10;
+
+		/// <summary>
+		/// Determines if the given event id matches
+		/// </summary>
+		/// <param name="EventId">The event id to compare</param>
+		/// <returns>True if the given event id matches</returns>
+		public static bool IsMatchingEventId(EventId? EventId)
+		{
+			return EventId == KnownLogEvents.Compiler || EventId == KnownLogEvents.AutomationTool_SourceFileLine || EventId == KnownLogEvents.MSBuild;
+		}
+
+		/// <summary>
+		/// Extracts a list of source files from an event
+		/// </summary>
+		/// <param name="Event">The event data</param>
+		/// <param name="SourceFiles">List of source files</param>
+		static void GetSourceFiles(ILogEventData Event, List<string> SourceFiles)
+		{
+			foreach (ILogEventLine Line in Event.Lines)
+			{
+				JsonElement Properties;
+				if (Line.Data.TryGetProperty("properties", out Properties) && Properties.ValueKind == JsonValueKind.Object)
+				{
+					foreach (JsonProperty Property in Properties.EnumerateObject())
+					{
+						if (Property.NameEquals("file") && Property.Value.ValueKind == JsonValueKind.String)
+						{
+							AddSourceFile(SourceFiles, Property.Value.GetString());
+						}
+						if (Property.Value.HasStringProperty("type", "SourceFile") && Property.Value.TryGetStringProperty("relativePath", out string? Value))
+						{
+							AddSourceFile(SourceFiles, Value);
+						}
+					}
+				}
+			}
+		}
+
+		/// <summary>
+		/// Add a new source file to a list of unique source files
+		/// </summary>
+		/// <param name="SourceFiles">List of source files</param>
+		/// <param name="RelativePath">File to add</param>
+		static void AddSourceFile(List<string> SourceFiles, string RelativePath)
+		{
+			int EndIdx = RelativePath.LastIndexOfAny(new char[] { '/', '\\' }) + 1;
+
+			string FileName = RelativePath.Substring(EndIdx);
+			if (!SourceFiles.Any(x => x.Equals(FileName, StringComparison.OrdinalIgnoreCase)))
+			{
+				SourceFiles.Add(FileName);
+			}
+		}
+
+		/// <inheritdoc/>
+		public bool TryGetFingerprint(IJob Job, INode Node, ILogEventData EventData, [NotNullWhen(true)] out NewIssueFingerprint? Fingerprint)
+		{
+			if (!IsMatchingEventId(EventData.EventId))
+			{
+				Fingerprint = null;
+				return false;
+			}
+
+			List<string> NewFileNames = new List<string>();
+			GetSourceFiles(EventData, NewFileNames);
+			Fingerprint = new NewIssueFingerprint(Type, NewFileNames, null);
+			return true;
+		}
+
+		/// <inheritdoc/>
+		public string GetSummary(IIssueFingerprint Fingerprint, IssueSeverity Severity)
+		{
+			string Type = (Severity == IssueSeverity.Warning) ? "Compile warnings" : "Compile errors";
+			string List = StringUtils.FormatList(Fingerprint.Keys.ToArray(), 2);
+			return $"{Type} in {List}";
+		}
+
+		/// <inheritdoc/>
+		public void RankSuspects(IIssueFingerprint Fingerprint, List<SuspectChange> Suspects)
+		{
+			foreach (SuspectChange Change in Suspects)
+			{
+				if (Change.ContainsCode)
+				{
+					if (Fingerprint.Keys.Any(x => Change.ModifiesFile(x)))
+					{
+						Change.Rank += 20;
+					}
+					else
+					{
+						Change.Rank += 10;
+					}
+				}
+			}
+		}
+	}
+}
diff --git a/Engine/Source/Programs/Horde/HordeServer/IssueHandlers/Impl/ContentIssueHandler.cs b/Engine/Source/Programs/Horde/HordeServer/IssueHandlers/Impl/ContentIssueHandler.cs
new file mode 100644
--- /dev/null
+++ b/Engine/Source/Programs/Horde/HordeServer/IssueHandlers/Impl/ContentIssueHandler.cs
@@ -0,0 +1,103 @@
+using HordeCommon;
+using HordeServer.Collections;
+using HordeServer.Models;
+using HordeServer.Utilities;
+using Microsoft.Extensions.Logging;
+using MongoDB.Bson;
+using MongoDB.Bson.Serialization;
+using MongoDB.Driver;
+using System;
+using System.Collections.Generic;
+using System.Diagnostics.CodeAnalysis;
+using System.IO;
+using System.Linq;
+using System.Security.Policy;
+using System.Text;
+using System.Threading.Tasks;
+using EpicGames.Core;
+
+namespace HordeServer.IssueHandlers.Impl
+{
+	/// <summary>
+	/// Instance of a particular compile error
+	/// </summary>
+	class ContentIssueHandler : IIssueHandler
+	{
+		/// <inheritdoc/>
+		public string Type => "Content";
+
+		/// <inheritdoc/>
+		public int Priority => 10;
+
+		/// <summary>
+		/// Determines if the given event id matches
+		/// </summary>
+		/// <param name="EventId">The event id to compare</param>
+		/// <returns>True if the given event id matches</returns>
+		public static bool IsMatchingEventId(EventId? EventId)
+		{
+			return EventId == KnownLogEvents.Engine_AssetLog;
+		}
+
+		/// <summary>
+		/// Adds all the assets from the given log event
+		/// </summary>
+		/// <param name="Event">The log event to parse</param>
+		/// <param name="AssetNames">Receives the referenced asset names</param>
+		public static void GetAssetNames(ILogEventData Event, HashSet<string> AssetNames)
+		{
+			foreach (ILogEventLine Line in Event.Lines)
+			{
+				string? RelativePath;
+				if (Line.Data.TryGetNestedProperty("properties.asset.relativePath", out RelativePath))
+				{
+					int EndIdx = RelativePath.LastIndexOfAny(new char[] { '/', '\\' }) + 1;
+					string FileName = RelativePath.Substring(EndIdx);
+					AssetNames.Add(FileName);
+				}
+			}
+		}
+
+		/// <inheritdoc/>
+		public bool TryGetFingerprint(IJob Job, INode Node, ILogEventData EventData, [NotNullWhen(true)] out NewIssueFingerprint? Fingerprint)
+		{
+			if (!IsMatchingEventId(EventData.EventId))
+			{
+				Fingerprint = null;
+				return false;
+			}
+
+			HashSet<string> NewAssetNames = new HashSet<string>();
+			GetAssetNames(EventData, NewAssetNames);
+			Fingerprint = new NewIssueFingerprint(Type, NewAssetNames, null);
+			return true;
+		}
+
+		/// <inheritdoc/>
+		public string GetSummary(IIssueFingerprint Fingerprint, IssueSeverity Severity)
+		{
+			string Type = (Severity == IssueSeverity.Warning) ? "Warnings" : "Errors";
+			string List = StringUtils.FormatList(Fingerprint.Keys.ToArray(), 2);
+			return $"{Type} in {List}";
+		}
+
+		/// <inheritdoc/>
+		public void RankSuspects(IIssueFingerprint Fingerprint, List<SuspectChange> Suspects)
+		{
+			foreach (SuspectChange Suspect in Suspects)
+			{
+				if (Suspect.ContainsContent)
+				{
+					if (Suspect.Details.Files.Any(x => Fingerprint.Keys.Any(y => x.Contains(y, StringComparison.OrdinalIgnoreCase))))
+					{
+						Suspect.Rank += 20;
+					}
+					else
+					{
+						Suspect.Rank += 10;
+					}
+				}
+			}
+		}
+	}
+}
diff --git a/Engine/Source/Programs/Horde/HordeServer/IssueHandlers/Impl/CopyrightIssueHandler.cs b/Engine/Source/Programs/Horde/HordeServer/IssueHandlers/Impl/CopyrightIssueHandler.cs
new file mode 100644
--- /dev/null
+++ b/Engine/Source/Programs/Horde/HordeServer/IssueHandlers/Impl/CopyrightIssueHandler.cs
@@ -0,0 +1,107 @@
+using HordeCommon;
+using HordeServer.Collections;
+using HordeCommon.Rpc;
+using HordeServer.Models;
+using HordeServer.Utilities;
+using Microsoft.Extensions.Logging;
+using MongoDB.Bson;
+using MongoDB.Bson.Serialization;
+using MongoDB.Driver;
+using System;
+using System.Collections.Generic;
+using System.Diagnostics.CodeAnalysis;
+using System.IO;
+using System.Linq;
+using System.Security.Policy;
+using System.Text;
+using System.Threading.Tasks;
+using EpicGames.Core;
+
+namespace HordeServer.IssueHandlers.Impl
+{
+	/// <summary>
+	/// Instance of a particular compile error
+	/// </summary>
+	class CopyrightIssueHandler : IIssueHandler
+	{
+		/// <inheritdoc/>
+		public string Type => "Copyright";
+
+		/// <inheritdoc/>
+		public int Priority => 10;
+
+		/// <summary>
+		/// Filenames containing errors or warnings
+		/// </summary>
+		public HashSet<string> FileNames { get; set; } = new HashSet<string>();
+
+		/// <summary>
+		/// Determines if the given event id matches
+		/// </summary>
+		/// <param name="EventId">The event id to compare</param>
+		/// <returns>True if the given event id matches</returns>
+		public static bool IsMatchingEventId(EventId? EventId)
+		{
+			return EventId == KnownLogEvents.AutomationTool_MissingCopyright;
+		}
+
+		/// <summary>
+		/// Extracts a list of source files from an event
+		/// </summary>
+		/// <param name="Event">The event data</param>
+		/// <param name="SourceFiles">List of source files</param>
+		public static void GetSourceFiles(ILogEventData Event, HashSet<string> SourceFiles)
+		{
+			foreach (ILogEventLine Line in Event.Lines)
+			{
+				string? RelativePath;
+				if (Line.Data.TryGetNestedProperty("properties.file.relativePath", out RelativePath) || Line.Data.TryGetNestedProperty("properties.file", out RelativePath))
+				{
+					int EndIdx = RelativePath.LastIndexOfAny(new char[] { '/', '\\' }) + 1;
+					string FileName = RelativePath.Substring(EndIdx);
+					SourceFiles.Add(FileName);
+				}
+			}
+		}
+
+		/// <inheritdoc/>
+		public bool TryGetFingerprint(IJob Job, INode Node, ILogEventData EventData, [NotNullWhen(true)] out NewIssueFingerprint? Fingerprint)
+		{
+			if (!IsMatchingEventId(EventData.EventId))
+			{
+				Fingerprint = null;
+				return false;
+			}
+
+			HashSet<string> NewFileNames = new HashSet<string>(StringComparer.OrdinalIgnoreCase);
+			GetSourceFiles(EventData, NewFileNames);
+			Fingerprint = new NewIssueFingerprint(Type, NewFileNames, null);
+			return true;
+		}
+
+		/// <inheritdoc/>
+		public void RankSuspects(IIssueFingerprint Fingerprint, List<SuspectChange> Suspects)
+		{
+			foreach (SuspectChange Suspect in Suspects)
+			{
+				if (Suspect.ContainsCode)
+				{
+					if (FileNames.Any(x => Suspect.ModifiesFile(x)))
+					{
+						Suspect.Rank += 20;
+					}
+					else
+					{
+						Suspect.Rank += 10;
+					}
+				}
+			}
+		}
+
+		/// <inheritdoc/>
+		public string GetSummary(IIssueFingerprint Fingerprint, IssueSeverity Severity)
+		{
+			return $"Missing copyright notice in {StringUtils.FormatList(Fingerprint.Keys.ToArray(), 2)}";
+		}
+	}
+}
diff --git a/Engine/Source/Programs/Horde/HordeServer/IssueHandlers/Impl/DefaultIssueHandler.cs b/Engine/Source/Programs/Horde/HordeServer/IssueHandlers/Impl/DefaultIssueHandler.cs
new file mode 100644
--- /dev/null
+++ b/Engine/Source/Programs/Horde/HordeServer/IssueHandlers/Impl/DefaultIssueHandler.cs
@@ -0,0 +1,53 @@
+using HordeServer.Collections;
+using HordeServer.Models;
+using System;
+using System.Collections.Generic;
+using System.Diagnostics.CodeAnalysis;
+using System.Linq;
+using System.Threading.Tasks;
+
+namespace HordeServer.IssueHandlers.Impl
+{
+	/// <summary>
+	/// Instance of a particular compile error
+	/// </summary>
+	class DefaultIssueHandler : IIssueHandler
+	{
+		/// <summary>
+		/// Name of the handler
+		/// </summary>
+		public const string Type = "Default";
+
+		/// <inheritdoc/>
+		string IIssueHandler.Type => Type;
+
+		/// <inheritdoc/>
+		public int Priority => 0;
+
+		/// <inheritdoc/>
+		public bool TryGetFingerprint(IJob Job, INode Node, ILogEventData EventData, [NotNullWhen(true)] out NewIssueFingerprint? Fingerprint)
+		{
+			Fingerprint = new NewIssueFingerprint(Type, new[] { Node.Name }, null);
+			return true;
+		}
+
+		/// <inheritdoc/>
+		public void RankSuspects(IIssueFingerprint Fingerprint, List<SuspectChange> Suspects)
+		{
+		}
+
+		/// <inheritdoc/>
+		public string GetSummary(IIssueFingerprint Fingerprint, IssueSeverity Severity)
+		{
+			string NodeName = Fingerprint.Keys.FirstOrDefault() ?? "(unknown)";
+			if(Severity == IssueSeverity.Warning)
+			{
+				return $"Warnings in {NodeName}";
+			}
+			else
+			{
+				return $"Errors in {NodeName}";
+			}
+		}
+	}
+}
diff --git a/Engine/Source/Programs/Horde/HordeServer/IssueHandlers/Impl/GauntletIssueHandler.cs b/Engine/Source/Programs/Horde/HordeServer/IssueHandlers/Impl/GauntletIssueHandler.cs
new file mode 100644
--- /dev/null
+++ b/Engine/Source/Programs/Horde/HordeServer/IssueHandlers/Impl/GauntletIssueHandler.cs
@@ -0,0 +1,175 @@
+using HordeCommon;
+using HordeServer.Models;
+using HordeServer.Utilities;
+using Microsoft.AspNetCore.Authentication;
+using Microsoft.AspNetCore.Connections.Features;
+using Microsoft.Extensions.Logging;
+using Serilog.Formatting.Json;
+using System;
+using System.Collections.Generic;
+using System.Linq;
+using System.Security.Cryptography.X509Certificates;
+using System.Text.Json;
+using System.Threading.Tasks;
+using System.Diagnostics.CodeAnalysis;
+using HordeServer.Collections;
+using EpicGames.Core;
+
+namespace HordeServer.IssueHandlers.Impl
+{
+	/// <summary>
+	/// Instance of a particular Gauntlet error
+	/// </summary>
+	class GauntletIssueHandler : IIssueHandler
+	{
+		/// <summary>
+		/// Prefix for unit test keys
+		/// </summary>
+		const string UnitTestPrefix = "UnitTest:";
+
+		/// <summary>
+		/// Prefix for screenshot test keys
+		/// </summary>
+		const string ScreenshotTestPrefix = "Screenshot:";
+
+		/// <inheritdoc/>
+		public string Type => "Gauntlet";
+
+		/// <inheritdoc/>
+		public int Priority => 10;
+
+		/// <summary>
+		/// Determines if the given event id matches
+		/// </summary>
+		/// <param name="EventId">The event id to compare</param>
+		/// <returns>True if the given event id matches</returns>
+		public static bool IsMatchingEventId(EventId? EventId)
+		{
+			return EventId == KnownLogEvents.Gauntlet_UnitTest || EventId == KnownLogEvents.Gauntlet_ScreenshotTest;
+		}
+
+		/// <summary>
+		/// Parses symbol names from a log event
+		/// </summary>
+		/// <param name="EventData">The log event data</param>
+		/// <param name="UnitTestNames">Receives a set of the unit test names</param>
+		private static void GetUnitTestNames(ILogEventData EventData, HashSet<string> UnitTestNames)
+		{
+			string? Group = null;
+			string? Name = null;
+
+			foreach (ILogEventLine Line in EventData.Lines)
+			{
+				string? Value;
+				if (Line.Data.TryGetNestedProperty("properties.group", out Value))
+				{
+					Group = Value;
+				}
+				if (Line.Data.TryGetNestedProperty("properties.name", out Value))
+				{
+					Name = Value;
+				}
+			}
+
+			if(Group != null && Name != null)
+			{
+				UnitTestNames.Add($"{UnitTestPrefix}:{Group}/{Name}");
+			}
+		}
+
+		/// <summary>
+		/// Parses screenshot test names from a log event
+		/// </summary>
+		/// <param name="EventData">The event data</param>
+		/// <param name="ScreenshotTestNames">Receives the parsed screenshot test names</param>
+		private static void GetScreenshotTestNames(ILogEventData EventData, HashSet<string> ScreenshotTestNames)
+		{
+			foreach (ILogEventLine Line in EventData.Lines)
+			{
+				if (Line.Data.TryGetProperty("properties", JsonValueKind.Object, out JsonElement Properties))
+				{
+					foreach (JsonProperty Property in Properties.EnumerateObject())
+					{
+						JsonElement Value = Property.Value;
+						if (Value.ValueKind == JsonValueKind.Object && Value.HasStringProperty("type", "Screenshot") && Value.TryGetStringProperty("name", out string? Name))
+						{
+							ScreenshotTestNames.Add($"{ScreenshotTestPrefix}:{Name}");
+						}
+					}
+				}
+			}
+		}
+
+		/// <inheritdoc/>
+		public bool TryGetFingerprint(IJob Job, INode Node, ILogEventData EventData, [NotNullWhen(true)] out NewIssueFingerprint? Fingerprint)
+		{
+			if(!IsMatchingEventId(EventData.EventId))
+			{
+				Fingerprint = null;
+				return false;
+			}
+
+			HashSet<string> Keys = new HashSet<string>();
+			GetUnitTestNames(EventData, Keys);
+			GetScreenshotTestNames(EventData, Keys);
+
+			if (Keys.Count == 0)
+			{
+				Fingerprint = null;
+				return false;
+			}
+
+			Fingerprint = new NewIssueFingerprint(Type, Keys, null);
+			return true;
+		}
+
+		/// <inheritdoc/>
+		public void RankSuspects(IIssueFingerprint Fingerprint, List<SuspectChange> Changes)
+		{
+			foreach (SuspectChange Change in Changes)
+			{
+				if (Change.ContainsCode)
+				{
+					Change.Rank += 10;
+				}
+			}
+		}
+
+		/// <inheritdoc/>
+		public string GetSummary(IIssueFingerprint Fingerprint, IssueSeverity Severity)
+		{
+			List<string> UnitTestNames = Fingerprint.Keys.Where(x => x.StartsWith(UnitTestPrefix, StringComparison.Ordinal)).Select(x => x.Substring(UnitTestPrefix.Length + 1)).ToList();
+			if (UnitTestNames.Count > 0)
+			{
+				HashSet<string> GroupNames = new HashSet<string>(StringComparer.OrdinalIgnoreCase);
+				HashSet<string> TestNames = new HashSet<string>(StringComparer.OrdinalIgnoreCase);
+				foreach (string UnitTestName in UnitTestNames)
+				{
+					int Idx = UnitTestName.IndexOf('/', StringComparison.OrdinalIgnoreCase);
+					if (Idx != -1)
+					{
+						GroupNames.Add(UnitTestName.Substring(0, Idx));
+					}
+					TestNames.Add(UnitTestName.Substring(Idx + 1));
+				}
+
+				if (GroupNames.Count == 1)
+				{
+					return $"{GroupNames.First()} test failures: {StringUtils.FormatList(TestNames.ToArray(), 3)}";
+				}
+				else
+				{
+					return $"{StringUtils.FormatList(GroupNames.OrderBy(x => x).ToArray(), 100)} test failures";
+				}
+			}
+
+			List<string> ScreenshotTestNames = Fingerprint.Keys.Where(x => x.StartsWith(ScreenshotTestPrefix, StringComparison.Ordinal)).Select(x => x.Substring(ScreenshotTestPrefix.Length + 1)).ToList();
+			if (ScreenshotTestNames.Count > 0)
+			{
+				return $"Screenshot test failures: {StringUtils.FormatList(ScreenshotTestNames.ToArray(), 3)}";
+			}
+
+			return "Test failures";
+		}
+	}
+}
