Initial drop of Horde Server/Agent for UE5/Main. Still missing some Epic-specific functionality required to build/run correctly.

diff --git a/Engine/Source/Programs/Horde/HordeServer/IssueHandlers/Impl/SymbolIssueHandler.cs b/Engine/Source/Programs/Horde/HordeServer/IssueHandlers/Impl/SymbolIssueHandler.cs
new file mode 100644
--- /dev/null
+++ b/Engine/Source/Programs/Horde/HordeServer/IssueHandlers/Impl/SymbolIssueHandler.cs
@@ -0,0 +1,103 @@
+using HordeCommon;
+using HordeServer.Collections;
+using HordeServer.Models;
+using HordeServer.Utilities;
+using Microsoft.Extensions.Logging;
+using MongoDB.Bson;
+using MongoDB.Bson.Serialization;
+using System;
+using System.Collections.Generic;
+using System.Diagnostics.CodeAnalysis;
+using System.Linq;
+using System.Text.RegularExpressions;
+using System.Threading.Tasks;
+using EpicGames.Core;
+
+namespace HordeServer.IssueHandlers.Impl
+{
+	/// <summary>
+	/// Instance of a particular compile error
+	/// </summary>
+	class SymbolIssueHandler : IIssueHandler
+	{
+		/// <inheritdoc/>
+		public string Type => "Symbol";
+
+		/// <inheritdoc/>
+		public int Priority => 10;
+
+		/// <summary>
+		/// Determines if the given event id matches
+		/// </summary>
+		/// <param name="EventId">The event id to compare</param>
+		/// <returns>True if the given event id matches</returns>
+		public static bool IsMatchingEventId(EventId? EventId)
+		{
+			return EventId == KnownLogEvents.Linker_UndefinedSymbol || EventId == KnownLogEvents.Linker_DuplicateSymbol || EventId == KnownLogEvents.Linker;
+		}
+
+		/// <summary>
+		/// Parses symbol names from a log event
+		/// </summary>
+		/// <param name="EventData">The log event data</param>
+		/// <param name="SymbolNames">Receives the list of symbol names</param>
+		public static void GetSymbolNames(ILogEventData EventData, SortedSet<string> SymbolNames)
+		{
+			foreach (ILogEventLine Line in EventData.Lines)
+			{
+				string? Identifier;
+				if (Line.Data.TryGetNestedProperty("properties.symbol.identifier", out Identifier))
+				{
+					SymbolNames.Add(Identifier);
+				}
+			}
+		}
+
+		/// <inheritdoc/>
+		public void RankSuspects(IIssueFingerprint Fingerprint, List<SuspectChange> Changes)
+		{
+			HashSet<string> Names = new HashSet<string>();
+			foreach (string Name in Fingerprint.Keys)
+			{
+				Names.UnionWith(Name.Split("::", StringSplitOptions.RemoveEmptyEntries));
+			}
+
+			foreach (SuspectChange Change in Changes)
+			{
+				if (Change.ContainsCode)
+				{
+					int Matches = Names.Count(x => Change.Details.Files.Any(y => y.Contains(x, StringComparison.OrdinalIgnoreCase)));
+					Change.Rank += 10 + (10 * Matches);
+				}
+			}
+		}
+
+		/// <inheritdoc/>
+		public string GetSummary(IIssueFingerprint Fingerprint, IssueSeverity Severity)
+		{
+			HashSet<string> Symbols = Fingerprint.Keys;
+			if (Symbols.Count == 1)
+			{
+				return $"Undefined symbol '{Symbols.First()}'";
+			}
+			else
+			{
+				return $"Undefined symbols: {StringUtils.FormatList(Symbols.ToArray(), 3)}";
+			}
+		}
+
+		public bool TryGetFingerprint(IJob Job, INode Node, ILogEventData EventData, [NotNullWhen(true)] out NewIssueFingerprint? Fingerprint)
+		{
+			if (!IsMatchingEventId(EventData.EventId))
+			{
+				Fingerprint = null;
+				return false;
+			}
+
+			SortedSet<string> SymbolNames = new SortedSet<string>();
+			GetSymbolNames(EventData, SymbolNames);
+			Fingerprint = new NewIssueFingerprint(Type, SymbolNames, null);
+			return true;
+		}
+	}
+}
diff --git a/Engine/Source/Programs/Horde/HordeServer/Logs/Builder/ILogBuilder.cs b/Engine/Source/Programs/Horde/HordeServer/Logs/Builder/ILogBuilder.cs
new file mode 100644
--- /dev/null
+++ b/Engine/Source/Programs/Horde/HordeServer/Logs/Builder/ILogBuilder.cs
@@ -0,0 +1,74 @@
+using HordeServer.Api;
+using HordeServer.Models;
+using MongoDB.Bson;
+using System;
+using System.Collections.Generic;
+using System.Linq;
+using System.Threading.Tasks;
+
+namespace HordeServer.Logs
+{
+	/// <summary>
+	/// Interface for the log file write cache
+	/// </summary>
+	public interface ILogBuilder
+	{
+		/// <summary>
+		/// Whether the cache should be flushed on shutdown
+		/// </summary>
+		bool FlushOnShutdown { get; }
+
+		/// <summary>
+		/// Append data to a key
+		/// </summary>
+		/// <param name="LogId">The log file id</param>
+		/// <param name="ChunkOffset">Offset of the chunk within the log file</param>
+		/// <param name="WriteOffset">Offset to write to</param>
+		/// <param name="WriteLineIndex">Line index being written</param>
+		/// <param name="WriteLineCount">Line count being written</param>
+		/// <param name="Data">Data to be appended</param>
+		/// <param name="Type">Type of data stored in this log file</param>
+		/// <returns>True if the data was appended to the given chunk. False if the chunk has been completed.</returns>
+		Task<bool> AppendAsync(ObjectId LogId, long ChunkOffset, long WriteOffset, int WriteLineIndex, int WriteLineCount, ReadOnlyMemory<byte> Data, LogType Type);
+
+		/// <summary>
+		/// Finish the current sub chunk
+		/// </summary>
+		/// <param name="LogId">The log file id</param>
+		/// <param name="Offset">Offset of the chunk within the log file</param>
+		/// <returns>Async task</returns>
+		Task CompleteSubChunkAsync(ObjectId LogId, long Offset);
+
+		/// <summary>
+		/// Finish the current chunk
+		/// </summary>
+		/// <param name="LogId">The log file id</param>
+		/// <param name="Offset">Offset of the chunk within the log file</param>
+		/// <returns>Async task</returns>
+		Task CompleteChunkAsync(ObjectId LogId, long Offset);
+
+		/// <summary>
+		/// Remove a complete chunk from the builder
+		/// </summary>
+		/// <param name="LogId">The log file id</param>
+		/// <param name="Offset">Offset of the chunk within the log file</param>
+		/// <returns>Async task</returns>
+		Task RemoveChunkAsync(ObjectId LogId, long Offset);
+
+		/// <summary>
+		/// Gets the current chunk for the given log file
+		/// </summary>
+		/// <param name="LogId">The log file id</param>
+		/// <param name="Offset">Offset of the chunk within the log file</param>
+		/// <param name="LineIndex">Line index of the chunk within the log file</param>
+		/// <returns></returns>
+		Task<LogChunkData?> GetChunkAsync(ObjectId LogId, long Offset, int LineIndex);
+
+		/// <summary>
+		/// Touches the timestamps of all the chunks after the given age, and returns them. Used for flushing the builder.
+		/// </summary>
+		/// <param name="MinAge">Minimum age of the chunks to enumerate. If specified, only chunks last modified longer than this period will be returned.</param>
+		/// <returns>List of chunks, identified by log id and chunk index</returns>
+		Task<List<(ObjectId, long)>> TouchChunksAsync(TimeSpan MinAge);
+	}
+}
diff --git a/Engine/Source/Programs/Horde/HordeServer/Logs/Builder/Impl/LocalLogBuilder.cs b/Engine/Source/Programs/Horde/HordeServer/Logs/Builder/Impl/LocalLogBuilder.cs
new file mode 100644
--- /dev/null
+++ b/Engine/Source/Programs/Horde/HordeServer/Logs/Builder/Impl/LocalLogBuilder.cs
@@ -0,0 +1,236 @@
+using HordeServer.Api;
+using HordeServer.Models;
+using HordeServer.Utilities;
+using MongoDB.Bson;
+using System;
+using System.Collections.Concurrent;
+using System.Collections.Generic;
+using System.IO;
+using System.Linq;
+using System.Threading.Tasks;
+
+namespace HordeServer.Logs.Builder
+{
+	/// <summary>
+	/// In-memory implementation of a log write buffer
+	/// </summary>
+	class LocalLogBuilder : ILogBuilder
+	{
+		/// <summary>
+		/// Stores information about a sub-chunk which is still being written to
+		/// </summary>
+		class PendingSubChunk
+		{
+			/// <summary>
+			/// The type of data stored in this sub chunk
+			/// </summary>
+			public LogType Type { get; }
+
+			/// <summary>
+			/// Data for the sub-chunk
+			/// </summary>
+			public MemoryStream Stream { get; } = new MemoryStream(4096);
+
+			/// <summary>
+			/// Constructor
+			/// </summary>
+			/// <param name="Type">Type of data stored in this subchunk</param>
+			public PendingSubChunk(LogType Type)
+			{
+				this.Type = Type;
+			}
+
+			/// <summary>
+			/// Converts this pending sub-chunk to a concrete LogSubChunkData object
+			/// </summary>
+			/// <param name="Offset">Offset of this sub-chunk within the file</param>
+			/// <param name="LineIndex">The line index</param>
+			/// <returns>New sub-chunk data object</returns>
+			public LogSubChunkData ToSubChunkData(long Offset, int LineIndex)
+			{
+				return new LogSubChunkData(Type, Offset, LineIndex, new ReadOnlyLogText(Stream.ToArray()));
+			}
+		}
+
+		/// <summary>
+		/// Stores information about a pending log chunk
+		/// </summary>
+		class PendingChunk
+		{
+			/// <summary>
+			/// Time that this chunk was created
+			/// </summary>
+			public DateTime CreateTimeUtc { get; set; } = DateTime.UtcNow;
+
+			/// <summary>
+			/// Array of sub chunks for this chunk
+			/// </summary>
+			public List<PendingSubChunk> SubChunks = new List<PendingSubChunk>();
+
+			/// <summary>
+			/// The next sub-chunk
+			/// </summary>
+			public PendingSubChunk NextSubChunk;
+
+			/// <summary>
+			/// Total length of this chunk
+			/// </summary>
+			public long Length;
+
+			/// <summary>
+			/// Whether the chunk is complete
+			/// </summary>
+			public bool Complete;
+
+			/// <summary>
+			/// Constructor
+			/// </summary>
+			/// <param name="Type">Type of data to store in this file</param>
+			public PendingChunk(LogType Type)
+			{
+				NextSubChunk = new PendingSubChunk(Type);
+			}
+
+			/// <summary>
+			/// Complete the current sub-chunk
+			/// </summary>
+			public void CompleteSubChunk()
+			{
+				if (NextSubChunk.Stream.Length > 0)
+				{
+					SubChunks.Add(NextSubChunk);
+					NextSubChunk = new PendingSubChunk(NextSubChunk.Type);
+				}
+			}
+		}
+
+		/// <summary>
+		/// Current log chunk state
+		/// </summary>
+		ConcurrentDictionary<(ObjectId, long), PendingChunk> PendingChunks = new ConcurrentDictionary<(ObjectId, long), PendingChunk>();
+
+		/// <inheritdoc/>
+		public bool FlushOnShutdown => true;
+
+		/// <inheritdoc/>
+		public Task<bool> AppendAsync(ObjectId LogId, long ChunkOffset, long WriteOffset, int WriteLineIndex, int WriteLineCount, ReadOnlyMemory<byte> Data, LogType Type)
+		{
+			PendingChunk? PendingChunk;
+			while (!PendingChunks.TryGetValue((LogId, ChunkOffset), out PendingChunk))
+			{
+				if(WriteOffset != ChunkOffset)
+				{
+					return Task.FromResult(false);
+				}
+
+				PendingChunk = new PendingChunk(Type);
+
+				if (PendingChunks.TryAdd((LogId, ChunkOffset), PendingChunk))
+				{
+					break;
+				}
+			}
+
+			lock (PendingChunk)
+			{
+				if (!PendingChunk.Complete && ChunkOffset + PendingChunk.Length == WriteOffset)
+				{
+					PendingChunk.NextSubChunk.Stream.Write(Data.Span);
+					PendingChunk.Length += Data.Span.Length;
+					return Task.FromResult(true);
+				}
+			}
+			return Task.FromResult(false);
+		}
+
+		/// <inheritdoc/>
+		public Task CompleteSubChunkAsync(ObjectId LogId, long Offset)
+		{
+			PendingChunk? PendingChunk;
+			if (PendingChunks.TryGetValue((LogId, Offset), out PendingChunk))
+			{
+				lock (PendingChunk)
+				{
+					PendingChunk.CompleteSubChunk();
+				}
+			}
+			return Task.CompletedTask;
+		}
+
+		/// <inheritdoc/>
+		public Task CompleteChunkAsync(ObjectId LogId, long Offset)
+		{
+			PendingChunk? PendingChunk;
+			if (PendingChunks.TryGetValue((LogId, Offset), out PendingChunk))
+			{
+				lock (PendingChunk)
+				{
+					if (!PendingChunk.Complete)
+					{
+						PendingChunk.Complete = true;
+						PendingChunk.CompleteSubChunk();
+					}
+				}
+			}
+			return Task.CompletedTask;
+		}
+
+		/// <inheritdoc/>
+		public Task RemoveChunkAsync(ObjectId LogId, long Offset)
+		{
+			PendingChunks.TryRemove((LogId, Offset), out _);
+			return Task.CompletedTask;
+		}
+
+		/// <inheritdoc/>
+		public Task<LogChunkData?> GetChunkAsync(ObjectId LogId, long Offset, int LineIndex)
+		{
+			PendingChunk? PendingChunk;
+			if (PendingChunks.TryGetValue((LogId, Offset), out PendingChunk))
+			{
+				long SubChunkOffset = Offset;
+				int SubChunkLineIndex = LineIndex;
+
+				List<LogSubChunkData> SubChunks = new List<LogSubChunkData>();
+				foreach(PendingSubChunk PendingSubChunk in PendingChunk.SubChunks)
+				{
+					LogSubChunkData SubChunk = PendingSubChunk.ToSubChunkData(SubChunkOffset, SubChunkLineIndex);
+					SubChunkOffset += SubChunk.Length;
+					SubChunkLineIndex += SubChunk.LineCount;
+					SubChunks.Add(SubChunk);
+				}
+
+				PendingSubChunk NextSubChunk = PendingChunk.NextSubChunk;
+				if (NextSubChunk.Stream.Length > 0)
+				{
+					LogSubChunkData SubChunkData = NextSubChunk.ToSubChunkData(SubChunkOffset, SubChunkLineIndex);
+					SubChunks.Add(SubChunkData);
+				}
+
+				return Task.FromResult<LogChunkData?>(new LogChunkData(Offset, LineIndex, SubChunks));
+			}
+
+			return Task.FromResult<LogChunkData?>(null);
+		}
+
+		/// <inheritdoc/>
+		public Task<List<(ObjectId, long)>> TouchChunksAsync(TimeSpan MinAge)
+		{
+			DateTime UtcNow = DateTime.UtcNow;
+
+			List<(ObjectId, long)> Chunks = new List<(ObjectId, long)>();
+			foreach (KeyValuePair<(ObjectId, long), PendingChunk> PendingChunk in PendingChunks.ToArray())
+			{
+				lock (PendingChunk.Value)
+				{
+					if (PendingChunk.Value.CreateTimeUtc < UtcNow - MinAge)
+					{
+						Chunks.Add(PendingChunk.Key);
+						PendingChunk.Value.CreateTimeUtc = UtcNow;
+					}
+				}
+			}
+			return Task.FromResult(Chunks);
+		}
+	}
+}
diff --git a/Engine/Source/Programs/Horde/HordeServer/Logs/Builder/Impl/RedisLogBuilder.cs b/Engine/Source/Programs/Horde/HordeServer/Logs/Builder/Impl/RedisLogBuilder.cs
new file mode 100644
--- /dev/null
+++ b/Engine/Source/Programs/Horde/HordeServer/Logs/Builder/Impl/RedisLogBuilder.cs
@@ -0,0 +1,301 @@
+using Datadog.Trace;
+using EpicGames.Core;
+using HordeServer.Api;
+using Microsoft.Extensions.Logging;
+using Microsoft.Extensions.Primitives;
+using MongoDB.Bson;
+using StackExchange.Redis;
+using System;
+using System.Collections.Generic;
+using System.Globalization;
+using System.Linq;
+using System.Text.RegularExpressions;
+using System.Threading.Tasks;
+using HordeServer.Utilities;
+
+namespace HordeServer.Logs.Builder
+{
+	/// <summary>
+	/// Redis-based cache for log file chunks
+	/// </summary>
+	class RedisLogBuilder : ILogBuilder
+	{
+		const string ItemsKey = "log-items";
+
+		struct ChunkKeys
+		{
+			public string Prefix { get; }
+			public string Type => $"{Prefix}-type";
+			public string LineIndex => $"{Prefix}-lineindex";
+			public string Length => $"{Prefix}-length";
+			public string ChunkData => $"{Prefix}-chunk";
+			public string SubChunkData => $"{Prefix}-subchunk";
+			public string Complete => $"{Prefix}-complete";
+
+			public ChunkKeys(ObjectId LogId, long Offset)
+			{
+				Prefix = $"log-{LogId}-chunk-{Offset}-builder";
+			}
+
+			public static bool TryParse(string Prefix, out ObjectId LogId, out long Offset)
+			{
+				Match Match = Regex.Match(Prefix, "log-([^-]+)-chunk-([^-]+)-builder");
+				if (!Match.Success || !ObjectId.TryParse(Match.Groups[1].Value, out LogId) || !long.TryParse(Match.Groups[2].Value, out Offset))
+				{
+					LogId = new ObjectId();
+					Offset = 0;
+					return false;
+				}
+				return true;
+			}
+		}
+
+		/// <inheritdoc/>
+		public bool FlushOnShutdown => false;
+
+		/// <summary>
+		/// The Redis database connection pool
+		/// </summary>
+		RedisConnectionPool RedisConnectionPool;
+
+		/// <summary>
+		/// Logger for debug output
+		/// </summary>
+		ILogger Logger;
+
+		/// <summary>
+		/// Constructor
+		/// </summary>
+		/// <param name="RedisConnectionPool">The redis database singleton</param>
+		/// <param name="Logger">Logger for debug output</param>
+		public RedisLogBuilder(RedisConnectionPool RedisConnectionPool, ILogger Logger)
+		{
+			this.RedisConnectionPool = RedisConnectionPool;
+			this.Logger = Logger;
+		}
+
+		/// <inheritdoc/>
+		public async Task<bool> AppendAsync(ObjectId LogId, long ChunkOffset, long WriteOffset, int WriteLineIndex, int WriteLineCount, ReadOnlyMemory<byte> Data, LogType Type)
+		{
+			IDatabase RedisDb = RedisConnectionPool.GetDatabase();
+			ChunkKeys Keys = new ChunkKeys(LogId, ChunkOffset);
+
+			if(ChunkOffset == WriteOffset)
+			{
+				using Scope Scope = Tracer.Instance.StartActive("Redis.CreateChunk");
+				Scope.Span.SetTag("LogId", LogId.ToString());
+				Scope.Span.SetTag("Offset", ChunkOffset.ToString(CultureInfo.InvariantCulture));
+				Scope.Span.SetTag("WriteOffset", WriteOffset.ToString(CultureInfo.InvariantCulture));
+
+				ITransaction? CreateTransaction = RedisDb.CreateTransaction();
+				CreateTransaction.AddCondition(Condition.SortedSetNotContains(ItemsKey, Keys.Prefix));
+				_ = CreateTransaction.SortedSetAddAsync(ItemsKey, Keys.Prefix, DateTime.UtcNow.Ticks);
+				_ = CreateTransaction.StringSetAsync(Keys.Type, (int)Type);
+				_ = CreateTransaction.StringSetAsync(Keys.Length, Data.Length);
+				_ = CreateTransaction.StringSetAsync(Keys.LineIndex, WriteLineIndex + WriteLineCount);
+				_ = CreateTransaction.StringSetAsync(Keys.SubChunkData, Data);
+				if (await CreateTransaction.ExecuteAsync())
+				{
+					Logger.LogTrace("Created {Size} bytes in {Key}", Data.Length, Keys.SubChunkData);
+					return true;
+				}
+			}
+
+			{
+				using Scope Scope = Tracer.Instance.StartActive("Redis.AppendChunk");
+				Scope.Span.SetTag("LogId", LogId.ToString());
+				Scope.Span.SetTag("Offset", ChunkOffset.ToString(CultureInfo.InvariantCulture));
+				Scope.Span.SetTag("WriteOffset", WriteOffset.ToString(CultureInfo.InvariantCulture));
+
+				ITransaction? AppendTransaction = RedisDb.CreateTransaction();
+				AppendTransaction.AddCondition(Condition.SortedSetContains(ItemsKey, Keys.Prefix));
+				AppendTransaction.AddCondition(Condition.KeyNotExists(Keys.Complete));
+				AppendTransaction.AddCondition(Condition.StringEqual(Keys.Type, (int)Type));
+				AppendTransaction.AddCondition(Condition.StringEqual(Keys.Length, (int)(WriteOffset - ChunkOffset)));
+				_ = AppendTransaction.StringAppendAsync(Keys.SubChunkData, Data);
+				_ = AppendTransaction.StringSetAsync(Keys.Length, (int)(WriteOffset - ChunkOffset) + Data.Length);
+				_ = AppendTransaction.StringSetAsync(Keys.LineIndex, WriteLineIndex + WriteLineCount);
+				if (await AppendTransaction.ExecuteAsync())
+				{
+					Logger.LogTrace("Appended {Size} to {Key}", Data.Length, Keys.SubChunkData);
+					return true;
+				}
+			}
+
+			return false;
+		}
+
+		/// <inheritdoc/>
+		public async Task CompleteSubChunkAsync(ObjectId LogId, long Offset)
+		{
+			IDatabase RedisDb = RedisConnectionPool.GetDatabase();
+			ChunkKeys Keys = new ChunkKeys(LogId, Offset);
+			for (; ; )
+			{
+				using Scope Scope = Tracer.Instance.StartActive("Redis.CompleteSubChunk");
+				Scope.Span.SetTag("LogId", LogId.ToString());
+				Scope.Span.SetTag("Offset", Offset.ToString(CultureInfo.InvariantCulture));
+
+				LogType Type = (LogType)(int)await RedisDb.StringGetAsync(Keys.Type);
+				int Length = (int)await RedisDb.StringGetAsync(Keys.Length);
+				int LineIndex = (int)await RedisDb.StringGetAsync(Keys.LineIndex);
+				long ChunkDataLength = await RedisDb.StringLengthAsync(Keys.ChunkData);
+
+				RedisValue SubChunkTextValue = await RedisDb.StringGetAsync(Keys.SubChunkData);
+				if (SubChunkTextValue.IsNullOrEmpty)
+				{
+					break;
+				}
+
+				ReadOnlyLogText SubChunkText = new ReadOnlyLogText(SubChunkTextValue);
+				LogSubChunkData SubChunkData = new LogSubChunkData(Type, Offset + Length - SubChunkText.Data.Length, LineIndex - SubChunkText.LineCount, SubChunkText);
+				byte[] SubChunkDataBytes = SubChunkData.ToByteArray();
+
+				ITransaction WriteTransaction = RedisDb.CreateTransaction();
+
+				WriteTransaction.AddCondition(Condition.StringLengthEqual(Keys.ChunkData, ChunkDataLength));
+				WriteTransaction.AddCondition(Condition.StringLengthEqual(Keys.SubChunkData, SubChunkText.Length));
+				WriteTransaction.AddCondition(Condition.StringEqual(Keys.Length, Length));
+				WriteTransaction.AddCondition(Condition.StringEqual(Keys.LineIndex, LineIndex));
+				Task<long> NewLength = WriteTransaction.StringAppendAsync(Keys.ChunkData, SubChunkDataBytes);
+				_ = WriteTransaction.KeyDeleteAsync(Keys.SubChunkData);
+
+				if (await WriteTransaction.ExecuteAsync())
+				{
+					Logger.LogDebug("Completed sub-chunk for log {LogId} chunk offset {Offset} -> sub-chunk size {SubChunkSize}, chunk size {ChunkSize}", LogId, Offset, SubChunkDataBytes.Length, await NewLength);
+					break;
+				}
+			}
+		}
+
+		/// <inheritdoc/>
+		public async Task CompleteChunkAsync(ObjectId LogId, long Offset)
+		{
+			IDatabase RedisDb = RedisConnectionPool.GetDatabase();
+			using Scope Scope = Tracer.Instance.StartActive("Redis.CompleteChunk");
+			Scope.Span.SetTag("LogId", LogId.ToString());
+			Scope.Span.SetTag("Offset", Offset.ToString(CultureInfo.InvariantCulture));
+
+			ChunkKeys Keys = new ChunkKeys(LogId, Offset);
+
+			ITransaction Transaction = RedisDb.CreateTransaction();
+			Transaction.AddCondition(Condition.SortedSetContains(ItemsKey, Keys.Prefix));
+			_ = Transaction.StringSetAsync(Keys.Complete, true);
+			if(!await Transaction.ExecuteAsync())
+			{
+				Logger.LogDebug("Log {LogId} chunk offset {Offset} is not in Redis builder", LogId, Offset);
+				return;
+			}
+
+			await CompleteSubChunkAsync(LogId, Offset);
+		}
+
+		/// <inheritdoc/>
+		public async Task RemoveChunkAsync(ObjectId LogId, long Offset)
+		{
+			IDatabase RedisDb = RedisConnectionPool.GetDatabase();
+			using Scope Scope = Tracer.Instance.StartActive("Redis.RemoveChunk");
+			Scope.Span.SetTag("LogId", LogId.ToString());
+			Scope.Span.SetTag("Offset", Offset.ToString(CultureInfo.InvariantCulture));
+
+			ChunkKeys Keys = new ChunkKeys(LogId, Offset);
+
+			ITransaction Transaction = RedisDb.CreateTransaction();
+			_ = Transaction.KeyDeleteAsync(Keys.Type);
+			_ = Transaction.KeyDeleteAsync(Keys.LineIndex);
+			_ = Transaction.KeyDeleteAsync(Keys.Length);
+			_ = Transaction.KeyDeleteAsync(Keys.ChunkData);
+			_ = Transaction.KeyDeleteAsync(Keys.SubChunkData);
+			_ = Transaction.KeyDeleteAsync(Keys.Complete);
+
+			_ = Transaction.SortedSetRemoveAsync(ItemsKey, Keys.Prefix);
+			await Transaction.ExecuteAsync();
+		}
+
+		/// <inheritdoc/>
+		public async Task<LogChunkData?> GetChunkAsync(ObjectId LogId, long Offset, int LineIndex)
+		{
+			IDatabase RedisDb = RedisConnectionPool.GetDatabase();
+			ChunkKeys Keys = new ChunkKeys(LogId, Offset);
+			for (; ; )
+			{
+				using Scope Scope = Tracer.Instance.StartActive("Redis.GetChunk");
+				Scope.Span.SetTag("LogId", LogId.ToString());
+				Scope.Span.SetTag("Offset", Offset.ToString(CultureInfo.InvariantCulture));
+
+				RedisValue ChunkDataValue = await RedisDb.StringGetAsync(Keys.ChunkData);
+
+				ReadOnlyMemory<byte> ChunkData = ChunkDataValue;
+
+				ITransaction Transaction = RedisDb.CreateTransaction();
+				Transaction.AddCondition(Condition.StringLengthEqual(Keys.ChunkData, ChunkData.Length));
+				Task<RedisValue> TypeTask = Transaction.StringGetAsync(Keys.Type);
+				Task<RedisValue> LastSubChunkDataTask = Transaction.StringGetAsync(Keys.SubChunkData);
+
+				if (await Transaction.ExecuteAsync())
+				{
+					MemoryReader Reader = new MemoryReader(ChunkData);
+
+					long SubChunkOffset = Offset;
+					int SubChunkLineIndex = LineIndex;
+
+					List<LogSubChunkData> SubChunks = new List<LogSubChunkData>();
+					while (Reader.Offset < Reader.Length)
+					{
+						LogSubChunkData SubChunkData = Reader.ReadLogSubChunkData(SubChunkOffset, SubChunkLineIndex);
+						SubChunkOffset += SubChunkData.Length;
+						SubChunkLineIndex += SubChunkData.LineCount;
+						SubChunks.Add(SubChunkData);
+					}
+
+					RedisValue Type = await TypeTask;
+
+					RedisValue LastSubChunkData = await LastSubChunkDataTask;
+					if (LastSubChunkData.Length() > 0)
+					{
+						ReadOnlyLogText SubChunkDataText = new ReadOnlyLogText(LastSubChunkData);
+						SubChunks.Add(new LogSubChunkData((LogType)(int)Type, SubChunkOffset, SubChunkLineIndex, SubChunkDataText));
+					}
+
+					if (SubChunks.Count == 0)
+					{
+						break;
+					}
+
+					return new LogChunkData(Offset, LineIndex, SubChunks);
+				}
+			}
+			return null;
+		}
+
+		/// <inheritdoc/>
+		public async Task<List<(ObjectId, long)>> TouchChunksAsync(TimeSpan MinAge)
+		{
+			IDatabase RedisDb = RedisConnectionPool.GetDatabase();
+			
+			// Find all the chunks that are suitable for expiry
+			DateTime UtcNow = DateTime.UtcNow;
+			SortedSetEntry[] Entries = await RedisDb.SortedSetRangeByScoreWithScoresAsync(ItemsKey, stop: (UtcNow - MinAge).Ticks);
+
+			// Update the score for each element in a transaction. If it succeeds, we can write the chunk. Otherwise another pod has beat us to it.
+			List<(ObjectId, long)> Results = new List<(ObjectId, long)>();
+			foreach (SortedSetEntry Entry in Entries)
+			{
+				if (ChunkKeys.TryParse(Entry.Element.ToString(), out ObjectId LogId, out long Offset))
+				{
+					ITransaction Transaction = RedisDb.CreateTransaction();
+					Transaction.AddCondition(Condition.SortedSetEqual(ItemsKey, Entry.Element, Entry.Score));
+					_ = Transaction.SortedSetAddAsync(ItemsKey, Entry.Element, UtcNow.Ticks);
+
+					if (!await Transaction.ExecuteAsync())
+					{
+						break;
+					}
+
+					Results.Add((LogId, Offset));
+				}
+			}
+			return Results;
+		}
+	}
+}
diff --git a/Engine/Source/Programs/Horde/HordeServer/Logs/LogChunkData.cs b/Engine/Source/Programs/Horde/HordeServer/Logs/LogChunkData.cs
new file mode 100644
--- /dev/null
+++ b/Engine/Source/Programs/Horde/HordeServer/Logs/LogChunkData.cs
@@ -0,0 +1,278 @@
+using EpicGames.Core;
+using HordeServer.Api;
+using HordeServer.Utilities;
+using System;
+using System.Collections.Generic;
+using System.IO;
+using System.Linq;
+using System.Text;
+using System.Threading.Tasks;
+
+namespace HordeServer.Logs
+{
+	/// <summary>
+	/// Data for a log chunk
+	/// </summary>
+	public class LogChunkData
+	{
+		/// <summary>
+		/// Offset of this chunk
+		/// </summary>
+		public long Offset { get; }
+
+		/// <summary>
+		/// Length of this chunk
+		/// </summary>
+		public int Length { get; }
+
+		/// <summary>
+		/// Line index of this chunk
+		/// </summary>
+		public int LineIndex { get; }
+
+		/// <summary>
+		/// Number of lines in this chunk
+		/// </summary>
+		public int LineCount { get; }
+			
+		/// <summary>
+		/// List of sub-chunks
+		/// </summary>
+		public IReadOnlyList<LogSubChunkData> SubChunks { get; }
+
+		/// <summary>
+		/// Offset of each sub-chunk within this chunk
+		/// </summary>
+		public int[] SubChunkOffset { get; }
+
+		/// <summary>
+		/// Total line count of the chunk after each sub-chunk
+		/// </summary>
+		public int[] SubChunkLineIndex { get; }
+
+		/// <summary>
+		/// Constructor
+		/// </summary>
+		/// <param name="Offset">Offset of this chunk</param>
+		/// <param name="LineIndex">First line index of this chunk</param>
+		/// <param name="SubChunks">Sub-chunks for this chunk</param>
+		public LogChunkData(long Offset, int LineIndex, IReadOnlyList<LogSubChunkData> SubChunks)
+		{
+			this.Offset = Offset;
+			this.LineIndex = LineIndex;
+			this.SubChunks = SubChunks;
+			this.SubChunkOffset = CreateSumLookup(SubChunks, x => x.Length);
+			this.SubChunkLineIndex = CreateSumLookup(SubChunks, x => x.LineCount);
+
+			if(SubChunks.Count > 0)
+			{
+				int LastSubChunkIdx = SubChunks.Count - 1;
+				Length = SubChunkOffset[LastSubChunkIdx] + SubChunks[LastSubChunkIdx].Length;
+				LineCount = SubChunkLineIndex[LastSubChunkIdx] + SubChunks[LastSubChunkIdx].LineCount;
+			}
+		}
+
+		/// <summary>
+		/// Creates a lookup table by summing a field across a list of subchunks
+		/// </summary>
+		/// <param name="SubChunks"></param>
+		/// <param name="Field"></param>
+		/// <returns></returns>
+		static int[] CreateSumLookup(IReadOnlyList<LogSubChunkData> SubChunks, Func<LogSubChunkData, int> Field)
+		{
+			int[] Total = new int[SubChunks.Count];
+
+			int Value = 0;
+			for (int Idx = 0; Idx + 1 < SubChunks.Count; Idx++)
+			{
+				Value += Field(SubChunks[Idx]);
+				Total[Idx + 1] = Value;
+			}
+
+			return Total;
+		}
+
+		/// <summary>
+		/// Gets the offset of a line within the file
+		/// </summary>
+		/// <param name="LineIdx">The line index</param>
+		/// <returns>Offset of the line within the file</returns>
+		public long GetLineOffsetWithinChunk(int LineIdx)
+		{
+			int SubChunkIdx = SubChunkLineIndex.BinarySearch(LineIdx);
+			if (SubChunkIdx < 0)
+			{
+				SubChunkIdx = ~SubChunkIdx - 1;
+			}
+			return SubChunkOffset[SubChunkIdx] + SubChunks[SubChunkIdx].InflateText().LineOffsets[LineIdx - SubChunkLineIndex[SubChunkIdx]];
+		}
+
+		/// <summary>
+		/// Gets the sub chunk index for the given line
+		/// </summary>
+		/// <param name="ChunkLineIdx">Line index within the chunk</param>
+		/// <returns>Subchunk line index</returns>
+		public int GetSubChunkForLine(int ChunkLineIdx)
+		{
+			int SubChunkIdx = SubChunkLineIndex.BinarySearch(ChunkLineIdx);
+			if (SubChunkIdx < 0)
+			{
+				SubChunkIdx = ~SubChunkIdx - 1;
+			}
+			return SubChunkIdx;
+		}
+
+		/// <summary>
+		/// Gets the index of the sub-chunk containing the given offset
+		/// </summary>
+		/// <param name="Offset">Offset to search for</param>
+		/// <returns>Index of the sub-chunk</returns>
+		public int GetSubChunkForOffsetWithinChunk(int Offset)
+		{
+			int SubChunkIdx = SubChunkOffset.BinarySearch(Offset);
+			if (SubChunkIdx < 0)
+			{
+				SubChunkIdx = ~SubChunkIdx - 1;
+			}
+			return SubChunkIdx;
+		}
+
+		/// <summary>
+		/// Signature for output data
+		/// </summary>
+		const int CurrentSignature = 'L' | ('C' << 8) | ('D' << 16);
+
+		/// <summary>
+		/// Read a log chunk from the given stream
+		/// </summary>
+		/// <param name="Reader">The reader to read from</param>
+		/// <param name="Offset">Offset of this chunk within the file</param>
+		/// <param name="LineIndex">Line index of this chunk</param>
+		/// <returns>New log chunk data</returns>
+		public static LogChunkData Read(MemoryReader Reader, long Offset, int LineIndex)
+		{
+			int Signature = Reader.ReadInt32();
+			if ((Signature & 0xffffff) != CurrentSignature)
+			{
+				List<LogSubChunkData> SubChunks = new List<LogSubChunkData>();
+				SubChunks.Add(new LogSubChunkData(LogType.Json, Offset, LineIndex, new ReadOnlyLogText(Reader.Memory)));
+				Reader.Offset = Reader.Memory.Length;
+				return new LogChunkData(Offset, LineIndex, SubChunks);
+			}
+
+			int Version = Signature >> 24;
+			if (Version == 0)
+			{
+				List<LogSubChunkData> SubChunks = ReadSubChunkList(Reader, Offset, LineIndex);
+				Reader.ReadVariableLengthBytes();
+				return new LogChunkData(Offset, LineIndex, SubChunks);
+			}
+			else
+			{
+				List<LogSubChunkData> SubChunks = ReadSubChunkList(Reader, Offset, LineIndex);
+				return new LogChunkData(Offset, LineIndex, SubChunks);
+			}
+		}
+
+		/// <summary>
+		/// Read a list of sub-chunks from the given stream
+		/// </summary>
+		/// <param name="Reader">The reader to read from</param>
+		/// <param name="SubChunkOffset">Offset of this chunk within the file</param>
+		/// <param name="SubChunkLineIndex">Line index of this chunk</param>
+		/// <returns>List of sub-chunks</returns>
+		static List<LogSubChunkData> ReadSubChunkList(MemoryReader Reader, long SubChunkOffset, int SubChunkLineIndex)
+		{
+			int NumSubChunks = Reader.ReadInt32();
+
+			List<LogSubChunkData> SubChunks = new List<LogSubChunkData>();
+			for (int Idx = 0; Idx < NumSubChunks; Idx++)
+			{
+				LogSubChunkData SubChunkData = Reader.ReadLogSubChunkData(SubChunkOffset, SubChunkLineIndex);
+				SubChunkOffset += SubChunkData.Length;
+				SubChunkLineIndex += SubChunkData.LineCount;
+				SubChunks.Add(SubChunkData);
+			}
+
+			return SubChunks;
+		}
+
+		/// <summary>
+		/// Construct an object from flat memory buffer
+		/// </summary>
+		/// <param name="Memory">Memory buffer</param>
+		/// <param name="Offset">Offset of this chunk within the file</param>
+		/// <param name="LineIndex">Line index of this chunk</param>
+		/// <returns>Log chunk data</returns>
+		public static LogChunkData FromMemory(ReadOnlyMemory<byte> Memory, long Offset, int LineIndex)
+		{
+			MemoryReader Reader = new MemoryReader(Memory);
+			LogChunkData ChunkData = Read(Reader, Offset, LineIndex);
+			Reader.CheckOffset(Memory.Length);
+			return ChunkData;
+		}
+
+		/// <summary>
+		/// Write the chunk data to a stream
+		/// </summary>
+		/// <returns>Serialized data</returns>
+		public void Write(MemoryWriter Writer)
+		{
+			Writer.WriteInt32(CurrentSignature | (1 << 24));
+
+			Writer.WriteInt32(SubChunks.Count);
+			foreach (LogSubChunkData SubChunk in SubChunks)
+			{
+				Writer.WriteLogSubChunkData(SubChunk);
+			}
+		}
+
+		/// <summary>
+		/// Construct an object from flat memory buffer
+		/// </summary>
+		/// <returns>Log chunk data</returns>
+		public byte[] ToByteArray()
+		{
+			byte[] Data = new byte[GetSerializedSize()];
+			MemoryWriter Writer = new MemoryWriter(Data);
+			Write(Writer);
+			Writer.CheckOffset(Data.Length);
+			return Data;
+		}
+
+		/// <summary>
+		/// Determines the size of the serialized buffer
+		/// </summary>
+		public int GetSerializedSize()
+		{
+			return sizeof(int) + (sizeof(int) + SubChunks.Sum(x => x.GetSerializedSize()));
+		}
+	}
+
+	/// <summary>
+	/// Extensions for the log chunk data
+	/// </summary>
+	static class LogChunkDataExtensions
+	{
+		/// <summary>
+		/// Read a log chunk from the given stream
+		/// </summary>
+		/// <param name="Reader">The reader to read from</param>
+		/// <param name="Offset">Offset of this chunk within the file</param>
+		/// <param name="LineIndex">Line index of this chunk</param>
+		/// <returns>New log chunk data</returns>
+		public static LogChunkData ReadLogChunkData(this MemoryReader Reader, long Offset, int LineIndex)
+		{
+			return LogChunkData.Read(Reader, Offset, LineIndex);
+		}
+
+		/// <summary>
+		/// Write the chunk data to a stream
+		/// </summary>
+		/// <returns>Serialized data</returns>
+		public static void WriteLogChunkData(this MemoryWriter Writer, LogChunkData ChunkData)
+		{
+			ChunkData.Write(Writer);
+		}
+	}
+}
diff --git a/Engine/Source/Programs/Horde/HordeServer/Logs/LogIndexData.cs b/Engine/Source/Programs/Horde/HordeServer/Logs/LogIndexData.cs
new file mode 100644
--- /dev/null
+++ b/Engine/Source/Programs/Horde/HordeServer/Logs/LogIndexData.cs
@@ -0,0 +1,637 @@
+using Datadog.Trace;
+using EpicGames.Core;
+using HordeServer.Utilities;
+using System;
+using System.Collections.Generic;
+using System.IO;
+using System.Linq;
+using System.Reflection;
+using System.Text;
+using System.Threading.Tasks;
+
+namespace HordeServer.Logs
+{
+	/// <summary>
+	/// Contains a source block of text to be indexed
+	/// </summary>
+	public class LogIndexBlock
+	{
+		/// <summary>
+		/// First line within the file
+		/// </summary>
+		public int LineIndex { get; }
+
+		/// <summary>
+		/// Number of lines in this block
+		/// </summary>
+		public int LineCount { get; }
+
+		/// <summary>
+		/// The plain text for this block
+		/// </summary>
+		public ILogText? CachedPlainText { get; private set; }
+
+		/// <summary>
+		/// The compressed plain text data
+		/// </summary>
+		public ReadOnlyMemory<byte> CompressedPlainText { get; }
+
+		/// <summary>
+		/// Constructor
+		/// </summary>
+		/// <param name="LineIndex">Index of the first line within this block</param>
+		/// <param name="LineCount">Number of lines in the block</param>
+		/// <param name="CachedPlainText">The decompressed plain text</param>
+		/// <param name="CompressedPlainText">The compressed text data</param>
+		public LogIndexBlock(int LineIndex, int LineCount, ILogText? CachedPlainText, ReadOnlyMemory<byte> CompressedPlainText)
+		{
+			this.LineIndex = LineIndex;
+			this.LineCount = LineCount;
+			this.CachedPlainText = CachedPlainText;
+			this.CompressedPlainText = CompressedPlainText;
+		}
+
+		/// <summary>
+		/// Accessor for the decompressed plaintext
+		/// </summary>
+		public ILogText InflatePlainText()
+		{
+			if (CachedPlainText == null)
+			{
+				CachedPlainText = new ReadOnlyLogText(CompressedPlainText.DecompressBzip2());
+			}
+			return CachedPlainText;
+		}
+	}
+
+	/// <summary>
+	/// Extension methods for <see cref="LogIndexBlock"/>
+	/// </summary>
+	public static class LogIndexBlockExtensions
+	{
+		/// <summary>
+		/// Deserialize an index block
+		/// </summary>
+		/// <param name="Reader">Reader to deserialize from</param>
+		/// <returns>The new index block</returns>
+		public static LogIndexBlock ReadLogIndexBlock(this MemoryReader Reader)
+		{
+			int LineIndex = Reader.ReadInt32();
+			int LineCount = Reader.ReadInt32();
+			ReadOnlyMemory<byte> CompressedPlainText = Reader.ReadVariableLengthBytes();
+			return new LogIndexBlock(LineIndex, LineCount, null, CompressedPlainText);
+		}
+
+		/// <summary>
+		/// Serialize an index block
+		/// </summary>
+		/// <param name="Writer">Writer to serialize to</param>
+		/// <param name="Block">The block to serialize</param>
+		public static void WriteLogIndexBlock(this MemoryWriter Writer, LogIndexBlock Block)
+		{
+			Writer.WriteInt32(Block.LineIndex);
+			Writer.WriteInt32(Block.LineCount);
+			Writer.WriteVariableLengthBytes(Block.CompressedPlainText.Span);
+		}
+
+		/// <summary>
+		/// Gets the serialized size of a block
+		/// </summary>
+		/// <param name="Block">The block to serialize</param>
+		public static int GetSerializedSize(this LogIndexBlock Block)
+		{
+			return sizeof(int) + sizeof(int) + (sizeof(int) + Block.CompressedPlainText.Length);
+		}
+	}
+
+	/// <summary>
+	/// Stats for a search
+	/// </summary>
+	public class LogSearchStats
+	{
+		/// <summary>
+		/// Number of blocks that were scanned
+		/// </summary>
+		public int NumScannedBlocks { get; set; }
+
+		/// <summary>
+		/// Number of bytes that had to be scanned for results
+		/// </summary>
+		public int NumScannedBytes { get; set; }
+
+		/// <summary>
+		/// Number of blocks that were skipped
+		/// </summary>
+		public int NumSkippedBlocks { get; set; }
+
+		/// <summary>
+		/// Number of blocks that had to be decompressed
+		/// </summary>
+		public int NumDecompressedBlocks { get; set; }
+
+		/// <summary>
+		/// Number of blocks that were searched but did not contain the search term
+		/// </summary>
+		public int NumFalsePositiveBlocks { get; set; }
+	}
+
+	/// <summary>
+	/// Data for a log chunk
+	/// </summary>
+	public class LogIndexData
+	{
+		/// <summary>
+		/// Index for tokens into the block list
+		/// </summary>
+		ReadOnlyTrie? CachedTrie;
+
+		/// <summary>
+		/// Number of bits in the index devoted to the block index
+		/// </summary>
+		int NumBlockBits;
+
+		/// <summary>
+		/// List of text blocks
+		/// </summary>
+		LogIndexBlock[] Blocks;
+
+		/// <summary>
+		/// Empty index data
+		/// </summary>
+		public static LogIndexData Empty { get; } = new LogIndexData(ReadOnlyTrie.Empty, 0, Array.Empty<LogIndexBlock>());
+
+		/// <summary>
+		/// Number of lines covered by the index
+		/// </summary>
+		public int LineCount => (Blocks.Length > 0) ? (Blocks[Blocks.Length - 1].LineIndex + Blocks[Blocks.Length - 1].LineCount) : 0;
+
+		/// <summary>
+		/// Constructor
+		/// </summary>
+		/// <param name="CachedTrie">Index into the text blocks</param>
+		/// <param name="NumBlockBits">Number of bits devoted to the block index</param>
+		/// <param name="Blocks">Bloom filters for this log file</param>
+		public LogIndexData(ReadOnlyTrie? CachedTrie, int NumBlockBits, LogIndexBlock[] Blocks)
+		{
+			this.CachedTrie = CachedTrie;
+			this.NumBlockBits = NumBlockBits;
+			this.Blocks = Blocks;
+		}
+
+		/// <summary>
+		/// Updates the blocks in this index relative to a give line index
+		/// </summary>
+		/// <param name="LineIndex"></param>
+		public void SetBaseLineIndex(int LineIndex)
+		{
+			for(int BlockIdx = 0; BlockIdx < Blocks.Length; BlockIdx++)
+			{
+				LogIndexBlock Block = Blocks[BlockIdx];
+				Blocks[BlockIdx] = new LogIndexBlock(LineIndex, Block.LineCount, Block.CachedPlainText, Block.CompressedPlainText);
+				LineIndex += Block.LineCount;
+			}
+		}
+
+		/// <summary>
+		/// Builds a trie for this index
+		/// </summary>
+		/// <returns>The trie for this data</returns>
+		ReadOnlyTrie BuildTrie()
+		{
+			if(CachedTrie == null)
+			{
+				ReadOnlyTrieBuilder Builder = new ReadOnlyTrieBuilder();
+				for (int BlockIdx = 0; BlockIdx < Blocks.Length; BlockIdx++)
+				{
+					LogIndexBlock Block = Blocks[BlockIdx];
+					LogToken.GetTokens(Block.InflatePlainText().Data.Span, Token => Builder.Add((Token << NumBlockBits) | (uint)BlockIdx));
+				}
+				CachedTrie = Builder.Build();
+			}
+			return CachedTrie;
+		}
+
+		/// <summary>
+		/// Create an index from an array of blocks
+		/// </summary>
+		/// <param name="Indexes">List of indexes to merge</param>
+		/// <returns>Index data</returns>
+		public static LogIndexData Merge(IEnumerable<LogIndexData> Indexes)
+		{
+			using Scope Scope = Tracer.Instance.StartActive("LogIndexData.Merge");
+
+			// Create the combined block list
+			LogIndexBlock[] NewBlocks = Indexes.SelectMany(x => x.Blocks).ToArray();
+
+			// Figure out how many bits to devote to the block size
+			int NewNumBlockBits = 0;
+			while(NewBlocks.Length > (1 << NewNumBlockBits))
+			{
+				NewNumBlockBits += 4;
+			}
+
+			// Add all the blocks into a combined index
+			ReadOnlyTrieBuilder NewTrieBuilder = new ReadOnlyTrieBuilder();
+
+			int BlockCount = 0;
+			foreach(LogIndexData Index in Indexes)
+			{
+				ulong BlockMask = ((1UL << Index.NumBlockBits) - 1);
+				foreach (ulong Value in Index.BuildTrie())
+				{
+					ulong Token = Value >> Index.NumBlockBits;
+					int BlockIdx = (int)(Value & BlockMask);
+
+					ulong NewToken = (Token << NewNumBlockBits) | (ulong)(long)(BlockCount + BlockIdx);
+					NewTrieBuilder.Add(NewToken);
+				}
+				BlockCount += Index.Blocks.Length;
+			}
+
+			// Construct the new index data
+			ReadOnlyTrie NewTrie = NewTrieBuilder.Build();
+			return new LogIndexData(NewTrie, NewNumBlockBits, NewBlocks);
+		}
+
+		/// <summary>
+		/// Search for the given text in the index
+		/// </summary>
+		/// <param name="FirstLineIndex">First line index to search from</param>
+		/// <param name="Text">Text to search for</param>
+		/// <param name="Stats">Receives stats for the search</param>
+		/// <returns>List of line numbers for the text</returns>
+		public IEnumerable<int> Search(int FirstLineIndex, SearchText Text, LogSearchStats Stats)
+		{
+			int LastBlockCount = 0;
+			foreach (int BlockIdx in EnumeratePossibleBlocks(Text.Bytes, FirstLineIndex))
+			{
+				LogIndexBlock Block = Blocks[BlockIdx];
+
+				Stats.NumScannedBlocks++;
+				Stats.NumDecompressedBlocks += (Block.CachedPlainText == null)? 1 : 0;
+
+				Stats.NumSkippedBlocks += BlockIdx - LastBlockCount;
+				LastBlockCount = BlockIdx + 1;
+
+				// Decompress the text
+				ILogText BlockText = Block.InflatePlainText();
+
+				// Find the initial offset within this block
+				int Offset = 0;
+				if(FirstLineIndex > Block.LineIndex)
+				{
+					int LineIndexWithinBlock = FirstLineIndex - Block.LineIndex;
+					Offset = BlockText.LineOffsets[LineIndexWithinBlock];
+				}
+
+				// Search within this block
+				for(; ;)
+				{
+					// Find the next offset
+					int NextOffset = BlockText.Data.Span.FindNextOcurrence(Offset, Text);
+					if(NextOffset == -1)
+					{
+						Stats.NumScannedBytes += BlockText.Length - Offset;
+						break;
+					}
+
+					// Update the stats
+					Stats.NumScannedBytes += NextOffset - Offset;
+					Offset = NextOffset;
+
+					// Check it's not another match within the same line
+					int LineIndexWithinBlock = BlockText.GetLineIndexForOffset(Offset);
+					yield return Block.LineIndex + LineIndexWithinBlock;
+
+					// Move to the next line
+					Offset = BlockText.LineOffsets[LineIndexWithinBlock + 1];
+				}
+
+				// If the last scanned bytes is zero, we didn't have any matches from this chunk
+				if(Offset == 0 && CachedTrie != null)
+				{
+					Stats.NumFalsePositiveBlocks++;
+				}
+			}
+			Stats.NumSkippedBlocks += Blocks.Length - LastBlockCount;
+		}
+
+		/// <summary>
+		/// Search for the given text in the index
+		/// </summary>
+		/// <param name="Text">Text to search for</param>
+		/// <param name="LineIndex">The first </param>
+		/// <returns>List of line numbers for the text</returns>
+		IEnumerable<int> EnumeratePossibleBlocks(ReadOnlyMemory<byte> Text, int LineIndex)
+		{
+			// Find the starting block index
+			int BlockIdx = Blocks.BinarySearch(x => x.LineIndex, LineIndex);
+			if(BlockIdx < 0)
+			{
+				BlockIdx = Math.Max(~BlockIdx - 1, 0);
+			}
+
+			// Make sure we're not starting out of range
+			if (BlockIdx < Blocks.Length)
+			{
+				// In order to be considered a possible block for a positive match, the block must contain an exact list of tokens parsed from the
+				// search text, along with a set of partial matches (for substrings that are not guaranteed to start on token boundaries). The former
+				// is relatively cheap to compute, and done by stepping through all enumerators for each token in a wave, only returning blocks
+				// which all contain the token. At the very least, this must contain blocks in the index from BlockIdx onwards.
+				List<IEnumerator<int>> Enumerators = new List<IEnumerator<int>>();
+				Enumerators.Add(Enumerable.Range(BlockIdx, Blocks.Length - BlockIdx).GetEnumerator());
+
+				// The second aspect of the search requires a more expensive search through the trie. This is done by running a set of arbitrary 
+				// delegates to filter the matches returned from the enumerators.
+				List<Predicate<int>> Predicates = new List<Predicate<int>>();
+
+				// If the index has a trie, tokenize the input and generate a list of enumerators and predicates.
+				if (CachedTrie != null)
+				{
+					// Find a list of filters for matching blocks
+					HashSet<ulong> Tokens = new HashSet<ulong>();
+					for (int TokenPos = 0; TokenPos < Text.Length;)
+					{
+						ReadOnlySpan<byte> Token = LogToken.GetTokenText(Text.Span, TokenPos);
+						if (TokenPos == 0)
+						{
+							GetUnalignedTokenPredicate(Token, Token.Length == Text.Length, Predicates);
+						}
+						else
+						{
+							GetAlignedTokenPredicate(Token, TokenPos + Token.Length == Text.Length, Tokens, Predicates);
+						}
+						TokenPos += Token.Length;
+					}
+
+					// Create an enumerator for each token
+					foreach (ulong Token in Tokens)
+					{
+						ulong MinValue = (Token << NumBlockBits) | (ulong)(long)BlockIdx;
+						ulong MaxValue = MinValue + (1UL << NumBlockBits) - 1;
+
+						IEnumerator<int> Enumerator = CachedTrie.EnumerateRange(MinValue, MaxValue).Select(x => (int)(x - MinValue)).GetEnumerator();
+						if (!Enumerator.MoveNext())
+						{
+							yield break;
+						}
+
+						Enumerators.Add(Enumerator);
+					}
+				}
+
+				// Enumerate the matches
+				for (; ; )
+				{
+					// Advance all the enumerators that are behind the current block index. If they are all equal, we have a match.
+					bool bMatch = true;
+					foreach (IEnumerator<int> Enumerator in Enumerators)
+					{
+						while (Enumerator.Current < BlockIdx)
+						{
+							if (!Enumerator.MoveNext())
+							{
+								yield break;
+							}
+						}
+
+						if (Enumerator.Current > BlockIdx)
+						{
+							BlockIdx = Enumerator.Current;
+							bMatch = false;
+						}
+					}
+
+					// Return the match and move to the next block
+					if (bMatch)
+					{
+						if (Predicates.All(Predicate => Predicate(BlockIdx)))
+						{
+							yield return BlockIdx;
+						}
+						BlockIdx++;
+					}
+				}
+			}
+		}
+
+		/// <summary>
+		/// Gets predicates for matching a token that starts 
+		/// </summary>
+		/// <param name="Text">The token text</param>
+		/// <param name="bAllowPartialMatch">Whether to allow a partial match of the token</param>
+		/// <param name="Tokens">Set of aligned tokens that are required</param>
+		/// <param name="Predicates">List of predicates for the search</param>
+		void GetAlignedTokenPredicate(ReadOnlySpan<byte> Text, bool bAllowPartialMatch, HashSet<ulong> Tokens, List<Predicate<int>> Predicates)
+		{
+			for (int Offset = 0; Offset < Text.Length; Offset += LogToken.MaxTokenBytes)
+			{
+				ulong Token = LogToken.GetWindowedTokenValue(Text, Offset);
+				if (Offset + LogToken.MaxTokenBytes > Text.Length && bAllowPartialMatch)
+				{
+					ulong TokenMask = LogToken.GetWindowedTokenMask(Text, Offset, true);
+					Predicates.Add(BlockIdx => BlockContainsToken(BlockIdx, Token, TokenMask));
+					break;
+				}
+				Tokens.Add(Token);
+			}
+		}
+
+		/// <summary>
+		/// Generates a predicate for matching a token which may or may not start on a regular token boundary
+		/// </summary>
+		/// <param name="Text">The token text</param>
+		/// <param name="bAllowPartialMatch">Whether to allow a partial match of the token</param>
+		/// <param name="Predicates">List of predicates for the search</param>
+		void GetUnalignedTokenPredicate(ReadOnlySpan<byte> Text, bool bAllowPartialMatch, List<Predicate<int>> Predicates)
+		{
+			byte[] TextCopy = Text.ToArray();
+
+			Lazy<HashSet<int>> Blocks = new Lazy<HashSet<int>>(() =>
+			{
+				HashSet<int> Union = new HashSet<int>();
+				for (int Shift = 0; Shift < LogToken.MaxTokenBytes; Shift++)
+				{
+					HashSet<int> Blocks = new HashSet<int>(BlocksContainingToken(TextCopy.AsSpan(), -Shift, bAllowPartialMatch));
+					for (int Offset = -Shift + LogToken.MaxTokenBytes; Offset < TextCopy.Length && Blocks.Count > 0; Offset += LogToken.MaxTokenBytes)
+					{
+						Blocks.IntersectWith(BlocksContainingToken(TextCopy.AsSpan(), Offset, bAllowPartialMatch));
+					}
+					if (Blocks.Count > 0)
+					{
+						Union.UnionWith(Blocks);
+					}
+				}
+				return Union;
+			});
+
+			Predicates.Add(BlockIdx => Blocks.Value.Contains(BlockIdx));
+		}
+
+		/// <summary>
+		/// Tests whether a block contains a particular token
+		/// </summary>
+		/// <param name="BlockIdx">Index of the block to search</param>
+		/// <param name="Token">The token to test</param>
+		/// <param name="TokenMask">Mask of which bits in the token are valid</param>
+		/// <returns>True if the given block contains a token</returns>
+		bool BlockContainsToken(int BlockIdx, ulong Token, ulong TokenMask)
+		{
+			Token = (Token << NumBlockBits) | (uint)BlockIdx;
+			TokenMask = (TokenMask << NumBlockBits) | ((1UL << NumBlockBits) - 1);
+			return CachedTrie!.EnumerateValues((Value, ValueMask) => (Value & TokenMask) == (Token & ValueMask)).Any();
+		}
+
+		/// <summary>
+		/// Tests whether a block contains a particular token
+		/// </summary>
+		/// <param name="Text">The token to test</param>
+		/// <param name="Offset">Offset of the window into the token to test</param>
+		/// <param name="bAllowPartialMatch">Whether to allow a partial match of the token</param>
+		/// <returns>True if the given block contains a token</returns>
+		IEnumerable<int> BlocksContainingToken(ReadOnlySpan<byte> Text, int Offset, bool bAllowPartialMatch)
+		{
+			ulong Token = LogToken.GetWindowedTokenValue(Text, Offset) << NumBlockBits;
+			ulong TokenMask = LogToken.GetWindowedTokenMask(Text, Offset, bAllowPartialMatch) << NumBlockBits;
+			ulong BlockMask = (1UL << NumBlockBits) - 1;
+			return CachedTrie!.EnumerateValues((Value, ValueMask) => (Value & TokenMask) == (Token & ValueMask)).Select(x => (int)(x & BlockMask)).Distinct();
+		}
+
+		/// <summary>
+		/// Deserialize the index from memory
+		/// </summary>
+		/// <param name="Reader">Reader to deserialize from</param>
+		/// <returns>Index data</returns>
+		public static LogIndexData Read(MemoryReader Reader)
+		{
+			int Version = Reader.ReadInt32();
+			if (Version == 0)
+			{
+				return LogIndexData.Empty;
+			}
+			else if (Version == 1)
+			{
+				ReadOnlyTrie Index = Reader.ReadTrie();
+				int NumBlockBits = Reader.ReadInt32();
+
+				LogIndexBlock[] Blocks = new LogIndexBlock[Reader.ReadInt32()];
+				for (int Idx = 0; Idx < Blocks.Length; Idx++)
+				{
+					int LineIndex = (Idx > 0) ? (Blocks[Idx - 1].LineIndex + Blocks[Idx - 1].LineCount) : 0;
+					int LineCount = Reader.ReadInt32();
+					ReadOnlyMemory<byte> CompressedPlainText = Reader.ReadVariableLengthBytes();
+					Reader.ReadTrie();
+					Blocks[Idx] = new LogIndexBlock(LineIndex, LineCount, null, CompressedPlainText);
+				}
+
+				return new LogIndexData(Index, NumBlockBits, Blocks);
+			}
+			else if (Version == 2)
+			{
+				ReadOnlyTrie Index = Reader.ReadTrie();
+				int NumBlockBits = Reader.ReadInt32();
+
+				LogIndexBlock[] Blocks = new LogIndexBlock[Reader.ReadInt32()];
+				for (int Idx = 0; Idx < Blocks.Length; Idx++)
+				{
+					int LineIndex = Reader.ReadInt32();
+					int LineCount = Reader.ReadInt32();
+					ReadOnlyMemory<byte> CompressedPlainText = Reader.ReadVariableLengthBytes();
+					Reader.ReadTrie();
+					Blocks[Idx] = new LogIndexBlock(LineIndex, LineCount, null, CompressedPlainText);
+				}
+
+				return new LogIndexData(Index, NumBlockBits, Blocks);
+			}
+			else if (Version == 3)
+			{
+				ReadOnlyTrie Index = Reader.ReadTrie();
+				int NumBlockBits = Reader.ReadInt32();
+				LogIndexBlock[] Blocks = Reader.ReadVariableLengthArray(() => Reader.ReadLogIndexBlock());
+				for(int Idx = 1; Idx < Blocks.Length; Idx++)
+				{
+					if(Blocks[Idx].LineIndex == 0)
+					{
+						int LineIndex = Blocks[Idx - 1].LineIndex + Blocks[Idx - 1].LineCount;
+						Blocks[Idx] = new LogIndexBlock(LineIndex, Blocks[Idx].LineCount, Blocks[Idx].CachedPlainText, Blocks[Idx].CompressedPlainText);
+					}
+				}
+				return new LogIndexData(Index, NumBlockBits, Blocks);
+			}
+			else
+			{
+				throw new InvalidDataException($"Invalid index version number {Version}");
+			}
+		}
+
+		/// <summary>
+		/// Serialize an index into memory
+		/// </summary>
+		/// <param name="Writer">Writer to serialize to</param>
+		public void Write(MemoryWriter Writer)
+		{
+			Writer.WriteInt32(3);
+			Writer.WriteTrie(BuildTrie());
+			Writer.WriteInt32(NumBlockBits);
+			Writer.WriteVariableLengthArray(Blocks, x => Writer.WriteLogIndexBlock(x));
+		}
+
+		/// <summary>
+		/// Deserialize the index from memory
+		/// </summary>
+		/// <param name="Memory">Memory to deserialize from</param>
+		/// <returns>Index data</returns>
+		public static LogIndexData FromMemory(ReadOnlyMemory<byte> Memory)
+		{
+			MemoryReader Reader = new MemoryReader(Memory);
+			return Read(Reader);
+		}
+
+		/// <summary>
+		/// Serializes the index
+		/// </summary>
+		/// <returns>Index data</returns>
+		public byte[] ToByteArray()
+		{
+			byte[] Buffer = new byte[GetSerializedSize()];
+
+			MemoryWriter Writer = new MemoryWriter(Buffer);
+			Write(Writer);
+			Writer.CheckOffset(Buffer.Length);
+
+			return Buffer;
+		}
+
+		/// <summary>
+		/// Gets the serialized size of this index
+		/// </summary>
+		/// <returns>The serialized size</returns>
+		public int GetSerializedSize()
+		{
+			return sizeof(int) + BuildTrie().GetSerializedSize() + sizeof(int) + (sizeof(int) + Blocks.Sum(x => x.GetSerializedSize()));
+		}
+	}
+
+	static class LogIndexDataExtensions
+	{
+		/// <summary>
+		/// Deserialize the index from memory
+		/// </summary>
+		/// <param name="Reader">Reader to deserialize from</param>
+		/// <returns>Index data</returns>
+		public static LogIndexData ReadLogIndexData(this MemoryReader Reader)
+		{
+			return LogIndexData.Read(Reader);
+		}
+
+		/// <summary>
+		/// Serialize an index into memory
+		/// </summary>
+		/// <param name="Writer">Writer to serialize to</param>
+		/// <param name="Index">The index to write</param>
+		public static void WriteLogIndexData(this MemoryWriter Writer, LogIndexData Index)
+		{
+			Index.Write(Writer);
+		}
+	}
+}
diff --git a/Engine/Source/Programs/Horde/HordeServer/Logs/LogSubChunkData.cs b/Engine/Source/Programs/Horde/HordeServer/Logs/LogSubChunkData.cs
new file mode 100644
--- /dev/null
+++ b/Engine/Source/Programs/Horde/HordeServer/Logs/LogSubChunkData.cs
@@ -0,0 +1,253 @@
+using EpicGames.Core;
+using HordeServer.Api;
+using HordeServer.Utilities;
+using ICSharpCode.SharpZipLib.BZip2;
+using Microsoft.AspNetCore.Routing.Constraints;
+using System;
+using System.Collections.Generic;
+using System.IO;
+using System.Linq;
+using System.Runtime.CompilerServices;
+using System.Text;
+using System.Threading.Tasks;
+
+namespace HordeServer.Logs
+{
+	/// <summary>
+	/// Pending data for a sub-chunk
+	/// </summary>
+	public class LogSubChunkData
+	{
+		/// <summary>
+		/// Type of data stored in this subchunk
+		/// </summary>
+		public LogType Type { get; }
+
+		/// <summary>
+		/// Offset within the file of this sub-chunk
+		/// </summary>
+		public long Offset { get; }
+
+		/// <summary>
+		/// Length of this sub-chunk
+		/// </summary>
+		public int Length { get; }
+
+		/// <summary>
+		/// Index of the first line in this sub-chunk
+		/// </summary>
+		public int LineIndex { get; }
+
+		/// <summary>
+		/// Gets the number of lines in this sub-chunk
+		/// </summary>
+		public int LineCount { get; }
+
+		/// <summary>
+		/// Text data
+		/// </summary>
+		ILogText? TextInternal;
+
+		/// <summary>
+		/// Compressed text data
+		/// </summary>
+		ReadOnlyMemory<byte> CompressedTextInternal;
+
+		/// <summary>
+		/// The index for this sub-chunk
+		/// </summary>
+		LogIndexData? IndexInternal;
+
+		/// <summary>
+		/// The log text
+		/// </summary>
+		public ILogText InflateText()
+		{
+			if (TextInternal == null)
+			{
+				TextInternal = new ReadOnlyLogText(CompressedTextInternal.DecompressBzip2());
+			}
+			return TextInternal;
+		}
+
+		/// <summary>
+		/// The compressed log text
+		/// </summary>
+		public ReadOnlyMemory<byte> DeflateText()
+		{
+			if(CompressedTextInternal.IsEmpty)
+			{
+				CompressedTextInternal = TextInternal!.Data.CompressBzip2();
+			}
+			return CompressedTextInternal;
+		}
+
+		/// <summary>
+		/// Index for tokens in this chunk
+		/// </summary>
+		public LogIndexData BuildIndex()
+		{
+			if(IndexInternal == null)
+			{
+				ILogText PlainText = InflateText();
+				if (Type != LogType.Text)
+				{
+					PlainText = PlainText.ToPlainText();
+				}
+
+				LogIndexBlock[] Blocks = new LogIndexBlock[1];
+				Blocks[0] = new LogIndexBlock(LineIndex, LineCount, PlainText, PlainText.Data.CompressBzip2());
+				IndexInternal = new LogIndexData(null, 0, Blocks);
+			}
+			return IndexInternal;
+		}
+
+		/// <summary>
+		/// Constructor
+		/// </summary>
+		/// <param name="Type">Type of data stored in this subchunk</param>
+		/// <param name="Offset">Offset within the file of this sub-chunk</param>
+		/// <param name="LineIndex">Index of the first line in this sub-chunk</param>
+		/// <param name="Text">Text to add</param>
+		public LogSubChunkData(LogType Type, long Offset, int LineIndex, ILogText Text)
+		{
+			this.Type = Type;
+			this.Offset = Offset;
+			this.Length = Text.Data.Length;
+			this.LineIndex = LineIndex;
+			this.LineCount = Text.LineCount;
+			this.TextInternal = Text;
+		}
+
+		/// <summary>
+		/// Constructor for raw data
+		/// </summary>
+		/// <param name="Type">Type of data stored in this subchunk</param>
+		/// <param name="Offset">Offset within the file of this sub-chunk</param>
+		/// <param name="Length">Length of the uncompressed data</param>
+		/// <param name="LineIndex">Index of the first line</param>
+		/// <param name="LineCount">Number of lines in the uncompressed text</param>
+		/// <param name="CompressedText">Compressed text data</param>
+		/// <param name="Index">Index data</param>
+		public LogSubChunkData(LogType Type, long Offset, int Length, int LineIndex, int LineCount, ReadOnlyMemory<byte> CompressedText, LogIndexData? Index)
+		{
+			this.Type = Type;
+			this.Offset = Offset;
+			this.Length = Length;
+			this.LineIndex = LineIndex;
+			this.LineCount = LineCount;
+			this.CompressedTextInternal = CompressedText;
+			this.IndexInternal = Index;
+		}
+
+		/// <summary>
+		/// Constructs a sub-chunk from a block of memory. Uses slices of the given memory buffer rather than copying the data.
+		/// </summary>
+		/// <param name="Reader">The reader to read from</param>
+		/// <param name="Offset">Offset of the sub-chunk within this file</param>
+		/// <param name="LineIndex">Index of the first line in this subchunk</param>
+		/// <returns>New subchunk data</returns>
+		public static LogSubChunkData Read(MemoryReader Reader, long Offset, int LineIndex)
+		{
+			int Version = Reader.ReadInt32(); // Version placeholder
+			if (Version == 0)
+			{
+				LogType Type = (LogType)Reader.ReadInt32();
+				int Length = Reader.ReadInt32();
+				int LineCount = Reader.ReadInt32();
+				ReadOnlyMemory<byte> CompressedText = Reader.ReadVariableLengthBytes();
+				Reader.ReadVariableLengthBytes();
+				return new LogSubChunkData(Type, Offset, Length, LineIndex, LineCount, CompressedText, null);
+			}
+			else if (Version == 1)
+			{
+				LogType Type = (LogType)Reader.ReadInt32();
+				int Length = Reader.ReadInt32();
+				int LineCount = Reader.ReadInt32();
+				ReadOnlyMemory<byte> CompressedText = Reader.ReadVariableLengthBytes();
+				ReadOnlyMemory<byte> CompressedPlainText = Reader.ReadVariableLengthBytes();
+				ReadOnlyTrie Trie = Reader.ReadTrie();
+				LogIndexBlock IndexBlock = new LogIndexBlock(LineIndex, LineCount, null, CompressedPlainText);
+				LogIndexData Index = new LogIndexData(Trie, 0, new[] { IndexBlock });
+				return new LogSubChunkData(Type, Offset, Length, LineIndex, LineCount, CompressedText, Index);
+			}
+			else
+			{
+				LogType Type = (LogType)Reader.ReadInt32();
+				int Length = Reader.ReadInt32();
+				int LineCount = Reader.ReadInt32();
+				ReadOnlyMemory<byte> CompressedText = Reader.ReadVariableLengthBytes();
+				LogIndexData Index = Reader.ReadLogIndexData();
+				Index.SetBaseLineIndex(LineIndex); // Fix for incorrectly saved data
+				return new LogSubChunkData(Type, Offset, Length, LineIndex, LineCount, CompressedText, Index);
+			}
+		}
+
+		/// <summary>
+		/// Serializes the sub-chunk to a stream
+		/// </summary>
+		/// <param name="Writer">Writer to output to</param>
+		public void Write(MemoryWriter Writer)
+		{
+			Writer.WriteInt32(2); // Version placeholder
+
+			Writer.WriteInt32((int)Type);
+			Writer.WriteInt32(Length);
+			Writer.WriteInt32(LineCount);
+
+			Writer.WriteVariableLengthBytes(DeflateText().Span);
+			Writer.WriteLogIndexData(BuildIndex());
+		}
+
+		/// <summary>
+		/// Serializes this object to a byte array
+		/// </summary>
+		/// <returns>Byte array</returns>
+		public byte[] ToByteArray()
+		{
+			byte[] Data = new byte[GetSerializedSize()];
+
+			MemoryWriter Writer = new MemoryWriter(Data);
+			Write(Writer);
+			Writer.CheckOffset(Data.Length);
+
+			return Data;
+		}
+
+		/// <summary>
+		/// Determines the size of serialized data
+		/// </summary>
+		public int GetSerializedSize()
+		{
+			return sizeof(int) + sizeof(int) + sizeof(int) + sizeof(int) + (sizeof(int) + DeflateText().Length) + BuildIndex().GetSerializedSize();
+		}
+	}
+
+	/// <summary>
+	/// Pending data for a sub-chunk
+	/// </summary>
+	public static class LogSubChunkDataExtensions
+	{
+		/// <summary>
+		/// Constructs a sub-chunk from a block of memory. Uses slices of the given memory buffer rather than copying the data.
+		/// </summary>
+		/// <param name="Reader">The reader to read from</param>
+		/// <param name="Offset">Offset of this sub-chunk within the file</param>
+		/// <param name="LineIndex">Index of the first line within this chunk</param>
+		/// <returns>New subchunk data</returns>
+		public static LogSubChunkData ReadLogSubChunkData(this MemoryReader Reader, long Offset, int LineIndex)
+		{
+			return LogSubChunkData.Read(Reader, Offset, LineIndex);
+		}
+
+		/// <summary>
+		/// Serializes the sub-chunk to a stream
+		/// </summary>
+		/// <param name="Writer">Writer to output to</param>
+		/// <param name="SubChunkData">The sub-chunk data to write</param>
+		public static void WriteLogSubChunkData(this MemoryWriter Writer, LogSubChunkData SubChunkData)
+		{
+			SubChunkData.Write(Writer);
+		}
+	}
+}
diff --git a/Engine/Source/Programs/Horde/HordeServer/Logs/LogText.cs b/Engine/Source/Programs/Horde/HordeServer/Logs/LogText.cs
new file mode 100644
--- /dev/null
+++ b/Engine/Source/Programs/Horde/HordeServer/Logs/LogText.cs
@@ -0,0 +1,495 @@
+using EpicGames.Core;
+using Microsoft.VisualBasic.CompilerServices;
+using System;
+using System.Collections;
+using System.Collections.Generic;
+using System.Diagnostics;
+using System.IO;
+using System.Linq;
+using System.Text;
+using System.Text.Json;
+using System.Threading.Tasks;
+
+namespace HordeServer.Logs
+{
+	/// <summary>
+	/// Interface for log data
+	/// </summary>
+	public interface ILogText
+	{
+		/// <summary>
+		/// The raw text data. Contains a complete set of lines followed by newline characters.
+		/// </summary>
+		public ReadOnlyMemory<byte> Data { get; }
+
+		/// <summary>
+		/// Length of the data in this block
+		/// </summary>
+		public int Length { get; }
+
+		/// <summary>
+		/// Offsets of lines within the data object, including a sentinel for the end of the data (LineCount + 1 entries).
+		/// </summary>
+		public IReadOnlyList<int> LineOffsets { get; }
+
+		/// <summary>
+		/// Number of lines in the 
+		/// </summary>
+		public int LineCount { get; }
+
+		/// <summary>
+		/// Computes the number of bytes used by this object
+		/// </summary>
+		public int AllocatedSize { get; }
+	}
+
+	/// <summary>
+	/// Helper methods for parsing log text
+	/// </summary>
+	public static class LogTextHelpers
+	{
+		/// <summary>
+		/// Finds the number of newline characters in the given span
+		/// </summary>
+		/// <param name="Span">Span of utf-8 characters to scan</param>
+		/// <param name="Offset">Starting offset within the span</param>
+		/// <returns>Number of lines in the given span</returns>
+		public static int GetLineCount(ReadOnlySpan<byte> Span, int Offset)
+		{
+			int LineCount = 0;
+			for (; Offset < Span.Length; Offset++)
+			{
+				if (Span[Offset] == (byte)'\n')
+				{
+					LineCount++;
+				}
+			}
+			return LineCount;
+		}
+
+		/// <summary>
+		/// Find the offsets of the start of each line within the given span
+		/// </summary>
+		/// <param name="Span">Span to search</param>
+		/// <param name="Offset">Offset within the span to start searching</param>
+		/// <param name="LineOffsets">Receives the list of line offsets</param>
+		/// <param name="LineCount">The current line count</param>
+		public static void FindLineOffsets(ReadOnlySpan<byte> Span, int Offset, Span<int> LineOffsets, int LineCount)
+		{
+			Debug.Assert(Span.Length == 0 || Span[Span.Length - 1] == (byte)'\n');
+			for(; ; )
+			{
+				// Store the start of this line
+				LineOffsets[LineCount++] = Offset;
+
+				// Check if we're at the end
+				if(Offset >= Span.Length)
+				{
+					break;
+				}
+
+				// Move to the end of this line
+				while (Span[Offset] != '\n')
+				{
+					Offset++;
+				}
+
+				// Move to the start of the next line
+				Offset++;
+			}
+		}
+	}
+
+	/// <summary>
+	/// Immutable block of log text
+	/// </summary>
+	public class ReadOnlyLogText : ILogText
+	{
+		/// <inheritdoc/>
+		public ReadOnlyMemory<byte> Data { get; private set; }
+
+		/// <inheritdoc/>
+		public int Length
+		{
+			get { return Data.Length; }
+		}
+
+		/// <inheritdoc/>
+		public IReadOnlyList<int> LineOffsets { get; private set; }
+
+		/// <inheritdoc/>
+		public int LineCount
+		{
+			get { return LineOffsets.Count - 1; }
+		}
+
+		/// <inheritdoc/>
+		public int AllocatedSize
+		{
+			get { return Data.Length + (LineOffsets.Count * sizeof(int)); }
+		}
+
+		/// <summary>
+		/// Constructor
+		/// </summary>
+		/// <param name="Data">Data to construct from</param>
+		public ReadOnlyLogText(ReadOnlyMemory<byte> Data)
+		{
+			ReadOnlySpan<byte> Span = Data.Span;
+
+			int LineCount = LogTextHelpers.GetLineCount(Span, 0);
+
+			int[] LineOffsets = new int[LineCount + 1];
+			LogTextHelpers.FindLineOffsets(Span, 0, LineOffsets, 0);
+			Debug.Assert(LineOffsets[LineCount] == Data.Length);
+
+			this.Data = Data;
+			this.LineOffsets = LineOffsets;
+		}
+	}
+
+	/// <summary>
+	/// Mutable log text object
+	/// </summary>
+	public class LogText : ILogText
+	{
+		/// <summary>
+		/// Accessor for Data
+		/// </summary>
+		byte[] InternalData;
+
+		/// <inheritdoc/>
+		public ReadOnlyMemory<byte> Data => InternalData.AsMemory(0, Length);
+
+		/// <summary>
+		/// Hacky accessor for InternalData, for serializing to Redis
+		/// </summary>
+		public byte[] InternalDataHACK => InternalData;
+
+		/// <inheritdoc/>
+		public int Length
+		{
+			get;
+			private set;
+		}
+
+		/// <summary>
+		/// Offsets of the start of each line within the data
+		/// </summary>
+		List<int> InternalLineOffsets = new List<int> { 0 };
+
+		/// <inheritdoc/>
+		public IReadOnlyList<int> LineOffsets => InternalLineOffsets;
+
+		/// <inheritdoc/>
+		public int LineCount => InternalLineOffsets.Count - 1;
+
+		/// <summary>
+		/// Chunk Size
+		/// </summary>
+		public int MaxLength
+		{
+			get { return InternalData.Length; }
+		}
+
+		/// <summary>
+		/// Computes the size of data used by this object
+		/// </summary>
+		public int AllocatedSize
+		{
+			get { return InternalData.Length + (InternalLineOffsets.Capacity * sizeof(int)); }
+		}
+
+		/// <summary>
+		/// Constructor
+		/// </summary>
+		public LogText()
+		{
+			this.InternalData = Array.Empty<byte>();
+		}
+
+		/// <summary>
+		/// Constructor
+		/// </summary>
+		/// <param name="Data">Data to initialize this chunk with. Ownership of this array is transferred to the chunk, and its length determines the chunk size.</param>
+		/// <param name="Length">Number of valid bytes within the initial data array</param>
+		public LogText(byte[] Data, int Length)
+		{
+			this.InternalData = Data;
+			this.Length = 0; // Updated below
+
+			UpdateLength(Length);
+		}
+
+		/// <summary>
+		/// Create a new chunk data object with the given data appended. The internal buffers are reused, with the assumption that
+		/// there is no contention over writing to the same location in the chunk.
+		/// </summary>
+		/// <param name="TextData">The data to append</param>
+		/// <returns>New chunk data object</returns>
+		public void Append(ReadOnlySpan<byte> TextData)
+		{
+			TextData.CopyTo(InternalData.AsSpan(Length, TextData.Length));
+			UpdateLength(Length + TextData.Length);
+		}
+
+		/// <summary>
+		/// Updates the plain text representation of this chunk
+		/// </summary>
+		public void AppendPlainText(ILogText SrcText, int SrcLineIndex, int SrcLineCount)
+		{
+			// Convert all the lines to plain text
+			for (int Idx = 0; Idx < SrcLineCount; Idx++)
+			{
+				int LineOffset = SrcText.LineOffsets[SrcLineIndex + Idx];
+				int NextLineOffset = SrcText.LineOffsets[SrcLineIndex + Idx + 1];
+				ReadOnlySpan<byte> InputLine = SrcText.Data.Slice(LineOffset, NextLineOffset - LineOffset).Span;
+
+				// Make sure the output buffer is large enough
+				int RequiredSpace = NextLineOffset - LineOffset;
+				if (RequiredSpace > InternalData.Length - Length)
+				{
+					int MaxLineOffset = SrcText.LineOffsets[SrcLineIndex + SrcLineCount];
+					Array.Resize(ref InternalData, InternalData.Length + RequiredSpace + (MaxLineOffset - NextLineOffset) / 2);
+				}
+
+				// Convert the line to plain text
+				Length = ConvertToPlainText(InputLine, InternalData, Length);
+				InternalLineOffsets.Add(Length);
+			}
+		}
+
+		/// <summary>
+		/// Determines if the given line is empty
+		/// </summary>
+		/// <param name="Input">The input data</param>
+		/// <returns>True if the given text is empty</returns>
+		public static bool IsEmptyOrWhitespace(ReadOnlySpan<byte> Input)
+		{
+			for(int Idx = 0; Idx < Input.Length; Idx++)
+			{
+				byte Byte = Input[Idx];
+				if(Byte != (byte)'\n' && Byte != '\r' && Byte != ' ')
+				{
+					return false;
+				}
+			}
+			return true;
+		}
+
+		/// <summary>
+		/// Converts a JSON log line to plain text
+		/// </summary>
+		/// <param name="Input">The JSON data</param>
+		/// <param name="Output">Output buffer for the converted line</param>
+		/// <param name="OutputOffset">Offset within the buffer to write the converted data</param>
+		/// <returns></returns>
+		public static int ConvertToPlainText(ReadOnlySpan<byte> Input, byte[] Output, int OutputOffset)
+		{
+			if(IsEmptyOrWhitespace(Input))
+			{
+				Output[OutputOffset] = (byte)'\n';
+				return OutputOffset + 1;
+			}
+
+			Utf8JsonReader Reader = new Utf8JsonReader(Input);
+			if (Reader.Read() && Reader.TokenType == JsonTokenType.StartObject)
+			{
+				while (Reader.Read() && Reader.TokenType == JsonTokenType.PropertyName)
+				{
+					if (!Reader.ValueTextEquals("message"))
+					{
+						Reader.Skip();
+						continue;
+					}
+					if (!Reader.Read() || Reader.TokenType != JsonTokenType.String)
+					{
+						Reader.Skip();
+						continue;
+					}
+
+					int UnescapedLength = UnescapeUtf8(Reader.ValueSpan, Output.AsSpan(OutputOffset));
+					OutputOffset += UnescapedLength;
+
+					Output[OutputOffset] = (byte)'\n';
+					OutputOffset++;
+
+					break;
+				}
+			}
+			return OutputOffset;
+		}
+
+		/// <summary>
+		/// Unescape a json utf8 string
+		/// </summary>
+		/// <param name="Source">Source span of bytes</param>
+		/// <param name="Target">Target span of bytes</param>
+		/// <returns>Length of the converted data</returns>
+		static int UnescapeUtf8(ReadOnlySpan<byte> Source, Span<byte> Target)
+		{
+			int Length = 0;
+			for (; ; )
+			{
+				// Copy up to the next backslash
+				int Backslash = Source.IndexOf((byte)'\\');
+				if (Backslash == -1)
+				{
+					Source.CopyTo(Target);
+					Length += Source.Length;
+					break;
+				}
+				else if (Backslash > 0)
+				{
+					Source.Slice(0, Backslash).CopyTo(Target);
+					Source = Source.Slice(Backslash);
+					Target = Target.Slice(Backslash);
+					Length += Backslash;
+				}
+
+				// Check what the escape code is
+				if (Source[1] == 'u')
+				{
+					char[] Chars = { (char)((StringUtils.ParseHexByte(Source, 2) << 8) | StringUtils.ParseHexByte(Source, 4)) };
+					int EncodedLength = Encoding.UTF8.GetBytes(Chars.AsSpan(), Target);
+					Source = Source.Slice(6);
+					Target = Target.Slice(EncodedLength);
+					Length += EncodedLength;
+				}
+				else
+				{
+					Target[0] = Source[1] switch
+					{
+						(byte)'\"' => (byte)'\"',
+						(byte)'\\' => (byte)'\\',
+						(byte)'\b' => (byte)'\b',
+						(byte)'\f' => (byte)'\f',
+						(byte)'\n' => (byte)'\n',
+						(byte)'\r' => (byte)'\r',
+						(byte)'\t' => (byte)'\t',
+						_ => Source[1]
+					};
+					Source = Source.Slice(2);
+					Target = Target.Slice(1);
+					Length++;
+				}
+			}
+			return Length;
+		}
+
+		/// <summary>
+		/// Generates placeholder data for a missing span
+		/// </summary>
+		/// <param name="ChunkIdx">Index of the chunk</param>
+		/// <param name="HostName">Host name of the machine that was holding the data</param>
+		/// <param name="TargetLength">Desired length of the buffer</param>
+		/// <param name="TargetLineCount">Desired line count for the buffer</param>
+		public void AppendMissingDataInfo(int ChunkIdx, string? HostName, int TargetLength, int TargetLineCount)
+		{
+			int NewLength = Length;
+
+			if(InternalData.Length < TargetLength)
+			{
+				Array.Resize(ref InternalData, TargetLength);
+			}
+
+			string Suffix = $" (server: {HostName})";
+			for(int NewLineCount = LineCount; NewLineCount < TargetLineCount; NewLineCount++)
+			{
+				byte[] ErrorBytes = Encoding.ASCII.GetBytes($"{{ \"level\":\"Error\",\"message\":\"[Missing data at chunk {ChunkIdx}, line {NewLineCount}{Suffix}]\" }}\n");
+				if (NewLength + ErrorBytes.Length > TargetLength)
+				{
+					break;
+				}
+
+				ErrorBytes.AsSpan().CopyTo(InternalData.AsSpan(NewLength));
+				NewLength += ErrorBytes.Length;
+
+				Suffix = String.Empty;
+			}
+
+			if (NewLength < TargetLength)
+			{
+				NewLength = Math.Max(NewLength - 1, 0);
+				InternalData.AsSpan(NewLength, (int)((TargetLength - 1) - NewLength)).Fill((byte)' ');
+				InternalData[TargetLength - 1] = (byte)'\n';
+			}
+
+			UpdateLength((int)TargetLength);
+		}
+
+		/// <summary>
+		/// Shrinks the data allocated by this chunk to the minimum required
+		/// </summary>
+		public void Shrink()
+		{
+			if (Length < InternalData.Length)
+			{
+				Array.Resize(ref InternalData, Length);
+			}
+		}
+
+		/// <summary>
+		/// Updates the length of this chunk, computing all the newline offsets
+		/// </summary>
+		/// <param name="NewLength">New length of the chunk</param>
+		private void UpdateLength(int NewLength)
+		{
+			if (NewLength > Length)
+			{
+				// Make sure the data ends with a newline
+				if (InternalData[NewLength - 1] != '\n')
+				{
+					throw new InvalidDataException("Chunk data must end with a newline");
+				}
+
+				// Calculate the new number of newlines
+				ReadOnlySpan<byte> NewData = InternalData.AsSpan(0, NewLength);
+				for(int Idx = Length; Idx < NewLength; Idx++)
+				{
+					if(NewData[Idx] == '\n')
+					{
+						InternalLineOffsets.Add(Idx + 1);
+					}
+				}
+
+				// Update the size of the buffer
+				Length = NewLength;
+			}
+		}
+	}
+
+	/// <summary>
+	/// Extension methods for ILogText
+	/// </summary>
+	public static class LogTextExtensions
+	{
+		/// <summary>
+		/// Find the line index for a particular offset
+		/// </summary>
+		/// <param name="LogText">The text to search</param>
+		/// <param name="Offset">Offset within the text</param>
+		/// <returns>The line index</returns>
+		public static int GetLineIndexForOffset(this ILogText LogText, int Offset)
+		{
+			int LineIdx = LogText.LineOffsets.BinarySearch(Offset);
+			if(LineIdx < 0)
+			{
+				LineIdx = ~LineIdx - 1;
+			}
+			return LineIdx;
+		}
+
+		/// <summary>
+		/// Converts a log text instance to plain text
+		/// </summary>
+		/// <param name="LogText">The text to convert</param>
+		/// <returns>The plain text instance</returns>
+		public static ILogText ToPlainText(this ILogText LogText)
+		{
+			LogText Other = new LogText();
+			Other.AppendPlainText(LogText, 0, LogText.LineCount);
+			return Other;
+		}
+	}
+}
diff --git a/Engine/Source/Programs/Horde/HordeServer/Logs/LogToken.cs b/Engine/Source/Programs/Horde/HordeServer/Logs/LogToken.cs
new file mode 100644
--- /dev/null
+++ b/Engine/Source/Programs/Horde/HordeServer/Logs/LogToken.cs
@@ -0,0 +1,205 @@
+using System;
+using System.Collections.Generic;
+using System.Linq;
+using System.Threading.Tasks;
+
+namespace HordeServer.Logs
+{
+	/// <summary>
+	/// Functionality for decomposing log text into tokens
+	/// </summary>
+	public static class LogToken
+	{
+		/// <summary>
+		/// Maximum number of bytes in each token
+		/// </summary>
+		public const int MaxTokenBytes = 4;
+
+		/// <summary>
+		/// Number of bits in each token
+		/// </summary>
+		public const int MaxTokenBits = MaxTokenBytes * 8;
+
+		/// <summary>
+		/// Lookup from input byte to token type
+		/// </summary>
+		static readonly byte[] TokenTypes = GetTokenTypes();
+
+		/// <summary>
+		/// Lookup from input byte to token char
+		/// </summary>
+		static readonly byte[] TokenChars = GetTokenChars();
+
+		/// <summary>
+		/// Get a set of tokens in a trie
+		/// </summary>
+		/// <param name="Text">Text to scan</param>
+		/// <returns></returns>
+		public static ReadOnlyTrie GetTrie(ReadOnlySpan<byte> Text)
+		{
+			ReadOnlyTrieBuilder Builder = new ReadOnlyTrieBuilder();
+			GetTokens(Text, Builder.Add);
+			return Builder.Build();
+		}
+
+		/// <summary>
+		/// Gets a single token
+		/// </summary>
+		/// <param name="Text">The text to parse</param>
+		/// <returns>The token value</returns>
+		public static ulong GetToken(ReadOnlySpan<byte> Text)
+		{
+			ulong Token = 0;
+			GetTokens(Text, x => Token = x);
+			return Token;
+		}
+
+		/// <summary>
+		/// Get a set of tokens in a trie
+		/// </summary>
+		/// <param name="Text">Text to scan</param>
+		/// <returns></returns>
+		public static HashSet<ulong> GetTokenSet(ReadOnlySpan<byte> Text)
+		{
+			HashSet<ulong> Tokens = new HashSet<ulong>();
+			GetTokens(Text, x => Tokens.Add(x));
+			return Tokens;
+		}
+
+		/// <summary>
+		/// Decompose a span of text into tokens
+		/// </summary>
+		/// <param name="Text">Text to scan</param>
+		/// <param name="AddToken">Receives a set of tokens</param>
+		public static void GetTokens(ReadOnlySpan<byte> Text, Action<ulong> AddToken)
+		{
+			if (Text.Length > 0)
+			{
+				int Type = TokenTypes[Text[0]];
+				int NumTokenBits = 8;
+				ulong Token = TokenChars[Text[0]];
+
+				for (int TextIdx = 1; TextIdx < Text.Length; TextIdx++)
+				{
+					byte NextChar = TokenChars[Text[TextIdx]];
+					int NextType = TokenTypes[NextChar];
+					if (Type != NextType || NumTokenBits + 8 > MaxTokenBits)
+					{
+						AddToken(Token << (MaxTokenBits - NumTokenBits));
+						Token = 0;
+						NumTokenBits = 0;
+						Type = NextType;
+					}
+					Token = (Token << 8) | NextChar;
+					NumTokenBits += 8;
+				}
+
+				AddToken(Token << (MaxTokenBits - NumTokenBits));
+			}
+		}
+
+		/// <summary>
+		/// Gets the length of the first token in the given span
+		/// </summary>
+		/// <param name="Text">The text to search</param>
+		/// <param name="Pos">Start position for the search</param>
+		/// <returns>Length of the first token</returns>
+		public static ReadOnlySpan<byte> GetTokenText(ReadOnlySpan<byte> Text, int Pos)
+		{
+			int Type = TokenTypes[Text[Pos]];
+			for (int End = Pos + 1; ; End++)
+			{
+				if (End == Text.Length || TokenTypes[Text[End]] != Type)
+				{
+					return Text.Slice(Pos, End - Pos);
+				}
+			}
+		}
+
+		/// <summary>
+		/// Gets the length of the first token in the given span
+		/// </summary>
+		/// <param name="Text">The text to search</param>
+		/// <param name="Offset">Offset of the window to read from the token</param>
+		/// <returns>Length of the first token</returns>
+		public static ulong GetWindowedTokenValue(ReadOnlySpan<byte> Text, int Offset)
+		{
+			ulong Token = 0;
+			for (int Idx = 0; Idx < MaxTokenBytes; Idx++)
+			{
+				Token <<= 8;
+				if (Offset >= 0 && Offset < Text.Length)
+				{
+					Token |= TokenChars[Text[Offset]];
+				}
+				Offset++;
+			}
+			return Token;
+		}
+
+		/// <summary>
+		/// Gets the length of the first token in the given span
+		/// </summary>
+		/// <param name="Text">The text to search</param>
+		/// <param name="Offset">Offset of the window to read from the token</param>
+		/// <param name="bAllowPartialMatch">Whether to allow only matching the start of the string</param>
+		/// <returns>Length of the first token</returns>
+		public static ulong GetWindowedTokenMask(ReadOnlySpan<byte> Text, int Offset, bool bAllowPartialMatch)
+		{
+			ulong Token = 0;
+			for (int Idx = 0; Idx < MaxTokenBytes; Idx++)
+			{
+				Token <<= 8;
+				if (Offset >= 0 && (Offset < Text.Length || !bAllowPartialMatch))
+				{
+					Token |= 0xff;
+				}
+				Offset++;
+			}
+			return Token;
+		}
+
+		/// <summary>
+		/// Build the lookup table for token types
+		/// </summary>
+		/// <returns>Array whose elements map from an input byte to token type</returns>
+		static byte[] GetTokenTypes()
+		{
+			byte[] CharTypes = new byte[256];
+			for (int Idx = 'a'; Idx <= 'z'; Idx++)
+			{
+				CharTypes[Idx] = 1;
+			}
+			for (int Idx = 'A'; Idx <= 'Z'; Idx++)
+			{
+				CharTypes[Idx] = 1;
+			}
+			for (int Idx = '0'; Idx <= '9'; Idx++)
+			{
+				CharTypes[Idx] = 2;
+			}
+			CharTypes[' '] = 3;
+			CharTypes['\t'] = 3;
+			CharTypes['\n'] = 4;
+			return CharTypes;
+		}
+
+		/// <summary>
+		/// Build the lookup table for token types
+		/// </summary>
+		/// <returns>Array whose elements map from an input byte to token type</returns>
+		static byte[] GetTokenChars()
+		{
+			byte[] Chars = new byte[256];
+			for(int Idx = 0; Idx < 256; Idx++)
+			{
+				Chars[Idx] = (byte)Idx;
+			}
+			for (int Idx = 'A'; Idx <= 'Z'; Idx++)
+			{
+				Chars[Idx] = (byte)('a' + Idx - 'A');
+			}
+			return Chars;
+		}
+	}
+}
diff --git a/Engine/Source/Programs/Horde/HordeServer/Logs/Storage/ILogStorage.cs b/Engine/Source/Programs/Horde/HordeServer/Logs/Storage/ILogStorage.cs
new file mode 100644
--- /dev/null
+++ b/Engine/Source/Programs/Horde/HordeServer/Logs/Storage/ILogStorage.cs
@@ -0,0 +1,54 @@
+using HordeServer.Models;
+using HordeServer.Services;
+using HordeServer.Storage;
+using HordeServer.Utilities;
+using Microsoft.Extensions.Caching.Memory;
+using MongoDB.Bson;
+using System;
+using System.Collections.Generic;
+using System.Linq;
+using System.Threading.Tasks;
+
+namespace HordeServer.Logs
+{
+	/// <summary>
+	/// Caching interface for reading and writing log data.
+	/// </summary>
+	public interface ILogStorage : IDisposable
+	{
+		/// <summary>
+		/// Attempts to read an index for the given log file
+		/// </summary>
+		/// <param name="LogId">Unique id of the log file</param>
+		/// <param name="Length">Length of the file covered by the index</param>
+		/// <returns>Index for the log file</returns>
+		Task<LogIndexData?> ReadIndexAsync(ObjectId LogId, long Length);
+
+		/// <summary>
+		/// Log file to write an index for
+		/// </summary>
+		/// <param name="LogId">Unique id of the log file</param>
+		/// <param name="Length">Length of the file covered by the index</param>
+		/// <param name="Index">The log file index</param>
+		/// <returns>Async task</returns>
+		Task WriteIndexAsync(ObjectId LogId, long Length, LogIndexData Index);
+
+		/// <summary>
+		/// Retrieves an item from the cache
+		/// </summary>
+		/// <param name="LogId">Unique id of the log file</param>
+		/// <param name="Offset">Offset of the chunk to read</param>
+		/// <param name="LineIndex">First line of the chunk</param>
+		/// <returns>Data for the given key, or null if it's not present</returns>
+		Task<LogChunkData?> ReadChunkAsync(ObjectId LogId, long Offset, int LineIndex);
+
+		/// <summary>
+		/// Writes a chunk to storage
+		/// </summary>
+		/// <param name="LogId">Unique id of the log file</param>
+		/// <param name="Offset">Offset of the chunk to write</param>
+		/// <param name="ChunkData">Information about the chunk data</param>
+		/// <returns>Async task</returns>
+		Task WriteChunkAsync(ObjectId LogId, long Offset, LogChunkData ChunkData);
+	}
+}
diff --git a/Engine/Source/Programs/Horde/HordeServer/Logs/Storage/Impl/LocalLogStorage.cs b/Engine/Source/Programs/Horde/HordeServer/Logs/Storage/Impl/LocalLogStorage.cs
new file mode 100644
--- /dev/null
+++ b/Engine/Source/Programs/Horde/HordeServer/Logs/Storage/Impl/LocalLogStorage.cs
@@ -0,0 +1,142 @@
+using HordeServer.Models;
+using Microsoft.Extensions.Caching.Memory;
+using MongoDB.Bson;
+using System;
+using System.Collections.Generic;
+using System.Linq;
+using System.Threading.Tasks;
+
+namespace HordeServer.Logs.Storage
+{
+	/// <summary>
+	/// In-memory cache for chunk and index data
+	/// </summary>
+	class LocalLogStorage : ILogStorage
+	{
+		/// <summary>
+		/// The memory cache instance
+		/// </summary>
+		MemoryCache Cache;
+
+		/// <summary>
+		/// The inner storage
+		/// </summary>
+		ILogStorage Inner;
+
+		/// <summary>
+		/// Constructor
+		/// </summary>
+		/// <param name="NumItems">Maximum size for the memory cache</param>
+		/// <param name="Inner">The inner storage provider</param>
+		public LocalLogStorage(int NumItems, ILogStorage Inner)
+		{
+			MemoryCacheOptions Options = new MemoryCacheOptions();
+			Options.SizeLimit = NumItems;
+
+			this.Cache = new MemoryCache(Options);
+			this.Inner = Inner;
+		}
+
+		/// <inheritdoc/>
+		public void Dispose()
+		{
+			Cache.Dispose();
+			Inner.Dispose();
+		}
+
+		/// <summary>
+		/// Gets the cache key for a particular index
+		/// </summary>
+		/// <param name="LogId">The log file to retrieve an index for</param>
+		/// <param name="Length">Length of the file covered by the index</param>
+		/// <returns>Cache key for the index</returns>
+		static string IndexKey(ObjectId LogId, long Length)
+		{
+			return $"{LogId}/index-{Length}";
+		}
+
+		/// <summary>
+		/// Gets the cache key for a particular chunk
+		/// </summary>
+		/// <param name="LogId">The log file to retrieve an index for</param>
+		/// <param name="Offset">The chunk offset</param>
+		/// <returns>Cache key for the chunk</returns>
+		static string ChunkKey(ObjectId LogId, long Offset)
+		{
+			return $"{LogId}/chunk-{Offset}";
+		}
+
+		/// <summary>
+		/// Adds an entry to the cache
+		/// </summary>
+		/// <param name="Key">The cache key</param>
+		/// <param name="Value">Value to store</param>
+		void AddEntry(string Key, object? Value)
+		{
+			using (ICacheEntry Entry = Cache.CreateEntry(Key))
+			{
+				if (Value == null)
+				{
+					Entry.SetAbsoluteExpiration(TimeSpan.FromSeconds(30.0));
+					Entry.SetSize(0);
+				}
+				else
+				{
+					Entry.SetSlidingExpiration(TimeSpan.FromHours(4));
+					Entry.SetSize(1);
+				}
+				Entry.SetValue(Value);
+			}
+		}
+
+		/// <summary>
+		/// Reads a value from the cache, or from the inner storage provider, adding the result to the cache
+		/// </summary>
+		/// <typeparam name="T">Type of object to read</typeparam>
+		/// <param name="Key">The cache key</param>
+		/// <param name="ReadInner">Delegate to read from the inner storage provider</param>
+		/// <returns>The retrieved value</returns>
+		async Task<T?> ReadValueAsync<T>(string Key, Func<Task<T?>> ReadInner) where T : class
+		{
+			T? Value;
+			if (!Cache.TryGetValue(Key, out Value))
+			{
+				try
+				{
+					Value = await ReadInner();
+				}
+				finally
+				{
+					AddEntry(Key, Value);
+				}
+			}
+			return Value;
+		}
+
+		/// <inheritdoc/>
+		public Task<LogIndexData?> ReadIndexAsync(ObjectId LogId, long Length)
+		{
+			return ReadValueAsync(IndexKey(LogId, Length), () => Inner.ReadIndexAsync(LogId, Length));
+		}
+
+		/// <inheritdoc/>
+		public Task WriteIndexAsync(ObjectId LogId, long Length, LogIndexData IndexData)
+		{
+			AddEntry(IndexKey(LogId, Length), IndexData);
+			return Inner.WriteIndexAsync(LogId, Length, IndexData);
+		}
+
+		/// <inheritdoc/>
+		public Task<LogChunkData?> ReadChunkAsync(ObjectId LogId, long Offset, int LineIndex)
+		{
+			return ReadValueAsync(ChunkKey(LogId, Offset), () => Inner.ReadChunkAsync(LogId, Offset, LineIndex));
+		}
+
+		/// <inheritdoc/>
+		public Task WriteChunkAsync(ObjectId LogId, long Offset, LogChunkData ChunkData)
+		{
+			AddEntry(ChunkKey(LogId, Offset), ChunkData);
+			return Inner.WriteChunkAsync(LogId, Offset, ChunkData);
+		}
+	}
+}
diff --git a/Engine/Source/Programs/Horde/HordeServer/Logs/Storage/Impl/NullLogStorage.cs b/Engine/Source/Programs/Horde/HordeServer/Logs/Storage/Impl/NullLogStorage.cs
new file mode 100644
--- /dev/null
+++ b/Engine/Source/Programs/Horde/HordeServer/Logs/Storage/Impl/NullLogStorage.cs
@@ -0,0 +1,43 @@
+using MongoDB.Bson;
+using System;
+using System.Collections.Generic;
+using System.Linq;
+using System.Threading.Tasks;
+
+namespace HordeServer.Logs.Storage.Impl
+{
+	/// <summary>
+	/// Empty implementation of log storage
+	/// </summary>
+	public sealed class NullLogStorage : ILogStorage
+	{
+		/// <inheritdoc/>
+		public void Dispose()
+		{
+		}
+
+		/// <inheritdoc/>
+		public Task<LogIndexData?> ReadIndexAsync(ObjectId LogId, long Length)
+		{
+			return Task.FromResult<LogIndexData?>(null);
+		}
+
+		/// <inheritdoc/>
+		public Task WriteIndexAsync(ObjectId LogId, long Length, LogIndexData Index)
+		{
+			return Task.CompletedTask;
+		}
+
+		/// <inheritdoc/>
+		public Task<LogChunkData?> ReadChunkAsync(ObjectId LogId, long Offset, int LineIndex)
+		{
+			return Task.FromResult<LogChunkData?>(null);
+		}
+
+		/// <inheritdoc/>
+		public Task WriteChunkAsync(ObjectId LogId, long Offset, LogChunkData ChunkData)
+		{
+			return Task.CompletedTask;
+		}
+	}
+}
diff --git a/Engine/Source/Programs/Horde/HordeServer/Logs/Storage/Impl/PersistentLogStorage.cs b/Engine/Source/Programs/Horde/HordeServer/Logs/Storage/Impl/PersistentLogStorage.cs
new file mode 100644
--- /dev/null
+++ b/Engine/Source/Programs/Horde/HordeServer/Logs/Storage/Impl/PersistentLogStorage.cs
@@ -0,0 +1,107 @@
+using EpicGames.Core;
+using HordeServer.Models;
+using HordeServer.Storage;
+using HordeServer.Utilities;
+using Microsoft.Extensions.Logging;
+using MongoDB.Bson;
+using System;
+using System.Collections.Generic;
+using System.Linq;
+using System.Threading.Tasks;
+
+namespace HordeServer.Logs.Readers
+{
+	/// <summary>
+	/// Bulk storage for log file data
+	/// </summary>
+	class PersistentLogStorage : ILogStorage
+	{
+		/// <summary>
+		/// The bulk storage provider to use
+		/// </summary>
+		IStorageBackend StorageProvider;
+
+		/// <summary>
+		/// Log provider
+		/// </summary>
+		ILogger Logger;
+
+		/// <summary>
+		/// Constructor
+		/// </summary>
+		/// <param name="StorageProvider">The storage provider</param>
+		/// <param name="Logger">Logging provider</param>
+		public PersistentLogStorage(IStorageBackend StorageProvider, ILogger Logger)
+		{
+			this.StorageProvider = StorageProvider;
+			this.Logger = Logger;
+		}
+
+		/// <inheritdoc/>
+		public void Dispose()
+		{
+			StorageProvider.Dispose();
+		}
+
+		/// <inheritdoc/>
+		public async Task<LogIndexData?> ReadIndexAsync(ObjectId LogId, long Length)
+		{
+			Logger.LogDebug("Reading log {LogId} index length {Length} from persistent storage", LogId, Length);
+
+			string Path = $"{LogId}/index_{Length}";
+			ReadOnlyMemory<byte>? Data = await StorageProvider.ReadAsync(Path);
+			if (Data == null)
+			{
+				return null;
+			}
+			return LogIndexData.FromMemory(Data.Value);
+		}
+
+		/// <inheritdoc/>
+		public Task WriteIndexAsync(ObjectId LogId, long Length, LogIndexData IndexData)
+		{
+			Logger.LogDebug("Writing log {LogId} index length {Length} to persistent storage", LogId, Length);
+
+			string Path = $"{LogId}/index_{Length}";
+			ReadOnlyMemory<byte> Data = IndexData.ToByteArray();
+			return StorageProvider.WriteAsync(Path, Data);
+		}
+
+		/// <inheritdoc/>
+		public async Task<LogChunkData?> ReadChunkAsync(ObjectId LogId, long Offset, int LineIndex)
+		{
+			Logger.LogDebug("Reading log {LogId} chunk offset {Offset} from persistent storage", LogId, Offset);
+
+			string Path = $"{LogId}/offset_{Offset}";
+			ReadOnlyMemory<byte>? Data = await StorageProvider.ReadAsync(Path);
+			if(Data == null)
+			{
+				return null;
+			}
+
+			MemoryReader Reader = new MemoryReader(Data.Value);
+			LogChunkData ChunkData = Reader.ReadLogChunkData(Offset, LineIndex);
+
+			if (Reader.Offset != Data.Value.Length)
+			{
+				throw new Exception($"Serialization of persistent chunk {Path} is not at expected offset (expected {Data.Value.Length}, actual {Reader.Offset})");
+			}
+
+			return ChunkData;
+		}
+
+		/// <inheritdoc/>
+		public Task WriteChunkAsync(ObjectId LogId, long Offset, LogChunkData ChunkData)
+		{
+			Logger.LogDebug("Writing log {LogId} chunk offset {Offset} to persistent storage", LogId, Offset);
+
+			string Path = $"{LogId}/offset_{Offset}";
+			byte[] Data = new byte[ChunkData.GetSerializedSize()];
+			MemoryWriter Writer = new MemoryWriter(Data);
+			Writer.WriteLogChunkData(ChunkData);
+			Writer.CheckOffset(Data.Length);
+
+			return StorageProvider.WriteAsync(Path, Data);
+		}
+	}
+}
