Initial drop of Horde Server/Agent for UE5/Main. Still missing some Epic-specific functionality required to build/run correctly.

diff --git a/Engine/Source/Programs/Horde/HordeServer/Logs/Storage/Impl/RedisLogStorage.cs b/Engine/Source/Programs/Horde/HordeServer/Logs/Storage/Impl/RedisLogStorage.cs
new file mode 100644
--- /dev/null
+++ b/Engine/Source/Programs/Horde/HordeServer/Logs/Storage/Impl/RedisLogStorage.cs
@@ -0,0 +1,169 @@
+using Microsoft.Extensions.Logging;
+using MongoDB.Bson;
+using StackExchange.Redis;
+using System;
+using System.Collections.Generic;
+using System.Diagnostics.CodeAnalysis;
+using System.Linq;
+using System.Threading.Tasks;
+
+namespace HordeServer.Logs.Storage.Impl
+{
+	/// <summary>
+	/// Redis log file storage
+	/// </summary>
+	public sealed class RedisLogStorage : ILogStorage
+	{
+		/// <summary>
+		/// The Redis database
+		/// </summary>
+		IDatabase RedisDb;
+
+		/// <summary>
+		/// Logging device
+		/// </summary>
+		ILogger Logger;
+
+		/// <summary>
+		/// Inner storage layer
+		/// </summary>
+		ILogStorage Inner;
+
+		/// <summary>
+		/// Constructor
+		/// </summary>
+		/// <param name="RedisDb">Redis database</param>
+		/// <param name="Logger">Logger instance</param>
+		/// <param name="Inner">Inner storage layer</param>
+		public RedisLogStorage(IDatabase RedisDb, ILogger Logger, ILogStorage Inner)
+		{
+			this.RedisDb = RedisDb;
+			this.Logger = Logger;
+			this.Inner = Inner;
+		}
+
+		/// <inheritdoc/>
+		public void Dispose()
+		{
+			Inner.Dispose();
+		}
+
+		/// <summary>
+		/// Gets the key for a log file's index
+		/// </summary>
+		/// <param name="LogId">The log file id</param>
+		/// <param name="Length">Length of the file covered by the index</param>
+		/// <returns>The index key</returns>
+		static string IndexKey(ObjectId LogId, long Length)
+		{
+			return $"log-{LogId}-index-{Length}";
+		}
+
+		/// <summary>
+		/// Gets the key for a chunk's data
+		/// </summary>
+		/// <param name="LogId">The log file id</param>
+		/// <param name="Offset">Offset of the chunk within the log file</param>
+		/// <returns>The chunk key</returns>
+		static string ChunkKey(ObjectId LogId, long Offset)
+		{
+			return $"log-{LogId}-chunk-{Offset}";
+		}
+
+		/// <summary>
+		/// Adds an index to the cache
+		/// </summary>
+		/// <param name="Key">Key for the item to store</param>
+		/// <param name="Value">Value to be stored</param>
+		/// <returns>Async task</returns>
+		[SuppressMessage("Design", "CA1031:Do not catch general exception types")]
+		async Task AddAsync(string Key, ReadOnlyMemory<byte> Value)
+		{
+			try
+			{
+				await RedisDb.StringSetAsync(Key, Value, expiry: TimeSpan.FromHours(1.0));
+				Logger.LogDebug("Added key {Key} to Redis cache ({Size} bytes)", Key, Value.Length);
+			}
+			catch (Exception Ex)
+			{
+				Logger.LogError(Ex, "Error writing Redis key {Key}", Key);
+			}
+		}
+
+		/// <summary>
+		/// Adds an index to the cache
+		/// </summary>
+		/// <param name="Key">Key for the item to store</param>
+		/// <param name="Deserialize">Delegate to deserialize the item</param>
+		/// <returns>Async task</returns>
+		[SuppressMessage("Design", "CA1031:Do not catch general exception types")]
+		async Task<T?> GetAsync<T>(string Key, Func<ReadOnlyMemory<byte>, T> Deserialize) where T : class
+		{
+			try
+			{
+				RedisValue Value = await RedisDb.StringGetAsync(Key);
+				if (Value.IsNullOrEmpty)
+				{
+					Logger.LogDebug("Redis cache miss for {Key}", Key);
+					return null;
+				}
+				else
+				{
+					Logger.LogDebug("Redis cache hit for {Key}", Key);
+					return Deserialize(Value);
+				}
+			}
+			catch (Exception Ex)
+			{
+				Logger.LogError(Ex, "Error reading Redis key {Key}", Key);
+				return null;
+			}
+		}
+
+		/// <inheritdoc/>
+		public async Task<LogIndexData?> ReadIndexAsync(ObjectId LogId, long Length)
+		{
+			string Key = IndexKey(LogId, Length);
+
+			LogIndexData? IndexData = await GetAsync(Key, Memory => LogIndexData.FromMemory(Memory));
+			if(IndexData == null)
+			{
+				IndexData = await Inner.ReadIndexAsync(LogId, Length);
+				if(IndexData != null)
+				{
+					await AddAsync(Key, IndexData.ToByteArray());
+				}
+			}
+			return IndexData;
+		}
+
+		/// <inheritdoc/>
+		public async Task WriteIndexAsync(ObjectId LogId, long Length, LogIndexData Index)
+		{
+			await Inner.WriteIndexAsync(LogId, Length, Index);
+		}
+
+		/// <inheritdoc/>
+		public async Task<LogChunkData?> ReadChunkAsync(ObjectId LogId, long Offset, int LineIndex)
+		{
+			string Key = ChunkKey(LogId, Offset);
+
+			LogChunkData? ChunkData = await GetAsync(Key, Memory => LogChunkData.FromMemory(Memory, Offset, LineIndex));
+			if (ChunkData == null)
+			{
+				ChunkData = await Inner.ReadChunkAsync(LogId, Offset, LineIndex);
+				if (ChunkData != null)
+				{
+					await AddAsync(Key, ChunkData.ToByteArray());
+				}
+			}
+			return ChunkData;
+		}
+
+		/// <inheritdoc/>
+		public async Task WriteChunkAsync(ObjectId LogId, long Offset, LogChunkData ChunkData)
+		{
+			await Inner.WriteChunkAsync(LogId, Offset, ChunkData);
+		}
+	}
+}
diff --git a/Engine/Source/Programs/Horde/HordeServer/Logs/Storage/Impl/SequencedLogStorage.cs b/Engine/Source/Programs/Horde/HordeServer/Logs/Storage/Impl/SequencedLogStorage.cs
new file mode 100644
--- /dev/null
+++ b/Engine/Source/Programs/Horde/HordeServer/Logs/Storage/Impl/SequencedLogStorage.cs
@@ -0,0 +1,191 @@
+using HordeServer.Models;
+using MongoDB.Bson;
+using System;
+using System.Collections.Generic;
+using System.Linq;
+using System.Threading.Tasks;
+
+namespace HordeServer.Logs.Storage.Impl
+{
+	/// <summary>
+	/// Storage layer which caches pending read tasks, to avoid fetching the same item more than once
+	/// </summary>
+	class SequencedLogStorage : ILogStorage
+	{
+		/// <summary>
+		/// Inner storage implementation
+		/// </summary>
+		ILogStorage Inner;
+
+		/// <summary>
+		/// Pending index reads
+		/// </summary>
+		Dictionary<(ObjectId, long), Task<LogIndexData?>> IndexReadTasks = new Dictionary<(ObjectId, long), Task<LogIndexData?>>();
+
+		/// <summary>
+		/// Pending index reads
+		/// </summary>
+		Dictionary<(ObjectId, long), Task> IndexWriteTasks = new Dictionary<(ObjectId, long), Task>();
+
+		/// <summary>
+		/// Pending chunk reads
+		/// </summary>
+		Dictionary<(ObjectId, long), Task<LogChunkData?>> ChunkReadTasks = new Dictionary<(ObjectId, long), Task<LogChunkData?>>();
+
+		/// <summary>
+		/// Pending chunk reads
+		/// </summary>
+		Dictionary<(ObjectId, long), Task> ChunkWriteTasks = new Dictionary<(ObjectId, long), Task>();
+
+		/// <summary>
+		/// Constructor
+		/// </summary>
+		/// <param name="Inner">The inner storage provider</param>
+		public SequencedLogStorage(ILogStorage Inner)
+		{
+			this.Inner = Inner;
+		}
+
+		/// <inheritdoc/>
+		public void Dispose()
+		{
+			Inner.Dispose();
+		}
+
+		/// <inheritdoc/>
+		public Task<LogIndexData?> ReadIndexAsync(ObjectId LogId, long Length)
+		{
+			Task<LogIndexData?>? Task;
+			lock (IndexReadTasks)
+			{
+				if (!IndexReadTasks.TryGetValue((LogId, Length), out Task))
+				{
+					Task = InnerReadIndexAsync(LogId, Length);
+					IndexReadTasks.Add((LogId, Length), Task);
+				}
+			}
+			return Task;
+		}
+
+		/// <inheritdoc/>
+		public Task WriteIndexAsync(ObjectId LogId, long Length, LogIndexData IndexData)
+		{
+			Task? Task;
+			lock (IndexWriteTasks)
+			{
+				if (!IndexWriteTasks.TryGetValue((LogId, Length), out Task))
+				{
+					Task = InnerWriteIndexAsync(LogId, Length, IndexData);
+					IndexWriteTasks.Add((LogId, Length), Task);
+				}
+			}
+			return Task;
+		}
+
+		/// <inheritdoc/>
+		public Task<LogChunkData?> ReadChunkAsync(ObjectId LogId, long Offset, int LineIndex)
+		{
+			Task<LogChunkData?>? Task;
+			lock (ChunkReadTasks)
+			{
+				if (!ChunkReadTasks.TryGetValue((LogId, Offset), out Task))
+				{
+					Task = InnerReadChunkAsync(LogId, Offset, LineIndex);
+					ChunkReadTasks[(LogId, Offset)] = Task;
+				}
+			}
+			return Task;
+		}
+
+		/// <inheritdoc/>
+		public Task WriteChunkAsync(ObjectId LogId, long Offset, LogChunkData ChunkData)
+		{
+			Task? Task;
+			lock (ChunkWriteTasks)
+			{
+				if (!ChunkWriteTasks.TryGetValue((LogId, Offset), out Task))
+				{
+					Task = InnerWriteChunkAsync(LogId, Offset, ChunkData);
+					ChunkWriteTasks[(LogId, Offset)] = Task;
+				}
+			}
+			return Task;
+		}
+
+		/// <summary>
+		/// Wrapper for reading an index from the inner storage provider
+		/// </summary>
+		/// <param name="LogId">The log file to read the index for</param>
+		/// <param name="Length">Length of the file that's indexed</param>
+		/// <returns>The index data</returns>
+		async Task<LogIndexData?> InnerReadIndexAsync(ObjectId LogId, long Length)
+		{
+			await Task.Yield();
+
+			LogIndexData? IndexData = await Inner.ReadIndexAsync(LogId, Length);
+			lock (IndexReadTasks)
+			{
+				IndexReadTasks.Remove((LogId, Length));
+			}
+			return IndexData;
+		}
+
+		/// <summary>
+		/// Wrapper for reading an index from the inner storage provider
+		/// </summary>
+		/// <param name="LogId">The log file to read the index for</param>
+		/// <param name="Length">Length of the indexed data</param>
+		/// <param name="IndexData">The index data to write</param>
+		/// <returns>The index data</returns>
+		async Task InnerWriteIndexAsync(ObjectId LogId, long Length, LogIndexData IndexData)
+		{
+			await Task.Yield();
+
+			await Inner.WriteIndexAsync(LogId, Length, IndexData);
+
+			lock (IndexWriteTasks)
+			{
+				IndexWriteTasks.Remove((LogId, Length));
+			}
+		}
+
+		/// <summary>
+		/// Wrapper for reading a chunk from the inner storage provider
+		/// </summary>
+		/// <param name="LogId">The log file to read the index for</param>
+		/// <param name="Offset">Offset of the chunk to read</param>
+		/// <param name="LineIndex">Index of the first line in this chunk</param>
+		/// <returns>The index data</returns>
+		async Task<LogChunkData?> InnerReadChunkAsync(ObjectId LogId, long Offset, int LineIndex)
+		{
+			await Task.Yield();
+
+			LogChunkData? ChunkData = await Inner.ReadChunkAsync(LogId, Offset, LineIndex);
+			lock (ChunkReadTasks)
+			{
+				ChunkReadTasks.Remove((LogId, Offset));
+			}
+			return ChunkData;
+		}
+
+		/// <summary>
+		/// Wrapper for reading a chunk from the inner storage provider
+		/// </summary>
+		/// <param name="LogId">The log file to write the chunk for</param>
+		/// <param name="Offset">Offset of the chunk within the log</param>
+		/// <param name="ChunkData">The chunk data</param>
+		/// <returns>The index data</returns>
+		async Task<LogChunkData> InnerWriteChunkAsync(ObjectId LogId, long Offset, LogChunkData ChunkData)
+		{
+			await Task.Yield();
+
+			await Inner.WriteChunkAsync(LogId, Offset, ChunkData);
+
+			lock (ChunkWriteTasks)
+			{
+				ChunkWriteTasks.Remove((LogId, Offset));
+			}
+			return ChunkData;
+		}
+	}
+}
diff --git a/Engine/Source/Programs/Horde/HordeServer/Models/Agent.cs b/Engine/Source/Programs/Horde/HordeServer/Models/Agent.cs
new file mode 100644
--- /dev/null
+++ b/Engine/Source/Programs/Horde/HordeServer/Models/Agent.cs
@@ -0,0 +1,830 @@
+// Copyright Epic Games, Inc. All Rights Reserved.
+
+using Google.Protobuf.WellKnownTypes;
+using HordeServer.Api;
+using HordeCommon;
+using HordeCommon.Rpc.Tasks;
+using HordeServer.Services;
+using HordeServer.Utilities;
+using MongoDB.Bson;
+using MongoDB.Bson.Serialization;
+using MongoDB.Bson.Serialization.Attributes;
+using System;
+using System.Collections.Generic;
+using System.Diagnostics.CodeAnalysis;
+using System.IO;
+using System.Linq;
+using System.Security.Cryptography;
+using Serilog;
+using PoolId = HordeServer.Utilities.StringId<HordeServer.Models.IPool>;
+using StreamId = HordeServer.Utilities.StringId<HordeServer.Models.IStream>;
+using AgentSoftwareVersion = HordeServer.Utilities.StringId<HordeServer.Collections.IAgentSoftwareCollection>;
+using AgentSoftwareChannelName = HordeServer.Utilities.StringId<HordeServer.Services.AgentSoftwareChannels>;
+using System.Diagnostics;
+using System.Threading.Tasks;
+
+namespace HordeServer.Models
+{
+	/// <summary>
+	/// Information about a workspace synced to an agent
+	/// </summary>
+	public class AgentWorkspace
+	{
+		/// <summary>
+		/// Identifier for the AutoSDK workspace
+		/// </summary>
+		public const string AutoSdkIdentifier = "AutoSDK";
+
+		/// <summary>
+		/// Name of the Perforce cluster to use
+		/// </summary>
+		public string? Cluster { get; set; }
+
+		/// <summary>
+		/// User to log into Perforce with (eg. buildmachine)
+		/// </summary>
+		public string? UserName { get; set; }
+
+		/// <summary>
+		/// Identifier to distinguish this workspace from other workspaces
+		/// </summary>
+		public string Identifier { get; set; }
+
+		/// <summary>
+		/// The stream to sync
+		/// </summary>
+		public string Stream { get; set; }
+
+		/// <summary>
+		/// Custom view for the workspace
+		/// </summary>
+		public List<string>? View { get; set; }
+
+		/// <summary>
+		/// Whether to use an incremental workspace
+		/// </summary>
+		public bool bIncremental { get; set; }
+
+		/// <summary>
+		/// Constructor
+		/// </summary>
+		/// <param name="Cluster">Name of the Perforce cluster</param>
+		/// <param name="UserName">User to log into Perforce with (eg. buildmachine)</param>
+		/// <param name="Identifier">Identifier to distinguish this workspace from other workspaces</param>
+		/// <param name="Stream">The stream to sync</param>
+		/// <param name="View">Custom view for the workspace</param>
+		/// <param name="bIncremental">Whether to use an incremental workspace</param>
+		public AgentWorkspace(string? Cluster, string? UserName, string Identifier, string Stream, List<string>? View, bool bIncremental)
+		{
+			if (!String.IsNullOrEmpty(Cluster))
+			{
+				this.Cluster = Cluster;
+			}
+			if (!String.IsNullOrEmpty(UserName))
+			{
+				this.UserName = UserName;
+			}
+			this.Identifier = Identifier;
+			this.Stream = Stream;
+			this.View = View;
+			this.bIncremental = bIncremental;
+		}
+
+		/// <summary>
+		/// Constructor
+		/// </summary>
+		/// <param name="Workspace">RPC message to construct from</param>
+		public AgentWorkspace(HordeCommon.Rpc.Messages.AgentWorkspace Workspace)
+			: this(Workspace.ConfiguredCluster, Workspace.ConfiguredUserName, Workspace.Identifier, Workspace.Stream, (Workspace.View.Count > 0) ? Workspace.View.ToList() : null, Workspace.Incremental)
+		{
+		}
+
+		/// <summary>
+		/// Gets a digest of the settings for this workspace
+		/// </summary>
+		/// <returns>Digest for the workspace settings</returns>
+		public string GetDigest()
+		{
+#pragma warning disable CA5351 // Do Not Use Broken Cryptographic Algorithms
+			using (MD5 Hasher = MD5.Create())
+			{
+				byte[] Data = BsonExtensionMethods.ToBson(this);
+				return BitConverter.ToString(Hasher.ComputeHash(Data)).Replace("-", "", StringComparison.Ordinal);
+			}
+#pragma warning restore CA5351
+		}
+
+		/// <inheritdoc/>
+		public override bool Equals(object? Obj)
+		{
+			AgentWorkspace? Other = Obj as AgentWorkspace;
+			if (Other == null)
+			{
+				return false;
+			}
+			if (Cluster != Other.Cluster || UserName != Other.UserName || Identifier != Other.Identifier || Stream != Other.Stream || bIncremental != Other.bIncremental)
+			{
+				return false;
+			}
+			if (!Enumerable.SequenceEqual(View ?? new List<string>(), Other.View ?? new List<string>()))
+			{
+				return false;
+			}
+			return true;
+		}
+
+		/// <inheritdoc/>
+		public override int GetHashCode()
+		{
+			return HashCode.Combine(Cluster, UserName, Identifier, Stream, bIncremental); // Ignore 'View' for now
+		}
+
+		/// <summary>
+		/// Checks if two workspace sets are equivalent, ignoring order
+		/// </summary>
+		/// <param name="WorkspacesA">First list of workspaces</param>
+		/// <param name="WorkspacesB">Second list of workspaces</param>
+		/// <returns>True if the sets are equivalent</returns>
+		public static bool SetEquals(IReadOnlyList<AgentWorkspace> WorkspacesA, IReadOnlyList<AgentWorkspace> WorkspacesB)
+		{
+			HashSet<AgentWorkspace> WorkspacesSetA = new HashSet<AgentWorkspace>(WorkspacesA);
+			return WorkspacesSetA.SetEquals(WorkspacesB);
+		}
+
+		/// <summary>
+		/// Converts this workspace to an RPC message
+		/// </summary>
+		/// <param name="Server">The Perforce server</param>
+		/// <param name="Credentials">Credentials for the server</param>
+		/// <returns>The RPC message</returns>
+		public HordeCommon.Rpc.Messages.AgentWorkspace ToRpcMessage(IPerforceServer Server, PerforceCredentials? Credentials)
+		{
+			// Construct the message
+			HordeCommon.Rpc.Messages.AgentWorkspace Result = new HordeCommon.Rpc.Messages.AgentWorkspace();
+			Result.ConfiguredCluster = Cluster;
+			Result.ConfiguredUserName = UserName;
+			Result.ServerAndPort = Server.ServerAndPort;
+			Result.UserName = Credentials?.UserName ?? UserName;
+			Result.Password = Credentials?.Password;
+			Result.Identifier = Identifier;
+			Result.Stream = Stream;
+			if (View != null)
+			{
+				Result.View.AddRange(View);
+			}
+			Result.Incremental = bIncremental;
+			return Result;
+		}
+	}
+
+	/// <summary>
+	/// Information about a device allocated for a lease
+	/// </summary>
+	public class AgentLeaseDevice
+	{
+		/// <summary>
+		/// Index of the device in the agent's device list
+		/// </summary>
+		[BsonRequired]
+		public int Index { get; set; }
+
+		/// <summary>
+		/// Handle of the device in the lease, ie. the logical name in the context of the work being done, not the physical device name.
+		/// </summary>
+		[BsonRequired]
+		public string? Handle { get; set; }
+
+		/// <summary>
+		/// Resources claimed by the lease. If null, the whole device is claimed, and nothing else can be allocated to it.
+		/// </summary>
+		[BsonIgnoreIfNull]
+		public Dictionary<string, int>? Resources { get; set; }
+
+		/// <summary>
+		/// Private constructor for serialization
+		/// </summary>
+		[BsonConstructor]
+		private AgentLeaseDevice()
+		{
+			Handle = null!;
+		}
+
+		/// <summary>
+		/// Constructor
+		/// </summary>
+		/// <param name="Index"></param>
+		/// <param name="Handle"></param>
+		/// <param name="Resources"></param>
+		public AgentLeaseDevice(int Index, string? Handle, Dictionary<string, int>? Resources)
+		{
+			this.Index = Index;
+			this.Handle = Handle;
+			this.Resources = Resources;
+		}
+	}
+
+	/// <summary>
+	/// Document describing an active lease
+	/// </summary>
+	public class AgentLease
+	{
+		/// <summary>
+		/// Name of this lease
+		/// </summary>
+		[BsonRequired]
+		public ObjectId Id { get; set; }
+
+		/// <summary>
+		/// Name of this lease
+		/// </summary>
+		public string Name { get; set; }
+
+		/// <summary>
+		/// The current state of the lease
+		/// </summary>
+		public LeaseState State { get; set; }
+
+		/// <summary>
+		/// The stream for this lease
+		/// </summary>
+		public StreamId? StreamId { get; set; }
+
+		/// <summary>
+		/// The pool for this lease
+		/// </summary>
+		public PoolId? PoolId { get; set; }
+
+		/// <summary>
+		/// Optional log for this lease
+		/// </summary>
+		public ObjectId? LogId { get; set; }
+
+		/// <summary>
+		/// Time at which the lease started
+		/// </summary>
+		[BsonRequired]
+		public DateTime StartTime { get; set; }
+
+		/// <summary>
+		/// Time at which the lease should be terminated
+		/// </summary>
+		public DateTime? ExpiryTime { get; set; }
+
+		/// <summary>
+		/// Flag indicating whether this lease has been accepted by the agent
+		/// </summary>
+		public bool Active { get; set; }
+
+		/// <summary>
+		/// Requirements for this lease
+		/// </summary>
+		public AgentRequirements Requirements { get; set; }
+
+		/// <summary>
+		/// For leases in the pending state, encodes an "any" protobuf containing the payload for the agent to execute the lease.
+		/// </summary>
+		public byte[]? Payload { get; set; }
+
+		/// <summary>
+		/// List of devices allocated to this lease
+		/// </summary>
+		[BsonIgnoreIfNull]
+		public List<AgentLeaseDevice>? Devices { get; set; }
+
+		/// <summary>
+		/// Private constructor
+		/// </summary>
+		[BsonConstructor]
+		private AgentLease()
+		{
+			Name = String.Empty;
+			Requirements = null!;
+			Devices = null!;
+		}
+
+		/// <summary>
+		/// Constructor
+		/// </summary>
+		/// <param name="Id">Identifier for the lease</param>
+		/// <param name="Name">Name of this lease</param>
+		/// <param name="StreamId"></param>
+		/// <param name="PoolId"></param>
+		/// <param name="LogId">Unique id for the log</param>
+		/// <param name="State">State for the lease</param>
+		/// <param name="Payload">Encoded "any" protobuf describing the contents of the payload</param>
+		/// <param name="Requirements">Requirements for this lease</param>
+		/// <param name="Devices">Device mapping for this lease</param>
+		public AgentLease(ObjectId Id, string Name, StreamId? StreamId, PoolId? PoolId, ObjectId? LogId, LeaseState State, byte[]? Payload, AgentRequirements Requirements, List<AgentLeaseDevice>? Devices)
+		{
+			this.Id = Id;
+			this.Name = Name;
+			this.StreamId = StreamId;
+			this.PoolId = PoolId;
+			this.LogId = LogId;
+			this.State = State;
+			this.Payload = Payload;
+			this.Requirements = Requirements;
+			this.Devices = Devices;
+			StartTime = DateTime.UtcNow;
+		}
+
+		/// <summary>
+		/// Determines if this is a conform lease
+		/// </summary>
+		/// <returns>True if this is a conform lease</returns>
+		public bool IsConformLease()
+		{
+			if (Payload != null)
+			{
+				Any BasePayload = Any.Parser.ParseFrom(Payload);
+				if (BasePayload.Is(ConformTask.Descriptor))
+				{
+					return true;
+				}
+			}
+			return false;
+		}
+
+		/// <summary>
+		/// Gets user-readable payload information
+		/// </summary>
+		/// <param name="Payload">The payload data</param>
+		/// <returns>Dictionary of key/value pairs for the payload</returns>
+		public static Dictionary<string, string>? GetPayloadDetails(ReadOnlyMemory<byte>? Payload)
+		{
+			Dictionary<string, string>? Details = null;
+			if (Payload != null)
+			{
+				Any BasePayload = Any.Parser.ParseFrom(Payload.Value.ToArray());
+
+				Details = new Dictionary<string, string>();
+
+				ConformTask ConformTask;
+				if(BasePayload.TryUnpack(out ConformTask))
+				{
+					Details["Type"] = "Conform";
+					Details["LogId"] = ConformTask.LogId;
+				}
+
+				ExecuteJobTask JobTask;
+				if (BasePayload.TryUnpack(out JobTask))
+				{
+					Details["Type"] = "Job";
+					Details["JobId"] = JobTask.JobId;
+					Details["BatchId"] = JobTask.BatchId;
+					Details["LogId"] = JobTask.LogId;
+				}
+
+				UpgradeTask UpgradeTask;
+				if (BasePayload.TryUnpack(out UpgradeTask))
+				{
+					Details["Type"] = "Upgrade";
+					Details["SoftwareId"] = UpgradeTask.SoftwareId;
+					Details["LogId"] = UpgradeTask.LogId;
+				}
+			}
+			return Details;
+		}
+
+		/// <summary>
+		/// Converts this lease to an RPC message
+		/// </summary>
+		/// <returns>RPC message</returns>
+		public HordeCommon.Rpc.Messages.Lease ToRpcMessage()
+		{
+			HordeCommon.Rpc.Messages.Lease Lease = new HordeCommon.Rpc.Messages.Lease();
+			Lease.Id = Id.ToString();
+			Lease.Payload = Google.Protobuf.WellKnownTypes.Any.Parser.ParseFrom(Payload); 
+			Lease.State = State;
+			return Lease;
+		}
+	}
+
+	/// <summary>
+	/// Mirrors an Agent document in the database
+	/// </summary>
+	public interface IAgent
+	{
+		/// <summary>
+		/// Randomly generated unique id for this agent.
+		/// </summary>
+		public AgentId Id { get; }
+
+		/// <summary>
+		/// The current session id, if it's online
+		/// </summary>
+		public ObjectId? SessionId { get; }
+
+		/// <summary>
+		/// Time at which the current session expires. 
+		/// </summary>
+		public DateTime? SessionExpiresAt { get; }
+
+		/// <summary>
+		/// Current status of this agent
+		/// </summary>
+		public AgentStatus Status { get; }
+
+		/// <summary>
+		/// Whether the agent is enabled
+		/// </summary>
+		public bool Enabled { get; }
+		
+		/// <summary>
+		/// Arbitrary comment for the agent (useful for disable reasons etc)
+		/// </summary>
+		public string? Comment { get; }
+
+		/// <summary>
+		/// Whether the agent is ephemeral
+		/// </summary>
+		public bool Ephemeral { get; }
+
+		/// <summary>
+		/// Whether the agent should be included on the dashboard. This is set to true for ephemeral agents once they are no longer online, or agents that are explicitly deleted.
+		/// </summary>
+		public bool Deleted { get; }
+
+		/// <summary>
+		/// Version of the software running on this agent
+		/// </summary>
+		public AgentSoftwareVersion? Version { get; }
+
+		/// <summary>
+		/// Channel for the software running on this agent. Uses <see cref="AgentSoftwareService.DefaultChannelName"/> if not specified
+		/// </summary>
+		public AgentSoftwareChannelName? Channel { get; }
+
+		/// <summary>
+		/// Last upgrade that was attempted
+		/// </summary>
+		public string? LastUpgradeVersion { get; }
+
+		/// <summary>
+		/// Time that which the last upgrade was attempted
+		/// </summary>
+		public DateTime? LastUpgradeTime { get; }
+
+		/// <summary>
+		/// List of manually assigned pools for agent
+		/// </summary>
+		public IReadOnlyList<PoolId> ExplicitPools { get; }
+
+		/// <summary>
+		/// Whether a conform is requested
+		/// </summary>
+		public bool RequestConform { get; }
+
+		/// <summary>
+		/// Whether a machine restart is requested
+		/// </summary>
+		public bool RequestRestart { get; }
+
+		/// <summary>
+		/// Whether the machine should be shutdown
+		/// </summary>
+		public bool RequestShutdown { get; }
+
+		/// <summary>
+		/// List of workspaces currently synced to this machine
+		/// </summary>
+		public IReadOnlyList<AgentWorkspace> Workspaces { get; }
+
+		/// <summary>
+		/// Time at which the last conform job ran
+		/// </summary>
+		public DateTime LastConformTime { get; }
+
+		/// <summary>
+		/// Number of times a conform job has failed
+		/// </summary>
+		public int? ConformAttemptCount { get; }
+
+		/// <summary>
+		/// Capabilities of this agent
+		/// </summary>
+		public AgentCapabilities Capabilities { get; }
+
+		/// <summary>
+		/// Array of active leases.
+		/// </summary>
+		public IReadOnlyList<AgentLease> Leases { get; }
+
+		/// <summary>
+		/// ACL for modifying this agent
+		/// </summary>
+		public Acl? Acl { get; }
+
+		/// <summary>
+		/// Last time that the agent was modified
+		/// </summary>
+		public DateTime UpdateTime { get; }
+
+		/// <summary>
+		/// Update counter for this document. Any updates should compare-and-swap based on the value of this counter, or increment it in the case of server-side updates.
+		/// </summary>
+		public uint UpdateIndex { get; }
+	}
+
+	/// <summary>
+	/// Extension methods for IAgent
+	/// </summary>
+	public static class AgentExtensions
+	{
+		/// <summary>
+		/// Determines whether this agent is online
+		/// </summary>
+		/// <returns></returns>
+		public static bool IsSessionValid(this IAgent Agent, DateTime UtcNow)
+		{
+			return Agent.SessionId.HasValue && Agent.SessionExpiresAt.HasValue && UtcNow < Agent.SessionExpiresAt.Value;
+		}
+
+		/// <summary>
+		/// Gets a list of pools for the given agent. Includes all automatically assigned pools based on agent capabilities.
+		/// </summary>
+		/// <param name="Agent">The agent instance</param>
+		/// <param name="Pools">List of all available pools</param>
+		/// <returns>List of pools</returns>
+		public static IEnumerable<IPool> GetPools(this IAgent Agent, IEnumerable<IPool> Pools)
+		{
+			return Pools.Where(x => Agent.InPool(x));
+		}
+
+		/// <summary>
+		/// Checks whether an agent is in the given pool
+		/// </summary>
+		/// <param name="Agent">The agent to check</param>
+		/// <param name="Pool">The pool to test against</param>
+		/// <returns>True if the agent is in the pool</returns>
+		public static bool InPool(this IAgent Agent, IPool Pool)
+		{
+			if (Agent.ExplicitPools.Contains(Pool.Id))
+			{
+				return true;
+			}
+			else if (Pool.Requirements != null && AgentExtensions.TryCreateLease(Agent.Capabilities, Pool.Requirements, Enumerable.Empty<AgentLease>(), out _))
+			{
+				return true;
+			}
+			return false;
+		}
+
+		/// <summary>
+		/// Determine whether it's possible to add a lease for the given resources
+		/// </summary>
+		/// <param name="Agent">The agent to create a lease for</param>
+		/// <param name="Pool">The pool required</param>
+		/// <param name="Requirements">Requirements for this lease</param>
+		/// <param name="LeasedDevices">On success, recieves the list of leased devices</param>
+		/// <returns>True if the new lease can be granted</returns>
+		public static bool TryCreateLease(this IAgent Agent, IPool Pool, AgentRequirements? Requirements, out List<AgentLeaseDevice>? LeasedDevices)
+		{
+			if (!Agent.Enabled || Agent.Status != AgentStatus.Ok || !Agent.InPool(Pool))
+			{
+				LeasedDevices = null!;
+				return false;
+			}
+			return TryCreateLease(Agent.Capabilities, Requirements, Agent.Leases, out LeasedDevices);
+		}
+
+		/// <summary>
+		/// Attempts to match the requirements for a particular agent, and return the list of leased devices that fulfils the requirements
+		/// </summary>
+		/// <param name="Capabilities">Capabilities of the gent</param>
+		/// <param name="Requirements">Requirements for the lease</param>
+		/// <param name="Leases">List of current leases</param>
+		/// <param name="LeasedDevices">Receives the list of device leases</param>
+		/// <returns>True if successful</returns>
+		public static bool TryCreateLease(AgentCapabilities Capabilities, AgentRequirements? Requirements, IEnumerable<AgentLease> Leases, out List<AgentLeaseDevice>? LeasedDevices)
+		{
+			List<AgentLeaseDevice>? Result = null;
+			if (Requirements != null)
+			{
+				if (!Requirements.Shared && Leases.Any())
+				{
+					LeasedDevices = null;
+					return false;
+				}
+
+				if (Requirements.Properties != null)
+				{
+					foreach (string Property in Requirements.Properties)
+					{
+						if (Capabilities.Properties == null || !Capabilities.Properties.Contains(Property))
+						{
+							LeasedDevices = null!;
+							return false;
+						}
+					}
+				}
+
+				if(Requirements.Devices != null)
+				{
+					if (!TryCreateDeviceLeases(Capabilities.Devices, Requirements.Devices, 0, 0, Leases.SelectMany(x => x.Devices), out Result))
+					{
+						LeasedDevices = null!;
+						return false;
+					}
+				}
+			}
+
+			LeasedDevices = Result;
+			return true;
+		}
+
+		/// <summary>
+		/// Try to match up the required devices for a lease to the available devices on an agent. This is done through a recursive search with backtracking for completeness,
+		/// due to not knowing which devices in the lease will correspond to each device in the agent, but it should be quick in practice due to agents having a small number 
+		/// of homogenous devices.
+		/// </summary>
+		/// <param name="CapableDevices">Device capabilities for the agent</param>
+		/// <param name="RequiredDevices">Device requirements for the lease</param>
+		/// <param name="RequiredDeviceIdx">Current index of the device in the lease requirements</param>
+		/// <param name="LeasedDeviceMask">Bitmask for devices that have currently been allocated. This is limited to 32, which is unlikely to be a problem in practice.</param>
+		/// <param name="CurrentLeasedDevices">The current set of leased devices</param>
+		/// <param name="LeasedDevices">On success, receives a list of the leased devices, mapping each requested device to a device in the agent device array.</param>
+		/// <returns>True if the leases were allocated</returns>
+		private static bool TryCreateDeviceLeases(List<DeviceCapabilities> CapableDevices, List<DeviceRequirements> RequiredDevices, int RequiredDeviceIdx, uint LeasedDeviceMask, IEnumerable<AgentLeaseDevice> CurrentLeasedDevices, [MaybeNullWhen(false)] out List<AgentLeaseDevice> LeasedDevices)
+		{
+			// If we've assigned all the required devices now, we can create the list of assigned devices
+			if (RequiredDeviceIdx == RequiredDevices.Count)
+			{
+				LeasedDevices = new List<AgentLeaseDevice>(RequiredDevices.Count);
+				foreach(DeviceRequirements Device in RequiredDevices)
+				{
+					LeasedDevices.Add(new AgentLeaseDevice(-1, Device.Handle, Device.Resources));
+				}
+				return true;
+			}
+
+			// Otherwise try to assign the next device
+			DeviceRequirements RequiredDevice = RequiredDevices[RequiredDeviceIdx];
+			for (int DeviceIdx = 0; DeviceIdx < CapableDevices.Count; DeviceIdx++)
+			{
+				uint DeviceFlag = 1U << DeviceIdx;
+				if ((LeasedDeviceMask & DeviceFlag) == 0 && MatchDevice(CapableDevices[DeviceIdx], RequiredDevice, CurrentLeasedDevices.Where(x => x.Index == DeviceIdx)))
+				{
+					List<AgentLeaseDevice>? Result;
+					if (TryCreateDeviceLeases(CapableDevices, RequiredDevices, RequiredDeviceIdx + 1, LeasedDeviceMask | DeviceFlag, CurrentLeasedDevices, out Result))
+					{
+						LeasedDevices = Result;
+						LeasedDevices[RequiredDeviceIdx].Index = DeviceIdx;
+						return true;
+					}
+				}
+			}
+
+			// Otherwise failed
+			LeasedDevices = null!;
+			return false;
+		}
+
+		/// <summary>
+		/// Determines if an agent device has the given requirements
+		/// </summary>
+		/// <param name="Capabilities">The device capabilities</param>
+		/// <param name="Requirements">Requirements for the lease</param>
+		/// <param name="Leases">Current leases for the device</param>
+		/// <returns>True if the device can satisfy the given requirements</returns>
+		private static bool MatchDevice(DeviceCapabilities Capabilities, DeviceRequirements Requirements, IEnumerable<AgentLeaseDevice> Leases)
+		{
+			// Check the device has all the required properties
+			if (Requirements.Properties != null)
+			{
+				foreach (string Property in Requirements.Properties)
+				{
+					if (Capabilities.Properties == null || !Capabilities.Properties.Contains(Property))
+					{
+						return false;
+					}
+				}
+			}
+
+			// Check the device has all the required resources
+			if (Requirements.Resources == null)
+			{
+				// Requires exclusive access to the device
+				if (Leases.Any())
+				{
+					return false;
+				}
+			}
+			else
+			{
+				// Requires shared access to the device. Check the device can be shared.
+				if (Capabilities.Resources == null)
+				{
+					return false;
+				}
+
+				// Check there are enough available of each resource type
+				foreach (KeyValuePair<string, int> Resource in Requirements.Resources)
+				{
+					// Make sure the device has enough of the named resource to start with
+					int RemainingCount;
+					if (!Capabilities.Resources.TryGetValue(Resource.Key, out RemainingCount) || RemainingCount < Resource.Value)
+					{
+						return false;
+					}
+
+					// Check each existing lease of this device
+					foreach (AgentLeaseDevice Lease in Leases)
+					{
+						// If the lease has an exclusive reservation, we can't use it
+						if (Lease.Resources == null)
+						{
+							return false;
+						}
+
+						// Update the remaining count for this resource
+						int UsedCount;
+						if (Lease.Resources.TryGetValue(Resource.Key, out UsedCount))
+						{
+							RemainingCount -= UsedCount;
+							if (RemainingCount < Resource.Value)
+							{
+								return false;
+							}
+						}
+					}
+				}
+			}
+
+			return true;
+		}
+
+		/// <summary>
+		/// Converts this workspace to an RPC message
+		/// </summary>
+		/// <param name="Agent">The agent to get a workspace for</param>
+		/// <param name="Workspace">The workspace definition</param>
+		/// <param name="Globals">The global state</param>
+		/// <param name="LoadBalancer">The Perforce load balancer</param>
+		/// <param name="WorkspaceMessages">List of messages</param>
+		/// <returns>The RPC message</returns>
+		public static async Task<bool> TryAddWorkspaceMessage(this IAgent Agent, AgentWorkspace Workspace, Globals Globals, PerforceLoadBalancer LoadBalancer, IList<HordeCommon.Rpc.Messages.AgentWorkspace> WorkspaceMessages)
+		{
+			// Find the matching Perforce cluster
+			PerforceCluster? ClusterInfo;
+			if (Workspace.Cluster == null)
+			{
+				ClusterInfo = Globals.PerforceClusters.FirstOrDefault();
+			}
+			else
+			{
+				ClusterInfo = Globals.PerforceClusters.FirstOrDefault(x => String.Equals(x.Name, Workspace.Cluster, StringComparison.OrdinalIgnoreCase));
+			}
+
+			// Find a matching server, trying to use a previously selected one if possible
+			string? ServerAndPort = WorkspaceMessages.FirstOrDefault(x => x.ConfiguredCluster == Workspace.Cluster)?.ServerAndPort;
+			if (ServerAndPort == null)
+			{
+				if (ClusterInfo == null)
+				{
+					return false;
+				}
+
+				IPerforceServer? Server = await LoadBalancer.SelectServerAsync(ClusterInfo, Agent);
+				if(Server == null)
+				{
+					return false;
+				}
+				ServerAndPort = Server.ServerAndPort;
+			}
+
+			// Find the matching credentials for the desired user
+			PerforceCredentials? Credentials = null;
+			if (ClusterInfo != null)
+			{
+				if (Workspace.UserName == null)
+				{
+					Credentials = ClusterInfo.Credentials.FirstOrDefault();
+				}
+				else
+				{
+					Credentials = ClusterInfo.Credentials.FirstOrDefault(x => String.Equals(x.UserName, Workspace.UserName, StringComparison.OrdinalIgnoreCase));
+				}
+			}
+
+			// Construct the message
+			HordeCommon.Rpc.Messages.AgentWorkspace Result = new HordeCommon.Rpc.Messages.AgentWorkspace();
+			Result.ConfiguredCluster = Workspace.Cluster;
+			Result.ConfiguredUserName = Workspace.UserName;
+			Result.Cluster = ClusterInfo?.Name;
+			Result.ServerAndPort = ServerAndPort;
+			Result.UserName = Credentials?.UserName ?? Workspace.UserName;
+			Result.Password = Credentials?.Password;
+			Result.Identifier = Workspace.Identifier;
+			Result.Stream = Workspace.Stream;
+			if (Workspace.View != null)
+			{
+				Result.View.AddRange(Workspace.View);
+			}
+			Result.Incremental = Workspace.bIncremental;
+			WorkspaceMessages.Add(Result);
+
+			return true;
+		}
+	}
+}
diff --git a/Engine/Source/Programs/Horde/HordeServer/Models/AgentSoftware.cs b/Engine/Source/Programs/Horde/HordeServer/Models/AgentSoftware.cs
new file mode 100644
--- /dev/null
+++ b/Engine/Source/Programs/Horde/HordeServer/Models/AgentSoftware.cs
@@ -0,0 +1,46 @@
+// Copyright Epic Games, Inc. All Rights Reserved.
+
+using HordeServer.Api;
+using HordeServer.Collections;
+using HordeServer.Utilities;
+using MongoDB.Bson;
+using MongoDB.Bson.Serialization.Attributes;
+using System;
+using System.Collections.Generic;
+using System.IO;
+using System.IO.Compression;
+using System.Linq;
+using System.Text.Json;
+using System.Threading.Tasks;
+
+using AgentSoftwareVersion = HordeServer.Utilities.StringId<HordeServer.Collections.IAgentSoftwareCollection>;
+using AgentSoftwareChannelName = HordeServer.Utilities.StringId<HordeServer.Services.AgentSoftwareChannels>;
+
+namespace HordeServer.Models
+{
+	/// <summary>
+	/// A software channel
+	/// </summary>
+	public interface IAgentSoftwareChannel
+	{
+		/// <summary>
+		/// The channel id
+		/// </summary>
+		public AgentSoftwareChannelName Name { get; set; }
+
+		/// <summary>
+		/// Name of the user that made the last modification
+		/// </summary>
+		public string? ModifiedBy { get; set; }
+
+		/// <summary>
+		/// Last modification time
+		/// </summary>
+		public DateTime ModifiedTime { get; set; }
+
+		/// <summary>
+		/// The software version number
+		/// </summary>
+		public AgentSoftwareVersion Version { get; set; }
+	}
+}
diff --git a/Engine/Source/Programs/Horde/HordeServer/Models/Artifact.cs b/Engine/Source/Programs/Horde/HordeServer/Models/Artifact.cs
new file mode 100644
--- /dev/null
+++ b/Engine/Source/Programs/Horde/HordeServer/Models/Artifact.cs
@@ -0,0 +1,83 @@
+// Copyright Epic Games, Inc. All Rights Reserved.
+
+using HordeServer.Api;
+using HordeServer.Utilities;
+using MongoDB.Bson;
+using MongoDB.Bson.Serialization.Attributes;
+
+namespace HordeServer.Models
+{
+	/// <summary>
+	/// Information about an artifact
+	/// </summary>
+	public class Artifact
+	{
+		/// <summary>
+		/// Identifier for the Artifact. Randomly generated.
+		/// </summary>
+		[BsonRequired, BsonId]
+		public ObjectId Id { get; set; }
+
+		/// <summary>
+		/// Unique id of the job containing this artifact
+		/// </summary>
+		[BsonRequired]
+		public ObjectId JobId { get; set; }
+
+		/// <summary>
+		/// Unique id of the job containing this artifact
+		/// </summary>
+		[BsonRequired]
+		public string Name { get; set; }
+
+		/// <summary>
+		/// Unique id of the step containing this artifact
+		/// </summary>
+		public SubResourceId? StepId { get; set; }
+
+		/// <summary>
+		/// Total size of the file
+		/// </summary>
+		[BsonRequired]
+		public long Length { get; set; }
+
+		/// <summary>
+		/// Type of artifact
+		/// </summary>
+		public string MimeType { get; set; }
+
+		/// <summary>
+		/// Version of the file in the db
+		/// </summary>
+		[BsonRequired]
+		public int UpdateIndex { get; set; }
+
+		/// <summary>
+		/// Private constructor for BSON serializer
+		/// </summary>
+		[BsonConstructor]
+		private Artifact()
+		{
+			this.Name = null!;
+			this.MimeType = null!;
+		}
+
+		/// <summary>
+		/// Constructor
+		/// </summary>
+		/// <param name="JobId">Unique id for the job containing this artifact</param>
+		/// <param name="StepId">Optional unique id for the step containing this artifact</param>
+		/// <param name="Name">Name</param>
+		/// <param name="Length">File Length</param>
+		/// <param name="MimeType">Type of artifact</param>
+		public Artifact(ObjectId JobId, SubResourceId? StepId, string Name, long Length, string MimeType)
+		{
+			this.Id = ObjectId.GenerateNewId();
+			this.JobId = JobId;
+			this.StepId = StepId;
+			this.Name = Name;
+			this.Length = Length;
+			this.MimeType = MimeType;
+		}
+	}
+}
diff --git a/Engine/Source/Programs/Horde/HordeServer/Models/Change.cs b/Engine/Source/Programs/Horde/HordeServer/Models/Change.cs
new file mode 100644
--- /dev/null
+++ b/Engine/Source/Programs/Horde/HordeServer/Models/Change.cs
@@ -0,0 +1,208 @@
+using HordeServer.Api;
+using HordeServer.Utilities;
+using MongoDB.Bson;
+using MongoDB.Bson.Serialization.Attributes;
+using System;
+using System.Collections.Generic;
+using System.Diagnostics.CodeAnalysis;
+using System.Globalization;
+using System.Linq;
+using System.Text.RegularExpressions;
+using System.Threading.Tasks;
+
+using StreamId = HordeServer.Utilities.StringId<HordeServer.Models.IStream>;
+
+namespace HordeServer.Models
+{
+	/// <summary>
+	/// Summary information for a change
+	/// </summary>
+	public class ChangeSummary
+	{
+		/// <summary>
+		/// The changelist number
+		/// </summary>
+		public int Number { get; set; }
+
+		/// <summary>
+		/// Author of the change
+		/// </summary>
+		public string Author { get; set; }
+
+		/// <summary>
+		/// Abbreviated changelist description
+		/// </summary>
+		public string Description { get; set; }
+
+		/// <summary>
+		/// Constructor
+		/// </summary>
+		/// <param name="Number">Changelist number</param>
+		/// <param name="Author">Author of the change</param>
+		/// <param name="Description">Changelist description</param>
+		public ChangeSummary(int Number, string Author, string Description)
+		{
+			this.Number = Number;
+			this.Author = Author;
+			this.Description = Description;
+		}
+	}
+
+	/// <summary>
+	/// Flags identifying content of a changelist
+	/// </summary>
+	[Flags]
+	public enum ChangeContentFlags
+	{
+		/// <summary>
+		/// The change contains code
+		/// </summary>
+		ContainsCode = 1,
+
+		/// <summary>
+		/// The change contains content
+		/// </summary>
+		ContainsContent = 2,
+	}
+
+	/// <summary>
+	/// Information about a commit
+	/// </summary>
+	public class ChangeDetails
+	{
+		/// <summary>
+		/// The source changelist number
+		/// </summary>
+		public int Number { get; set; }
+
+		/// <summary>
+		/// Name of the user that authored this change
+		/// </summary>
+		public string Author { get; set; }
+
+		/// <summary>
+		/// The description text
+		/// </summary>
+		public string Description { get; set; }
+
+		/// <summary>
+		/// List of files that were modified, relative to the stream base
+		/// </summary>
+		public List<string> Files { get; set; }
+
+		/// <summary>
+		/// Private constructor for serialization
+		/// </summary>
+		private ChangeDetails()
+		{
+			Author = null!;
+			Description = null!;
+			Files = null!;
+		}
+
+		/// <summary>
+		/// Constructor
+		/// </summary>
+		/// <param name="Number">Changelist number</param>
+		/// <param name="Author">Author of the change</param>
+		/// <param name="Description">Changelist description</param>
+		/// <param name="Files">List of files modified, relative to the stream base</param>
+		public ChangeDetails(int Number, string Author, string Description, List<string> Files)
+		{
+			this.Number = Number;
+			this.Author = Author;
+			this.Description = Description;
+			this.Files = Files;
+		}
+
+		/// <summary>
+		/// Determines if this change is a code change
+		/// </summary>
+		/// <returns>True if this change is a code change</returns>
+		public ChangeContentFlags GetContentFlags()
+		{
+			ChangeContentFlags Scope = 0;
+
+			// Check whether the files are code or content
+			string[] CodeExtensions = { ".cs", ".h", ".cpp", ".inl", ".usf", ".ush", ".uproject", ".uplugin" };
+			foreach (string File in Files)
+			{
+				if (CodeExtensions.Any(Extension => File.EndsWith(Extension, StringComparison.OrdinalIgnoreCase)))
+				{
+					Scope |= ChangeContentFlags.ContainsCode;
+				}
+				else
+				{
+					Scope |= ChangeContentFlags.ContainsContent;
+				}
+
+				if (Scope == (ChangeContentFlags.ContainsCode | ChangeContentFlags.ContainsContent))
+				{
+					break;
+				}
+			}
+			return Scope;
+		}
+
+		/// <summary>
+		/// Attempts to parse the Robomerge source from this commit information
+		/// </summary>
+		/// <param name="Source">On success, receives the source information</param>
+		/// <returns>True if the commit was merged from another stream</returns>
+		public bool TryParseRobomergeSource([NotNullWhen(true)] out (string, int)? Source)
+		{
+			Match Match = Regex.Match(Description, @"#ROBOMERGE-SOURCE: CL (\d+) in (//[^ ]*)/...", RegexOptions.Multiline);
+			if (Match.Success)
+			{
+				Source = (Match.Groups[2].Value, int.Parse(Match.Groups[1].Value, CultureInfo.InvariantCulture));
+				return true;
+			}
+			else
+			{
+				Source = null;
+				return false;
+			}
+		}
+	}
+
+	/// <summary>
+	/// Summary of a file in the depot
+	/// </summary>
+	public class FileSummary
+	{
+		/// <summary>
+		/// Depot path to the file
+		/// </summary>
+		public string DepotPath { get; set; }
+
+		/// <summary>
+		/// Whether the file exists
+		/// </summary>
+		public bool Exists { get; set; }
+
+		/// <summary>
+		/// Last changelist number that the file was modified
+		/// </summary>
+		public int Change { get; set; }
+
+		/// <summary>
+		/// Error about this file, if it exists
+		/// </summary>
+		public string? Error { get; set; }
+
+		/// <summary>
+		/// Information about a file
+		/// </summary>
+		/// <param name="DepotPath">Depot </param>
+		/// <param name="Exists"></param>
+		/// <param name="Change"></param>
+		/// <param name="Error"></param>
+		public FileSummary(string DepotPath, bool Exists, int Change, string? Error = null)
+		{
+			this.DepotPath = DepotPath;
+			this.Exists = Exists;
+			this.Change = Change;
+			this.Error = Error;
+		}
+	}
+}
diff --git a/Engine/Source/Programs/Horde/HordeServer/Models/ConformList.cs b/Engine/Source/Programs/Horde/HordeServer/Models/ConformList.cs
new file mode 100644
--- /dev/null
+++ b/Engine/Source/Programs/Horde/HordeServer/Models/ConformList.cs
@@ -0,0 +1,26 @@
+using HordeServer.Utilities;
+using MongoDB.Bson;
+using System;
+using System.Collections.Generic;
+using System.Linq;
+using System.Threading.Tasks;
+
+namespace HordeServer.Models
+{
+	/// <summary>
+	/// List of machines that are currently conforming
+	/// </summary>
+	[SingletonDocument("5ff62d06a3bdb49f82da3cac")]
+	public class ConformList : SingletonBase
+	{
+		/// <summary>
+		/// Maximum allowed at once
+		/// </summary>
+		public int MaxCount { get; set; }
+
+		/// <summary>
+		/// List of current conforms that are running
+		/// </summary>
+		public List<ObjectId> LeaseIds { get; set; } = new List<ObjectId>();
+	}
+}
diff --git a/Engine/Source/Programs/Horde/HordeServer/Models/Event.cs b/Engine/Source/Programs/Horde/HordeServer/Models/Event.cs
new file mode 100644
--- /dev/null
+++ b/Engine/Source/Programs/Horde/HordeServer/Models/Event.cs
@@ -0,0 +1,82 @@
+using HordeServer.Api;
+using System;
+using System.Collections.Generic;
+using System.Diagnostics.CodeAnalysis;
+using System.Linq;
+using System.Threading.Tasks;
+
+using StreamId = HordeServer.Utilities.StringId<HordeServer.Models.IStream>;
+using TemplateRefId = HordeServer.Utilities.StringId<HordeServer.Models.TemplateRef>;
+
+namespace HordeServer.Models
+{
+	/// <summary>
+	/// Base interface for events that users can subscribe for notifications to
+	/// </summary>
+	public interface IEvent
+	{
+		/// <summary>
+		/// Converts this event into an API event record
+		/// </summary>
+		/// <returns>The API event record</returns>
+		public EventRecord ToRecord();
+	}
+
+	/// <summary>
+	/// Event that occurs when a job completes
+	/// </summary>
+	public interface IJobCompleteEvent : IEvent
+	{
+		/// <summary>
+		/// The stream id
+		/// </summary>
+		StreamId StreamId { get; }
+
+		/// <summary>
+		/// The template id
+		/// </summary>
+		TemplateRefId TemplateId { get; }
+	}
+
+	/// <summary>
+	/// Event that occurs when a label completes
+	/// </summary>
+	public interface ILabelCompleteEvent : IEvent
+	{
+		/// <summary>
+		/// The stream id
+		/// </summary>
+		StreamId StreamId { get; }
+
+		/// <summary>
+		/// The template id
+		/// </summary>
+		TemplateRefId TemplateId { get; }
+
+		/// <summary>
+		/// Name of the label to monitor
+		/// </summary>
+		string LabelName { get; }
+	}
+
+	/// <summary>
+	/// Event that occurs when a job step completes
+	/// </summary>
+	public interface IStepCompleteEvent : IEvent
+	{
+		/// <summary>
+		/// The stream id
+		/// </summary>
+		StreamId StreamId { get; }
+
+		/// <summary>
+		/// The template id
+		/// </summary>
+		TemplateRefId TemplateId { get; }
+
+		/// <summary>
+		/// Name of the step to monitor
+		/// </summary>
+		string StepName { get; }
+	}
+}
diff --git a/Engine/Source/Programs/Horde/HordeServer/Models/Globals.cs b/Engine/Source/Programs/Horde/HordeServer/Models/Globals.cs
new file mode 100644
--- /dev/null
+++ b/Engine/Source/Programs/Horde/HordeServer/Models/Globals.cs
@@ -0,0 +1,324 @@
+using Google.Protobuf.WellKnownTypes;
+using HordeServer.Services;
+using HordeServer.Utilities;
+using Microsoft.IdentityModel.Tokens;
+using MongoDB.Bson;
+using MongoDB.Bson.Serialization.Attributes;
+using System;
+using System.Collections.Generic;
+using System.ComponentModel.DataAnnotations;
+using System.Diagnostics;
+using System.Linq;
+using System.Security.Cryptography;
+using System.Threading.Tasks;
+
+namespace HordeServer.Models
+{
+	/// <summary>
+	/// Base class for singleton documents
+	/// </summary>
+	public abstract class SingletonBase
+	{
+		/// <summary>
+		/// Unique id for this singleton
+		/// </summary>
+		[BsonId]
+		public ObjectId Id { get; set; }
+
+		/// <summary>
+		/// The revision index of this document
+		/// </summary>
+		public int Revision { get; set; }
+
+		/// <summary>
+		/// Callback to allow the singleton to fix up itself after being read
+		/// </summary>
+		public virtual void PostLoad()
+		{
+		}
+	}
+
+	/// <summary>
+	/// How frequently the maintence window repeats
+	/// </summary>
+	public enum ScheduledDowntimeFrequency
+	{
+		/// <summary>
+		/// Once
+		/// </summary>
+		Once,
+
+		/// <summary>
+		/// Every day
+		/// </summary>
+		Daily,
+
+		/// <summary>
+		/// Every week
+		/// </summary>
+		Weekly,
+	}
+
+	/// <summary>
+	/// Settings for the maintenance window
+	/// </summary>
+	public class ScheduledDowntime
+	{
+		/// <summary>
+		/// Start time
+		/// </summary>
+		public DateTimeOffset StartTime { get; set; }
+
+		/// <summary>
+		/// Finish time
+		/// </summary>
+		public DateTimeOffset FinishTime { get; set; }
+
+		/// <summary>
+		/// Frequency that the window repeats
+		/// </summary>
+		public ScheduledDowntimeFrequency Frequency { get; set; } = ScheduledDowntimeFrequency.Once;
+
+		/// <summary>
+		/// Gets the next scheduled downtime
+		/// </summary>
+		/// <param name="Now">The current time</param>
+		/// <returns>Start and finish time</returns>
+		public (DateTimeOffset StartTime, DateTimeOffset FinishTime) GetNext(DateTimeOffset Now)
+		{
+			TimeSpan Offset = TimeSpan.Zero;
+			if (Frequency == ScheduledDowntimeFrequency.Daily)
+			{
+				double Days = (Now - StartTime).TotalDays;
+				if (Days >= 1.0)
+				{
+					Days -= Days % 1.0;
+				}
+				Offset = TimeSpan.FromDays(Days);
+			}
+			else if (Frequency == ScheduledDowntimeFrequency.Weekly)
+			{
+				double Days = (Now - StartTime).TotalDays;
+				if (Days >= 7.0)
+				{
+					Days -= Days % 7.0;
+				}
+				Offset = TimeSpan.FromDays(Days);
+			}
+			return (StartTime + Offset, FinishTime + Offset);
+		}
+
+		/// <summary>
+		/// Determines if this schedule is active
+		/// </summary>
+		/// <param name="Now">The current time</param>
+		/// <returns>True if downtime is active</returns>
+		public bool IsActive(DateTimeOffset Now)
+		{
+			if (Frequency == ScheduledDowntimeFrequency.Once)
+			{
+				return Now >= StartTime && Now < FinishTime;
+			}
+			else if (Frequency == ScheduledDowntimeFrequency.Daily)
+			{
+				double Days = (Now - StartTime).TotalDays;
+				if (Days < 0.0)
+				{
+					return false;
+				}
+				else
+				{
+					return (Days % 1.0) < (FinishTime - StartTime).TotalDays;
+				}
+			}
+			else if (Frequency == ScheduledDowntimeFrequency.Weekly)
+			{
+				double Days = (Now - StartTime).TotalDays;
+				if(Days < 0.0)
+				{
+					return false;
+				}
+				else
+				{
+					return (Days % 7.0) < (FinishTime - StartTime).TotalDays;
+				}
+			}
+			else
+			{
+				return false;
+			}
+		}
+	}
+
+	/// <summary>
+	/// User notice
+	/// </summary>
+	public class Notice
+	{
+		/// <summary>
+		/// Unique id for this notice
+		/// </summary>
+		public string Id { get; set; } = String.Empty;
+
+		/// <summary>
+		/// Start time to display this message
+		/// </summary>
+		public DateTime? StartTime { get; set; }
+
+		/// <summary>
+		/// Finish time to display this message
+		/// </summary>
+		public DateTime? FinishTime { get; set; }
+
+		/// <summary>
+		/// Message to display
+		/// </summary>
+		public string Message { get; set; } = String.Empty;
+	}
+
+	/// <summary>
+	/// Information about an individual Perforce server
+	/// </summary>
+	public class PerforceServer
+	{
+		/// <summary>
+		/// The server and port. The server may be a DNS entry with multiple records, in which case it will be actively load balanced.
+		/// </summary>
+		public string ServerAndPort { get; set; } = "perforce:1666";
+
+		/// <summary>
+		/// Whether to query the healthcheck address under each server
+		/// </summary>
+		public bool HealthCheck { get; set; }
+
+		/// <summary>
+		/// Whether to resolve the DNS entries and load balance between different hosts
+		/// </summary>
+		public bool ResolveDns { get; set; }
+
+		/// <summary>
+		/// List of properties for an agent to be eligable to use this server
+		/// </summary>
+		public List<string>? Properties { get; set; }
+	}
+
+	/// <summary>
+	/// Credentials for a Perforce user
+	/// </summary>
+	public class PerforceCredentials
+	{
+		/// <summary>
+		/// The username
+		/// </summary>
+		public string UserName { get; set; } = String.Empty;
+
+		/// <summary>
+		/// Password for the user
+		/// </summary>
+		public string Password { get; set; } = String.Empty;
+	}
+
+	/// <summary>
+	/// Information about a cluster of Perforce servers. 
+	/// </summary>
+	public class PerforceCluster
+	{
+		/// <summary>
+		/// The default cluster name
+		/// </summary>
+		public const string DefaultName = "Default";
+
+		/// <summary>
+		/// Name of the cluster
+		/// </summary>
+		[Required]
+		public string Name { get; set; } = null!;
+
+		/// <summary>
+		/// List of servers
+		/// </summary>
+		public List<PerforceServer> Servers { get; set; } = new List<PerforceServer>();
+
+		/// <summary>
+		/// List of server credentials
+		/// </summary>
+		public List<PerforceCredentials> Credentials { get; set; } = new List<PerforceCredentials>();
+	}
+
+	/// <summary>
+	/// Global server settings
+	/// </summary>
+	[SingletonDocument("5e3981cb28b8ec59cd07184a")]
+	public class Globals : SingletonBase
+	{
+		/// <summary>
+		/// The unique id for all globals objects
+		/// </summary>
+		public static ObjectId StaticId { get; } = new ObjectId("5e3981cb28b8ec59cd07184a");
+
+		/// <summary>
+		/// Unique instance id of this database
+		/// </summary>
+		public ObjectId InstanceId { get; set; }
+
+		/// <summary>
+		/// Set to true to drop the database on restart.
+		/// </summary>
+		public bool ForceReset { get; set; }
+
+		/// <summary>
+		/// Manually added status messages
+		/// </summary>
+		public List<Notice> Notices { get; set; } = new List<Notice>();
+
+		/// <summary>
+		/// List of Perforce clusters
+		/// </summary>
+		public List<PerforceCluster> PerforceClusters { get; set; } = new List<PerforceCluster>();
+
+		/// <summary>
+		/// Randomly generated JWT signing key
+		/// </summary>
+		public byte[]? JwtSigningKey { get; set; }
+
+		/// <summary>
+		/// The current schema version
+		/// </summary>
+		public int? SchemaVersion { get; set; }
+
+		/// <summary>
+		/// The scheduled downtime
+		/// </summary>
+		public List<ScheduledDowntime> ScheduledDowntime { get; set; } = new List<ScheduledDowntime>();
+
+		/// <summary>
+		/// Private constructor for serialization
+		/// </summary>
+		public Globals()
+		{
+			InstanceId = ObjectId.GenerateNewId();
+		}
+
+		/// <summary>
+		/// Rotate the signing key
+		/// </summary>
+		public void RotateSigningKey()
+		{
+			using (RNGCryptoServiceProvider Generator = new RNGCryptoServiceProvider())
+			{
+				JwtSigningKey = new byte[128];
+				Generator.GetBytes(JwtSigningKey);
+			}
+		}
+
+		/// <inheritdoc/>
+		public override void PostLoad()
+		{
+			base.PostLoad();
+
+			DateTimeOffset Now = DateTimeOffset.Now;
+			Notices.RemoveAll(x => x.FinishTime != null && x.FinishTime < Now);
+			ScheduledDowntime.RemoveAll(x => x.FinishTime < Now);
+		}
+	}
+}
diff --git a/Engine/Source/Programs/Horde/HordeServer/Models/Graph.cs b/Engine/Source/Programs/Horde/HordeServer/Models/Graph.cs
new file mode 100644
--- /dev/null
+++ b/Engine/Source/Programs/Horde/HordeServer/Models/Graph.cs
@@ -0,0 +1,370 @@
+// Copyright Epic Games, Inc. All Rights Reserved.
+
+using EpicGames.Core;
+using HordeServer.Api;
+using HordeCommon;
+using HordeServer.Utilities;
+using MongoDB.Bson;
+using MongoDB.Bson.Serialization.Attributes;
+using Newtonsoft.Json;
+using System;
+using System.Collections.Generic;
+using System.Diagnostics;
+using System.Diagnostics.CodeAnalysis;
+using System.Linq;
+
+namespace HordeServer.Models
+{
+	/// <summary>
+	/// Represents a node in the graph
+	/// </summary>
+	public interface INode
+	{
+		/// <summary>
+		/// The name of this node 
+		/// </summary>
+		public string Name { get; }
+
+		/// <summary>
+		/// Indices of nodes which must have succeeded for this node to run
+		/// </summary>
+		public NodeRef[] InputDependencies { get; }
+
+		/// <summary>
+		/// Indices of nodes which must have completed for this node to run
+		/// </summary>
+		public NodeRef[] OrderDependencies { get; }
+
+		/// <summary>
+		/// The priority that this node should be run at, within this job
+		/// </summary>
+		public Priority Priority { get; }
+
+		/// <summary>
+		/// Whether this node can be run multiple times
+		/// </summary>
+		public bool AllowRetry { get; }
+
+		/// <summary>
+		/// This node can start running early, before dependencies of other nodes in the same group are complete
+		/// </summary>
+		public bool RunEarly { get; }
+
+		/// <summary>
+		/// Whether to include warnings in the output (defaults to true)
+		/// </summary>
+		public bool Warnings { get; }
+
+		/// <summary>
+		/// List of credentials required for this node. Each entry maps an environment variable name to a credential in the form "CredentialName.PropertyName".
+		/// </summary>
+		public IReadOnlyDictionary<string, string>? Credentials { get; }
+
+		/// <summary>
+		/// Properties for this node
+		/// </summary>
+		public IReadOnlyDictionary<string, string>? Properties { get; }
+	}
+
+	/// <summary>
+	/// Information about a sequence of nodes which can execute on a single agent
+	/// </summary>
+	public interface INodeGroup
+	{
+		/// <summary>
+		/// The type of agent to execute this group
+		/// </summary>
+		public string AgentType { get; }
+
+		/// <summary>
+		/// Nodes in this group
+		/// </summary>
+		public IReadOnlyList<INode> Nodes { get; }
+	}
+
+	/// <summary>
+	/// Reference to a node within another grup
+	/// </summary>
+	[DebuggerDisplay("Group: {GroupIdx}, Node: {NodeIdx}")]
+	public class NodeRef
+	{
+		/// <summary>
+		/// The group index of the referenced node
+		/// </summary>
+		public int GroupIdx { get; set; }
+
+		/// <summary>
+		/// The node index of the referenced node
+		/// </summary>
+		public int NodeIdx { get; set; }
+
+		/// <summary>
+		/// Private constructor for serialization
+		/// </summary>
+		[BsonConstructor]
+		private NodeRef()
+		{
+			GroupIdx = 0;
+			NodeIdx = 0;
+		}
+
+		/// <summary>
+		/// Constructor
+		/// </summary>
+		/// <param name="GroupIdx">Index of thr group containing the node</param>
+		/// <param name="NodeIdx">Index of the node within the group</param>
+		public NodeRef(int GroupIdx, int NodeIdx)
+		{
+			this.GroupIdx = GroupIdx;
+			this.NodeIdx = NodeIdx;
+		}
+
+		/// <inheritdoc/>
+		public override bool Equals(object? Other)
+		{
+			NodeRef? OtherNodeRef = Other as NodeRef;
+			return OtherNodeRef != null && OtherNodeRef.GroupIdx == GroupIdx && OtherNodeRef.NodeIdx == NodeIdx;
+		}
+
+		/// <inheritdoc/>
+		public override int GetHashCode()
+		{
+			return HashCode.Combine(GroupIdx, NodeIdx);
+		}
+
+		/// <summary>
+		/// Converts this reference to a node name
+		/// </summary>
+		/// <param name="Groups">List of groups that this reference points to</param>
+		/// <returns>Name of the referenced node</returns>
+		public INode ToNode(IReadOnlyList<INodeGroup> Groups)
+		{
+			return Groups[GroupIdx].Nodes[NodeIdx];
+		}
+	}
+
+	/// <summary>
+	/// An collection of node references
+	/// </summary>
+	public interface IAggregate
+	{
+		/// <summary>
+		/// Name of the aggregate
+		/// </summary>
+		public string Name { get; }
+
+		/// <summary>
+		/// List of nodes for the aggregate to be valid
+		/// </summary>
+		public IReadOnlyList<NodeRef> Nodes { get; }
+	}
+
+	/// <summary>
+	/// Label indicating the status of a set of nodes
+	/// </summary>
+	public interface ILabel
+	{
+		/// <summary>
+		/// Label to show in the dashboard. Null if does not need to be shown.
+		/// </summary>
+		public string? DashboardName { get; }
+
+		/// <summary>
+		/// Category for the label. May be null.
+		/// </summary>
+		public string? DashboardCategory { get; }
+
+		/// <summary>
+		/// Name to display for this label in UGS
+		/// </summary>
+		public string? UgsName { get; }
+
+		/// <summary>
+		/// Project which this label applies to, for UGS
+		/// </summary>
+		public string? UgsProject { get; }
+
+		/// <summary>
+		/// Which change to display the label on
+		/// </summary>
+		public LabelChange Change { get; }
+
+		/// <summary>
+		/// List of required nodes for the aggregate to be valid
+		/// </summary>
+		public List<NodeRef> RequiredNodes { get; }
+
+		/// <summary>
+		/// List of optional nodes to include in the aggregate state
+		/// </summary>
+		public List<NodeRef> IncludedNodes { get; }
+	}
+
+	/// <summary>
+	/// Extension methods for ILabel
+	/// </summary>
+	public static class LabelExtensions
+	{
+		/// <summary>
+		/// Enumerate all the required dependencies of this node group
+		/// </summary>
+		/// <param name="Label">The label instance</param>
+		/// <param name="Groups">List of groups for the job containing this aggregate</param>
+		/// <returns>Sequence of nodes</returns>
+		public static IEnumerable<INode> GetDependencies(this ILabel Label, IReadOnlyList<INodeGroup> Groups)
+		{
+			foreach (NodeRef RequiredNode in Label.RequiredNodes)
+			{
+				yield return RequiredNode.ToNode(Groups);
+			}
+			foreach (NodeRef IncludedNode in Label.IncludedNodes)
+			{
+				yield return IncludedNode.ToNode(Groups);
+			}
+		}
+	}
+
+	/// <summary>
+	/// A unique dependency graph instance
+	/// </summary>
+	public interface IGraph
+	{
+		/// <summary>
+		/// Hash of this graph
+		/// </summary>
+		public ContentHash Id { get; }
+
+		/// <summary>
+		/// Schema version for this document
+		/// </summary>
+		public int Schema { get; }
+
+		/// <summary>
+		/// List of groups for this graph
+		/// </summary>
+		public IReadOnlyList<INodeGroup> Groups { get; }
+
+		/// <summary>
+		/// List of aggregates for this graph
+		/// </summary>
+		public IReadOnlyList<IAggregate> Aggregates { get; }
+
+		/// <summary>
+		/// Status labels for this graph
+		/// </summary>
+		public IReadOnlyList<ILabel> Labels { get; }
+	}
+
+	/// <summary>
+	/// Extension methods for graphs
+	/// </summary>
+	public static class GraphExtensions
+	{
+		/// <summary>
+		/// Gets the node from a node reference
+		/// </summary>
+		/// <param name="Graph">The graph instance</param>
+		/// <param name="Ref">The node reference</param>
+		/// <returns>The node for the given reference</returns>
+		public static INode GetNode(this IGraph Graph, NodeRef Ref)
+		{
+			return Graph.Groups[Ref.GroupIdx].Nodes[Ref.NodeIdx];
+		}
+
+		/// <summary>
+		/// Tries to find a node by name
+		/// </summary>
+		/// <param name="Graph">The graph to search</param>
+		/// <param name="NodeName">Name of the node</param>
+		/// <param name="Ref">Receives the node reference</param>
+		/// <returns>True if the node was found, false otherwise</returns>
+		public static bool TryFindNode(this IGraph Graph, string NodeName, out NodeRef Ref)
+		{
+			for(int GroupIdx = 0; GroupIdx < Graph.Groups.Count; GroupIdx++)
+			{
+				INodeGroup Group = Graph.Groups[GroupIdx];
+				for(int NodeIdx = 0; NodeIdx < Group.Nodes.Count; NodeIdx++)
+				{
+					INode Node = Group.Nodes[NodeIdx];
+					if (String.Equals(Node.Name, NodeName, StringComparison.OrdinalIgnoreCase))
+					{
+						Ref = new NodeRef(GroupIdx, NodeIdx);
+						return true;
+					}
+				}
+			}
+
+			Ref = new NodeRef(0, 0);
+			return false;
+		}
+
+		/// <summary>
+		/// Tries to find a node by name
+		/// </summary>
+		/// <param name="Graph">The graph to search</param>
+		/// <param name="NodeName">Name of the node</param>
+		/// <param name="Node">Receives the node</param>
+		/// <returns>True if the node was found, false otherwise</returns>
+		public static bool TryFindNode(this IGraph Graph, string NodeName, out INode? Node)
+		{
+			NodeRef Ref;
+			if (TryFindNode(Graph, NodeName, out Ref))
+			{
+				Node = Graph.Groups[Ref.GroupIdx].Nodes[Ref.NodeIdx];
+				return true;
+			}
+			else
+			{
+				Node = null;
+				return false;
+			}
+		}
+
+		/// <summary>
+		/// Tries to find a node by name
+		/// </summary>
+		/// <param name="Graph">The graph to search</param>
+		/// <param name="Name">Name of the node</param>
+		/// <param name="AggregateIdx">Receives the aggregate index</param>
+		/// <returns>True if the node was found, false otherwise</returns>
+		public static bool TryFindAggregate(this IGraph Graph, string Name, out int AggregateIdx)
+		{
+			AggregateIdx = Graph.Aggregates.FindIndex(x => x.Name.Equals(Name, StringComparison.OrdinalIgnoreCase));
+			return AggregateIdx != -1;
+		}
+
+		/// <summary>
+		/// Tries to find a node by name
+		/// </summary>
+		/// <param name="Graph">The graph to search</param>
+		/// <param name="Name">Name of the node</param>
+		/// <param name="Aggregate">Receives the aggregate</param>
+		/// <returns>True if the node was found, false otherwise</returns>
+		public static bool TryFindAggregate(this IGraph Graph, string Name, [NotNullWhen(true)] out IAggregate? Aggregate)
+		{
+			int AggregateIdx;
+			if (TryFindAggregate(Graph, Name, out AggregateIdx))
+			{
+				Aggregate = Graph.Aggregates[AggregateIdx];
+				return true;
+			}
+			else
+			{
+				Aggregate = null;
+				return false;
+			}
+		}
+
+		/// <summary>
+		/// Gets a list of dependencies for the given node
+		/// </summary>
+		/// <param name="Graph">The graph instance</param>
+		/// <param name="Node">The node to return dependencies for</param>
+		/// <returns>List of dependencies</returns>
+		public static IEnumerable<INode> GetDependencies(this IGraph Graph, INode Node)
+		{
+			return Enumerable.Concat(Node.InputDependencies, Node.OrderDependencies).Select(x => Graph.GetNode(x));
+		}
+	}
+}
+
diff --git a/Engine/Source/Programs/Horde/HordeServer/Models/IIssue.cs b/Engine/Source/Programs/Horde/HordeServer/Models/IIssue.cs
new file mode 100644
--- /dev/null
+++ b/Engine/Source/Programs/Horde/HordeServer/Models/IIssue.cs
@@ -0,0 +1,158 @@
+using EpicGames.Core;
+using HordeServer.Api;
+using HordeServer.Collections;
+using HordeCommon;
+using HordeServer.Services;
+using HordeServer.Utilities;
+using MongoDB.Bson;
+using MongoDB.Bson.Serialization.Attributes;
+using MongoDB.Bson.Serialization.Serializers;
+using System;
+using System.Collections.Generic;
+using System.Diagnostics.CodeAnalysis;
+using System.Linq;
+using System.Runtime.CompilerServices;
+
+using StreamId = HordeServer.Utilities.StringId<HordeServer.Models.IStream>;
+using TemplateRefId = HordeServer.Utilities.StringId<HordeServer.Models.TemplateRef>;
+
+namespace HordeServer.Models
+{
+	/// <summary>
+	/// The severity of an issue
+	/// </summary>
+	public enum IssueSeverity
+	{
+		/// <summary>
+		/// Unspecified severity
+		/// </summary>
+		Unspecified,
+
+		/// <summary>
+		/// This error represents a warning
+		/// </summary>
+		Warning,
+
+		/// <summary>
+		/// This issue represents an error
+		/// </summary>
+		Error,
+	}
+
+	/// <summary>
+	/// Stores information about a build health issue
+	/// </summary>
+	public interface IIssue
+	{
+		/// <summary>
+		/// The unique object id
+		/// </summary>
+		public int Id { get; }
+
+		/// <summary>
+		/// Summary for this issue
+		/// </summary>
+		public string Summary { get; }
+
+		/// <summary>
+		/// Summary set by users
+		/// </summary>
+		public string? UserSummary { get; }
+
+		/// <summary>
+		/// Fingerprint for this issue
+		/// </summary>
+		public IIssueFingerprint Fingerprint { get; }
+
+		/// <summary>
+		/// Severity of this issue
+		/// </summary>
+		public IssueSeverity Severity { get; }
+
+		/// <summary>
+		/// Current owner of this issue
+		/// </summary>
+		public string? Owner { get; }
+
+		/// <summary>
+		/// User id of the owner
+		/// </summary>
+		public ObjectId? OwnerId { get; }
+
+		/// <summary>
+		/// User that nominated the current owner
+		/// </summary>
+		public string? NominatedBy { get; }
+
+		/// <summary>
+		/// User id of the person that nominated the owner
+		/// </summary>
+		public ObjectId? NominatedById { get; }
+
+		/// <summary>
+		/// Time at which the issue was created
+		/// </summary>
+		public DateTime CreatedAt { get; }
+
+		/// <summary>
+		/// Time that the current owner was nominated
+		/// </summary>
+		public DateTime? NominatedAt { get; }
+
+		/// <summary>
+		/// Time at which the issue was acknowledged
+		/// </summary>
+		public DateTime? AcknowledgedAt { get; }
+
+		/// <summary>
+		/// Time at which the issue was resolved
+		/// </summary>
+		public DateTime? ResolvedAt { get; }
+
+		/// <summary>
+		/// Fix changelist for this issue
+		/// </summary>
+		public int? FixChange { get; }
+
+		/// <summary>
+		/// Whether all suspects should be notified about this issue
+		/// </summary>
+		public bool NotifySuspects { get; }
+
+		/// <summary>
+		/// Suspects for this issue
+		/// </summary>
+		public IReadOnlyList<IIssueSuspect> Suspects { get; }
+
+		/// <summary>
+		/// Update index for this instance
+		/// </summary>
+		public int UpdateIndex { get; }
+	}
+
+	/// <summary>
+	/// Suspect for an issue
+	/// </summary>
+	public interface IIssueSuspect
+	{
+		/// <summary>
+		/// Author of the change
+		/// </summary>
+		string Author { get; }
+
+		/// <summary>
+		/// The user id of the change's author
+		/// </summary>
+		ObjectId AuthorId { get; }
+
+		/// <summary>
+		/// The change suspected of causing this issue (in the origin stream)
+		/// </summary>
+		int Change { get; }
+
+		/// <summary>
+		/// Time at which the author declined the issue
+		/// </summary>
+		DateTime? DeclinedAt { get; }
+	}
+}
\ No newline at end of file
diff --git a/Engine/Source/Programs/Horde/HordeServer/Models/IIssueDetails.cs b/Engine/Source/Programs/Horde/HordeServer/Models/IIssueDetails.cs
new file mode 100644
--- /dev/null
+++ b/Engine/Source/Programs/Horde/HordeServer/Models/IIssueDetails.cs
@@ -0,0 +1,64 @@
+using HordeServer.Services;
+using System;
+using System.Collections.Generic;
+using System.Linq;
+using System.Threading.Tasks;
+
+namespace HordeServer.Models
+{
+	/// <summary>
+	/// Cached issue information
+	/// </summary>
+	public interface IIssueDetails
+	{
+		/// <summary>
+		/// The issue instance
+		/// </summary>
+		IIssue Issue { get; }
+
+		/// <summary>
+		/// List of spans for the issue
+		/// </summary>
+		IReadOnlyList<IIssueSpan> Spans { get; }
+
+		/// <summary>
+		/// List of steps for the issue
+		/// </summary>
+		IReadOnlyList<IIssueStep> Steps { get; }
+
+		/// <summary>
+		/// Determines whether the given user should be notified about the given issue
+		/// </summary>
+		/// <returns>True if the user should be notified for this change</returns>
+		bool ShowNotifications();
+
+		/// <summary>
+		/// Determines if the issue is relevant to the given user
+		/// </summary>
+		/// <param name="User">The user to query</param>
+		/// <returns>True if the issue is relevant to the given user</returns>
+		bool IncludeForUser(string User);
+	}
+
+	/// <summary>
+	/// Extension methods for IIssueService implementations
+	/// </summary>
+	public static class IssueDetailsExtensions
+	{
+		/// <summary>
+		/// Gets an issue details object for a specific issue id
+		/// </summary>
+		/// <param name="IssueService">The issue service</param>
+		/// <param name="IssueId">Issue id to query </param>
+		/// <returns></returns>
+		public static async Task<IIssueDetails?> GetIssueDetailsAsync(this IIssueService IssueService, int IssueId)
+		{
+			IIssue? Issue = await IssueService.GetIssueAsync(IssueId);
+			if(Issue == null)
+			{
+				return null;
+			}
+			return await IssueService.GetIssueDetailsAsync(Issue);
+		}
+	}
+}
diff --git a/Engine/Source/Programs/Horde/HordeServer/Models/IIssueFingerprint.cs b/Engine/Source/Programs/Horde/HordeServer/Models/IIssueFingerprint.cs
new file mode 100644
--- /dev/null
+++ b/Engine/Source/Programs/Horde/HordeServer/Models/IIssueFingerprint.cs
@@ -0,0 +1,77 @@
+using HordeServer.Utilities;
+using System;
+using System.Collections.Generic;
+using System.Linq;
+using System.Threading.Tasks;
+
+namespace HordeServer.Models
+{
+	/// <summary>
+	/// Fingerprint for an issue
+	/// </summary>
+	public interface IIssueFingerprint
+	{
+		/// <summary>
+		/// The type of issue
+		/// </summary>
+		public string Type { get; }
+
+		/// <summary>
+		/// List of keys which identify this issue.
+		/// </summary>
+		public CaseInsensitiveStringSet Keys { get; }
+
+		/// <summary>
+		/// Set of keys which should trigger a negative match
+		/// </summary>
+		public CaseInsensitiveStringSet? RejectKeys { get; }
+	}
+
+	/// <summary>
+	/// Extension methods for <see cref="IIssueFingerprint"/>
+	/// </summary>
+	public static class IssueFingerprintExtensions
+	{
+		/// <summary>
+		/// Checks if a fingerprint matches another fingerprint
+		/// </summary>
+		/// <param name="Fingerprint">The first fingerprint to compare</param>
+		/// <param name="Other">The other fingerprint to compare to</param>
+		/// <returns>True is the fingerprints match</returns>
+		public static bool IsMatch(this IIssueFingerprint Fingerprint, IIssueFingerprint Other)
+		{
+			if (!Fingerprint.Type.Equals(Other.Type, StringComparison.Ordinal))
+			{
+				return false;
+			}
+			if (Fingerprint.Keys.Count > 0 && Other.Keys.Count > 0 && !Fingerprint.Keys.Any(x => Other.Keys.Contains(x)))
+			{
+				return false;
+			}
+			if (Fingerprint.RejectKeys != null && Fingerprint.RejectKeys.Any(x => Other.Keys.Contains(x)))
+			{
+				return false;
+			}
+			return true;
+		}
+
+		/// <summary>
+		/// Checks if a fingerprint matches another fingerprint for creating a new span
+		/// </summary>
+		/// <param name="Fingerprint">The first fingerprint to compare</param>
+		/// <param name="Other">The other fingerprint to compare to</param>
+		/// <returns>True is the fingerprints match</returns>
+		public static bool IsMatchForNewSpan(this IIssueFingerprint Fingerprint, IIssueFingerprint Other)
+		{
+			if (!Fingerprint.Type.Equals(Other.Type, StringComparison.Ordinal))
+			{
+				return false;
+			}
+			if (Fingerprint.RejectKeys != null && Fingerprint.RejectKeys.Any(x => Other.Keys.Contains(x)))
+			{
+				return false;
+			}
+			return true;
+		}
+	}
+}
diff --git a/Engine/Source/Programs/Horde/HordeServer/Models/IIssueStep.cs b/Engine/Source/Programs/Horde/HordeServer/Models/IIssueStep.cs
new file mode 100644
--- /dev/null
+++ b/Engine/Source/Programs/Horde/HordeServer/Models/IIssueStep.cs
@@ -0,0 +1,61 @@
+using MongoDB.Bson;
+using System;
+using System.Collections.Generic;
+using System.Linq;
+using System.Threading.Tasks;
+
+using SubResourceId = HordeServer.Models.SubResourceId;
+
+namespace HordeServer.Models
+{
+	/// <summary>
+	/// Identifies a particular changelist and job that contributes to a span
+	/// </summary>
+	public interface IIssueStep
+	{
+		/// <summary>
+		/// The span that this step belongs to
+		/// </summary>
+		public ObjectId SpanId { get; }
+
+		/// <summary>
+		/// The changelist number
+		/// </summary>
+		public int Change { get; }
+
+		/// <summary>
+		/// Name of the job
+		/// </summary>
+		public string JobName { get; }
+
+		/// <summary>
+		/// The unique job id
+		/// </summary>
+		public ObjectId JobId { get; }
+
+		/// <summary>
+		/// Unique id of the batch within the job
+		/// </summary>
+		public SubResourceId BatchId { get; }
+
+		/// <summary>
+		/// Unique id of the step within the job
+		/// </summary>
+		public SubResourceId StepId { get; }
+
+		/// <summary>
+		/// Time that the step started
+		/// </summary>
+		public DateTime StepTime { get; }
+
+		/// <summary>
+		/// Whether to notify suspects for this issue
+		/// </summary>
+		public bool NotifySuspects { get; }
+
+		/// <summary>
+		/// The log id for this step
+		/// </summary>
+		public ObjectId? LogId { get; }
+	}
+}
diff --git a/Engine/Source/Programs/Horde/HordeServer/Models/IUserClaim.cs b/Engine/Source/Programs/Horde/HordeServer/Models/IUserClaim.cs
new file mode 100644
--- /dev/null
+++ b/Engine/Source/Programs/Horde/HordeServer/Models/IUserClaim.cs
@@ -0,0 +1,72 @@
+using System;
+using System.Collections.Generic;
+using System.Linq;
+using System.Security.Claims;
+using System.Threading.Tasks;
+
+namespace HordeServer.Models
+{
+	/// <summary>
+	/// Claim for a user
+	/// </summary>
+	public interface IUserClaim
+	{
+		/// <summary>
+		/// Name of the claim
+		/// </summary>
+		public string Type { get; }
+
+		/// <summary>
+		/// Value of the claim
+		/// </summary>
+		public string Value { get; }
+	}
+
+	/// <summary>
+	/// New claim document
+	/// </summary>
+	public class UserClaim : IUserClaim
+	{
+		/// <inheritdoc/>
+		public string Type { get; set; }
+
+		/// <inheritdoc/>
+		public string Value { get; set; }
+
+		/// <summary>
+		/// Constructor
+		/// </summary>
+		/// <param name="Type"></param>
+		/// <param name="Value"></param>
+		public UserClaim(string Type, string Value)
+		{
+			this.Type = Type;
+			this.Value = Value;
+		}
+
+		/// <summary>
+		/// Constructor
+		/// </summary>
+		/// <param name="Other">Claim to construct from</param>
+		public UserClaim(IUserClaim Other)
+			: this(Other.Type, Other.Value)
+		{
+		}
+
+		/// <summary>
+		/// Constructs a UserClaim from a Claim object
+		/// </summary>
+		/// <param name="Claim">Claim object</param>
+		/// <returns></returns>
+		public static UserClaim FromClaim(Claim Claim)
+		{
+			return new UserClaim(Claim.Type, Claim.Value);
+		}
+
+		/// <summary>
+		/// Conversion operator from NET claims
+		/// </summary>
+		/// <param name="Claim"></param>
+		public static implicit operator UserClaim(Claim Claim) => FromClaim(Claim);
+	}
+}
