Initial drop of Horde Server/Agent for UE5/Main. Still missing some Epic-specific functionality required to build/run correctly.

diff --git a/Engine/Source/Programs/Horde/HordeServer/Models/Job.cs b/Engine/Source/Programs/Horde/HordeServer/Models/Job.cs
new file mode 100644
--- /dev/null
+++ b/Engine/Source/Programs/Horde/HordeServer/Models/Job.cs
@@ -0,0 +1,1494 @@
+// Copyright Epic Games, Inc. All Rights Reserved.
+
+using EpicGames.Core;
+using HordeServer.Api;
+using HordeCommon;
+using HordeServer.Utilities;
+using MongoDB.Bson;
+using MongoDB.Bson.Serialization.Attributes;
+using MongoDB.Driver;
+using System;
+using System.Collections.Generic;
+using System.ComponentModel;
+using System.Diagnostics.CodeAnalysis;
+using System.Linq;
+using System.Runtime.CompilerServices;
+using System.Text;
+using System.Text.Json.Serialization;
+using System.Threading.Tasks;
+
+using ReportPlacement = HordeCommon.Rpc.ReportPlacement;
+using StreamId = HordeServer.Utilities.StringId<HordeServer.Models.IStream>;
+using TemplateRefId = HordeServer.Utilities.StringId<HordeServer.Models.TemplateRef>;
+using PoolId = HordeServer.Utilities.StringId<HordeServer.Models.IPool>;
+
+namespace HordeServer.Models
+{
+	/// <summary>
+	/// Report for a job or jobstep
+	/// </summary>
+	public interface IReport
+	{
+		/// <summary>
+		/// Name of the report
+		/// </summary>
+		string Name { get; }
+
+		/// <summary>
+		/// Where to render the report
+		/// </summary>
+		ReportPlacement Placement { get; }
+
+		/// <summary>
+		/// The artifact id
+		/// </summary>
+		ObjectId ArtifactId { get; }
+	}
+
+	/// <summary>
+	/// Implementation of IReport
+	/// </summary>
+	public class Report : IReport
+	{
+		/// <inheritdoc/>
+		public string Name { get; set; } = String.Empty;
+
+		/// <inheritdoc/>
+		public ReportPlacement Placement { get; set; }
+
+		/// <inheritdoc/>
+		public ObjectId ArtifactId { get; set; }
+	}
+
+	/// <summary>
+	/// Embedded jobstep document
+	/// </summary>
+	public interface IJobStep
+	{
+		/// <summary>
+		/// Unique ID assigned to this jobstep. A new id is generated whenever a jobstep's order is changed.
+		/// </summary>
+		public SubResourceId Id { get; }
+
+		/// <summary>
+		/// Index of the node which this jobstep is to execute
+		/// </summary>
+		public int NodeIdx { get; }
+
+		/// <summary>
+		/// Current state of the job step. This is updated automatically when runs complete.
+		/// </summary>
+		public JobStepState State { get; }
+
+		/// <summary>
+		/// Current outcome of the jobstep
+		/// </summary>
+		public JobStepOutcome Outcome { get; }
+
+		/// <summary>
+		/// The log id for this step
+		/// </summary>
+		public ObjectId? LogId { get; }
+
+		/// <summary>
+		/// Unique id for notifications
+		/// </summary>
+		public ObjectId? NotificationTriggerId { get; }
+
+		/// <summary>
+		/// Time at which the batch transitioned to the ready state (UTC).
+		/// </summary>
+		public DateTime? ReadyTimeUtc { get; }
+
+		/// <summary>
+		/// Time at which the batch transitioned to the executing state (UTC).
+		/// </summary>
+		public DateTime? StartTimeUtc { get; }
+
+		/// <summary>
+		/// Time at which the run finished (UTC)
+		/// </summary>
+		public DateTime? FinishTimeUtc { get; }
+
+		/// <summary>
+		/// Override for the priority of this step
+		/// </summary>
+		public Priority? Priority { get; }
+
+		/// <summary>
+		/// If a retry is requested, stores the name of the user that requested it
+		/// </summary>
+		public string? RetryByUser { get; }
+		
+		/// <summary>
+		/// Signal if a step should be aborted
+		/// </summary>
+		public bool AbortRequested { get; }
+		
+		/// <summary>
+		/// If an abort is requested, stores the name of the user that requested it
+		/// </summary>
+		public string? AbortByUser { get; }
+
+		/// <summary>
+		/// List of reports for this step
+		/// </summary>
+		public IReadOnlyList<IReport>? Reports { get; }
+
+		/// <summary>
+		/// Reports for this jobstep.
+		/// </summary>
+		public Dictionary<string, string>? Properties { get; }
+	}
+
+	/// <summary>
+	/// Extension methods for job steps
+	/// </summary>
+	public static class JobStepExtensions
+	{
+		/// <summary>
+		/// Determines if a jobstep has failed or is skipped. Can be used to determine whether dependent steps will be able to run.
+		/// </summary>
+		/// <returns>True if the step is failed or skipped</returns>
+		public static bool IsFailedOrSkipped(this IJobStep Step)
+		{
+			return Step.State == JobStepState.Skipped || Step.Outcome == JobStepOutcome.Failure;
+		}
+
+		/// <summary>
+		/// Determines if a jobstep is done by checking to see if it is completed, skipped, or aborted.
+		/// </summary>
+		/// <returns>True if the step is completed, skipped, or aborted</returns>
+		public static bool IsPending(this IJobStep Step)
+		{
+			return Step.State != JobStepState.Aborted && Step.State != JobStepState.Completed && Step.State != JobStepState.Skipped;
+		}
+	}
+
+	/// <summary>
+	/// Stores information about a batch of job steps
+	/// </summary>
+	public interface IJobStepBatch
+	{
+		/// <summary>
+		/// Unique id for this group
+		/// </summary>
+		public SubResourceId Id { get; }
+
+		/// <summary>
+		/// The log file id for this batch
+		/// </summary>
+		public ObjectId? LogId { get; }
+
+		/// <summary>
+		/// Index of the group being executed
+		/// </summary>
+		public int GroupIdx { get; }
+
+		/// <summary>
+		/// The state of this group
+		/// </summary>
+		public JobStepBatchState State { get; }
+
+		/// <summary>
+		/// Error associated with this group
+		/// </summary>
+		public JobStepBatchError Error { get; }
+
+		/// <summary>
+		/// Steps within this run
+		/// </summary>
+		public IReadOnlyList<IJobStep> Steps { get; }
+
+		/// <summary>
+		/// The pool that this agent was taken from
+		/// </summary>
+		public PoolId? PoolId { get; }
+
+		/// <summary>
+		/// The agent assigned to execute this group
+		/// </summary>
+		public AgentId? AgentId { get; }
+
+		/// <summary>
+		/// The agent session that is executing this group
+		/// </summary>
+		public ObjectId? SessionId { get; }
+
+		/// <summary>
+		/// The lease that's executing this group
+		/// </summary>
+		public ObjectId? LeaseId { get; }
+
+		/// <summary>
+		/// The weighted priority of this batch for the scheduler
+		/// </summary>
+		public int SchedulePriority { get; }
+
+		/// <summary>
+		/// Time at which the group became ready (UTC).
+		/// </summary>
+		public DateTime? ReadyTimeUtc { get; }
+
+		/// <summary>
+		/// Time at which the group started (UTC).
+		/// </summary>
+		public DateTime? StartTimeUtc { get; }
+
+		/// <summary>
+		/// Time at which the group finished (UTC)
+		/// </summary>
+		public DateTime? FinishTimeUtc { get; }
+	}
+
+	/// <summary>
+	/// Extension methods for IJobStepBatch
+	/// </summary>
+	public static class JobStepBatchExtensions
+	{
+		/// <summary>
+		/// Attempts to get a step with the given id
+		/// </summary>
+		/// <param name="Batch">The batch to search</param>
+		/// <param name="StepId">The step id</param>
+		/// <param name="Step">On success, receives the step object</param>
+		/// <returns>True if the step was found</returns>
+		public static bool TryGetStep(this IJobStepBatch Batch, SubResourceId StepId, [NotNullWhen(true)] out IJobStep Step)
+		{
+			Step = Batch.Steps.FirstOrDefault(x => x.Id == StepId);
+			return Step != null;
+		}
+
+		/// <summary>
+		/// Determines if new steps can be appended to this batch. We do not allow this after the last step has been completed, because the agent is shutting down.
+		/// </summary>
+		/// <param name="Batch">The batch to search</param>
+		/// <returns>True if new steps can be appended to this batch</returns>
+		public static bool CanBeAppendedTo(this IJobStepBatch Batch)
+		{
+			for (int Idx = Batch.Steps.Count - 1; Idx >= 0; Idx--)
+			{
+				IJobStep Step = Batch.Steps[Idx];
+				if (Step.State != JobStepState.Skipped)
+				{
+					return Step.State != JobStepState.Completed && Step.State != JobStepState.Aborted;
+				}
+			}
+			return false;
+		}
+
+		/// <summary>
+		/// Gets the wait time for this batch
+		/// </summary>
+		/// <param name="Batch">The batch to search</param>
+		/// <returns>Wait time for the batch</returns>
+		public static TimeSpan? GetWaitTime(this IJobStepBatch Batch)
+		{
+			if (Batch.StartTimeUtc == null || Batch.ReadyTimeUtc == null)
+			{
+				return null;
+			}
+			else
+			{
+				return Batch.StartTimeUtc.Value - Batch.ReadyTimeUtc.Value;
+			}
+		}
+
+		/// <summary>
+		/// Gets the initialization time for this batch
+		/// </summary>
+		/// <param name="Batch">The batch to search</param>
+		/// <returns>Initialization time for this batch</returns>
+		public static TimeSpan? GetInitTime(this IJobStepBatch Batch)
+		{
+			if (Batch.StartTimeUtc != null)
+			{
+				foreach (IJobStep Step in Batch.Steps)
+				{
+					if (Step.StartTimeUtc != null)
+					{
+						return Step.StartTimeUtc - Batch.StartTimeUtc.Value;
+					}
+				}
+			}
+			return null;
+		}
+
+		/// <summary>
+		/// Get the dependencies required for this batch to start, taking run-early nodes into account
+		/// </summary>
+		/// <param name="Batch">The batch to search</param>
+		/// <param name="Groups">List of node groups</param>
+		/// <returns>Set of nodes that must have completed for this batch to start</returns>
+		public static HashSet<INode> GetStartDependencies(this IJobStepBatch Batch, IReadOnlyList<INodeGroup> Groups)
+		{
+			// Find all the nodes that this group will start with.
+			List<INode> Nodes = Batch.Steps.ConvertAll(x => Groups[Batch.GroupIdx].Nodes[x.NodeIdx]);
+			if (Nodes.Any(x => x.RunEarly))
+			{
+				Nodes.RemoveAll(x => !x.RunEarly);
+			}
+
+			// Find all their dependencies
+			HashSet<INode> Dependencies = new HashSet<INode>();
+			foreach (INode Node in Nodes)
+			{
+				Dependencies.UnionWith(Node.InputDependencies.Select(x => Groups[x.GroupIdx].Nodes[x.NodeIdx]));
+				Dependencies.UnionWith(Node.OrderDependencies.Select(x => Groups[x.GroupIdx].Nodes[x.NodeIdx]));
+			}
+
+			// Exclude all the dependencies within the same group
+			Dependencies.ExceptWith(Groups[Batch.GroupIdx].Nodes);
+			return Dependencies;
+		}
+	}
+
+	/// <summary>
+	/// Cumulative timing information to reach a certain point in a job
+	/// </summary>
+	public class TimingInfo
+	{
+		/// <summary>
+		/// Wait time on the critical path
+		/// </summary>
+		public TimeSpan? TotalWaitTime { get; set; }
+
+		/// <summary>
+		/// Sync time on the critical path
+		/// </summary>
+		public TimeSpan? TotalInitTime { get; set; }
+
+		/// <summary>
+		/// Duration to this point
+		/// </summary>
+		public TimeSpan? TotalTimeToComplete { get; set; }
+
+		/// <summary>
+		/// Average wait time to this point
+		/// </summary>
+		public TimeSpan? AverageTotalWaitTime { get; set; }
+
+		/// <summary>
+		/// Average sync time to this point
+		/// </summary>
+		public TimeSpan? AverageTotalInitTime { get; set; }
+
+		/// <summary>
+		/// Average duration to this point
+		/// </summary>
+		public TimeSpan? AverageTotalTimeToComplete { get; set; }
+
+		/// <summary>
+		/// Individual step timing information
+		/// </summary>
+		public IJobStepTiming? StepTiming { get; set; }
+
+		/// <summary>
+		/// Constructor
+		/// </summary>
+		public TimingInfo()
+		{
+			TotalWaitTime = TimeSpan.Zero;
+			TotalInitTime = TimeSpan.Zero;
+			TotalTimeToComplete = TimeSpan.Zero;
+
+			AverageTotalWaitTime = TimeSpan.Zero;
+			AverageTotalInitTime = TimeSpan.Zero;
+			AverageTotalTimeToComplete = TimeSpan.Zero;
+		}
+
+		/// <summary>
+		/// Copy constructor
+		/// </summary>
+		/// <param name="Other">The timing info object to copy from</param>
+		public TimingInfo(TimingInfo Other)
+		{
+			TotalWaitTime = Other.TotalWaitTime;
+			TotalInitTime = Other.TotalInitTime;
+			TotalTimeToComplete = Other.TotalTimeToComplete;
+
+			AverageTotalWaitTime = Other.AverageTotalWaitTime;
+			AverageTotalInitTime = Other.AverageTotalInitTime;
+			AverageTotalTimeToComplete = Other.AverageTotalTimeToComplete;
+		}
+
+		/// <summary>
+		/// Modifies this timing to wait for another timing
+		/// </summary>
+		/// <param name="Other">The other node to wait for</param>
+		public void WaitFor(TimingInfo Other)
+		{
+			if (TotalTimeToComplete != null)
+			{
+				if (Other.TotalTimeToComplete == null || Other.TotalTimeToComplete.Value > TotalTimeToComplete.Value)
+				{
+					TotalInitTime = Other.TotalInitTime;
+					TotalWaitTime = Other.TotalWaitTime;
+					TotalTimeToComplete = Other.TotalTimeToComplete;
+				}
+			}
+
+			if (AverageTotalTimeToComplete != null)
+			{
+				if (Other.AverageTotalTimeToComplete == null || Other.AverageTotalTimeToComplete.Value > AverageTotalTimeToComplete.Value)
+				{
+					AverageTotalInitTime = Other.AverageTotalInitTime;
+					AverageTotalWaitTime = Other.AverageTotalWaitTime;
+					AverageTotalTimeToComplete = Other.AverageTotalTimeToComplete;
+				}
+			}
+		}
+
+		/// <summary>
+		/// Waits for all the given timing info objects to complete
+		/// </summary>
+		/// <param name="Others">Other timing info objects to wait for</param>
+		public void WaitForAll(IEnumerable<TimingInfo> Others)
+		{
+			foreach(TimingInfo Other in Others)
+			{
+				WaitFor(Other);
+			}
+		}
+
+		/// <summary>
+		/// Constructs a new TimingInfo object which represents the last TimingInfo to finish
+		/// </summary>
+		/// <param name="Others">TimingInfo objects to wait for</param>
+		/// <returns>New TimingInfo instance</returns>
+		public static TimingInfo Max(IEnumerable<TimingInfo> Others)
+		{
+			TimingInfo TimingInfo = new TimingInfo();
+			TimingInfo.WaitForAll(Others);
+			return TimingInfo;
+		}
+	}
+
+	/// <summary>
+	/// Information about a chained job trigger
+	/// </summary>
+	public interface IChainedJob
+	{
+		/// <summary>
+		/// The target to monitor
+		/// </summary>
+		public string Target { get; }
+
+		/// <summary>
+		/// The template to trigger on success
+		/// </summary>
+		public TemplateRefId TemplateRefId { get; }
+
+		/// <summary>
+		/// The triggered job id
+		/// </summary>
+		public ObjectId? JobId { get; }
+	}
+
+	/// <summary>
+	/// Document describing a job
+	/// </summary>
+	[SuppressMessage("Compiler", "CA1056:URI parameters should not be strings")]
+	public interface IJob
+	{
+		/// <summary>
+		/// Job argument indicating a target that should be built
+		/// </summary>
+		public const string TargetArgumentPrefix = "-Target=";
+
+		/// <summary>
+		/// Name of the node which parses the buildgraph script
+		/// </summary>
+		public const string SetupNodeName = "Setup Build";
+
+		/// <summary>
+		/// Identifier for the job. Randomly generated.
+		/// </summary>
+		public ObjectId Id { get; }
+
+		/// <summary>
+		/// The stream that this job belongs to
+		/// </summary>
+		public StreamId StreamId { get; }
+
+		/// <summary>
+		/// The template ref id
+		/// </summary>
+		public TemplateRefId TemplateId { get; }
+
+		/// <summary>
+		/// The template that this job was created from
+		/// </summary>
+		public ContentHash? TemplateHash { get; }
+
+		/// <summary>
+		/// Hash of the graph definition
+		/// </summary>
+		public ContentHash GraphHash { get; }
+
+		/// <summary>
+		/// Id of the user that started this job
+		/// </summary>
+		public ObjectId? StartedByUserId { get; }
+
+		/// <summary>
+		/// Name of the user that started this job
+		/// </summary>
+		public string? StartedByUser { get; }
+
+		/// <summary>
+		/// Name of the user that aborted this job. Set to null if the job is not aborted.
+		/// </summary>
+		public string? AbortedByUser { get; }
+
+		/// <summary>
+		/// Name of the job.
+		/// </summary>
+		public string Name { get; }
+
+		/// <summary>
+		/// The changelist number to build
+		/// </summary>
+		public int Change { get; }
+
+		/// <summary>
+		/// The code changelist number for this build
+		/// </summary>
+		public int CodeChange { get; }
+
+		/// <summary>
+		/// The preflight changelist number
+		/// </summary>
+		public int PreflightChange { get; }
+
+		/// <summary>
+		/// The cloned preflight changelist number (if the prefight change is duplicated via p4 reshelve)
+		/// </summary>
+		public int ClonedPreflightChange { get; }
+
+		/// <summary>
+		/// Priority of this job
+		/// </summary>
+		public Priority Priority { get; }
+
+		/// <summary>
+		/// For preflights, submit the change if the job is successful
+		/// </summary>
+		public bool AutoSubmit { get; }
+
+		/// <summary>
+		/// The submitted changelist number
+		/// </summary>
+		public int? AutoSubmitChange { get; }
+
+		/// <summary>
+		/// Message produced by trying to auto-submit the change
+		/// </summary>
+		public string? AutoSubmitMessage { get; }
+
+		/// <summary>
+		/// Time that the job was created (in UTC)
+		/// </summary>
+		public DateTime CreateTimeUtc { get; }
+
+		/// <summary>
+		/// Largest value of the CombinedPriority value for batches in the ready state.
+		/// </summary>
+		public int SchedulePriority { get; }
+
+		/// <summary>
+		/// Array of jobstep runs
+		/// </summary>
+		public IReadOnlyList<IJobStepBatch> Batches { get; }
+
+		/// <summary>
+		/// Optional user-defined properties for this job
+		/// </summary>
+		public IReadOnlyList<string> Arguments { get; }
+
+		/// <summary>
+		/// Issues associated with this job
+		/// </summary>
+		public IReadOnlyList<int> Issues { get; }
+
+		/// <summary>
+		/// Unique id for notifications
+		/// </summary>
+		public ObjectId? NotificationTriggerId { get; }
+
+		/// <summary>
+		/// Whether to show badges in UGS for this job
+		/// </summary>
+		public bool ShowUgsBadges { get; }
+
+		/// <summary>
+		/// Whether to show alerts in UGS for this job
+		/// </summary>
+		public bool ShowUgsAlerts { get; }
+
+		/// <summary>
+		/// Notification channel for this job.
+		/// </summary>
+		public string? NotificationChannel { get; }
+		
+		/// <summary>
+		/// Notification channel filter for this job.
+		/// </summary>
+		public string? NotificationChannelFilter { get; }
+		
+		/// <summary>
+		/// URL to a Helix Swarm server to call once job finishes. Used for adding build outcome to Swarm reviews.
+		/// </summary>
+		public string? HelixSwarmCallbackUrl { get; }
+
+		/// <summary>
+		/// Mapping of label ids to notification trigger ids for notifications
+		/// </summary>
+		public IReadOnlyDictionary<int, ObjectId> LabelIdxToTriggerId { get; }
+
+		/// <summary>
+		/// List of reports for this step
+		/// </summary>
+		public IReadOnlyList<IReport>? Reports { get; }
+
+		/// <summary>
+		/// List of downstream job triggers
+		/// </summary>
+		public IReadOnlyList<IChainedJob> ChainedJobs { get; }
+
+		/// <summary>
+		/// Next id for batches or groups
+		/// </summary>
+		public SubResourceId NextSubResourceId { get; }
+
+		/// <summary>
+		/// The last update time
+		/// </summary>
+		public DateTime UpdateTimeUtc { get; }
+
+		/// <summary>
+		/// The ACL for this job.
+		/// </summary>
+		public Acl? Acl { get; }
+
+		/// <summary>
+		/// Update counter for this document. Any updates should compare-and-swap based on the value of this counter, or increment it in the case of server-side updates.
+		/// </summary>
+		public int UpdateIndex { get; }
+	}
+
+	/// <summary>
+	/// Extension methods for jobs
+	/// </summary>
+	public static class JobExtensions
+	{
+		/// <summary>
+		/// Attempts to get a step with a given ID
+		/// </summary>
+		/// <param name="Job">The job document</param>
+		/// <param name="StepId">The step id</param>
+		/// <param name="Step">On success, receives the step object</param>
+		/// <returns>True if the step was found</returns>
+		public static bool TryGetStep(this IJob Job, SubResourceId StepId, [NotNullWhen(true)] out IJobStep? Step)
+		{
+			foreach (IJobStepBatch Batch in Job.Batches)
+			{
+				if (Batch.TryGetStep(StepId, out Step))
+				{
+					return true;
+				}
+			}
+
+			Step = null;
+			return false;
+		}
+
+		/// <summary>
+		/// Gets the current job state
+		/// </summary>
+		/// <param name="Job">The job document</param>
+		/// <returns>Job state</returns>
+		public static JobState GetState(this IJob Job)
+		{
+			bool bWaiting = false;
+			foreach (IJobStepBatch Batch in Job.Batches)
+			{
+				foreach (IJobStep Step in Batch.Steps)
+				{
+					if (Step.State == JobStepState.Running)
+					{
+						return JobState.Running;
+					}
+					else if (Step.State == JobStepState.Ready || Step.State == JobStepState.Waiting)
+					{
+						if (Batch.State == JobStepBatchState.Running)
+						{
+							return JobState.Running;
+						}
+						else
+						{
+							bWaiting = true;
+						}
+					}
+				}
+			}
+			return bWaiting ? JobState.Waiting : JobState.Complete;
+		}
+
+		/// <summary>
+		/// Gets the outcome for a particular named target. May be an aggregate or node name.
+		/// </summary>
+		/// <param name="Job">The job to check</param>
+		/// <returns>The step outcome</returns>
+		public static (JobStepState, JobStepOutcome) GetTargetState(this IJob Job)
+		{
+			IReadOnlyDictionary<NodeRef, IJobStep> NodeToStep = GetStepForNodeMap(Job);
+			return GetTargetState(NodeToStep.Values);
+		}
+
+		/// <summary>
+		/// Gets the outcome for a particular named target. May be an aggregate or node name.
+		/// </summary>
+		/// <param name="Job">The job to check</param>
+		/// <param name="Graph">Graph for the job</param>
+		/// <param name="Target">Target to find an outcome for</param>
+		/// <returns>The step outcome</returns>
+		public static (JobStepState, JobStepOutcome)? GetTargetState(this IJob Job, IGraph Graph, string? Target)
+		{
+			if (Target == null)
+			{
+				return GetTargetState(Job);
+			}
+
+			NodeRef NodeRef;
+			if (Graph.TryFindNode(Target, out NodeRef))
+			{
+				IJobStep? Step;
+				if (Job.TryGetStepForNode(NodeRef, out Step))
+				{
+					return (Step.State, Step.Outcome);
+				}
+				else
+				{
+					return null;
+				}
+			}
+
+			IAggregate? Aggregate;
+			if (Graph.TryFindAggregate(Target, out Aggregate))
+			{
+				IReadOnlyDictionary<NodeRef, IJobStep> StepForNode = GetStepForNodeMap(Job);
+
+				List<IJobStep> Steps = new List<IJobStep>();
+				foreach (NodeRef AggregateNodeRef in Aggregate.Nodes)
+				{
+					IJobStep? Step;
+					if (!StepForNode.TryGetValue(AggregateNodeRef, out Step))
+					{
+						return null;
+					}
+					Steps.Add(Step);
+				}
+
+				return GetTargetState(Steps);
+			}
+
+			return null;
+		}
+
+		/// <summary>
+		/// Gets the outcome for a particular named target. May be an aggregate or node name.
+		/// </summary>
+		/// <param name="Steps">Steps to include</param>
+		/// <returns>The step outcome</returns>
+		public static (JobStepState, JobStepOutcome) GetTargetState(IEnumerable<IJobStep> Steps)
+		{
+			bool bAnySkipped = false;
+			bool bAnyWarnings = false;
+			bool bAnyFailed = false;
+			bool bAnyPending = false;
+			foreach (IJobStep Step in Steps)
+			{
+				bAnyPending |= Step.IsPending();
+				bAnySkipped |= Step.State == JobStepState.Aborted || Step.State == JobStepState.Skipped;
+				bAnyFailed |= (Step.Outcome == JobStepOutcome.Failure);
+				bAnyWarnings |= (Step.Outcome == JobStepOutcome.Warnings);
+			}
+
+			JobStepState NewState = bAnyPending ? JobStepState.Running : JobStepState.Completed;
+			JobStepOutcome NewOutcome = bAnyFailed ? JobStepOutcome.Failure : bAnyWarnings ? JobStepOutcome.Warnings : bAnySkipped ? JobStepOutcome.Unspecified : JobStepOutcome.Success;
+			return (NewState, NewOutcome);
+		}
+
+		/// <summary>
+		/// Gets the outcome for a particular named target. May be an aggregate or node name.
+		/// </summary>
+		/// <param name="Job">The job to check</param>
+		/// <param name="Graph">Graph for the job</param>
+		/// <param name="Target">Target to find an outcome for</param>
+		/// <returns>The step outcome</returns>
+		public static JobStepOutcome GetTargetOutcome(this IJob Job, IGraph Graph, string Target)
+		{
+			NodeRef NodeRef;
+			if (Graph.TryFindNode(Target, out NodeRef))
+			{
+				IJobStep? Step;
+				if (Job.TryGetStepForNode(NodeRef, out Step))
+				{
+					return Step.Outcome;
+				}
+				else
+				{
+					return JobStepOutcome.Unspecified;
+				}
+			}
+
+			IAggregate? Aggregate;
+			if (Graph.TryFindAggregate(Target, out Aggregate))
+			{
+				IReadOnlyDictionary<NodeRef, IJobStep> StepForNode = GetStepForNodeMap(Job);
+
+				bool Warnings = false;
+				foreach (NodeRef AggregateNodeRef in Aggregate.Nodes)
+				{
+					IJobStep? Step;
+					if (!StepForNode.TryGetValue(AggregateNodeRef, out Step))
+					{
+						return JobStepOutcome.Unspecified;
+					}
+					if (Step.Outcome == JobStepOutcome.Failure)
+					{
+						return JobStepOutcome.Failure;
+					}
+					Warnings |= (Step.Outcome == JobStepOutcome.Warnings);
+				}
+				return Warnings ? JobStepOutcome.Warnings : JobStepOutcome.Success;
+			}
+
+			return JobStepOutcome.Unspecified;
+		}
+
+		/// <summary>
+		/// Gets the job step for a particular node
+		/// </summary>
+		/// <param name="Job">The job to search</param>
+		/// <param name="NodeRef">The node ref</param>
+		/// <param name="JobStep">Receives the jobstep on success</param>
+		/// <returns>True if the jobstep was founds</returns>
+		public static bool TryGetStepForNode(this IJob Job, NodeRef NodeRef, [NotNullWhen(true)] out IJobStep? JobStep)
+		{
+			JobStep = null;
+			foreach (IJobStepBatch Batch in Job.Batches)
+			{
+				if (Batch.GroupIdx == NodeRef.GroupIdx)
+				{
+					foreach (IJobStep BatchStep in Batch.Steps)
+					{
+						if (BatchStep.NodeIdx == NodeRef.NodeIdx)
+						{
+							JobStep = BatchStep;
+						}
+					}
+				}
+			}
+			return JobStep != null;
+		}
+
+		/// <summary>
+		/// Gets a dictionary that maps <see cref="NodeRef"/> objects to their associated
+		/// <see cref="IJobStep"/> objects on a <see cref="IJob"/>.
+		/// </summary>
+		/// <param name="Job">The job document</param>
+		/// <returns>Map of <see cref="NodeRef"/> to <see cref="IJobStep"/></returns>
+		public static IReadOnlyDictionary<NodeRef, IJobStep> GetStepForNodeMap(this IJob Job)
+		{
+			Dictionary<NodeRef, IJobStep> StepForNode = new Dictionary<NodeRef, IJobStep>();
+			foreach (IJobStepBatch Batch in Job.Batches)
+			{
+				foreach (IJobStep BatchStep in Batch.Steps)
+				{
+					NodeRef BatchNodeRef = new NodeRef(Batch.GroupIdx, BatchStep.NodeIdx);
+					StepForNode[BatchNodeRef] = BatchStep;
+				}
+			}
+			return StepForNode;
+		}
+
+		/// <summary>
+		/// Gets the estimated timing info for all nodes in the job
+		/// </summary>
+		/// <param name="Job">The job document</param>
+		/// <param name="Graph">Graph for this job</param>
+		/// <param name="JobTiming">Job timing information</param>
+		/// <returns>Map of node to expected timing info</returns>
+		public static Dictionary<INode, TimingInfo> GetTimingInfo(this IJob Job, IGraph Graph, IJobTiming JobTiming)
+		{
+			TimeSpan CurrentTime = DateTime.UtcNow - Job.CreateTimeUtc;
+
+			Dictionary<INode, TimingInfo> NodeToTimingInfo = Graph.Groups.SelectMany(x => x.Nodes).ToDictionary(x => x, x => new TimingInfo());
+			foreach (IJobStepBatch Batch in Job.Batches)
+			{
+				INodeGroup Group = Graph.Groups[Batch.GroupIdx];
+
+				// Step through the batch, keeping track of the time that things finish.
+				TimingInfo TimingInfo = new TimingInfo();
+
+				// Wait for the dependencies for the batch to start
+				HashSet<INode> DependencyNodes = Batch.GetStartDependencies(Graph.Groups);
+				TimingInfo.WaitForAll(DependencyNodes.Select(x => NodeToTimingInfo[x]));
+
+				// If the batch has actually started, correct the expected time to use this instead
+				if (Batch.StartTimeUtc != null)
+				{
+					TimingInfo.TotalTimeToComplete = Batch.StartTimeUtc - Job.CreateTimeUtc;
+				}
+
+				// Get the average times for this batch
+				TimeSpan? AverageWaitTime = GetAverageWaitTime(Graph, Batch, JobTiming);
+				TimeSpan? AverageInitTime = GetAverageInitTime(Graph, Batch, JobTiming);
+
+				// Update the wait times and initialization times along this path
+				TimingInfo.TotalWaitTime = TimingInfo.TotalWaitTime + (Batch.GetWaitTime() ?? AverageWaitTime);
+				TimingInfo.TotalInitTime = TimingInfo.TotalInitTime + (Batch.GetInitTime() ?? AverageInitTime);
+
+				// Update the average wait and initialization times too
+				TimingInfo.AverageTotalWaitTime = TimingInfo.AverageTotalWaitTime + AverageWaitTime;
+				TimingInfo.AverageTotalInitTime = TimingInfo.AverageTotalInitTime + AverageInitTime;
+
+				// Step through the batch, updating the expected times as we go
+				foreach (IJobStep Step in Batch.Steps)
+				{
+					INode Node = Group.Nodes[Step.NodeIdx];
+
+					// Get the timing for this step
+					IJobStepTiming? StepTimingInfo;
+					JobTiming.TryGetStepTiming(Node.Name, out StepTimingInfo);
+
+					// If the step has already started, update the actual time to reach this point
+					if(Step.StartTimeUtc != null)
+					{
+						TimingInfo.TotalTimeToComplete = Step.StartTimeUtc.Value - Job.CreateTimeUtc;
+					}
+
+					// If the step hasn't started yet, make sure the start time is later than the current time
+					if (Step.StartTimeUtc == null && CurrentTime > TimingInfo.TotalTimeToComplete)
+					{
+						TimingInfo.TotalTimeToComplete = CurrentTime;
+					}
+
+					// Wait for all the node dependencies to complete
+					TimingInfo.WaitForAll(Graph.GetDependencies(Node).Select(x => NodeToTimingInfo[x]));
+
+					// If the step has actually finished, correct the time to use that instead
+					if (Step.FinishTimeUtc != null)
+					{
+						TimingInfo.TotalTimeToComplete = Step.FinishTimeUtc.Value - Job.CreateTimeUtc;
+					}
+					else
+					{
+						TimingInfo.TotalTimeToComplete = TimingInfo.TotalTimeToComplete + NullableTimeSpanFromSeconds(StepTimingInfo?.AverageDuration);
+					}
+
+					// If the step hasn't finished yet, make sure the start time is later than the current time
+					if (Step.FinishTimeUtc == null && CurrentTime > TimingInfo.TotalTimeToComplete)
+					{
+						TimingInfo.TotalTimeToComplete = CurrentTime;
+					}
+
+					// Update the average time to complete
+					TimingInfo.AverageTotalTimeToComplete = TimingInfo.AverageTotalTimeToComplete + NullableTimeSpanFromSeconds(StepTimingInfo?.AverageDuration);
+
+					// Add it to the lookup
+					TimingInfo NodeTimingInfo = new TimingInfo(TimingInfo);
+					NodeTimingInfo.StepTiming = StepTimingInfo;
+					NodeToTimingInfo[Node] = NodeTimingInfo;
+				}
+			}
+			return NodeToTimingInfo;
+		}
+
+		/// <summary>
+		/// Gets the average wait time for this batch
+		/// </summary>
+		/// <param name="Graph">Graph for the job</param>
+		/// <param name="Batch">The batch to get timing info for</param>
+		/// <param name="JobTiming">The job timing information</param>
+		/// <returns>Wait time for the batch</returns>
+		public static TimeSpan? GetAverageWaitTime(IGraph Graph, IJobStepBatch Batch, IJobTiming JobTiming)
+		{
+			TimeSpan? WaitTime = null;
+			foreach (IJobStep Step in Batch.Steps)
+			{
+				INode Node = Graph.Groups[Batch.GroupIdx].Nodes[Step.NodeIdx];
+				if(JobTiming.TryGetStepTiming(Node.Name, out IJobStepTiming? TimingInfo))
+				{
+					if (TimingInfo.AverageWaitTime != null)
+					{
+						TimeSpan StepWaitTime = TimeSpan.FromSeconds(TimingInfo.AverageWaitTime.Value);
+						if (WaitTime == null || StepWaitTime > WaitTime.Value)
+						{
+							WaitTime = StepWaitTime;
+						}
+					}
+				}
+			}
+			return WaitTime;
+		}
+
+		/// <summary>
+		/// Gets the average initialization time for this batch
+		/// </summary>
+		/// <param name="Graph">Graph for the job</param>
+		/// <param name="Batch">The batch to get timing info for</param>
+		/// <param name="JobTiming">The job timing information</param>
+		/// <returns>Initialization time for this batch</returns>
+		public static TimeSpan? GetAverageInitTime(IGraph Graph, IJobStepBatch Batch, IJobTiming JobTiming)
+		{
+			TimeSpan? InitTime = null;
+			foreach (IJobStep Step in Batch.Steps)
+			{
+				INode Node = Graph.Groups[Batch.GroupIdx].Nodes[Step.NodeIdx];
+				if (JobTiming.TryGetStepTiming(Node.Name, out IJobStepTiming? TimingInfo))
+				{
+					if (TimingInfo.AverageInitTime != null)
+					{
+						TimeSpan StepInitTime = TimeSpan.FromSeconds(TimingInfo.AverageInitTime.Value);
+						if (InitTime == null || StepInitTime > InitTime.Value)
+						{
+							InitTime = StepInitTime;
+						}
+					}
+				}
+			}
+			return InitTime;
+		}
+
+		/// <summary>
+		/// Creates a nullable timespan from a nullable number of seconds
+		/// </summary>
+		/// <param name="Seconds">The number of seconds to construct from</param>
+		/// <returns>TimeSpan object</returns>
+		static TimeSpan? NullableTimeSpanFromSeconds(float? Seconds)
+		{
+			if (Seconds == null)
+			{
+				return null;
+			}
+			else
+			{
+				return TimeSpan.FromSeconds(Seconds.Value);
+			}
+		}
+
+		/// <summary>
+		/// Attempts to get a batch with the given id
+		/// </summary>
+		/// <param name="Job">The job document</param>
+		/// <param name="BatchId">The batch id</param>
+		/// <param name="Batch">On success, receives the batch object</param>
+		/// <returns>True if the batch was found</returns>
+		public static bool TryGetBatch(this IJob Job, SubResourceId BatchId, [NotNullWhen(true)] out IJobStepBatch? Batch)
+		{
+			Batch = Job.Batches.FirstOrDefault(x => x.Id == BatchId);
+			return Batch != null;
+		}
+
+		/// <summary>
+		/// Attempts to get a batch with the given id
+		/// </summary>
+		/// <param name="Job">The job document</param>
+		/// <param name="BatchId">The batch id</param>
+		/// <param name="StepId">The step id</param>
+		/// <param name="Step">On success, receives the step object</param>
+		/// <returns>True if the batch was found</returns>
+		public static bool TryGetStep(this IJob Job, SubResourceId BatchId, SubResourceId StepId, [NotNullWhen(true)] out IJobStep? Step)
+		{
+			IJobStepBatch? Batch;
+			if (!TryGetBatch(Job, BatchId, out Batch))
+			{
+				Step = null;
+				return false;
+			}
+			return Batch.TryGetStep(StepId, out Step);
+		}
+
+		/// <summary>
+		/// Finds the set of nodes affected by a label
+		/// </summary>
+		/// <param name="Job">The job document</param>
+		/// <param name="Graph">Graph definition for the job</param>
+		/// <param name="LabelIdx">Index of the label. -1 or Graph.Labels.Count are treated as referring to the default lable.</param>
+		/// <returns>Set of nodes affected by the given label</returns>
+		public static HashSet<NodeRef> GetNodesForLabel(this IJob Job, IGraph Graph, int LabelIdx)
+		{
+			if (LabelIdx != -1 && LabelIdx != Graph.Labels.Count)
+			{
+				// Return all the nodes included by the label
+				return new HashSet<NodeRef>(Graph.Labels[LabelIdx].IncludedNodes);
+			}
+			else
+			{
+				// Set of nodes which are not covered by an existing label, initially containing everything
+				HashSet<NodeRef> UnlabeledNodes = new HashSet<NodeRef>();
+				for (int GroupIdx = 0; GroupIdx < Graph.Groups.Count; GroupIdx++)
+				{
+					INodeGroup Group = Graph.Groups[GroupIdx];
+					for (int NodeIdx = 0; NodeIdx < Group.Nodes.Count; NodeIdx++)
+					{
+						UnlabeledNodes.Add(new NodeRef(GroupIdx, NodeIdx));
+					}
+				}
+
+				// Remove all the nodes that are part of an active label
+				IReadOnlyDictionary<NodeRef, IJobStep> StepForNode = Job.GetStepForNodeMap();
+				foreach (ILabel Label in Graph.Labels)
+				{
+					if (Label.RequiredNodes.Any(x => StepForNode.ContainsKey(x)))
+					{
+						UnlabeledNodes.ExceptWith(Label.IncludedNodes);
+					}
+				}
+				return UnlabeledNodes;
+			}
+		}
+
+		/// <summary>
+		/// Create a list of aggregate responses, combining the graph definitions with the state of the job
+		/// </summary>
+		/// <param name="Job">The job document</param>
+		/// <param name="Graph">Graph definition for the job</param>
+		/// <param name="Responses">List to receive all the responses</param>
+		/// <returns>The default label state</returns>
+		public static GetDefaultLabelStateResponse? GetLabelStateResponses(this IJob Job, IGraph Graph, List<GetLabelStateResponse> Responses)
+		{
+			// Create a lookup from noderef to step information
+			IReadOnlyDictionary<NodeRef, IJobStep> StepForNode = Job.GetStepForNodeMap();
+
+			// Set of nodes which are not covered by an existing label, initially containing everything
+			HashSet<NodeRef> UnlabeledNodes = new HashSet<NodeRef>();
+			for (int GroupIdx = 0; GroupIdx < Graph.Groups.Count; GroupIdx++)
+			{
+				INodeGroup Group = Graph.Groups[GroupIdx];
+				for (int NodeIdx = 0; NodeIdx < Group.Nodes.Count; NodeIdx++)
+				{
+					UnlabeledNodes.Add(new NodeRef(GroupIdx, NodeIdx));
+				}
+			}
+
+			// Create the responses
+			foreach (ILabel Label in Graph.Labels)
+			{
+				// Refresh the state for this label
+				LabelState NewState = LabelState.Unspecified;
+				foreach (NodeRef RequiredNodeRef in Label.RequiredNodes)
+				{
+					if (StepForNode.ContainsKey(RequiredNodeRef))
+					{
+						NewState = LabelState.Complete;
+						break;
+					}
+				}
+
+				// Refresh the outcome
+				LabelOutcome NewOutcome = LabelOutcome.Success;
+				if (NewState == LabelState.Complete)
+				{
+					GetLabelState(Label.IncludedNodes, StepForNode, out NewState, out NewOutcome);
+					UnlabeledNodes.ExceptWith(Label.IncludedNodes);
+				}
+
+				// Create the response
+				Responses.Add(new GetLabelStateResponse(NewState, NewOutcome));
+			}
+
+			// Remove all the nodes that don't have a step
+			UnlabeledNodes.RemoveWhere(x => !StepForNode.ContainsKey(x));
+
+			// Remove successful "setup build" nodes from the list
+			if (Graph.Groups.Count > 1 && Graph.Groups[0].Nodes.Count > 0)
+			{
+				INode Node = Graph.Groups[0].Nodes[0];
+				if (Node.Name == IJob.SetupNodeName)
+				{
+					NodeRef NodeRef = new NodeRef(0, 0);
+					if (UnlabeledNodes.Contains(NodeRef))
+					{
+						IJobStep Step = StepForNode[NodeRef];
+						if (Step.State == JobStepState.Completed && Step.Outcome == JobStepOutcome.Success && Responses.Count > 0)
+						{
+							UnlabeledNodes.Remove(NodeRef);
+						}
+					}
+				}
+			}
+
+			// Add a response for everything not included elsewhere.
+			GetLabelState(UnlabeledNodes, StepForNode, out LabelState OtherState, out LabelOutcome OtherOutcome);
+			return new GetDefaultLabelStateResponse(OtherState, OtherOutcome, UnlabeledNodes.Select(x => Graph.GetNode(x).Name).ToList());
+		}
+
+		/// <summary>
+		/// Get the states of all labels for this job
+		/// </summary>
+		/// <param name="Job">The job to get states for</param>
+		/// <param name="Graph">The graph for this job</param>
+		/// <returns>Collection of label states by label index</returns>
+		public static IReadOnlyList<(LabelState, LabelOutcome)> GetLabelStates(this IJob Job, IGraph Graph)
+		{
+			IReadOnlyDictionary<NodeRef, IJobStep> StepForNodeRef = Job.GetStepForNodeMap();
+
+			List<(LabelState, LabelOutcome)> States = new List<(LabelState, LabelOutcome)>();
+			for (int Idx = 0; Idx < Graph.Labels.Count; Idx++)
+			{
+				ILabel Label = Graph.Labels[Idx];
+
+				// Default the label to the unspecified state
+				LabelState NewState = LabelState.Unspecified;
+				LabelOutcome NewOutcome = LabelOutcome.Unspecified;
+
+				// Check if the label should be included
+				if (Label.RequiredNodes.Any(x => StepForNodeRef.ContainsKey(x)))
+				{
+					// Combine the state of the steps contributing towards this label
+					bool bAnySkipped = false;
+					bool bAnyWarnings = false;
+					bool bAnyFailed = false;
+					bool bAnyPending = false;
+					foreach (NodeRef IncludedNode in Label.IncludedNodes)
+					{
+						IJobStep? Step;
+						if (StepForNodeRef.TryGetValue(IncludedNode, out Step))
+						{
+							bAnyPending |= Step.IsPending();
+							bAnySkipped |= Step.State == JobStepState.Aborted || Step.State == JobStepState.Skipped;
+							bAnyFailed |= (Step.Outcome == JobStepOutcome.Failure);
+							bAnyWarnings |= (Step.Outcome == JobStepOutcome.Warnings);
+						}
+					}
+
+					// Figure out the overall label state
+					NewState = bAnyPending ? LabelState.Running : LabelState.Complete;
+					NewOutcome = bAnyFailed ? LabelOutcome.Failure : bAnyWarnings ? LabelOutcome.Warnings : bAnySkipped? LabelOutcome.Unspecified : LabelOutcome.Success;
+				}
+
+				States.Add((NewState, NewOutcome));
+			}
+			return States;	
+		}
+
+		/// <summary>
+		/// Get the states of all UGS badges for this job
+		/// </summary>
+		/// <param name="Job">The job to get states for</param>
+		/// <param name="Graph">The graph for this job</param>
+		/// <returns>List of badge states</returns>
+		public static Dictionary<int, UgsBadgeState> GetUgsBadgeStates(this IJob Job, IGraph Graph)
+		{
+			IReadOnlyList<(LabelState, LabelOutcome)> LabelStates = GetLabelStates(Job, Graph);
+			return Job.GetUgsBadgeStates(Graph, LabelStates);
+		}
+
+		/// <summary>
+		/// Get the states of all UGS badges for this job
+		/// </summary>
+		/// <param name="Job">The job to get states for</param>
+		/// <param name="Graph">The graph for this job</param>
+		/// <param name="LabelStates">The existing label states to get the UGS badge states from</param>
+		/// <returns>List of badge states</returns>
+		public static Dictionary<int, UgsBadgeState> GetUgsBadgeStates(this IJob Job, IGraph Graph, IReadOnlyList<(LabelState, LabelOutcome)> LabelStates)
+		{
+			Dictionary<int, UgsBadgeState> UgsBadgeStates = new Dictionary<int, UgsBadgeState>();
+			for (int LabelIdx = 0; LabelIdx < LabelStates.Count; ++LabelIdx)
+			{
+				if (Graph.Labels[LabelIdx].UgsName == null)
+				{
+					continue;
+				}
+
+				(LabelState State, LabelOutcome Outcome) Label = LabelStates[LabelIdx];
+				switch (Label.State)
+				{
+					case LabelState.Complete:
+					{
+						switch (Label.Outcome)
+						{
+							case LabelOutcome.Success:
+							{
+								UgsBadgeStates.Add(LabelIdx, UgsBadgeState.Success);
+								break;
+							}
+
+							case LabelOutcome.Warnings:
+							{
+								UgsBadgeStates.Add(LabelIdx, UgsBadgeState.Warning);
+								break;
+							}
+
+							case LabelOutcome.Failure:
+							{
+								UgsBadgeStates.Add(LabelIdx, UgsBadgeState.Failure);
+								break;
+							}
+
+							case LabelOutcome.Unspecified:
+							{
+								UgsBadgeStates.Add(LabelIdx, UgsBadgeState.Skipped);
+								break;
+							}
+						}
+						break;
+					}
+
+					case LabelState.Running:
+					{
+						UgsBadgeStates.Add(LabelIdx, UgsBadgeState.Starting);
+						break;
+					}
+
+					case LabelState.Unspecified:
+					{
+						UgsBadgeStates.Add(LabelIdx, UgsBadgeState.Skipped);
+						break;
+					}
+				}
+			}
+			return UgsBadgeStates;
+		}
+
+		/// <summary>
+		/// Gets the state of a job, as a label that includes all steps
+		/// </summary>
+		/// <param name="Job">The job to query</param>
+		/// <param name="StepForNode">Map from node to step</param>
+		/// <param name="NewState">Receives the state of the label</param>
+		/// <param name="NewOutcome">Receives the outcome of the label</param>
+		public static void GetJobState(this IJob Job, IReadOnlyDictionary<NodeRef, IJobStep> StepForNode, out LabelState NewState, out LabelOutcome NewOutcome)
+		{
+			List<NodeRef> Nodes = new List<NodeRef>();
+			foreach (IJobStepBatch Batch in Job.Batches)
+			{
+				foreach (IJobStep Step in Batch.Steps)
+				{
+					Nodes.Add(new NodeRef(Batch.GroupIdx, Step.NodeIdx));
+				}
+			}
+			GetLabelState(Nodes, StepForNode, out NewState, out NewOutcome);
+		}
+		
+		/// <summary>
+		/// Gets the state of a label
+		/// </summary>
+		/// <param name="IncludedNodes">Nodes to include in this label</param>
+		/// <param name="StepForNode">Map from node to step</param>
+		/// <param name="NewState">Receives the state of the label</param>
+		/// <param name="NewOutcome">Receives the outcome of the label</param>
+		public static void GetLabelState(IEnumerable<NodeRef> IncludedNodes, IReadOnlyDictionary<NodeRef, IJobStep> StepForNode, out LabelState NewState, out LabelOutcome NewOutcome)
+		{
+			NewState = LabelState.Complete;
+			NewOutcome = LabelOutcome.Success;
+			foreach (NodeRef IncludedNodeRef in IncludedNodes)
+			{
+				IJobStep? IncludedStep;
+				if (StepForNode.TryGetValue(IncludedNodeRef, out IncludedStep))
+				{
+					// Update the state
+					if (IncludedStep.State != JobStepState.Completed && IncludedStep.State != JobStepState.Skipped && IncludedStep.State != JobStepState.Aborted)
+					{
+						NewState = LabelState.Running;
+					}
+
+					// Update the outcome
+					if (IncludedStep.State == JobStepState.Skipped || IncludedStep.State == JobStepState.Aborted || IncludedStep.Outcome == JobStepOutcome.Failure)
+					{
+						NewOutcome = LabelOutcome.Failure;
+					}
+					else if (IncludedStep.Outcome == JobStepOutcome.Warnings && NewOutcome == LabelOutcome.Success)
+					{
+						NewOutcome = LabelOutcome.Warnings;
+					}
+				}
+			}
+		}
+
+		/// <summary>
+		/// Creates a response containing information about this object
+		/// </summary>
+		/// <param name="Job">The job document</param>
+		/// <param name="Graph">The graph definition</param>
+		/// <param name="bIncludeBatches">Whether to include the job batches in the response</param>
+		/// <param name="bIncludeLabels">Whether to include the job aggregates in the response</param>
+		/// <param name="bIncludeAcl">Whether to include the ACL in the response</param>
+		/// <returns>The response object</returns>
+		public static GetJobResponse ToResponse(this IJob Job, IGraph Graph, bool bIncludeBatches, bool bIncludeLabels, bool bIncludeAcl)
+		{
+			GetAclResponse? AclResponse = (bIncludeAcl && Job.Acl != null) ? new GetAclResponse(Job.Acl) : null;
+			GetJobResponse Response = new GetJobResponse(Job, AclResponse);
+			if (bIncludeBatches || bIncludeLabels)
+			{
+				if (bIncludeBatches)
+				{
+					Response.Batches = Job.Batches.ConvertAll(x => new GetBatchResponse(x));
+				}
+				if (bIncludeLabels)
+				{
+					Response.Labels = new List<GetLabelStateResponse>();
+					Response.DefaultLabel = GetLabelStateResponses(Job, Graph, Response.Labels);
+				}
+			}
+			return Response;
+		}
+
+		/// <summary>
+		/// Creates a response containing information about this object
+		/// </summary>
+		/// <param name="Job">The job document</param>
+		/// <param name="Graph">The graph for this job</param>
+		/// <param name="bIncludeAcl">Whether to include the ACL in the response</param>
+		/// <param name="Filter">Filter for the properties to return</param>
+		/// <returns>Object containing the requested properties</returns>
+		public static object ToResponse(this IJob Job, IGraph Graph, bool bIncludeAcl, PropertyFilter? Filter)
+		{
+			if(Filter == null)
+			{
+				return Job.ToResponse(Graph, true, true, bIncludeAcl);
+			}
+			else
+			{
+				return Filter.ApplyTo(Job.ToResponse(Graph, Filter.Includes(nameof(GetJobResponse.Batches)), Filter.Includes(nameof(GetJobResponse.Labels)) || Filter.Includes(nameof(GetJobResponse.DefaultLabel)), bIncludeAcl));
+			}
+		}
+
+		/// <summary>
+		/// Creates an RPC response object
+		/// </summary>
+		/// <param name="Job">The job document</param>
+		/// <returns></returns>
+		public static HordeCommon.Rpc.GetJobResponse ToRpcResponse(this IJob Job)
+		{
+			HordeCommon.Rpc.GetJobResponse Response = new HordeCommon.Rpc.GetJobResponse();
+			Response.StreamId = Job.StreamId.ToString();
+			Response.Change = Job.Change;
+			Response.CodeChange = Job.CodeChange;
+			Response.PreflightChange = Job.PreflightChange;
+			Response.ClonedPreflightChange = Job.ClonedPreflightChange;
+			Response.Arguments.Add(Job.Arguments);
+			return Response;
+		}
+	}
+
+	/// <summary>
+	/// Projection of a job definition to just include permissions info
+	/// </summary>
+	public interface IJobPermissions
+	{
+		/// <summary>
+		/// ACL for the job
+		/// </summary>
+		public Acl? Acl { get; }
+
+		/// <summary>
+		/// The stream containing 
+		/// </summary>
+		public StreamId StreamId { get; }
+	}
+}
diff --git a/Engine/Source/Programs/Horde/HordeServer/Models/JobStepRef.cs b/Engine/Source/Programs/Horde/HordeServer/Models/JobStepRef.cs
new file mode 100644
--- /dev/null
+++ b/Engine/Source/Programs/Horde/HordeServer/Models/JobStepRef.cs
@@ -0,0 +1,183 @@
+using HordeCommon;
+using HordeServer.Utilities;
+using MongoDB.Bson;
+using MongoDB.Bson.Serialization;
+using MongoDB.Bson.Serialization.Attributes;
+using System;
+using System.Collections.Generic;
+using System.Linq;
+using System.Threading.Tasks;
+
+using StreamId = HordeServer.Utilities.StringId<HordeServer.Models.IStream>;
+using TemplateRefId = HordeServer.Utilities.StringId<HordeServer.Models.TemplateRef>;
+using PoolId = HordeServer.Utilities.StringId<HordeServer.Models.IPool>;
+
+namespace HordeServer.Models
+{
+	/// <summary>
+	/// Unique id struct for JobStepRef objects. Includes a job id, batch id, and step id to uniquely identify the step.
+	/// </summary>
+	[BsonSerializer(typeof(JobStepRefIdSerializer))]
+	public struct JobStepRefId
+	{
+		/// <summary>
+		/// The job id
+		/// </summary>
+		public ObjectId JobId { get; set; }
+
+		/// <summary>
+		/// The batch id within the job
+		/// </summary>
+		public SubResourceId BatchId { get; set; }
+
+		/// <summary>
+		/// The step id
+		/// </summary>
+		public SubResourceId StepId { get; set; }
+
+		/// <summary>
+		/// Constructor
+		/// </summary>
+		/// <param name="JobId">The job id</param>
+		/// <param name="BatchId">The batch id within the job</param>
+		/// <param name="StepId">The step id</param>
+		public JobStepRefId(ObjectId JobId, SubResourceId BatchId, SubResourceId StepId)
+		{
+			this.JobId = JobId;
+			this.BatchId = BatchId;
+			this.StepId = StepId;
+		}
+
+		/// <summary>
+		/// Parse a job step id from a string
+		/// </summary>
+		/// <param name="Text">Text to parse</param>
+		/// <returns>The parsed id</returns>
+		public static JobStepRefId Parse(string Text)
+		{
+			string[] Components = Text.Split(':');
+			return new JobStepRefId(ObjectId.Parse(Components[0]), SubResourceId.Parse(Components[1]), SubResourceId.Parse(Components[2]));
+		}
+
+		/// <summary>
+		/// Formats this id as a string
+		/// </summary>
+		/// <returns>Formatted id</returns>
+		public override string ToString()
+		{
+			return $"{JobId}:{BatchId}:{StepId}";
+		}
+	}
+
+	/// <summary>
+	/// Serializer for JobStepRefId objects
+	/// </summary>
+	public sealed class JobStepRefIdSerializer : IBsonSerializer<JobStepRefId>
+	{
+		/// <inheritdoc/>
+		public Type ValueType
+		{
+			get { return typeof(JobStepRefId); }
+		}
+
+		/// <inheritdoc/>
+		void IBsonSerializer.Serialize(BsonSerializationContext Context, BsonSerializationArgs Args, object Value)
+		{
+			Serialize(Context, Args, (JobStepRefId)Value);
+		}
+
+		/// <inheritdoc/>
+		object IBsonSerializer.Deserialize(BsonDeserializationContext Context, BsonDeserializationArgs Args)
+		{
+			return ((IBsonSerializer<JobStepRefId>)this).Deserialize(Context, Args);
+		}
+
+		/// <inheritdoc/>
+		public JobStepRefId Deserialize(BsonDeserializationContext Context, BsonDeserializationArgs Args)
+		{
+			return JobStepRefId.Parse(Context.Reader.ReadString());
+		}
+
+		/// <inheritdoc/>
+		public void Serialize(BsonSerializationContext Context, BsonSerializationArgs Args, JobStepRefId Value)
+		{
+			Context.Writer.WriteString(((JobStepRefId)Value).ToString());
+		}
+	}
+
+	/// <summary>
+	/// Searchable reference to a jobstep
+	/// </summary>
+	public interface IJobStepRef
+	{
+		/// <summary>
+		/// Globally unique identifier for the jobstep being referenced
+		/// </summary>
+		public JobStepRefId Id { get; }
+
+		/// <summary>
+		/// Name of the job
+		/// </summary>
+		public string JobName { get; }
+
+		/// <summary>
+		/// Name of the name
+		/// </summary>
+		public string NodeName { get; }
+
+		/// <summary>
+		/// Unique id of the stream containing the job
+		/// </summary>
+		public StreamId StreamId { get; }
+
+		/// <summary>
+		/// Template for the job being executed
+		/// </summary>
+		public TemplateRefId TemplateId { get; }
+
+		/// <summary>
+		/// The change number being built
+		/// </summary>
+		public int Change { get; }
+
+		/// <summary>
+		/// Log for this step
+		/// </summary>
+		public ObjectId? LogId { get; }
+
+		/// <summary>
+		/// The agent type
+		/// </summary>
+		public PoolId? PoolId { get; }
+
+		/// <summary>
+		/// The agent id
+		/// </summary>
+		public AgentId? AgentId { get; }
+
+		/// <summary>
+		/// Outcome of the step, once complete.
+		/// </summary>
+		public JobStepOutcome? Outcome { get; }
+
+		/// <summary>
+		/// Time taken for the batch containing this batch to start after it became ready
+		/// </summary>
+		public float BatchWaitTime { get; }
+
+		/// <summary>
+		/// Time taken for this batch to initialize
+		/// </summary>
+		public float BatchInitTime { get; }
+
+		/// <summary>
+		/// Time at which the step started.
+		/// </summary>
+		public DateTime StartTimeUtc { get; }
+
+		/// <summary>
+		/// Time at which the step finished.
+		/// </summary>
+		public DateTime? FinishTimeUtc { get; }
+	}
+}
diff --git a/Engine/Source/Programs/Horde/HordeServer/Models/JobTiming.cs b/Engine/Source/Programs/Horde/HordeServer/Models/JobTiming.cs
new file mode 100644
--- /dev/null
+++ b/Engine/Source/Programs/Horde/HordeServer/Models/JobTiming.cs
@@ -0,0 +1,43 @@
+using System;
+using System.Collections.Generic;
+using System.Diagnostics.CodeAnalysis;
+using System.Linq;
+using System.Threading.Tasks;
+
+namespace HordeServer.Models
+{
+	/// <summary>
+	/// Information about the timing for an individual step
+	/// </summary>
+	public interface IJobStepTiming
+	{
+		/// <summary>
+		/// Wait time before executing the group containing this node
+		/// </summary>
+		public float? AverageWaitTime { get; }
+
+		/// <summary>
+		/// Time taken for the group containing this node to initialize
+		/// </summary>
+		public float? AverageInitTime { get; }
+
+		/// <summary>
+		/// Time spent executing this node
+		/// </summary>
+		public float? AverageDuration { get; }
+	}
+
+	/// <summary>
+	/// Timing information for a particular job
+	/// </summary>
+	public interface IJobTiming
+	{
+		/// <summary>
+		/// Gets timing information for a particular step
+		/// </summary>
+		/// <param name="Name">Name of the node being executed</param>
+		/// <param name="Timing">Receives the timing information for the given step</param>
+		/// <returns>True if the timing was found</returns>
+		public bool TryGetStepTiming(string Name, [NotNullWhen(true)] out IJobStepTiming? Timing);
+	}
+}
diff --git a/Engine/Source/Programs/Horde/HordeServer/Models/Lease.cs b/Engine/Source/Programs/Horde/HordeServer/Models/Lease.cs
new file mode 100644
--- /dev/null
+++ b/Engine/Source/Programs/Horde/HordeServer/Models/Lease.cs
@@ -0,0 +1,79 @@
+// Copyright Epic Games, Inc. All Rights Reserved.
+
+using HordeServer.Api;
+using HordeCommon;
+using HordeServer.Utilities;
+using MongoDB.Bson;
+using MongoDB.Bson.Serialization.Attributes;
+using Serilog;
+using System;
+using System.Collections.Generic;
+using System.Linq;
+using System.Threading.Tasks;
+
+using PoolId = HordeServer.Utilities.StringId<HordeServer.Models.IPool>;
+using StreamId = HordeServer.Utilities.StringId<HordeServer.Models.IStream>;
+
+namespace HordeServer.Models
+{
+	/// <summary>
+	/// Document describing a lease. This exists to permanently record a lease; the agent object tracks internal state of any active leases through AgentLease objects.
+	/// </summary>
+	public interface ILease
+	{
+		/// <summary>
+		/// The unique id of this lease
+		/// </summary>
+		public ObjectId Id { get; }
+
+		/// <summary>
+		/// Name of this lease
+		/// </summary>
+		public string Name { get; }
+
+		/// <summary>
+		/// Unique id of the agent 
+		/// </summary>
+		public AgentId AgentId { get; }
+
+		/// <summary>
+		/// Unique id of the agent session
+		/// </summary>
+		public ObjectId SessionId { get; }
+
+		/// <summary>
+		/// The stream this lease belongs to
+		/// </summary>
+		public StreamId? StreamId { get; }
+
+		/// <summary>
+		/// Pool for the work being executed
+		/// </summary>
+		public PoolId? PoolId { get; }
+
+		/// <summary>
+		/// The log for this lease, if applicable
+		/// </summary>
+		public ObjectId? LogId { get; }
+
+		/// <summary>
+		/// Time at which this lease started
+		/// </summary>
+		public DateTime StartTime { get; }
+
+		/// <summary>
+		/// Time at which this lease completed
+		/// </summary>
+		public DateTime? FinishTime { get; }
+
+		/// <summary>
+		/// Payload for this lease. A packed Google.Protobuf.Any object.
+		/// </summary>
+		public ReadOnlyMemory<byte> Payload { get; }
+
+		/// <summary>
+		/// Outcome of the lease
+		/// </summary>
+		public LeaseOutcome Outcome { get; }
+	}
+}
diff --git a/Engine/Source/Programs/Horde/HordeServer/Models/LogEvent.cs b/Engine/Source/Programs/Horde/HordeServer/Models/LogEvent.cs
new file mode 100644
--- /dev/null
+++ b/Engine/Source/Programs/Horde/HordeServer/Models/LogEvent.cs
@@ -0,0 +1,127 @@
+using HordeCommon;
+using HordeCommon.Rpc;
+using Microsoft.Extensions.Logging;
+using MongoDB.Bson;
+using System;
+using System.Collections.Generic;
+using System.Diagnostics.CodeAnalysis;
+using System.Linq;
+using System.Text.Json;
+using System.Threading.Tasks;
+
+namespace HordeServer.Models
+{
+	/// <summary>
+	/// Interface for a log event line
+	/// </summary>
+	public interface ILogEventLine
+	{
+		/// <summary>
+		/// Log level
+		/// </summary>
+		LogLevel Level { get; }
+
+		/// <summary>
+		/// The event id
+		/// </summary>
+		EventId? EventId { get; }
+
+		/// <summary>
+		/// The complete rendered message, in plaintext
+		/// </summary>
+		string Message { get; }
+
+		/// <summary>
+		/// Data for the line
+		/// </summary>
+		JsonElement Data { get; }
+	}
+
+	/// <summary>
+	/// Interface for event data
+	/// </summary>
+	public interface ILogEventData
+	{
+		/// <summary>
+		/// The log level
+		/// </summary>
+		EventSeverity Severity { get; }
+
+		/// <summary>
+		/// The type of event
+		/// </summary>
+		EventId? EventId { get; }
+
+		/// <summary>
+		/// The complete rendered message, in plaintext
+		/// </summary>
+		string Message { get; }
+
+		/// <summary>
+		/// Gets this event data as a BSON document
+		/// </summary>
+		IReadOnlyList<ILogEventLine> Lines { get; }
+	}
+
+	/// <summary>
+	/// Represents a node in the graph
+	/// </summary>
+	public interface ILogEvent
+	{
+		/// <summary>
+		/// Unique id of the log containing this event
+		/// </summary>
+		public ObjectId LogId { get; }
+
+		/// <summary>
+		/// Severity of the event
+		/// </summary>
+		public EventSeverity Severity { get; }
+
+		/// <summary>
+		/// Index of the first line for this event
+		/// </summary>
+		public int LineIndex { get; }
+
+		/// <summary>
+		/// Number of lines in the event
+		/// </summary>
+		public int LineCount { get; }
+
+		/// <summary>
+		/// Span id for this log event
+		/// </summary>
+		public ObjectId? SpanId { get; }
+	}
+
+	/// <summary>
+	/// Represents a node in the graph
+	/// </summary>
+	public class NewLogEventData
+	{
+		/// <summary>
+		/// Unique id of the log containing this event
+		/// </summary>
+		public ObjectId LogId { get; set; }
+
+		/// <summary>
+		/// Severity of the event
+		/// </summary>
+		public EventSeverity Severity { get; set; }
+
+		/// <summary>
+		/// Index of the first line for this event
+		/// </summary>
+		public int LineIndex { get; set; }
+
+		/// <summary>
+		/// Number of lines in the event
+		/// </summary>
+		public int LineCount { get; set; }
+
+		/// <summary>
+		/// The span this this event belongs to
+		/// </summary>
+		public ObjectId? SpanId { get; set; }
+	}
+}
diff --git a/Engine/Source/Programs/Horde/HordeServer/Models/LogFile.cs b/Engine/Source/Programs/Horde/HordeServer/Models/LogFile.cs
new file mode 100644
--- /dev/null
+++ b/Engine/Source/Programs/Horde/HordeServer/Models/LogFile.cs
@@ -0,0 +1,121 @@
+// Copyright Epic Games, Inc. All Rights Reserved.
+
+using EpicGames.Core;
+using HordeServer.Api;
+using HordeServer.Services;
+using HordeServer.Utilities;
+using MongoDB.Bson;
+using MongoDB.Bson.Serialization.Attributes;
+using System;
+using System.Collections.Generic;
+using System.Diagnostics.CodeAnalysis;
+using System.Linq;
+using System.Threading.Tasks;
+
+namespace HordeServer.Models
+{
+	/// <summary>
+	/// Information about a log file chunk
+	/// </summary>
+	public interface ILogChunk
+	{
+		/// <summary>
+		/// Offset of the chunk within the log
+		/// </summary>
+		long Offset { get; }
+
+		/// <summary>
+		/// Length of this chunk. If zero, the chunk is still being written to.
+		/// </summary>
+		int Length { get; }
+
+		/// <summary>
+		/// Index of the first line within this chunk. If a line straddles two chunks, this is the index of the split line.
+		/// </summary>
+		int LineIndex { get; }
+
+		/// <summary>
+		/// If the chunk has yet to be pushed to persistent storage, includes the name of the server that is currently storing it.
+		/// </summary>
+		string? Server { get; }
+	}
+
+	/// <summary>
+	/// Information about a log file
+	/// </summary>
+	public interface ILogFile
+	{
+		/// <summary>
+		/// Identifier for the LogFile. Randomly generated.
+		/// </summary>
+		public ObjectId Id { get; }
+
+		/// <summary>
+		/// Unique id of the job containing this log
+		/// </summary>
+		public ObjectId JobId { get; }
+
+		/// <summary>
+		/// The session allowed to write to this log
+		/// </summary>
+		public ObjectId? SessionId { get; }
+
+		/// <summary>
+		/// Maximum line index in the file
+		/// </summary>
+		public int? MaxLineIndex { get; }
+
+		/// <summary>
+		/// Length of the file which is indexed
+		/// </summary>
+		public long? IndexLength { get; }
+
+		/// <summary>
+		/// Type of data stored in this log 
+		/// </summary>
+		public LogType Type { get; }
+
+		/// <summary>
+		/// Chunks within this file
+		/// </summary>
+		public IReadOnlyList<ILogChunk> Chunks { get; }
+	}
+
+	/// <summary>
+	/// Extension methods for log files
+	/// </summary>
+	public static class LogFileExtensions
+	{
+		/// <summary>
+		/// Gets the chunk index containing the given offset.
+		/// </summary>
+		/// <param name="Chunks">The chunks to search</param>
+		/// <param name="Offset">The offset to search for</param>
+		/// <returns>The chunk index containing the given offset</returns>
+		public static int GetChunkForOffset(this IReadOnlyList<ILogChunk> Chunks, long Offset)
+		{
+			int ChunkIndex = Chunks.BinarySearch(x => x.Offset, Offset);
+			if (ChunkIndex < 0)
+			{
+				ChunkIndex = ~ChunkIndex - 1;
+			}
+			return ChunkIndex;
+		}
+
+		/// <summary>
+		/// Gets the starting chunk index for the given line
+		/// </summary>
+		/// <param name="Chunks">The chunks to search</param>
+		/// <param name="LineIndex">Index of the line to query</param>
+		/// <returns>Index of the chunk to fetch</returns>
+		public static int GetChunkForLine(this IReadOnlyList<ILogChunk> Chunks, int LineIndex)
+		{
+			int ChunkIndex = Chunks.BinarySearch(x => x.LineIndex, LineIndex);
+			if(ChunkIndex < 0)
+			{
+				ChunkIndex = ~ChunkIndex - 1;
+			}
+			return ChunkIndex;
+		}
+	}
+}
diff --git a/Engine/Source/Programs/Horde/HordeServer/Models/NotificationTrigger.cs b/Engine/Source/Programs/Horde/HordeServer/Models/NotificationTrigger.cs
new file mode 100644
--- /dev/null
+++ b/Engine/Source/Programs/Horde/HordeServer/Models/NotificationTrigger.cs
@@ -0,0 +1,51 @@
+// Copyright Epic Games, Inc. All Rights Reserved.
+
+using HordeServer.Api;
+using MongoDB.Bson;
+using MongoDB.Bson.Serialization.Attributes;
+using System.Collections.Generic;
+
+namespace HordeServer.Models
+{
+	/// <summary>
+	/// An individual subscription
+	/// </summary>
+	public interface INotificationSubscription
+	{
+		/// <summary>
+		/// Name of the user
+		/// </summary>
+		public ObjectId UserId { get; }
+
+		/// <summary>
+		/// Whether to receive email notifications
+		/// </summary>
+		public bool Email { get; }
+
+		/// <summary>
+		/// Whether to receive Slack notifications
+		/// </summary>
+		public bool Slack { get; }
+	}
+
+	/// <summary>
+	/// Trigger for notifications to be sent
+	/// </summary>
+	public interface INotificationTrigger
+	{
+		/// <summary>
+		/// Unique id for this subscription list
+		/// </summary>
+		ObjectId Id { get; }
+
+		/// <summary>
+		/// Whether this trigger has been fired
+		/// </summary>
+		bool Fired { get; }
+
+		/// <summary>
+		/// List of subscriptions to this event
+		/// </summary>
+		public IReadOnlyList<INotificationSubscription> Subscriptions { get; }
+	}
+}
diff --git a/Engine/Source/Programs/Horde/HordeServer/Models/Pool.cs b/Engine/Source/Programs/Horde/HordeServer/Models/Pool.cs
new file mode 100644
--- /dev/null
+++ b/Engine/Source/Programs/Horde/HordeServer/Models/Pool.cs
@@ -0,0 +1,75 @@
+using HordeServer.Api;
+using HordeServer.Utilities;
+using MongoDB.Bson;
+using MongoDB.Bson.Serialization.Attributes;
+using System;
+using System.Collections.Generic;
+using System.Linq;
+using System.Security.Permissions;
+using System.Threading.Tasks;
+
+using PoolId = HordeServer.Utilities.StringId<HordeServer.Models.IPool>;
+
+namespace HordeServer.Models
+{
+	/// <summary>
+	/// A pool of machines
+	/// </summary>
+	public interface IPool
+	{
+		/// <summary>
+		/// Unique id for this pool
+		/// </summary>
+		public PoolId Id { get; }
+
+		/// <summary>
+		/// Name of the pool
+		/// </summary>
+		public string Name { get; }
+
+		/// <summary>
+		/// List of workspaces currently assigned to this pool
+		/// </summary>
+		public IReadOnlyList<AgentWorkspace> Workspaces { get; }
+
+		/// <summary>
+		/// Requirements for agents to be included in this pool
+		/// </summary>
+		public AgentRequirements? Requirements { get; }
+
+		/// <summary>
+		/// Arbitrary properties related to this pool
+		/// </summary>
+		public IReadOnlyDictionary<string, string> Properties { get; }
+
+		/// <summary>
+		/// Whether to enable autoscaling for this pool
+		/// </summary>
+		public bool EnableAutoscaling { get; }
+
+		/// <summary>
+		/// The minimum number of agents to keep in the pool
+		/// </summary>
+		public int? MinAgents { get; }
+
+		/// <summary>
+		/// The minimum number of idle agents to hold in reserve
+		/// </summary>
+		public int? NumReserveAgents { get; }
+
+		/// <summary>
+		/// Last time the pool was (auto) scaled up
+		/// </summary>
+		public DateTime? LastScaleUpTime { get; }
+		
+		/// <summary>
+		/// Last time the pool was (auto) scaled down
+		/// </summary>
+		public DateTime? LastScaleDownTime { get; }
+
+		/// <summary>
+		/// Update index for this document
+		/// </summary>
+		public int UpdateIndex { get; }
+	}
+}
diff --git a/Engine/Source/Programs/Horde/HordeServer/Models/Project.cs b/Engine/Source/Programs/Horde/HordeServer/Models/Project.cs
new file mode 100644
--- /dev/null
+++ b/Engine/Source/Programs/Horde/HordeServer/Models/Project.cs
@@ -0,0 +1,212 @@
+// Copyright Epic Games, Inc. All Rights Reserved.
+
+using EpicGames.Core;
+using HordeServer.Api;
+using HordeServer.Utilities;
+using MongoDB.Bson;
+using MongoDB.Bson.Serialization.Attributes;
+using MongoDB.Driver;
+using System;
+using System.Collections.Generic;
+using System.Linq;
+using System.Text.RegularExpressions;
+using System.Threading.Tasks;
+
+using ProjectId = HordeServer.Utilities.StringId<HordeServer.Models.IProject>;
+
+namespace HordeServer.Models
+{
+	/// <summary>
+	/// Specifies a category of streams to display on the dashboard
+	/// </summary>
+	public class StreamCategory
+	{
+		/// <summary>
+		/// Name of this group
+		/// </summary>
+		[BsonRequired]
+		public string Name { get; set; }
+
+		/// <summary>
+		/// Row to display this category on
+		/// </summary>
+		public int Row { get; set; }
+
+		/// <summary>
+		/// Whether to show this category on the nav menu
+		/// </summary>
+		public bool ShowOnNavMenu { get; set; } = true;
+
+		/// <summary>
+		/// List of stream name patterns to include for this project
+		/// </summary>
+		public List<string> IncludePatterns { get; set; } = new List<string>();
+
+		/// <summary>
+		/// Patterns for stream names to be excluded
+		/// </summary>
+		public List<string> ExcludePatterns { get; set; } = new List<string>();
+
+		/// <summary>
+		/// Private constructor
+		/// </summary>
+		private StreamCategory()
+		{
+			Name = null!;
+		}
+
+		/// <summary>
+		/// Constructor
+		/// </summary>
+		/// <param name="Name">Name of this group</param>
+		public StreamCategory(string Name)
+		{
+			this.Name = Name;
+		}
+
+		/// <summary>
+		/// Constructs a category from a request object
+		/// </summary>
+		/// <param name="Request">The request object</param>
+		public StreamCategory(CreateProjectCategoryRequest Request)
+		{
+			this.Name = Request.Name;
+			this.Row = Request.Row;
+			this.ShowOnNavMenu = Request.ShowOnNavMenu;
+			this.IncludePatterns = Request.IncludePatterns;
+			this.ExcludePatterns = Request.ExcludePatterns;
+		}
+
+		/// <summary>
+		/// Tests if a given name matches a pattern
+		/// </summary>
+		/// <param name="Name">The name to test</param>
+		/// <param name="Pattern">The pattern to match against</param>
+		/// <returns>True if the pattern matches, false otherwise</returns>
+		public static bool MatchPattern(string Name, string Pattern)
+		{
+			return Regex.IsMatch(Name, Pattern);
+		}
+	}
+
+	/// <summary>
+	/// Represents a project
+	/// </summary>
+	public interface IProject
+	{
+		/// <summary>
+		/// Identifier for the project.
+		/// </summary>
+		public ProjectId Id { get; }
+
+		/// <summary>
+		/// Name of the project
+		/// </summary>
+		public string Name { get; }
+
+		/// <summary>
+		/// Order to display on the dashboard
+		/// </summary>
+		public int Order { get; }
+
+		/// <summary>
+		/// The ACL for this object
+		/// </summary>
+		public Acl? Acl { get; }
+
+		/// <summary>
+		/// List of stream categories for this project. Controls how streams are displayed on the dashboard.
+		/// </summary>
+		public IReadOnlyList<StreamCategory> Categories { get; }
+
+		/// <summary>
+		/// Optional user-defined properties for this project
+		/// </summary>
+		public IReadOnlyDictionary<string, string> Properties { get; }
+	}
+
+	/// <summary>
+	/// Extension methods for projects
+	/// </summary>
+	public static class ProjectExtensions
+	{
+		/// <summary>
+		/// Converts this object to a public response
+		/// </summary>
+		/// <param name="Project">The project instance</param>
+		/// <param name="bIncludeStreams">Whether to include streams in the response</param>
+		/// <param name="bIncludeCategories">Whether to include categories in the response</param>
+		/// <param name="Streams">The list of streams</param>
+		/// <param name="bIncludeAcl">Whether to include the ACL in the response</param>
+		/// <returns>Response instance</returns>
+		public static GetProjectResponse ToResponse(this IProject Project, bool bIncludeStreams, bool bIncludeCategories, List<IStream>? Streams, bool bIncludeAcl)
+		{
+			List<GetProjectStreamResponse>? StreamResponses = null;
+			if(bIncludeStreams)
+			{
+				StreamResponses = Streams!.ConvertAll(x => new GetProjectStreamResponse(x.Id.ToString(), x.Name));
+			}
+
+			List<GetProjectCategoryResponse>? CategoryResponses = null;
+			if (bIncludeCategories)
+			{
+				CategoryResponses = Project.Categories.ConvertAll(x => new GetProjectCategoryResponse(x));
+				if (Streams != null)
+				{
+					foreach (IStream Stream in Streams)
+					{
+						GetProjectCategoryResponse? CategoryResponse = CategoryResponses.FirstOrDefault(x => MatchCategory(Stream.Name, x));
+						if(CategoryResponse == null)
+						{
+							int Row = (CategoryResponses.Count > 0) ? CategoryResponses.Max(x => x.Row) : 0;
+							if (CategoryResponses.Count(x => x.Row == Row) >= 3)
+							{
+								Row++;
+							}
+
+							StreamCategory OtherCategory = new StreamCategory("Other");
+							OtherCategory.Row = Row;
+							OtherCategory.IncludePatterns.Add(".*");
+
+							CategoryResponse = new GetProjectCategoryResponse(OtherCategory);
+							CategoryResponses.Add(CategoryResponse);
+						}
+						CategoryResponse.Streams!.Add(Stream.Id.ToString());
+					}
+				}
+			}
+
+			GetAclResponse? AclResponse = (bIncludeAcl && Project.Acl != null) ? new GetAclResponse(Project.Acl) : null;
+			return new GetProjectResponse(Project.Id.ToString(), Project.Name, Project.Order, StreamResponses, CategoryResponses, new Dictionary<string, string>(Project.Properties), AclResponse);
+		}
+
+		/// <summary>
+		/// Tests if a category response matches a given stream name
+		/// </summary>
+		/// <param name="Name">The stream name</param>
+		/// <param name="Category">The category response</param>
+		/// <returns>True if the category matches</returns>
+		static bool MatchCategory(string Name, GetProjectCategoryResponse Category)
+		{
+			if (Category.IncludePatterns.Any(x => StreamCategory.MatchPattern(Name, x)))
+			{
+				if (!Category.ExcludePatterns.Any(x => StreamCategory.MatchPattern(Name, x)))
+				{
+					return true;
+				}
+			}
+			return false;
+		}
+	}
+
+	/// <summary>
+	/// Projection of a project definition to just include permissions info
+	/// </summary>
+	public interface IProjectPermissions
+	{
+		/// <summary>
+		/// ACL for the project
+		/// </summary>
+		public Acl? Acl { get; }
+	}
+}
diff --git a/Engine/Source/Programs/Horde/HordeServer/Models/Schedule.cs b/Engine/Source/Programs/Horde/HordeServer/Models/Schedule.cs
new file mode 100644
--- /dev/null
+++ b/Engine/Source/Programs/Horde/HordeServer/Models/Schedule.cs
@@ -0,0 +1,306 @@
+// Copyright Epic Games, Inc. All Rights Reserved.
+
+using EpicGames.Core;
+using HordeServer.Api;
+using HordeServer.Utilities;
+using MongoDB.Bson;
+using MongoDB.Bson.Serialization.Attributes;
+using System;
+using System.Collections.Generic;
+using System.Linq;
+using System.Threading.Tasks;
+
+using TemplateRefId = HordeServer.Utilities.StringId<HordeServer.Models.TemplateRef>;
+
+namespace HordeServer.Models
+{
+	/// <summary>
+	/// Specifies a pattern of times that this schedule should run. Each schedule may have multiple patterns.
+	/// </summary>
+	public class SchedulePattern
+	{
+		/// <summary>
+		/// Which days of the week the schedule should run
+		/// </summary>
+		public List<DayOfWeek>? DaysOfWeek { get; set; }
+
+		/// <summary>
+		/// Time during the day for the first schedule to trigger. Measured in minutes from midnight.
+		/// </summary>
+		public int MinTime { get; set; }
+
+		/// <summary>
+		/// Time during the day for the last schedule to trigger. Measured in minutes from midnight.
+		/// </summary>
+		public int? MaxTime { get; set; }
+
+		/// <summary>
+		/// Interval between each schedule triggering
+		/// </summary>
+		public int? Interval { get; set; }
+
+		/// <summary>
+		/// Private constructor for serialization
+		/// </summary>
+		private SchedulePattern()
+		{
+		}
+
+		/// <summary>
+		/// Constructor
+		/// </summary>
+		/// <param name="DaysOfWeek">Which days of the week the schedule should run</param>
+		/// <param name="MinTime">Time during the day for the first schedule to trigger. Measured in minutes from midnight.</param>
+		/// <param name="MaxTime">Time during the day for the last schedule to trigger. Measured in minutes from midnight.</param>
+		/// <param name="Interval">Interval between each schedule triggering</param>
+		public SchedulePattern(List<DayOfWeek>? DaysOfWeek, int MinTime, int? MaxTime, int? Interval)
+		{
+			this.DaysOfWeek = DaysOfWeek;
+			this.MinTime = MinTime;
+			this.MaxTime = MaxTime;
+			this.Interval = Interval;
+		}
+
+		/// <summary>
+		/// Calculates the trigger index based on the given time in minutes
+		/// </summary>
+		/// <param name="LastTime">Time in minutes for the last trigger</param>
+		/// <param name="TimeZone">The timezone for running the schedule</param>
+		/// <returns>Index of the trigger</returns>
+		public DateTimeOffset GetNextTriggerTime(DateTimeOffset LastTime, TimeZoneInfo TimeZone)
+		{
+			// Convert last time into the correct timezone for running the scheule
+			LastTime = TimeZoneInfo.ConvertTime(LastTime, TimeZone);
+
+			// Get the base time (ie. the start of this day) for anchoring the schedule
+			DateTimeOffset BaseTime = new DateTimeOffset(LastTime.Year, LastTime.Month, LastTime.Day, 0, 0, 0, LastTime.Offset);
+			for (; ; )
+			{
+				if (DaysOfWeek == null || DaysOfWeek.Contains(BaseTime.DayOfWeek))
+				{
+					// Get the last time in minutes from the start of this day
+					int LastTimeMinutes = (int)(LastTime - BaseTime).TotalMinutes;
+
+					// Get the time of the first trigger of this day. If the last time is less than this, this is the next trigger.
+					if (LastTimeMinutes < MinTime)
+					{
+						return BaseTime.AddMinutes(MinTime);
+					}
+
+					// Otherwise, get the time for the last trigger in the day.
+					if (Interval.HasValue && Interval.Value > 0)
+					{
+						int ActualMaxTime = MaxTime ?? ((24 * 60) - 1);
+						if (LastTimeMinutes < ActualMaxTime)
+						{
+							int LastIndex = (LastTimeMinutes - MinTime) / Interval.Value;
+							int NextIndex = LastIndex + 1;
+
+							int NextTimeMinutes = MinTime + (NextIndex * Interval.Value);
+							if (NextTimeMinutes <= ActualMaxTime)
+							{
+								return BaseTime.AddMinutes(NextTimeMinutes);
+							}
+						}
+					}
+				}
+				BaseTime = BaseTime.AddDays(1.0);
+			}
+		}
+	}
+
+	/// <summary>
+	/// Gate allowing a schedule to trigger.
+	/// </summary>
+	public class ScheduleGate
+	{
+		/// <summary>
+		/// The template containing the dependency
+		/// </summary>
+		public TemplateRefId TemplateRefId { get; set; }
+
+		/// <summary>
+		/// Target to wait for
+		/// </summary>
+		public string Target { get; set; }
+
+		/// <summary>
+		/// Private constructor for serializtaion
+		/// </summary>
+		[BsonConstructor]
+		private ScheduleGate()
+		{
+			Target = String.Empty;
+		}
+
+		/// <summary>
+		/// Constructor
+		/// </summary>
+		/// <param name="TemplateRefId">The template containing the dependency</param>
+		/// <param name="Target">Target to wait for</param>
+		public ScheduleGate(TemplateRefId TemplateRefId, string Target)
+		{
+			this.TemplateRefId = TemplateRefId;
+			this.Target = Target;
+		}
+	}
+
+	/// <summary>
+	/// Represents a schedule
+	/// </summary>
+	public class Schedule
+	{
+		/// <summary>
+		/// Whether this schedule is currently enabled
+		/// </summary>
+		public bool Enabled { get; set; }
+
+		/// <summary>
+		/// Maximum number of builds triggered by this schedule active at once. Set to zero for unlimited.
+		/// </summary>
+		public int MaxActive { get; set; }
+
+		/// <summary>
+		/// Maximum number of changes the schedule can fall behind head revision. If greater than zero, builds will be triggered for every submitted changelist until the backlog is this size.
+		/// </summary>
+		public int MaxChanges { get; set; }
+
+		/// <summary>
+		/// Whether the build requires a change to be submitted
+		/// </summary>
+		[BsonIgnoreIfDefault, BsonDefaultValue(true)]
+		public bool RequireSubmittedChange { get; set; } = true;
+
+		/// <summary>
+		/// Reference to another job/target which much succeed for this schedule to trigger
+		/// </summary>
+		public ScheduleGate? Gate { get; set; }
+
+		/// <summary>
+		/// Whether this build requires a code change to trigger
+		/// </summary>
+		public List<ChangeContentFlags>? Filter { get; set; }
+
+		/// <summary>
+		/// Parameters for the template
+		/// </summary>
+		public Dictionary<string, string> TemplateParameters { get; set; }
+
+		/// <summary>
+		/// List of patterns for this schedule
+		/// </summary>
+		public List<SchedulePattern> Patterns { get; set; }
+
+		/// <summary>
+		/// Last changelist number that this was triggered for
+		/// </summary>
+		public int LastTriggerChange { get; set; }
+
+		/// <summary>
+		/// Last time that the schedule was triggered
+		/// </summary>
+		public DateTimeOffset LastTriggerTime { get; set; }
+
+		/// <summary>
+		/// List of jobs that are currently active
+		/// </summary>
+		public List<ObjectId> ActiveJobs { get; set; } = new List<ObjectId>();
+
+		/// <summary>
+		/// Private constructor for serializtaion
+		/// </summary>
+		[BsonConstructor]
+		private Schedule()
+		{
+			TemplateParameters = new Dictionary<string, string>();
+			Patterns = new List<SchedulePattern>();
+		}
+
+		/// <summary>
+		/// Constructor
+		/// </summary>
+		/// <param name="Enabled">Whether the schedule is currently enabled</param>
+		/// <param name="MaxActive">Maximum number of builds that may be active at once</param>
+		/// <param name="MaxChanges">Maximum number of changes the schedule can fall behind head revision</param>
+		/// <param name="RequireSubmittedChange">Whether a change has to be submitted for the schedule to trigger</param>
+		/// <param name="Gate">Reference to another job/target which much succeed for this schedule to trigger</param>
+		/// <param name="Filter">Filter for changes to consider</param>
+		/// <param name="TemplateParameters">Parameters for the template to run</param>
+		/// <param name="Patterns">List of patterns for the schedule</param>
+		public Schedule(bool Enabled = true, int MaxActive = 0, int MaxChanges = 0, bool RequireSubmittedChange = true, ScheduleGate? Gate = null, List<ChangeContentFlags>? Filter = null, Dictionary<string, string>? TemplateParameters = null, List<SchedulePattern>? Patterns = null)
+		{
+			this.Enabled = Enabled;
+			this.MaxActive = MaxActive;
+			this.MaxChanges = MaxChanges;
+			this.RequireSubmittedChange = RequireSubmittedChange;
+			this.Gate = Gate;
+			this.Filter = Filter;
+			this.TemplateParameters = TemplateParameters ?? new Dictionary<string, string>();
+			this.Patterns = Patterns ?? new List<SchedulePattern>();
+			this.LastTriggerTime = DateTimeOffset.Now;
+		}
+
+		/// <summary>
+		/// Copies the state fields from another schedule object
+		/// </summary>
+		/// <param name="Other">The schedule object to copy from</param>
+		public void CopyState(Schedule Other)
+		{
+			LastTriggerChange = Other.LastTriggerChange;
+			LastTriggerTime = Other.LastTriggerTime;
+			ActiveJobs.AddRange(Other.ActiveJobs);
+		}
+
+		/// <summary>
+		/// Gets the flags to filter changes by
+		/// </summary>
+		/// <returns>Set of filter flags</returns>
+		public ChangeContentFlags? GetFilterFlags()
+		{
+			if(Filter == null)
+			{
+				return null;
+			}
+
+			ChangeContentFlags Flags = 0;
+			foreach (ChangeContentFlags Flag in Filter)
+			{
+				Flags |= Flag;
+			}
+			return Flags;
+		}
+
+		/// <summary>
+		/// Get the next time that the schedule will trigger
+		/// </summary>
+		/// <param name="TimeZone">Timezone to evaluate the trigger</param>
+		/// <returns>Next time at which the schedule will trigger</returns>
+		public DateTimeOffset? GetNextTriggerTime(TimeZoneInfo TimeZone)
+		{
+			return GetNextTriggerTime(LastTriggerTime, TimeZone);
+		}
+
+		/// <summary>
+		/// Get the next time that the schedule will trigger
+		/// </summary>
+		/// <param name="LastTime">Last time at which the schedule triggered</param>
+		/// <param name="TimeZone">Timezone to evaluate the trigger</param>
+		/// <returns>Next time at which the schedule will trigger</returns>
+		public DateTimeOffset? GetNextTriggerTime(DateTimeOffset LastTime, TimeZoneInfo TimeZone)
+		{
+			DateTimeOffset? NextTriggerTime = null;
+			if (Enabled)
+			{
+				foreach (SchedulePattern Pattern in Patterns)
+				{
+					DateTimeOffset PatternTriggerTime = Pattern.GetNextTriggerTime(LastTime, TimeZone);
+					if (NextTriggerTime == null || NextTriggerTime < PatternTriggerTime)
+					{
+						NextTriggerTime = PatternTriggerTime;
+					}
+				}
+			}
+			return NextTriggerTime;
+		}
+	}
+}
diff --git a/Engine/Source/Programs/Horde/HordeServer/Models/Session.cs b/Engine/Source/Programs/Horde/HordeServer/Models/Session.cs
new file mode 100644
--- /dev/null
+++ b/Engine/Source/Programs/Horde/HordeServer/Models/Session.cs
@@ -0,0 +1,49 @@
+using HordeServer.Api;
+using HordeServer.Utilities;
+using MongoDB.Bson;
+using MongoDB.Bson.Serialization.Attributes;
+using System;
+using System.Collections.Generic;
+using System.Linq;
+using System.Threading.Tasks;
+
+using AgentSoftwareVersion = HordeServer.Utilities.StringId<HordeServer.Collections.IAgentSoftwareCollection>;
+
+namespace HordeServer.Models
+{
+	/// <summary>
+	/// Information about an agent session.
+	/// </summary>
+	public interface ISession
+	{
+		/// <summary>
+		/// Unique id for this session
+		/// </summary>
+		public ObjectId Id { get; }
+
+		/// <summary>
+		/// The agent id
+		/// </summary>
+		public AgentId AgentId { get; }
+
+		/// <summary>
+		/// Start time for this session
+		/// </summary>
+		public DateTime StartTime { get; }
+
+		/// <summary>
+		/// Finishing time for this session
+		/// </summary>
+		public DateTime? FinishTime { get; }
+
+		/// <summary>
+		/// Capabilities of this agent at the time the session started
+		/// </summary>
+		public AgentCapabilities? Capabilities { get; }
+
+		/// <summary>
+		/// Version of the agent software
+		/// </summary>
+		public AgentSoftwareVersion? Version { get; }
+	}
+}
diff --git a/Engine/Source/Programs/Horde/HordeServer/Models/Software.cs b/Engine/Source/Programs/Horde/HordeServer/Models/Software.cs
new file mode 100644
--- /dev/null
+++ b/Engine/Source/Programs/Horde/HordeServer/Models/Software.cs
@@ -0,0 +1,52 @@
+// Copyright Epic Games, Inc. All Rights Reserved.
+
+using HordeServer.Api;
+using HordeServer.Utilities;
+using MongoDB.Bson;
+using MongoDB.Bson.Serialization.Attributes;
+using System;
+using System.Collections.Generic;
+using System.IO;
+using System.IO.Compression;
+using System.Linq;
+using System.Text.Json;
+using System.Threading.Tasks;
+
+namespace HordeServer.Models
+{
+	/// <summary>
+	/// Information about agent software
+	/// </summary>
+	public interface ISoftware
+	{
+		/// <summary>
+		/// Identifier for the app. Randomly generated.
+		/// </summary>
+		public ObjectId Id { get; }
+
+		/// <summary>
+		/// Zipped data for this app
+		/// </summary>
+		public byte[]? Data { get; }
+
+		/// <summary>
+		/// User that created this client
+		/// </summary>
+		public string? UploadedByUser { get; }
+
+		/// <summary>
+		/// Time at which the client was created
+		/// </summary>
+		public DateTime UploadedAtTime { get; }
+
+		/// <summary>
+		/// User that made this client the default
+		/// </summary>
+		public string? MadeDefaultByUser { get; }
+
+		/// <summary>
+		/// Timestamp at which this app was made the default
+		/// </summary>
+		public DateTime? MadeDefaultAtTime { get; }
+	}
+}
