Initial drop of Horde Server/Agent for UE5/Main. Still missing some Epic-specific functionality required to build/run correctly.

diff --git a/Engine/Source/Programs/Horde/HordeServer/Models/Stream.cs b/Engine/Source/Programs/Horde/HordeServer/Models/Stream.cs
new file mode 100644
--- /dev/null
+++ b/Engine/Source/Programs/Horde/HordeServer/Models/Stream.cs
@@ -0,0 +1,706 @@
+// Copyright Epic Games, Inc. All Rights Reserved.
+
+using EpicGames.Core;
+using Google.Protobuf.WellKnownTypes;
+using HordeServer.Api;
+using HordeCommon;
+using HordeServer.Utilities;
+using MongoDB.Bson;
+using MongoDB.Bson.Serialization.Attributes;
+using MongoDB.Driver;
+using System;
+using System.Collections.Generic;
+using System.Diagnostics.CodeAnalysis;
+using System.IO;
+using System.Linq;
+using System.Security.Cryptography;
+using System.Text;
+using System.Text.RegularExpressions;
+using System.Threading.Tasks;
+
+using PoolId = HordeServer.Utilities.StringId<HordeServer.Models.IPool>;
+using ProjectId = HordeServer.Utilities.StringId<HordeServer.Models.IProject>;
+using StreamId = HordeServer.Utilities.StringId<HordeServer.Models.IStream>;
+using TemplateRefId = HordeServer.Utilities.StringId<HordeServer.Models.TemplateRef>;
+using System.Runtime.CompilerServices;
+
+namespace HordeServer.Models
+{
+	/// <summary>
+	/// Exception thrown when stream validation fails
+	/// </summary>
+	public class InvalidStreamException : Exception
+	{
+		/// <inheritdoc/>
+		public InvalidStreamException()
+		{
+		}
+
+		/// <inheritdoc/>
+		public InvalidStreamException(string Message) : base(Message)
+		{
+		}
+
+		/// <inheritdoc/>
+		public InvalidStreamException(string Message, Exception InnerEx) : base(Message, InnerEx)
+		{
+		}
+	}
+
+	/// <summary>
+	/// Mapping from a BuildGraph agent type to a set of machines on the farm
+	/// </summary>
+	public class AgentType
+	{
+		/// <summary>
+		/// Name of the pool of agents to use
+		/// </summary>
+		[BsonRequired]
+		public PoolId Pool { get; set; }
+
+		/// <summary>
+		/// Name of the workspace to execute on
+		/// </summary>
+		public string? Workspace { get; set; }
+
+		/// <summary>
+		/// Path to the temporary storage dir
+		/// </summary>
+		public string? TempStorageDir { get; set; }
+
+		/// <summary>
+		/// Environment variables to be set when executing the job
+		/// </summary>
+		public Dictionary<string, string>? Environment { get; set; }
+
+		/// <summary>
+		/// Default constructor
+		/// </summary>
+		[BsonConstructor]
+		private AgentType()
+		{
+		}
+
+		/// <summary>
+		/// Constructor
+		/// </summary>
+		/// <param name="Pool">The pool for this agent</param>
+		/// <param name="Workspace">Name of the workspace to use</param>
+		/// <param name="TempStorageDir">Path to the temp storage directory</param>
+		public AgentType(PoolId Pool, string Workspace, string? TempStorageDir)
+		{
+			this.Pool = Pool;
+			this.Workspace = Workspace;
+			this.TempStorageDir = TempStorageDir;
+		}
+
+		/// <summary>
+		/// Constructor
+		/// </summary>
+		/// <param name="Request">The object to construct from</param>
+		public AgentType(CreateAgentTypeRequest Request)
+		{
+			Pool = new PoolId(Request.Pool);
+			Workspace = Request.Workspace;
+			TempStorageDir = Request.TempStorageDir;
+			Environment = Request.Environment;
+		}
+
+		/// <summary>
+		/// Constructs an AgentType object from an optional request
+		/// </summary>
+		/// <param name="Request">The request object</param>
+		/// <returns>New agent type object</returns>
+		[return: NotNullIfNotNull("Request")]
+		public static AgentType? FromRequest(CreateAgentTypeRequest? Request)
+		{
+			return (Request != null) ? new AgentType(Request) : null;
+		}
+
+		/// <summary>
+		/// Creates an API response object from this stream
+		/// </summary>
+		/// <returns>The response object</returns>
+		public Api.GetAgentTypeResponse ToApiResponse()
+		{
+			return new Api.GetAgentTypeResponse(Pool.ToString(), Workspace, TempStorageDir, Environment);
+		}
+
+		/// <summary>
+		/// Creates an API response object from this stream
+		/// </summary>
+		/// <returns>The response object</returns>
+		public HordeCommon.Rpc.GetAgentTypeResponse ToRpcResponse()
+		{
+			HordeCommon.Rpc.GetAgentTypeResponse Response = new HordeCommon.Rpc.GetAgentTypeResponse();
+			if (TempStorageDir != null)
+			{
+				Response.TempStorageDir = TempStorageDir;
+			}
+			if (Environment != null)
+			{
+				Response.Environment.Add(Environment);
+			}
+			return Response;
+		}
+	}
+
+	/// <summary>
+	/// Information about a workspace type
+	/// </summary>
+	public class WorkspaceType
+	{
+		/// <summary>
+		/// Name of the Perforce cluster to use
+		/// </summary>
+		[BsonIgnoreIfNull]
+		public string? Cluster { get; set; }
+
+		/// <summary>
+		/// The Perforce server and port
+		/// </summary>
+		[BsonIgnoreIfNull]
+		public string? ServerAndPort { get; set; }
+
+		/// <summary>
+		/// The Perforce username for syncing this workspace
+		/// </summary>
+		[BsonIgnoreIfNull]
+		public string? UserName { get; set; }
+
+		/// <summary>
+		/// The Perforce password for syncing this workspace
+		/// </summary>
+		[BsonIgnoreIfNull]
+		public string? Password { get; set; }
+
+		/// <summary>
+		/// Identifier to distinguish this workspace from other workspaces. Defaults to the workspace type name.
+		/// </summary>
+		[BsonIgnoreIfNull]
+		public string? Identifier { get; set; }
+
+		/// <summary>
+		/// Override for the stream to sync
+		/// </summary>
+		[BsonIgnoreIfNull]
+		public string? Stream { get; set; }
+
+		/// <summary>
+		/// Custom view for the workspace
+		/// </summary>
+		[BsonIgnoreIfNull]
+		public List<string>? View { get; set; }
+
+		/// <summary>
+		/// Whether to use an incrementally synced workspace
+		/// </summary>
+		public bool Incremental { get; set; }
+
+		/// <summary>
+		/// Default constructor
+		/// </summary>
+		public WorkspaceType()
+		{
+		}
+
+		/// <summary>
+		/// Constructor
+		/// </summary>
+		/// <param name="Request">The object to construct from</param>
+		public WorkspaceType(CreateWorkspaceTypeRequest Request)
+		{
+			Cluster = Request.Cluster;
+			ServerAndPort = Request.ServerAndPort;
+			UserName = Request.UserName;
+			Password = Request.Password;
+			Identifier = Request.Identifier;
+			Stream = Request.Stream;
+			View = Request.View;
+			Incremental = Request.Incremental;
+		}
+
+		/// <summary>
+		/// Constructs an AgentType object from an optional request
+		/// </summary>
+		/// <param name="Request">The request object</param>
+		/// <returns>New agent type object</returns>
+		[return: NotNullIfNotNull("Request")]
+		public static WorkspaceType? FromRequest(CreateWorkspaceTypeRequest? Request)
+		{
+			return (Request != null) ? new WorkspaceType(Request) : null;
+		}
+
+		/// <summary>
+		/// Creates an API response object from this stream
+		/// </summary>
+		/// <returns>The response object</returns>
+		public Api.GetWorkspaceTypeResponse ToApiResponse()
+		{
+			return new Api.GetWorkspaceTypeResponse(Cluster, ServerAndPort, UserName, Identifier, Stream, View, Incremental);
+		}
+	}
+
+	/// <summary>
+	/// Allows triggering another downstream job on succesful completion of a step or aggregate
+	/// </summary>
+	public class ChainedJobTemplate
+	{
+		/// <summary>
+		/// Name of the target that needs to complete successfully
+		/// </summary>
+		public string Trigger { get; set; } = String.Empty;
+
+		/// <summary>
+		/// The new template to trigger
+		/// </summary>
+		public TemplateRefId TemplateRefId { get; set; }
+
+		/// <summary>
+		/// Default constructor for serialization
+		/// </summary>
+		private ChainedJobTemplate()
+		{
+		}
+
+		/// <summary>
+		/// Constructor
+		/// </summary>
+		/// <param name="Trigger">Name of the target that needs to complete</param>
+		/// <param name="TemplateRefId">The new template to trigger</param>
+		public ChainedJobTemplate(string Trigger, TemplateRefId TemplateRefId)
+		{
+			this.Trigger = Trigger;
+			this.TemplateRefId = TemplateRefId;
+		}
+
+		/// <summary>
+		/// Constructor
+		/// </summary>
+		/// <param name="Request">Request to construct from</param>
+		public ChainedJobTemplate(CreateChainedJobTemplateRequest Request)
+			: this(Request.Trigger, new TemplateRefId(Request.TemplateId))
+		{
+		}
+	}
+
+	/// <summary>
+	/// Reference to a template
+	/// </summary>
+	public class TemplateRef
+	{
+		/// <summary>
+		/// The template name (duplicated from the template object)
+		/// </summary>
+		public string Name { get; set; }
+
+		/// <summary>
+		/// Hash of the template definition
+		/// </summary>
+		public ContentHash Hash { get; set; } = ContentHash.Empty;
+
+		/// <summary>
+		/// Whether to show badges in UGS for this schedule
+		/// </summary>
+		public bool ShowUgsBadges { get; set; }
+
+		/// <summary>
+		/// Whether to show desktop alerts for build health issues created from jobs this type
+		/// </summary>
+		public bool ShowUgsAlerts { get; set; }
+
+		/// <summary>
+		/// Notification channel for this template. Overrides the stream channel if set.
+		/// </summary>
+		public string? NotificationChannel { get; set; }
+
+		/// <summary>
+		/// Notification channel filter for this template. Errors|Warnings|Success
+		/// </summary>
+		public string? NotificationChannelFilter { get; set; }
+
+		/// <summary>
+		/// Channel for triage notification messages
+		/// </summary>
+		public string? TriageChannel { get; set; }
+
+		/// <summary>
+		/// List of schedules for this template
+		/// </summary>
+		[BsonIgnoreIfNull]
+		public Schedule? Schedule { get; set; }
+
+		/// <summary>
+		/// List of downstream templates to trigger at the same change
+		/// </summary>
+		[BsonIgnoreIfNull]
+		public List<ChainedJobTemplate>? ChainedJobs { get; set; }
+
+		/// <summary>
+		/// Custom permissions for this template
+		/// </summary>
+		public Acl? Acl { get; set; }
+
+		/// <summary>
+		/// Private constructor for serialization
+		/// </summary>
+		private TemplateRef()
+		{
+			this.Name = null!;
+		}
+
+		/// <summary>
+		/// Constructor
+		/// </summary>
+		/// <param name="Template">The template being referenced</param>
+		/// <param name="ShowUgsBadges">Whether to show badges in UGS for this job</param>
+		/// <param name="ShowUgsAlerts">Whether to show alerts in UGS for this job</param>
+		/// <param name="NotificationChannel">Notification channel for this template</param>
+		/// <param name="NotificationChannelFilter">Notification channel filter for this template</param>
+		/// <param name="TriageChannel"></param>
+		/// <param name="Schedule">Schedule for this template</param>
+		/// <param name="Triggers">List of downstream templates to trigger</param>
+		/// <param name="Acl">ACL for this template</param>
+		public TemplateRef(ITemplate Template, bool ShowUgsBadges = false, bool ShowUgsAlerts = false, string? NotificationChannel = null, string? NotificationChannelFilter = null, string? TriageChannel = null, Schedule? Schedule = null, List<ChainedJobTemplate>? Triggers = null, Acl? Acl = null)
+		{
+			this.Name = Template.Name;
+			this.Hash = Template.Id;
+			this.ShowUgsBadges = ShowUgsBadges;
+			this.ShowUgsAlerts = ShowUgsAlerts;
+			this.NotificationChannel = NotificationChannel;
+			this.NotificationChannelFilter = NotificationChannelFilter;
+			this.TriageChannel = TriageChannel;
+			this.Schedule = Schedule;
+			this.ChainedJobs = Triggers;
+			this.Acl = Acl;
+		}
+	}
+
+	/// <summary>
+	/// Definition of a query to execute to find the changelist to run a build at
+	/// </summary>
+	public class DefaultPreflight
+	{
+		/// <summary>
+		/// The template id to execute
+		/// </summary>
+		public TemplateRefId? TemplateRefId { get; set; }
+
+		/// <summary>
+		/// The job type to query for the change to use
+		/// </summary>
+		public TemplateRefId? ChangeTemplateRefId { get; set; }
+
+		/// <summary>
+		/// Constructor
+		/// </summary>
+		/// <param name="TemplateRefId"></param>
+		/// <param name="ChangeTemplateRefId">The job type to query for the change to use</param>
+		public DefaultPreflight(TemplateRefId? TemplateRefId, TemplateRefId? ChangeTemplateRefId)
+		{
+			this.TemplateRefId = TemplateRefId;
+			this.ChangeTemplateRefId = ChangeTemplateRefId;
+		}
+
+		/// <summary>
+		/// Convert to a request object
+		/// </summary>
+		/// <returns></returns>
+		public DefaultPreflightRequest ToRequest()
+		{
+			return new DefaultPreflightRequest { TemplateId = TemplateRefId?.ToString(), ChangeTemplateId = ChangeTemplateRefId?.ToString() };
+		}
+	}
+
+	/// <summary>
+	/// Extension methods for template refs
+	/// </summary>
+	static class TemplateRefExtensions
+	{
+		/// <summary>
+		/// Adds a new template ref to a list
+		/// </summary>
+		/// <param name="TemplateRefs">List of template refs</param>
+		/// <param name="TemplateRef">The template ref to add</param>
+		public static void AddRef(this Dictionary<TemplateRefId, TemplateRef> TemplateRefs, TemplateRef TemplateRef)
+		{
+			TemplateRefs.Add(new TemplateRefId(TemplateRef.Name), TemplateRef);
+		}
+	}
+
+	/// <summary>
+	/// Information about a stream
+	/// </summary>
+	public interface IStream
+	{
+		/// <summary>
+		/// Name of the stream.
+		/// </summary>
+		public StreamId Id { get; }
+
+		/// <summary>
+		/// The project that this stream belongs to
+		/// </summary>
+		public ProjectId ProjectId { get; }
+
+		/// <summary>
+		/// The stream name
+		/// </summary>
+		public string Name { get; }
+
+		/// <summary>
+		/// Depot path to poll for this stream's configuration from
+		/// </summary>
+		public string? ConfigPath { get; }
+
+		/// <summary>
+		/// Revision of the config file currently being used
+		/// </summary>
+		public int? ConfigChange { get; }
+
+		/// <summary>
+		/// Order to display on the dashboard's drop-down list
+		/// </summary>
+		public int Order { get; }
+
+		/// <summary>
+		/// Notification channel for all jobs in this stream
+		/// </summary>
+		public string? NotificationChannel { get; }
+
+		/// <summary>
+		/// Notification channel filter for all jobs in this stream. Errors|Warnings|Success
+		/// </summary>
+		public string? NotificationChannelFilter { get; }
+
+		/// <summary>
+		/// Channel to post issue triage notifications
+		/// </summary>
+		public string? TriageChannel { get; }
+
+		/// <summary>
+		/// Default template to use for preflights
+		/// </summary>
+		public DefaultPreflight? DefaultPreflight { get; }
+
+		/// <summary>
+		/// List of pages to display in the dashboard
+		/// </summary>
+		public IReadOnlyList<StreamTab> Tabs { get; }
+
+		/// <summary>
+		/// Dictionary of agent types
+		/// </summary>
+		public IReadOnlyDictionary<string, AgentType> AgentTypes { get; }
+
+		/// <summary>
+		/// Dictionary of workspace types
+		/// </summary>
+		public IReadOnlyDictionary<string, WorkspaceType> WorkspaceTypes { get; }
+
+		/// <summary>
+		/// List of templates available for this stream
+		/// </summary>
+		public IReadOnlyDictionary<TemplateRefId, TemplateRef> Templates { get; }
+
+		/// <summary>
+		/// Optional user-defined properties for this stream
+		/// </summary>
+		public IReadOnlyDictionary<string, string> Properties { get; }
+
+		/// <summary>
+		/// Last time that we queried for commits
+		/// </summary>
+		public DateTime? LastCommitTime { get; }
+		
+		/// <summary>
+		/// Stream is paused for builds until specified time
+		/// </summary>
+		public DateTime? PausedUntil { get; }
+		
+		/// <summary>
+		/// Comment/reason for why the stream was paused
+		/// </summary>
+		public string? PauseComment { get; }
+
+		/// <summary>
+		/// The ACL for this object
+		/// </summary>
+		public Acl? Acl { get; }
+	}
+
+	/// <summary>
+	/// Extension methods for streams
+	/// </summary>
+	static class StreamExtensions
+	{
+		/// <summary>
+		/// Tries to get an agent workspace definition from the given type name
+		/// </summary>
+		/// <param name="Stream">The stream object</param>
+		/// <param name="AgentType">The agent type</param>
+		/// <param name="Workspace">Receives the agent workspace definition</param>
+		/// <returns>True if the agent type was valid, and an agent workspace could be created</returns>
+		public static bool TryGetAgentWorkspace(this IStream Stream, AgentType AgentType, [NotNullWhen(true)] out AgentWorkspace? Workspace)
+		{
+			// Get the workspace settings
+			if (AgentType.Workspace == null)
+			{
+				// Use the default settings (fast switching workspace, clean 
+				Workspace = new AgentWorkspace(null, null, GetDefaultWorkspaceIdentifier(Stream), Stream.Name, null, false);
+				return true;
+			}
+			else
+			{
+				// Try to get the matching workspace type
+				WorkspaceType? WorkspaceType;
+				if (!Stream.WorkspaceTypes.TryGetValue(AgentType.Workspace, out WorkspaceType))
+				{
+					Workspace = null;
+					return false;
+				}
+
+				// Get the workspace identifier
+				string Identifier;
+				if (WorkspaceType.Identifier != null && !String.IsNullOrEmpty(WorkspaceType.Identifier))
+				{
+					Identifier = WorkspaceType.Identifier;
+				}
+				else if (WorkspaceType.Incremental)
+				{
+					Identifier = $"{Stream.GetEscapedName()}+{AgentType.Workspace}";
+				}
+				else
+				{
+					Identifier = GetDefaultWorkspaceIdentifier(Stream);
+				}
+
+				// Create the new workspace
+				Workspace = new AgentWorkspace(WorkspaceType.Cluster, WorkspaceType.UserName, Identifier, WorkspaceType.Stream ?? Stream.Name, WorkspaceType.View, WorkspaceType.Incremental);
+				return true;
+			}
+		}
+
+		/// <summary>
+		/// The escaped name of this stream. Removes all non-identifier characters.
+		/// </summary>
+		/// <param name="Stream">The stream object</param>
+		/// <returns>Escaped name for the stream</returns>
+		public static string GetEscapedName(this IStream Stream)
+		{
+			return Regex.Replace(Stream.Name, @"[^a-zA-Z0-9_]", "+");
+		}
+
+		/// <summary>
+		/// Gets the default identifier for workspaces created for this stream. Just includes an escaped depot name.
+		/// </summary>
+		/// <param name="Stream">The stream object</param>
+		/// <returns>The default workspace identifier</returns>
+		private static string GetDefaultWorkspaceIdentifier(IStream Stream)
+		{
+			return Regex.Replace(Stream.GetEscapedName(), @"^(\+\+[^+]*).*$", "$1");
+		}
+
+		/// <summary>
+		/// Checks the stream definition for consistency
+		/// </summary>
+		/// <param name="Stream">The stream object</param>
+		public static void Validate(this IStream Stream)
+		{
+			// Check the default preflight template is valid
+			if (Stream.DefaultPreflight != null)
+			{
+				if (Stream.DefaultPreflight.TemplateRefId != null && !Stream.Templates.ContainsKey(Stream.DefaultPreflight.TemplateRefId.Value))
+				{
+					throw new InvalidStreamException($"Default preflight template was listed as '{Stream.DefaultPreflight.TemplateRefId.Value}', but no template was found by that name");
+				}
+				if (Stream.DefaultPreflight.ChangeTemplateRefId != null && !Stream.Templates.ContainsKey(Stream.DefaultPreflight.ChangeTemplateRefId.Value))
+				{
+					throw new InvalidStreamException($"Default preflight template for CL was listed as '{Stream.DefaultPreflight.ChangeTemplateRefId.Value}', but no template was found by that name");
+				}
+			}
+
+			// Check that all the templates are referenced by a tab
+			HashSet<TemplateRefId> RemainingTemplates = new HashSet<TemplateRefId>(Stream.Templates.Keys);
+			foreach(JobsTab JobsTab in Stream.Tabs.OfType<JobsTab>())
+			{
+				if (JobsTab.Templates != null)
+				{
+					RemainingTemplates.ExceptWith(JobsTab.Templates);
+				}
+			}
+			if(RemainingTemplates.Count > 0)
+			{
+				throw new InvalidStreamException(String.Join("\n", RemainingTemplates.Select(x => $"Template '{x}' is not listed on any tab for {Stream.Id}")));
+			}
+
+			// Check that all the agent types reference valid workspace names
+			foreach (KeyValuePair<string, AgentType> Pair in Stream.AgentTypes)
+			{
+				string? WorkspaceTypeName = Pair.Value.Workspace;
+				if (WorkspaceTypeName != null && !Stream.WorkspaceTypes.ContainsKey(WorkspaceTypeName))
+				{
+					throw new InvalidStreamException($"Agent type '{Pair.Key}' references undefined workspace type '{Pair.Value.Workspace}' in {Stream.Id}");
+				}
+			}
+		}
+
+		/// <summary>
+		/// Converts to a public response object
+		/// </summary>
+		/// <param name="Stream">The stream object</param>
+		/// <param name="bIncludeAcl">Whether to include the ACL in the response object</param>
+		/// <param name="ApiTemplateRefs">The template refs for this stream. Passed separately because they have their own ACL.</param>
+		/// <returns>New response instance</returns>
+		public static Api.GetStreamResponse ToApiResponse(this IStream Stream, bool bIncludeAcl, List<GetTemplateRefResponse> ApiTemplateRefs)
+		{
+			List<GetStreamTabResponse> ApiTabs = Stream.Tabs.ConvertAll(x => x.ToResponse());
+			Dictionary<string, GetAgentTypeResponse> ApiAgentTypes = Stream.AgentTypes.ToDictionary(x => x.Key, x => x.Value.ToApiResponse());
+			Dictionary<string, GetWorkspaceTypeResponse> ApiWorkspaceTypes = Stream.WorkspaceTypes.ToDictionary(x => x.Key, x => x.Value.ToApiResponse());
+			GetAclResponse? ApiAcl = (bIncludeAcl && Stream.Acl != null)? new GetAclResponse(Stream.Acl) : null;
+			return new Api.GetStreamResponse(Stream.Id.ToString(), Stream.ProjectId.ToString(), Stream.Name, Stream.ConfigPath, Stream.ConfigChange, Stream.Order, Stream.NotificationChannel, Stream.NotificationChannelFilter, Stream.TriageChannel, Stream.DefaultPreflight?.ToRequest(), ApiTabs, ApiAgentTypes, ApiWorkspaceTypes, ApiTemplateRefs, new Dictionary<string, string>(Stream.Properties), ApiAcl, Stream.PausedUntil, Stream.PauseComment);
+		}
+
+		/// <summary>
+		/// Converts to an RPC response object
+		/// </summary>
+		/// <param name="Stream">The stream object</param>
+		/// <returns>New response instance</returns>
+		public static HordeCommon.Rpc.GetStreamResponse ToRpcResponse(this IStream Stream)
+		{
+			HordeCommon.Rpc.GetStreamResponse Response = new HordeCommon.Rpc.GetStreamResponse();
+			Response.Name = Stream.Name;
+			Response.AgentTypes.Add(Stream.AgentTypes.ToDictionary(x => x.Key, x => x.Value.ToRpcResponse()));
+			Response.Properties.Add(Stream.Properties);
+			Response.LastCommitTime = Stream.LastCommitTime.HasValue? Timestamp.FromDateTime(Stream.LastCommitTime.Value) : new Timestamp();
+			return Response;
+		}
+
+		/// <summary>
+		/// Check if stream is paused for new builds
+		/// </summary>
+		/// <param name="Stream">The stream object</param>
+		/// <param name="CurrentTime">Current time (allow tests to pass in a fake clock)</param>
+		/// <returns>If stream is paused</returns>
+		public static bool IsPaused(this IStream Stream, DateTime CurrentTime)
+		{
+			return Stream.PausedUntil != null && Stream.PausedUntil > CurrentTime;
+		}
+	}
+
+	/// <summary>
+	/// Projection of a stream definition to just include permissions info
+	/// </summary>
+	public interface IStreamPermissions
+	{
+		/// <summary>
+		/// ACL for the stream
+		/// </summary>
+		public Acl? Acl { get; }
+
+		/// <summary>
+		/// The project containing this stream
+		/// </summary>
+		public ProjectId ProjectId { get; }
+	}
+}
diff --git a/Engine/Source/Programs/Horde/HordeServer/Models/StreamTab.cs b/Engine/Source/Programs/Horde/HordeServer/Models/StreamTab.cs
new file mode 100644
--- /dev/null
+++ b/Engine/Source/Programs/Horde/HordeServer/Models/StreamTab.cs
@@ -0,0 +1,266 @@
+using HordeServer.Api;
+using MongoDB.Bson.Serialization.Attributes;
+using System;
+using System.Collections.Generic;
+using System.Linq;
+using System.Threading.Tasks;
+
+using TemplateRefId = HordeServer.Utilities.StringId<HordeServer.Models.TemplateRef>;
+
+namespace HordeServer.Models
+{
+	/// <summary>
+	/// Information about a page to display in the dashboard for a stream
+	/// </summary>
+	[BsonKnownTypes(typeof(JobsTab))]
+	public abstract class StreamTab
+	{
+		/// <summary>
+		/// Title of this page
+		/// </summary>
+		[BsonRequired]
+		public string Title { get; set; }
+
+		/// <summary>
+		/// Private constructor for serialization
+		/// </summary>
+		private StreamTab()
+		{
+			Title = null!;
+		}
+
+		/// <summary>
+		/// Constructor
+		/// </summary>
+		/// <param name="Title">Title of this page</param>
+		public StreamTab(string Title)
+		{
+			this.Title = Title;
+		}
+
+		/// <summary>
+		/// Constructor
+		/// </summary>
+		/// <param name="Request">Public request object</param>
+		public StreamTab(CreateStreamTabRequest Request)
+		{
+			Title = Request.Title;
+		}
+
+		/// <summary>
+		/// Creates an instance from a request
+		/// </summary>
+		/// <param name="Request">The request object</param>
+		/// <returns>New tab instance</returns>
+		public static StreamTab FromRequest(CreateStreamTabRequest Request)
+		{
+			CreateJobsTabRequest? JobsRequest = Request as CreateJobsTabRequest;
+			if (JobsRequest != null)
+			{
+				return new JobsTab(JobsRequest);
+			}
+
+			throw new Exception($"Unknown tab request type '{Request.GetType()}'");
+		}
+
+		/// <summary>
+		/// Creates a response object
+		/// </summary>
+		/// <returns>Response object</returns>
+		public abstract GetStreamTabResponse ToResponse();
+	}
+
+	/// <summary>
+	/// Describes a column to display on the jobs page
+	/// </summary>
+	[BsonKnownTypes(typeof(JobsTabLabelColumn), typeof(JobsTabParameterColumn))]
+	public abstract class JobsTabColumn
+	{
+		/// <summary>
+		/// Heading for this column
+		/// </summary>
+		public string Heading { get; set; }
+
+		/// <summary>
+		/// Relative width of this column.
+		/// </summary>
+		public int? RelativeWidth { get; set; }
+
+		/// <summary>
+		/// Private constructor for serialization
+		/// </summary>
+		protected JobsTabColumn()
+		{
+			Heading = null!;
+			RelativeWidth = 1;
+		}
+
+		/// <summary>
+		/// Constructor
+		/// </summary>
+		/// <param name="Heading">Heading for this column</param>
+		/// <param name="RelativeWidth">Relative width of this column.</param>
+		public JobsTabColumn(string Heading, int? RelativeWidth)
+		{
+			this.Heading = Heading;
+			this.RelativeWidth = RelativeWidth;
+		}
+
+		/// <summary>
+		/// Constructor
+		/// </summary>
+		/// <param name="Request">Public request object</param>
+		public JobsTabColumn(CreateJobsTabColumnRequest Request)
+		{
+			this.Heading = Request.Heading;
+			this.RelativeWidth = Request.RelativeWidth;
+		}
+
+		/// <summary>
+		/// Converts this object to a JSON response object
+		/// </summary>
+		/// <returns>Response object</returns>
+		public abstract GetJobsTabColumnResponse ToResponse();
+	}
+
+	/// <summary>
+	/// Column that contains a set of labels
+	/// </summary>
+	public class JobsTabLabelColumn : JobsTabColumn
+	{
+		/// <summary>
+		/// Category of labels to display in this column. If null, includes any label not matched by another column.
+		/// </summary>
+		public string? Category { get; set; }
+
+		/// <summary>
+		/// Private constructor for serialization
+		/// </summary>
+		protected JobsTabLabelColumn()
+		{
+		}
+
+		/// <summary>
+		/// Constructor
+		/// </summary>
+		/// <param name="Heading">Heading for this column</param>
+		/// <param name="Category">Category of labels to display in this column</param>
+		/// <param name="RelativeWidth">Relative width of this column.</param>
+		public JobsTabLabelColumn(string Heading, string? Category, int? RelativeWidth)
+			: base(Heading, RelativeWidth)
+		{
+			this.Category = Category;
+		}
+
+		/// <summary>
+		/// Converts this object to a JSON response object
+		/// </summary>
+		/// <returns>Response object</returns>
+		public override GetJobsTabColumnResponse ToResponse()
+		{
+			return new GetJobsTabLabelColumnResponse(Heading, Category, RelativeWidth);
+		}
+	}
+
+	/// <summary>
+	/// Include the value for a parameter on the jobs tab
+	/// </summary>
+	public class JobsTabParameterColumn : JobsTabColumn
+	{
+		/// <summary>
+		/// Name of a parameter to show in this column. Should be in the form of a prefix, eg. "-set:Foo="
+		/// </summary>
+		public string? Parameter { get; set; }
+
+		/// <summary>
+		/// Private constructor for serialization
+		/// </summary>
+		protected JobsTabParameterColumn()
+		{
+		}
+
+		/// <summary>
+		/// Constructor
+		/// </summary>
+		/// <param name="Heading">Heading for this column</param>
+		/// <param name="Parameter">Category of labels to display in this column</param>
+		/// <param name="RelativeWidth">Relative width of this column.</param>
+		public JobsTabParameterColumn(string Heading, string? Parameter, int? RelativeWidth)
+			: base(Heading, RelativeWidth)
+		{
+			this.Parameter = Parameter;
+		}
+
+		/// <summary>
+		/// Converts this object to a JSON response object
+		/// </summary>
+		/// <returns>Response object</returns>
+		public override GetJobsTabColumnResponse ToResponse()
+		{
+			return new GetJobsTabParameterColumnResponse(Heading, Parameter, RelativeWidth);
+		}
+	}
+
+	/// <summary>
+	/// Describes a job page
+	/// </summary>
+	[BsonDiscriminator("Jobs")]
+	public class JobsTab : StreamTab
+	{
+		/// <summary>
+		/// Whether to show job names
+		/// </summary>
+		public bool ShowNames { get; set; }
+
+		/// <summary>
+		/// Names of jobs to include on this page. If there is only one name specified, the name column does not need to be displayed.
+		/// </summary>
+		public List<string>? JobNames { get; set; }
+
+		/// <summary>
+		/// List of templates to display on this tab.
+		/// </summary>
+		public List<TemplateRefId>? Templates { get; set; }
+
+		/// <summary>
+		/// Columns to display in the table
+		/// </summary>
+		public List<JobsTabColumn>? Columns { get; set; }
+
+		/// <summary>
+		/// Constructor
+		/// </summary>
+		/// <param name="Title">Title for this page</param>
+		/// <param name="ShowNames">Show names of jobs on the page</param>
+		/// <param name="Templates">Template names to include on this page</param>
+		/// <param name="JobNames">Names of jobs to include on this page</param>
+		/// <param name="Columns">Columns ot display in the table</param>
+		public JobsTab(string Title, bool ShowNames, List<TemplateRefId>? Templates, List<string>? JobNames, List<JobsTabColumn>? Columns)
+			: base(Title)
+		{
+			this.ShowNames = ShowNames;
+			this.Templates = Templates;
+			this.JobNames = JobNames;
+			this.Columns = Columns;
+		}
+
+		/// <summary>
+		/// Constructor
+		/// </summary>
+		/// <param name="Request">Public request object</param>
+		public JobsTab(CreateJobsTabRequest Request)
+			: base(Request)
+		{
+			ShowNames = Request.ShowNames;
+			Templates = Request.Templates?.ConvertAll(x => new TemplateRefId(x));
+			JobNames = Request.JobNames;
+			Columns = Request.Columns?.ConvertAll(x => x.ToModel());
+		}
+
+		/// <inheritdoc/>
+		public override GetStreamTabResponse ToResponse()
+		{
+			return new GetJobsTabResponse(Title, ShowNames, Templates?.ConvertAll(x => x.ToString()), JobNames, Columns?.ConvertAll(x => x.ToResponse()));
+		}
+	}
+}
diff --git a/Engine/Source/Programs/Horde/HordeServer/Models/SubResourceId.cs b/Engine/Source/Programs/Horde/HordeServer/Models/SubResourceId.cs
new file mode 100644
--- /dev/null
+++ b/Engine/Source/Programs/Horde/HordeServer/Models/SubResourceId.cs
@@ -0,0 +1,164 @@
+using MongoDB.Bson.Serialization;
+using MongoDB.Bson.Serialization.Attributes;
+using System;
+using System.Collections.Generic;
+using System.Diagnostics;
+using System.Globalization;
+using System.Linq;
+using System.Threading.Tasks;
+
+namespace HordeServer.Models
+{
+	/// <summary>
+	/// Identifier for subresources. Assigning unique ids to subresources prevents against race conditions using indices when subresources are added and removed.
+	/// 
+	/// Subresource identifiers are stored as 16-bit integers formatted as a 4-digit hex code, in order to keep URLs short. Calling Next() will generate a new
+	/// identifier with more entropy than just incrementing the value but an identical period before repeating, in order to make URL fragments more distinctive.
+	/// </summary>
+	[BsonSerializer(typeof(SubResourceIdSerializer))]
+	public struct SubResourceId : IEquatable<SubResourceId>
+	{
+		/// <summary>
+		/// The unique identifier value
+		/// </summary>
+		public ushort Value { get; }
+
+		/// <summary>
+		/// Constructor
+		/// </summary>
+		/// <param name="Value">New identifier for this subresource</param>
+		public SubResourceId(ushort Value)
+		{
+			this.Value = Value;
+		}
+
+		/// <summary>
+		/// Creates a new random subresource id. We use random numbers for this to increase distinctiveness.
+		/// </summary>
+		/// <returns>New subresource id</returns>
+		public static SubResourceId Random()
+		{
+			return new SubResourceId((ushort)Stopwatch.GetTimestamp());
+		}
+
+		/// <summary>
+		/// Updates the current value, and returns a copy of the previous value.
+		/// </summary>
+		/// <returns>New subresource identifier</returns>
+		public SubResourceId Next()
+		{
+			// 771 is a factor of 65535, so we won't repeat when wrapping round 64k.
+			return new SubResourceId((ushort)((int)Value + 771));
+		}
+
+		/// <summary>
+		/// Parse a subresource id from a string
+		/// </summary>
+		/// <param name="Text">Text to parse</param>
+		/// <returns>New subresource id</returns>
+		public static SubResourceId Parse(string Text)
+		{
+			return new SubResourceId(ushort.Parse(Text, NumberStyles.HexNumber, CultureInfo.InvariantCulture));
+		}
+
+		/// <summary>
+		/// Converts this identifier to a string
+		/// </summary>
+		/// <returns>String representation of this id</returns>
+		public override string ToString()
+		{
+			return Value.ToString("x4", CultureInfo.InvariantCulture);
+		}
+
+		/// <inheritdoc/>
+		public override bool Equals(object? Other)
+		{
+			SubResourceId? OtherId = Other as SubResourceId?;
+			return OtherId != null && Equals(OtherId.Value);
+		}
+
+		/// <inheritdoc/>
+		public bool Equals(SubResourceId Other)
+		{
+			return Value == Other.Value;
+		}
+
+		/// <inheritdoc/>
+		public override int GetHashCode()
+		{
+			return (int)Value;
+		}
+
+		/// <summary>
+		/// Equality operator for identifiers
+		/// </summary>
+		/// <param name="Left">First identifier to compare</param>
+		/// <param name="Right">Second identifier to compare</param>
+		/// <returns>True if the identifiers are equal</returns>
+		public static bool operator ==(SubResourceId Left, SubResourceId Right)
+		{
+			return Left.Value == Right.Value;
+		}
+
+		/// <summary>
+		/// Inequality operator for identifiers
+		/// </summary>
+		/// <param name="Left">First identifier to compare</param>
+		/// <param name="Right">Second identifier to compare</param>
+		/// <returns>True if the identifiers are equal</returns>
+		public static bool operator !=(SubResourceId Left, SubResourceId Right)
+		{
+			return Left.Value != Right.Value;
+		}
+	}
+
+	/// <summary>
+	/// Extension methods for manipulating subresource ids
+	/// </summary>
+	public static class SubResourceIdExtensions
+	{
+		/// <summary>
+		/// Parse a string as a subresource identifier
+		/// </summary>
+		/// <param name="Text">Text to parse</param>
+		/// <returns>The new subresource identifier</returns>
+		public static SubResourceId ToSubResourceId(this string Text)
+		{
+			return new SubResourceId(ushort.Parse(Text, NumberStyles.HexNumber, CultureInfo.InvariantCulture));
+		}
+	}
+
+	/// <summary>
+	/// Serializer for subresource ids
+	/// </summary>
+	public class SubResourceIdSerializer : IBsonSerializer<SubResourceId>
+	{
+		/// <inheritdoc/>
+		public Type ValueType => typeof(SubResourceId);
+
+		/// <inheritdoc/>
+		public object Deserialize(BsonDeserializationContext Context, BsonDeserializationArgs Args)
+		{
+			return new SubResourceId((ushort)Context.Reader.ReadInt32());
+		}
+
+		/// <inheritdoc/>
+		public void Serialize(BsonSerializationContext Context, BsonSerializationArgs Args, object Value)
+		{
+			SubResourceId Id = (SubResourceId)Value;
+			Context.Writer.WriteInt32((int)Id.Value);
+		}
+
+		/// <inheritdoc/>
+		public void Serialize(BsonSerializationContext Context, BsonSerializationArgs Args, SubResourceId Id)
+		{
+			Context.Writer.WriteInt32((int)Id.Value);
+		}
+
+		/// <inheritdoc/>
+		SubResourceId IBsonSerializer<SubResourceId>.Deserialize(BsonDeserializationContext Context, BsonDeserializationArgs Vrgs)
+		{
+			return new SubResourceId((ushort)Context.Reader.ReadInt32());
+		}
+	}
+}
diff --git a/Engine/Source/Programs/Horde/HordeServer/Models/Subscription.cs b/Engine/Source/Programs/Horde/HordeServer/Models/Subscription.cs
new file mode 100644
--- /dev/null
+++ b/Engine/Source/Programs/Horde/HordeServer/Models/Subscription.cs
@@ -0,0 +1,50 @@
+using HordeCommon;
+using HordeServer.Api;
+using MongoDB.Bson;
+using System;
+using System.Collections.Generic;
+using System.Linq;
+using System.Threading.Tasks;
+
+using StreamId = HordeServer.Utilities.StringId<HordeServer.Models.IStream>;
+using TemplateRefId = HordeServer.Utilities.StringId<HordeServer.Models.TemplateRef>;
+
+namespace HordeServer.Models
+{
+	/// <summary>
+	/// The type of notification to send
+	/// </summary>
+	public enum NotificationType
+	{
+		/// <summary>
+		/// Send a DM on Slack
+		/// </summary>
+		Slack,
+	}
+
+	/// <summary>
+	/// Subscription to a type of event
+	/// </summary>
+	public interface ISubscription
+	{
+		/// <summary>
+		/// Unique id for this subscription
+		/// </summary>
+		public string Id { get; }
+
+		/// <summary>
+		/// Name of the event to subscribe to
+		/// </summary>
+		public IEvent Event { get; }
+
+		/// <summary>
+		/// User to notify
+		/// </summary>
+		public ObjectId UserId { get; }
+
+		/// <summary>
+		/// Type of notification to receive
+		/// </summary>
+		public NotificationType NotificationType { get; }
+	}
+}
diff --git a/Engine/Source/Programs/Horde/HordeServer/Models/Template.cs b/Engine/Source/Programs/Horde/HordeServer/Models/Template.cs
new file mode 100644
--- /dev/null
+++ b/Engine/Source/Programs/Horde/HordeServer/Models/Template.cs
@@ -0,0 +1,505 @@
+// Copyright Epic Games, Inc. All Rights Reserved.
+
+using EpicGames.Core;
+using HordeServer.Api;
+using HordeCommon;
+using HordeServer.Utilities;
+using MongoDB.Bson;
+using MongoDB.Bson.Serialization.Attributes;
+using System;
+using System.Collections.Generic;
+using System.Linq;
+using System.Runtime.CompilerServices;
+using System.Text.Json.Serialization;
+using System.Threading.Tasks;
+
+namespace HordeServer.Models
+{
+	/// <summary>
+	/// Base class for parameters used to configure templates via the new build dialog
+	/// </summary>
+	[BsonDiscriminator(RootClass = true)]
+	[BsonKnownTypes(typeof(GroupParameter), typeof(TextParameter), typeof(ListParameter), typeof(BoolParameter))]
+	public abstract class Parameter
+	{
+		/// <summary>
+		/// Gets the default arguments for this parameter and its children
+		/// </summary>
+		/// <param name="DefaultArguments">List of default arguments</param>
+		public abstract void GetDefaultArguments(List<string> DefaultArguments);
+
+		/// <summary>
+		/// Convert this parameter to data for serialization
+		/// </summary>
+		/// <returns>Serializable parameter data</returns>
+		public abstract ParameterData ToData();
+	}
+
+	/// <summary>
+	/// Used to group a number of other parameters
+	/// </summary>
+	public class GroupParameter : Parameter
+	{
+		/// <summary>
+		/// Label to display next to this parameter
+		/// </summary>
+		public string Label { get; set; }
+
+		/// <summary>
+		/// How to display this group
+		/// </summary>
+		public GroupParameterStyle Style { get; set; }
+
+		/// <summary>
+		/// List of child parameters
+		/// </summary>
+		public List<Parameter> Children { get; set; }
+
+		/// <summary>
+		/// Private constructor for serialization
+		/// </summary>
+		private GroupParameter()
+		{
+			this.Label = null!;
+			this.Children = null!;
+		}
+
+		/// <summary>
+		/// Constructor
+		/// </summary>
+		/// <param name="Label">Name of the group</param>
+		/// <param name="Style">How to display this group</param>
+		/// <param name="Children">List of child parameters</param>
+		public GroupParameter(string Label, GroupParameterStyle Style, List<Parameter> Children)
+		{
+			this.Label = Label;
+			this.Style = Style;
+			this.Children = Children;
+		}
+
+		/// <summary>
+		/// Gets the default arguments for this parameter and its children
+		/// </summary>
+		/// <param name="DefaultArguments">List of default arguments</param>
+		public override void GetDefaultArguments(List<string> DefaultArguments)
+		{
+			foreach (Parameter Child in Children)
+			{
+				Child.GetDefaultArguments(DefaultArguments);
+			}
+		}
+
+		/// <summary>
+		/// Convert this parameter to data for serialization
+		/// </summary>
+		/// <returns>Serializable parameter data</returns>
+		public override ParameterData ToData()
+		{
+			return new GroupParameterData(Label, Style, Children.ConvertAll(x => x.ToData()));
+		}
+	}
+
+	/// <summary>
+	/// Free-form text entry parameter
+	/// </summary>
+	public sealed class TextParameter : Parameter
+	{
+		/// <summary>
+		/// Label to display next to this parameter. Should default to the parameter name.
+		/// </summary>
+		public string Label { get; set; }
+
+		/// <summary>
+		/// Argument to add (will have the value of this field appended)
+		/// </summary>
+		public string Argument { get; set; }
+
+		/// <summary>
+		/// Default value for this argument
+		/// </summary>
+		public string Default { get; set; }
+
+		/// <summary>
+		/// Hint text to display when the field is empty
+		/// </summary>
+		public string? Hint { get; set; }
+
+		/// <summary>
+		/// Regex used to validate values entered into this text field.
+		/// </summary>
+		[BsonIgnoreIfNull]
+		public string? Validation { get; set; }
+
+		/// <summary>
+		/// Message displayed to explain valid values if validation fails.
+		/// </summary>
+		[BsonIgnoreIfNull]
+		public string? Description { get; set; }
+
+		/// <summary>
+		/// Tool tip text to display
+		/// </summary>
+		public string? ToolTip { get; set; }
+
+		/// <summary>
+		/// Private constructor for serialization
+		/// </summary>
+		private TextParameter()
+		{
+			Label = null!;
+			Argument = null!;
+			Default = String.Empty;
+		}
+
+		/// <summary>
+		/// Constructor
+		/// </summary>
+		/// <param name="Label">Parameter to pass this value to the BuildGraph script with</param>
+		/// <param name="Argument">Label to show next to the parameter</param>
+		/// <param name="Default">Default value for this argument</param>
+		/// <param name="Hint">Hint text to display</param>
+		/// <param name="Validation">Regex used to validate entries</param>
+		/// <param name="Description">Message displayed for invalid values</param>
+		/// <param name="ToolTip">Tool tip text to display</param>
+		public TextParameter(string Label, string Argument, string Default, string? Hint, string? Validation, string? Description, string? ToolTip)
+		{
+			this.Label = Label;
+			this.Argument = Argument;
+			this.Default = Default;
+			this.Hint = Hint;
+			this.Validation = Validation;
+			this.Description = Description;
+			this.ToolTip = ToolTip;
+		}
+
+		/// <summary>
+		/// Gets the default arguments for this parameter and its children
+		/// </summary>
+		/// <param name="DefaultArguments">List of default arguments</param>
+		public override void GetDefaultArguments(List<string> DefaultArguments)
+		{
+			DefaultArguments.Add(Argument + Default);
+		}
+
+		/// <summary>
+		/// Convert this parameter to data for serialization
+		/// </summary>
+		/// <returns>Serializable parameter data</returns>
+		public override ParameterData ToData()
+		{
+			return new TextParameterData(Label, Argument, Default, Hint, Validation, Description, ToolTip);
+		}
+	}
+
+	/// <summary>
+	/// Possible option for a list parameter
+	/// </summary>
+	public class ListParameterItem
+	{
+		/// <summary>
+		/// Group to display this entry in
+		/// </summary>
+		[BsonIgnoreIfNull]
+		public string? Group { get; set; }
+
+		/// <summary>
+		/// Text to display for this option.
+		/// </summary>
+		public string Text { get; set; }
+
+		/// <summary>
+		/// Argument to add with this parameter.
+		/// </summary>
+		[BsonIgnoreIfNull]
+		public string? ArgumentIfEnabled { get; set; }
+
+		/// <summary>
+		/// Argument to add with this parameter.
+		/// </summary>
+		[BsonIgnoreIfNull]
+		public string? ArgumentIfDisabled { get; set; }
+
+		/// <summary>
+		/// Whether this item is selected by default
+		/// </summary>
+		[BsonIgnoreIfDefault, BsonDefaultValue(false)]
+		public bool Default { get; set; }
+
+		/// <summary>
+		/// Constructor
+		/// </summary>
+		/// <param name="Group">The group to put this parameter in</param>
+		/// <param name="Text">Text to display for this option</param>
+		/// <param name="ArgumentIfEnabled">Argument to add if this option is enabled</param>
+		/// <param name="ArgumentIfDisabled">Argument to add if this option is disabled</param>
+		/// <param name="Default">Whether this item is selected by default</param>
+		public ListParameterItem(string? Group, string Text, string? ArgumentIfEnabled, string? ArgumentIfDisabled, bool Default)
+		{
+			this.Group = Group;
+			this.Text = Text;
+			this.ArgumentIfEnabled = ArgumentIfEnabled;
+			this.ArgumentIfDisabled = ArgumentIfDisabled;
+			this.Default = Default;
+		}
+
+		/// <summary>
+		/// Convert this parameter to data for serialization
+		/// </summary>
+		/// <returns>Serializable parameter data</returns>
+		public ListParameterItemData ToData()
+		{
+			return new ListParameterItemData(Group, Text, ArgumentIfEnabled, ArgumentIfDisabled, Default);
+		}
+	}
+
+	/// <summary>
+	/// Allows the user to select a value from a constrained list of choices
+	/// </summary>
+	public class ListParameter : Parameter
+	{
+		/// <summary>
+		/// Label to display next to this parameter.
+		/// </summary>
+		public string Label { get; set; }
+
+		/// <summary>
+		/// Style of picker parameter to use
+		/// </summary>
+		public ListParameterStyle Style { get; set; }
+
+		/// <summary>
+		/// List of values to display in the list
+		/// </summary>
+		public List<ListParameterItem> Items { get; set; }
+
+		/// <summary>
+		/// Tool tip text to display
+		/// </summary>
+		public string? ToolTip { get; set; }
+
+		/// <summary>
+		/// Private constructor for serialization
+		/// </summary>
+		private ListParameter()
+		{
+			this.Label = null!;
+			this.Items = null!;
+		}
+
+		/// <summary>
+		/// List of possible values
+		/// </summary>
+		/// <param name="Label">Label to show next to this parameter</param>
+		/// <param name="Style">Style of list parameter to use</param>
+		/// <param name="Entries">Entries for this list</param>
+		/// <param name="ToolTip">Tool tip text to display</param>
+		public ListParameter(string Label, ListParameterStyle Style, List<ListParameterItem> Entries, string? ToolTip)
+		{
+			this.Label = Label;
+			this.Style = Style;
+			this.Items = Entries;
+			this.ToolTip = ToolTip;
+		}
+
+		/// <summary>
+		/// Gets the default arguments for this parameter and its children
+		/// </summary>
+		/// <param name="DefaultArguments">List of default arguments</param>
+		public override void GetDefaultArguments(List<string> DefaultArguments)
+		{
+			foreach(ListParameterItem Item in Items)
+			{
+				if (Item.Default)
+				{
+					if (Item.ArgumentIfEnabled != null)
+					{
+						DefaultArguments.Add(Item.ArgumentIfEnabled);
+					}
+				}
+				else
+				{
+					if (Item.ArgumentIfDisabled != null)
+					{
+						DefaultArguments.Add(Item.ArgumentIfDisabled);
+					}
+				}
+			}
+		}
+
+		/// <summary>
+		/// Convert this parameter to data for serialization
+		/// </summary>
+		/// <returns>Serializable parameter data</returns>
+		public override ParameterData ToData()
+		{
+			return new ListParameterData(Label, Style, Items.ConvertAll(x => x.ToData()), ToolTip);
+		}
+	}
+
+	/// <summary>
+	/// Allows the user to toggle an option on or off
+	/// </summary>
+	public class BoolParameter : Parameter
+	{
+		/// <summary>
+		/// Label to display next to this parameter.
+		/// </summary>
+		public string Label { get; set; }
+
+		/// <summary>
+		/// Value if enabled
+		/// </summary>
+		[BsonIgnoreIfNull]
+		public string? ArgumentIfEnabled { get; set; }
+
+		/// <summary>
+		/// Value if disabled
+		/// </summary>
+		[BsonIgnoreIfNull]
+		public string? ArgumentIfDisabled { get; set; }
+
+		/// <summary>
+		/// Whether this option should be enabled by default
+		/// </summary>
+		[BsonIgnoreIfDefault, BsonDefaultValue(false)]
+		public bool Default { get; set; }
+
+		/// <summary>
+		/// Tool tip text to display
+		/// </summary>
+		public string? ToolTip { get; set; }
+
+		/// <summary>
+		/// Private constructor for serialization
+		/// </summary>
+		private BoolParameter()
+		{
+			Label = null!;
+		}
+
+		/// <summary>
+		/// Constructor
+		/// </summary>
+		/// <param name="Label">Label to display for this parameter</param>
+		/// <param name="ArgumentIfEnabled">Value if enabled</param>
+		/// <param name="ArgumentIfDisabled">Value if disabled</param>
+		/// <param name="Default">Default value for this argument</param>
+		/// <param name="ToolTip">Tool tip text to display</param>
+		public BoolParameter(string Label, string? ArgumentIfEnabled, string? ArgumentIfDisabled, bool Default, string? ToolTip)
+		{
+			this.Label = Label;
+			this.ArgumentIfEnabled = ArgumentIfEnabled;
+			this.ArgumentIfDisabled = ArgumentIfDisabled;
+			this.Default = Default;
+			this.ToolTip = ToolTip;
+		}
+
+		/// <summary>
+		/// Gets the default arguments for this parameter and its children
+		/// </summary>
+		/// <param name="DefaultArguments">List of default arguments</param>
+		public override void GetDefaultArguments(List<string> DefaultArguments)
+		{
+			string? DefaultArgument = Default ? ArgumentIfEnabled : ArgumentIfDisabled;
+			if (!String.IsNullOrEmpty(DefaultArgument))
+			{
+				DefaultArguments.Add(DefaultArgument);
+			}
+		}
+		/// <summary>
+		/// Convert this parameter to data for serialization
+		/// </summary>
+		/// <returns><see cref="BoolParameterData"/> instance</returns>
+		public override ParameterData ToData()
+		{
+			return new BoolParameterData(Label, ArgumentIfEnabled, ArgumentIfDisabled, Default, ToolTip);
+		}
+	}
+
+	/// <summary>
+	/// Counter to be updated when executing a template
+	/// </summary>
+	public interface ITemplateCounter
+	{
+		/// <summary>
+		/// Name of the property to set to the counter value
+		/// </summary>
+		string PropertyName { get; }
+
+		/// <summary>
+		/// Name of the counter
+		/// </summary>
+		string CounterName { get; }
+	}
+
+	/// <summary>
+	/// Document describing a job template. These objects are considered immutable once created and uniquely referenced by hash, in order to de-duplicate across all job runs.
+	/// </summary>
+	public interface ITemplate
+	{
+		/// <summary>
+		/// Hash of this template
+		/// </summary>
+		public ContentHash Id { get; }
+
+		/// <summary>
+		/// Name of the template.
+		/// </summary>
+		public string Name { get; }
+
+		/// <summary>
+		/// Priority of this job
+		/// </summary>
+		public Priority? Priority { get; }
+
+		/// <summary>
+		/// Whether to allow preflights for this job type
+		/// </summary>
+		public bool AllowPreflights { get; }
+
+		/// <summary>
+		/// Agent type to use for parsing the job state
+		/// </summary>
+		public string? InitialAgentType { get; }
+
+		/// <summary>
+		/// Path to a file within the stream to submit to generate a new changelist for jobs
+		/// </summary>
+		public string? SubmitNewChange { get; }
+
+		/// <summary>
+		/// List of property counters
+		/// </summary>
+		public IReadOnlyList<ITemplateCounter> Counters { get; }
+
+		/// <summary>
+		/// Optional predefined user-defined properties for this job
+		/// </summary>
+		public IReadOnlyList<string> Arguments { get; }
+
+		/// <summary>
+		/// Parameters for this template
+		/// </summary>
+		public IReadOnlyList<Parameter> Parameters { get; }
+	}
+
+	/// <summary>
+	/// Extension methods for templates
+	/// </summary>
+	public static class TemplateExtensions
+	{
+		/// <summary>
+		/// Gets the arguments for default options in this template. Does not include the standard template arguments.
+		/// </summary>
+		/// <returns>List of default arguments</returns>
+		public static List<string> GetDefaultArguments(this ITemplate Template)
+		{
+			List<string> DefaultArguments = new List<string>(Template.Arguments);
+			foreach (Parameter Parameter in Template.Parameters)
+			{
+				Parameter.GetDefaultArguments(DefaultArguments);
+			}
+			return DefaultArguments;
+		}
+	}
+}
diff --git a/Engine/Source/Programs/Horde/HordeServer/Models/TestData.cs b/Engine/Source/Programs/Horde/HordeServer/Models/TestData.cs
new file mode 100644
--- /dev/null
+++ b/Engine/Source/Programs/Horde/HordeServer/Models/TestData.cs
@@ -0,0 +1,57 @@
+using MongoDB.Bson;
+using System;
+using System.Collections.Generic;
+using System.Linq;
+using System.Threading.Tasks;
+
+using StreamId = HordeServer.Utilities.StringId<HordeServer.Models.IStream>;
+using TemplateRefId = HordeServer.Utilities.StringId<HordeServer.Models.TemplateRef>;
+
+namespace HordeServer.Models
+{
+	/// <summary>
+	/// Stores information about the results of a test
+	/// </summary>
+	public interface ITestData
+	{
+		/// <summary>
+		/// Unique id of the test data
+		/// </summary>
+		ObjectId Id { get; }
+
+		/// <summary>
+		/// Stream that generated the test data
+		/// </summary>
+		StreamId StreamId { get; }
+
+		/// <summary>
+		/// The template reference id
+		/// </summary>
+		TemplateRefId TemplateRefId { get; }
+
+		/// <summary>
+		/// The job which produced the data
+		/// </summary>
+		ObjectId JobId { get; }
+
+		/// <summary>
+		/// The step that ran
+		/// </summary>
+		SubResourceId StepId { get; }
+
+		/// <summary>
+		/// The changelist number that contained the data
+		/// </summary>
+		int Change { get; }
+
+		/// <summary>
+		/// Key used to identify the particular data
+		/// </summary>
+		string Key { get; }
+
+		/// <summary>
+		/// The data stored for this test
+		/// </summary>
+		BsonDocument Data { get; }
+	}
+}
diff --git a/Engine/Source/Programs/Horde/HordeServer/Models/Worker.cs b/Engine/Source/Programs/Horde/HordeServer/Models/Worker.cs
new file mode 100644
--- /dev/null
+++ b/Engine/Source/Programs/Horde/HordeServer/Models/Worker.cs
@@ -0,0 +1,347 @@
+// Copyright Epic Games, Inc. All Rights Reserved.
+
+using HordeServer.Api;
+using MongoDB.Bson;
+using MongoDB.Bson.Serialization.Attributes;
+using System;
+using System.Collections.Generic;
+using System.Diagnostics.CodeAnalysis;
+using System.Linq;
+using System.Threading.Tasks;
+
+namespace HordeServer.Models
+{
+	/// <summary>
+	/// Information about a device attached to an agent
+	/// </summary>
+	public class DeviceCapabilities
+	{
+		/// <summary>
+		/// Identifier for this device
+		/// </summary>
+		[BsonRequired]
+		public string Handle { get; set; }
+
+		/// <summary>
+		/// Properties of this device
+		/// </summary>
+		[BsonIgnoreIfNull]
+		public HashSet<string>? Properties { get; set; }
+
+		/// <summary>
+		/// Resources of this device
+		/// </summary>
+		[BsonIgnoreIfNull]
+		public Dictionary<string, int>? Resources { get; set; }
+
+		/// <summary>
+		/// Private constructor for serialization
+		/// </summary>
+		[BsonConstructor]
+		public DeviceCapabilities()
+		{
+			Handle = null!;
+		}
+
+		/// <summary>
+		/// Constructor
+		/// </summary>
+		/// <param name="Handle">Identifier of this device</param>
+		/// <param name="Properties">Properties of this device</param>
+		/// <param name="Resources">Resources for this device.</param>
+		public DeviceCapabilities(string Handle, HashSet<string>? Properties, Dictionary<string, int>? Resources)
+		{
+			this.Handle = Handle;
+			this.Properties = Properties;
+			this.Resources = Resources;
+		}
+
+		/// <summary>
+		/// Construct a device 
+		/// </summary>
+		/// <param name="Capabilities"></param>
+		public DeviceCapabilities(HordeCommon.Rpc.Messages.DeviceCapabilities Capabilities)
+			: this(Capabilities.Handle, new HashSet<string>(Capabilities.Properties), new Dictionary<string, int>(Capabilities.Resources))
+		{
+		}
+
+		/// <summary>
+		/// Tests whether this device has the given property
+		/// </summary>
+		/// <param name="Property"></param>
+		/// <returns></returns>
+		public bool HasProperty(string Property)
+		{
+			return Properties != null && Properties.Contains(Property);
+		}
+	}
+
+	/// <summary>
+	/// Capabilities of an agent
+	/// </summary>
+	public class AgentCapabilities
+	{
+		/// <summary>
+		/// Devices available on this agent
+		/// </summary>
+		public List<DeviceCapabilities> Devices { get; set; }
+
+		/// <summary>
+		/// Properties for this device
+		/// </summary>
+		[BsonIgnoreIfNull]
+		public HashSet<string>? Properties { get; set; }
+
+		/// <summary>
+		/// The primary device (the host machine)
+		/// </summary>
+		public DeviceCapabilities PrimaryDevice => Devices[0];
+
+		/// <summary>
+		/// Private constructor for serialization
+		/// </summary>
+		public AgentCapabilities()
+		{
+			this.Devices = new List<DeviceCapabilities>();
+		}
+
+		/// <summary>
+		/// Constructor
+		/// </summary>
+		/// <param name="Devices">Devices attached to this agent</param>
+		/// <param name="Properties">Set of properties for this agent</param>
+		public AgentCapabilities(List<DeviceCapabilities> Devices, HashSet<string>? Properties)
+		{
+			this.Devices = Devices;
+			this.Properties = Properties;
+		}
+
+		/// <summary>
+		/// Construct 
+		/// </summary>
+		/// <param name="Capabilities">RPC capabilities message</param>
+		public AgentCapabilities(HordeCommon.Rpc.Messages.AgentCapabilities Capabilities)
+			: this(Capabilities.Devices.Select(x => new DeviceCapabilities(x)).ToList(), new HashSet<string>(Capabilities.Properties))
+		{
+		}
+	}
+
+	/// <summary>
+	/// Request to create a dervice
+	/// </summary>
+	public class DeviceRequirements
+	{
+		/// <summary>
+		/// Name of this device
+		/// </summary>
+		public string? Handle { get; set; }
+
+		/// <summary>
+		/// List of properties for this device
+		/// </summary>
+		public HashSet<string>? Properties { get; set; }
+
+		/// <summary>
+		/// Count of resources for this device. This is not included in the RWAPI.
+		/// </summary>
+		public Dictionary<string, int>? Resources { get; set; }
+
+		/// <summary>
+		/// Private constructor for serialization
+		/// </summary>
+		private DeviceRequirements()
+		{
+			Handle = null!;
+		}
+
+		/// <summary>
+		/// Copy constructor
+		/// </summary>
+		/// <param name="Other">The other device requirements</param>
+		public DeviceRequirements(DeviceRequirements Other)
+		{
+			this.Handle = Other.Handle;
+			if (Other.Properties != null)
+			{
+				Properties = new HashSet<string>(Other.Properties);
+			}
+			if (Other.Resources != null)
+			{
+				Resources = new Dictionary<string, int>(Other.Resources, Other.Resources.Comparer);
+			}
+		}
+
+		/// <summary>
+		/// Constructor
+		/// </summary>
+		/// <param name="Handle">Name of this device</param>
+		/// <param name="Properties">Properties for this device</param>
+		/// <param name="Resources">Resources for this device</param>
+		public DeviceRequirements(string? Handle, HashSet<string>? Properties, Dictionary<string, int>? Resources)
+		{
+			this.Handle = Handle;
+			this.Properties = Properties;
+			this.Resources = Resources;
+		}
+
+		/// <summary>
+		/// Creates a NodeDevice object from a public API request
+		/// </summary>
+		/// <param name="Request">The input request</param>
+		/// <returns>New NodeDevice instance</returns>
+		public static DeviceRequirements FromRequest(CreateDeviceRequirementsRequest Request)
+		{
+			return new DeviceRequirements(Request.Name, Request.Properties, Request.Resources);
+		}
+
+		/// <summary>
+		/// Creates an RPC response for this instance
+		/// </summary>
+		/// <returns>Response object</returns>
+		public HordeCommon.Rpc.Messages.DeviceRequirements ToRpcResponse()
+		{
+			HordeCommon.Rpc.Messages.DeviceRequirements Requirements = new HordeCommon.Rpc.Messages.DeviceRequirements();
+			if (Properties != null)
+			{
+				Requirements.Properties.Add(Properties);
+			}
+			if (Resources != null)
+			{
+				Requirements.Resources.Add(Resources);
+			}
+			return Requirements;
+		}
+	}
+
+	/// <summary>
+	/// Requirements for an agent. Currently mirrors the Worker type in the RWAPI.
+	/// </summary>
+	public class AgentRequirements
+	{
+		/// <summary>
+		/// Devices attached to this agent
+		/// </summary>
+		public List<DeviceRequirements>? Devices { get; set; }
+
+		/// <summary>
+		/// Properties of this agent
+		/// </summary>
+		public HashSet<string>? Properties { get; set; }
+
+		/// <summary>
+		/// Configurable settings for this agent
+		/// </summary>
+		public Dictionary<string, string>? Configs { get; set; }
+
+		/// <summary>
+		/// Whether the agent can be shared with another step
+		/// </summary>
+		public bool Shared { get; set; }
+
+		/// <summary>
+		/// Constructor
+		/// </summary>
+		public AgentRequirements()
+		{
+			this.Devices = new List<DeviceRequirements>();
+		}
+
+		/// <summary>
+		/// Copy constructor
+		/// </summary>
+		/// <param name="Other">Instance to copy from</param>
+		public AgentRequirements(AgentRequirements Other)
+		{
+			if (Other.Devices != null)
+			{
+				this.Devices = Other.Devices.ConvertAll(x => new DeviceRequirements(x));
+			}
+			if(Other.Properties != null)
+			{
+				this.Properties = new HashSet<string>(Other.Properties);
+			}
+			if(Other.Configs != null)
+			{
+				this.Configs = new Dictionary<string, string>(Other.Configs);
+			}
+			this.Shared = Other.Shared;
+		}
+
+		/// <summary>
+		/// Constructor
+		/// </summary>
+		/// <param name="Devices">Devices attached to this agent</param>
+		/// <param name="Properties">Properties of this agent</param>
+		/// <param name="Configs">Configuration settings for this agent</param>
+		/// <param name="Shared">Whether the agent can be shared with another job</param>
+		public AgentRequirements(List<DeviceRequirements>? Devices, HashSet<string>? Properties, Dictionary<string, string>? Configs, bool Shared)
+		{
+			this.Devices = Devices;
+			this.Properties = Properties;
+			this.Configs = Configs;
+			this.Shared = Shared;
+		}
+
+		/// <summary>
+		/// Determine if this requirements object is empty
+		/// </summary>
+		/// <returns>TRue if the object is empty</returns>
+		public bool IsEmpty()
+		{
+			return Devices == null && Properties == null && Configs == null && !Shared;
+		}
+
+		/// <summary>
+		/// Constructs a NodeAgent object from a public API request
+		/// </summary>
+		/// <param name="Request">The API request</param>
+		/// <returns>The new node agent instance</returns>
+		[return: NotNullIfNotNull("Request")]
+		public static AgentRequirements? FromRequest(CreateAgentRequirementsRequest? Request)
+		{
+			if (Request == null)
+			{
+				return null;
+			}
+
+			List<DeviceRequirements>? Devices = null;
+			if (Request.Devices != null)
+			{
+				Devices = Request.Devices?.ConvertAll(x => DeviceRequirements.FromRequest(x));
+			}
+
+			HashSet<string>? Properties = null;
+			if (Request.Properties != null)
+			{
+				Properties = new HashSet<string>(Request.Properties, StringComparer.Ordinal);
+			}
+
+			Dictionary<string, string>? Configs = null;
+			if (Request.Configs != null)
+			{
+				Configs = new Dictionary<string, string>(Request.Configs, StringComparer.Ordinal);
+			}
+
+			return new AgentRequirements(Devices, Properties, Configs, Request.Shared);
+		}
+
+		/// <summary>
+		/// Creates an Rpc response object for this instance
+		/// </summary>
+		/// <returns>Response object</returns>
+		public HordeCommon.Rpc.Messages.AgentRequirements ToRpcResponse()
+		{
+			HordeCommon.Rpc.Messages.AgentRequirements Requirements = new HordeCommon.Rpc.Messages.AgentRequirements();
+			if (Devices != null)
+			{
+				Requirements.Devices.AddRange(Devices.Select(x => x.ToRpcResponse()));
+			}
+			if (Properties != null)
+			{
+				Requirements.Properties.AddRange(Properties);
+			}
+			Requirements.Shared = Shared;
+			return Requirements;
+		}
+	}
+}
diff --git a/Engine/Source/Programs/Horde/HordeServer/Notifications/INotificationService.cs b/Engine/Source/Programs/Horde/HordeServer/Notifications/INotificationService.cs
new file mode 100644
--- /dev/null
+++ b/Engine/Source/Programs/Horde/HordeServer/Notifications/INotificationService.cs
@@ -0,0 +1,78 @@
+// Copyright Epic Games, Inc. All Rights Reserved.
+
+using HordeServer.Api;
+using HordeServer.Models;
+using MongoDB.Bson;
+using System;
+using System.Collections.Generic;
+using System.Linq;
+using System.Security.Claims;
+using System.Threading.Tasks;
+
+namespace HordeServer.Notifications
+{
+	/// <summary>
+	/// Interface for the notification service
+	/// </summary>
+	public interface INotificationService
+	{
+		/// <summary>
+		/// Updates a subscription to a trigger
+		/// </summary>
+		/// <param name="TriggerId"></param>
+		/// <param name="User"></param>
+		/// <param name="Email">Whether to receive email notifications</param>
+		/// <param name="Slack">Whether to receive Slack notifications</param>
+		/// <returns></returns>
+		Task<bool> UpdateSubscriptionsAsync(ObjectId TriggerId, ClaimsPrincipal User, bool? Email, bool? Slack);
+
+		/// <summary>
+		/// Gets the current subscriptions for a user
+		/// </summary>
+		/// <param name="TriggerId"></param>
+		/// <param name="User"></param>
+		/// <returns>Subscriptions for that user</returns>
+		Task<INotificationSubscription?> GetSubscriptionsAsync(ObjectId TriggerId, ClaimsPrincipal User);
+
+		/// <summary>
+		/// Notify all subscribers that a job step has finished
+		/// </summary>
+		/// <param name="Job">The job containing the step that has finished</param>
+		/// <param name="Graph"></param>
+		/// <param name="BatchId">The batch id</param>
+		/// <param name="StepId">The step id</param>
+		/// <returns>Async task</returns>
+		void NotifyJobStepComplete(IJob Job, IGraph Graph, SubResourceId BatchId, SubResourceId StepId);
+
+		/// <summary>
+		/// Notify all subscribers that a job step's outcome has changed
+		/// </summary>
+		/// <param name="Job">The job containing the step that has changed</param>
+		/// <param name="OldLabelStates">The old label states for the job</param>
+		/// <param name="NewLabelStates">The new label states for the job</param>
+		/// <returns>Async task</returns>
+		void NotifyLabelUpdate(IJob Job, IReadOnlyList<(LabelState, LabelOutcome)> OldLabelStates, IReadOnlyList<(LabelState, LabelOutcome)> NewLabelStates);
+
+		/// <summary>
+		/// Notify slack channel about a stream update failure
+		/// </summary>
+		/// <param name="Stream">The stream</param>
+		/// <param name="ErrorMessage">Error message passed back</param>
+		/// <param name="Change">Latest change number of the file</param>
+		/// <param name="Author">Author of the change, if found from p4 service</param>
+		/// <param name="Description">Description of the change, if found from p4 service</param>
+		void NotifyUpdateStreamFailure(IStream Stream, string ErrorMessage, int Change, string? Author, string? Description);
+
+		/// <summary>
+		/// Notify slack channel about a stream update failure in relation to a file
+		/// </summary>
+		/// <param name="File">The file</param>
+		void NotifyUpdateStreamFailure(FileSummary File);
+
+		/// <summary>
+		/// Sends a notification to a user regarding a build health issue.
+		/// </summary>
+		/// <param name="Issue">The new issue that was created</param>
+		void NotifyIssueUpdated(IIssue Issue);
+	}
+}
diff --git a/Engine/Source/Programs/Horde/HordeServer/Notifications/INotificationSink.cs b/Engine/Source/Programs/Horde/HordeServer/Notifications/INotificationSink.cs
new file mode 100644
--- /dev/null
+++ b/Engine/Source/Programs/Horde/HordeServer/Notifications/INotificationSink.cs
@@ -0,0 +1,91 @@
+// Copyright Epic Games, Inc. All Rights Reserved.
+
+using HordeCommon;
+using HordeServer.Api;
+using HordeServer.Models;
+using MongoDB.Bson;
+using System;
+using System.Collections.Generic;
+using System.Linq;
+using System.Threading.Tasks;
+
+namespace HordeServer.Notifications
+{
+	/// <summary>
+	/// Implements a notification method
+	/// </summary>
+	public interface INotificationSink
+	{
+		/// <summary>
+		/// Send notifications that a job has completed
+		/// </summary>
+		/// <param name="JobStream"></param>
+		/// <param name="Job">The job containing the step</param>
+		/// <param name="Graph"></param>
+		/// <param name="Outcome"></param>
+		/// <returns>Async task</returns>
+		Task NotifyJobCompleteAsync(IStream JobStream, IJob Job, IGraph Graph, LabelOutcome Outcome);
+
+		/// <summary>
+		/// Send notifications that a job has completed
+		/// </summary>
+		/// <param name="User">User to notify</param>
+		/// <param name="JobStream"></param>
+		/// <param name="Job">The job containing the step</param>
+		/// <param name="Graph"></param>
+		/// <param name="Outcome"></param>
+		/// <returns>Async task</returns>
+		Task NotifyJobCompleteAsync(IUser User, IStream JobStream, IJob Job, IGraph Graph, LabelOutcome Outcome);
+
+		/// <summary>
+		/// Send notifications that a job step has completed
+		/// </summary>
+		/// <param name="User">User to notify</param>
+		/// <param name="JobStream">Stream containing the job</param>
+		/// <param name="Job">The job containing the step</param>
+		/// <param name="Batch">Unique id of the batch</param>
+		/// <param name="Step">The step id</param>
+		/// <param name="Node">Corresponding node for the step</param>
+		/// <param name="JobStepEventData"></param>
+		/// <returns>Async task</returns>
+		Task NotifyJobStepCompleteAsync(IUser User, IStream JobStream, IJob Job, IJobStepBatch Batch, IJobStep Step, INode Node, List<ILogEventData> JobStepEventData);
+
+		/// <summary>
+		/// Send notifications that a job step has completed
+		/// </summary>
+		/// <param name="User">User to notify</param>
+		/// <param name="Job">The job containing the step</param>
+		/// <param name="Stream"></param>
+		/// <param name="Label"></param>
+		/// <param name="LabelIdx"></param>
+		/// <param name="Outcome"></param>
+		/// <param name="StepData"></param>
+		/// <returns>Async task</returns>
+		Task NotifyLabelCompleteAsync(IUser User, IJob Job, IStream Stream, ILabel Label, int LabelIdx, LabelOutcome Outcome, List<(string, JobStepOutcome, Uri)> StepData);
+
+		/// <summary>
+		/// Send notifications that a new issue has been created or assigned
+		/// </summary>
+		/// <param name="Issue"></param>
+		/// <returns></returns>
+		Task NotifyIssueUpdatedAsync(IIssue Issue);
+
+		/// <summary>
+		/// Notification that a stream has failed to update
+		/// </summary>
+		/// <param name="Stream"></param>
+		/// <param name="ErrorMessage"></param>
+		/// <param name="Change"></param>
+		/// <param name="Author"></param>
+		/// <param name="Description"></param>
+		/// <returns></returns>
+		Task NotifyStreamUpdateFailedAsync(IStream Stream, string ErrorMessage, int Change, IUser? Author = null, string? Description = null);
+
+		/// <summary>
+		/// Notification that a stream has failed to update
+		/// </summary>
+		/// <param name="File"></param>
+		/// <returns></returns>
+		Task NotifyStreamUpdateFailedAsync(FileSummary File);
+	}
+}
diff --git a/Engine/Source/Programs/Horde/HordeServer/Notifications/Impl/NotificationService.cs b/Engine/Source/Programs/Horde/HordeServer/Notifications/Impl/NotificationService.cs
new file mode 100644
--- /dev/null
+++ b/Engine/Source/Programs/Horde/HordeServer/Notifications/Impl/NotificationService.cs
@@ -0,0 +1,622 @@
+// Copyright Epic Games, Inc. All Rights Reserved.
+
+using System;
+using System.Collections.Concurrent;
+using System.Collections.Generic;
+using System.Drawing;
+using System.Linq;
+using System.Net;
+using System.Net.Http;
+using System.Net.Mail;
+using System.Security.Claims;
+using System.Text;
+using System.Threading;
+using System.Threading.Tasks;
+using EpicGames.Core;
+using HordeServer.Api;
+using HordeServer.Collections;
+using HordeCommon;
+using HordeServer.Models;
+using HordeServer.Utilities;
+using HordeServer.Utilities.BlockKit;
+using HordeServer.Utilities.Slack.BlockKit;
+using Microsoft.Extensions.Hosting;
+using Microsoft.Extensions.Logging;
+using Microsoft.Extensions.Options;
+using MongoDB.Bson;
+using MongoDB.Driver;
+using Newtonsoft.Json;
+using Newtonsoft.Json.Linq;
+using StatsdClient;
+using JsonSerializer = System.Text.Json.JsonSerializer;
+using HordeServer.Notifications;
+using HordeServer.Services;
+
+namespace HordeServer.Notifications.Impl
+{
+	/// <summary>
+	/// Wraps funtionality for delivering notifications.
+	/// </summary>
+	public class NotificationService : BackgroundService, INotificationService
+	{
+		/// <summary>
+		/// The available notification sinks
+		/// </summary>
+		private List<INotificationSink> Sinks;
+
+		/// <summary>
+		/// Collection of subscriptions
+		/// </summary>
+		private readonly ISubscriptionCollection SubscriptionCollection;
+
+		/// <summary>
+		/// Collection of notification request documents.
+		/// </summary>
+		private readonly INotificationTriggerCollection TriggerCollection;
+
+		/// <summary>
+		/// Instance of the <see cref="GraphCollection"/>.
+		/// </summary>
+		private readonly IGraphCollection GraphCollection;
+
+		/// <summary>
+		/// 
+		/// </summary>
+		private readonly IUserCollection UserCollection;
+
+		/// <summary>
+		/// Job service instance
+		/// </summary>
+		private readonly JobService JobService;
+
+		/// <summary>
+		/// Instance of the <see cref="StreamService"/>.
+		/// </summary>
+		private readonly StreamService StreamService;
+
+		/// <summary>
+		/// 
+		/// </summary>
+		private readonly IIssueService IssueService;
+
+		/// <summary>
+		/// Instance of the <see cref="LogFileService"/>.
+		/// </summary>
+		private readonly ILogFileService LogFileService;
+		
+		/// <summary>
+		/// Instance of the <see cref="IDogStatsd"/>.
+		/// </summary>
+		private readonly IDogStatsd DogStatsd;
+		
+		/// <summary>
+		/// Settings for the application.
+		/// </summary>
+		private readonly IOptionsMonitor<ServerSettings> Settings;
+
+		/// <summary>
+		/// List of asychronous tasks currently executing
+		/// </summary>
+		private readonly ConcurrentQueue<Task> NewTasks = new ConcurrentQueue<Task>();
+
+		/// <summary>
+		/// Set when there are new tasks to wait for
+		/// </summary>
+		private AsyncEvent NewTaskEvent = new AsyncEvent();
+
+		/// <summary>
+		/// Settings for the application.
+		/// </summary>
+		private readonly ILogger<NotificationService> Logger;
+
+		/// <summary>
+		/// Constructor
+		/// </summary>
+		public NotificationService(IEnumerable<INotificationSink> Sinks, IOptionsMonitor<ServerSettings> Settings, ILogger<NotificationService> Logger, IGraphCollection GraphCollection, ISubscriptionCollection SubscriptionCollection, INotificationTriggerCollection TriggerCollection, IUserCollection UserCollection, JobService JobService, StreamService StreamService, IIssueService IssueService, ILogFileService LogFileService, IDogStatsd DogStatsd)
+		{
+			this.Sinks = Sinks.ToList();
+			this.Settings = Settings;
+			this.Logger = Logger;
+			this.GraphCollection = GraphCollection;
+			this.SubscriptionCollection = SubscriptionCollection;
+			this.TriggerCollection = TriggerCollection;
+			this.UserCollection = UserCollection;
+			this.JobService = JobService;
+			this.StreamService = StreamService;
+			this.IssueService = IssueService;
+			this.LogFileService = LogFileService;
+			this.DogStatsd = DogStatsd;
+
+			IssueService.OnIssueUpdated += NotifyIssueUpdated;
+			JobService.OnJobStepComplete += NotifyJobStepComplete;
+			JobService.OnLabelUpdate += NotifyLabelUpdate;
+		}
+
+		/// <inheritdoc/>
+		public override void Dispose()
+		{
+			base.Dispose();
+
+			IssueService.OnIssueUpdated -= NotifyIssueUpdated;
+			JobService.OnJobStepComplete -= NotifyJobStepComplete;
+			JobService.OnLabelUpdate -= NotifyLabelUpdate;
+
+			GC.SuppressFinalize(this);
+		}
+
+		/// <inheritdoc/>
+		public async Task<bool> UpdateSubscriptionsAsync(ObjectId TriggerId, ClaimsPrincipal User, bool? Email, bool? Slack)
+		{
+			ObjectId? UserId = User.GetUserId();
+			if (UserId == null)
+			{
+				Logger.LogWarning("Unable to find username for principal {User}", User.Identity.Name);
+				return false;
+			}
+
+			INotificationTrigger Trigger = await TriggerCollection.FindOrAddAsync(TriggerId);
+			await TriggerCollection.UpdateSubscriptionsAsync(Trigger, UserId.Value, Email, Slack);
+			return true;
+		}
+
+		/// <inheritdoc/>
+		public async Task<INotificationSubscription?> GetSubscriptionsAsync(ObjectId TriggerId, ClaimsPrincipal User)
+		{
+			ObjectId? UserId = User.GetUserId();
+			if (UserId == null)
+			{
+				return null;
+			}
+
+			INotificationTrigger? Trigger = await TriggerCollection.GetAsync(TriggerId);
+			if(Trigger == null)
+			{
+				return null;
+			}
+
+			return Trigger.Subscriptions.FirstOrDefault(x => x.UserId == UserId.Value);
+		}
+
+		/// <inheritdoc/>
+		public void NotifyJobStepComplete(IJob Job, IGraph Graph, SubResourceId BatchId, SubResourceId StepId)
+		{
+			// Enqueue job step complete notifications if needed
+			if (Job.TryGetStep(BatchId, StepId, out IJobStep? Step))
+			{
+				Logger.LogInformation("Queuing step notifications for {JobId}:{BatchId}:{StepId}", Job.Id, BatchId, StepId);
+				EnqueueTask(() => SendJobStepNotificationsAsync(Job, BatchId, StepId));
+			}
+
+			// Enqueue job complete notifications if needed
+			if (Job.GetState() == JobState.Complete)
+			{
+				Logger.LogInformation("Queuing job notifications for {JobId}:{BatchId}:{StepId}", Job.Id, BatchId, StepId);
+				EnqueueTask(() => SendJobNotificationsAsync(Job, Graph));
+				EnqueueTask(() => RecordJobCompleteMetrics(Job));
+			}
+		}
+
+		/// <inheritdoc/>
+		public void NotifyLabelUpdate(IJob Job, IReadOnlyList<(LabelState, LabelOutcome)> OldLabelStates, IReadOnlyList<(LabelState, LabelOutcome)> NewLabelStates)
+		{
+			// If job has any label trigger IDs, send label complete notifications if needed
+			for (int Idx = 0; Idx < OldLabelStates.Count && Idx < NewLabelStates.Count; Idx++)
+			{
+				if (OldLabelStates[Idx] != NewLabelStates[Idx])
+				{
+					EnqueueTask(() => SendAllLabelNotificationsAsync(Job, OldLabelStates, NewLabelStates));
+					break;
+				}
+			}
+		}
+
+		/// <inheritdoc/>
+		public void NotifyIssueUpdated(IIssue Issue)
+		{
+			Logger.LogInformation("Issue {IssueId} updated", Issue.Id);
+			EnqueueTasks(Sink => Sink.NotifyIssueUpdatedAsync(Issue));
+		}
+
+		/// <inheritdoc/>
+		public void NotifyUpdateStreamFailure(IStream Stream, string ErrorMessage, int Change, string? Author = null, string? Description = null)
+		{
+			EnqueueTask(() => NotifyUpdateStreamFailureAsync(Stream, ErrorMessage, Change, Author, Description));
+		}
+
+		async Task NotifyUpdateStreamFailureAsync(IStream Stream, string ErrorMessage, int Change, string? Author = null, string? Description = null)
+		{
+			IUser? User = (Author != null) ? await UserCollection.FindUserByLoginAsync(Author) : null;
+			EnqueueTasks(Sink => Sink.NotifyStreamUpdateFailedAsync(Stream, ErrorMessage, Change, User, Description));
+		}
+		/// <inheritdoc/>
+		public void NotifyUpdateStreamFailure(FileSummary File)
+		{
+			EnqueueTasks(Sink => Sink.NotifyStreamUpdateFailedAsync(File));
+		}
+
+		/// <summary>
+		/// Enqueues an async task
+		/// </summary>
+		/// <param name="TaskFunc">Function to generate an async task</param>
+		void EnqueueTask(Func<Task> TaskFunc)
+		{
+			NewTasks.Enqueue(Task.Run(TaskFunc));
+		}
+
+		/// <summary>
+		/// Enqueues an async task
+		/// </summary>
+		/// <param name="TaskFunc">Function to generate an async task</param>
+		void EnqueueTasks(Func<INotificationSink, Task> TaskFunc)
+		{
+			foreach (INotificationSink Sink in Sinks)
+			{
+				EnqueueTask(() => TaskFunc(Sink));
+			}
+		}
+
+		/// <inheritdoc/>
+		[System.Diagnostics.CodeAnalysis.SuppressMessage("Design", "CA1031:Do not catch general exception types", Justification = "<Pending>")]
+		protected async override Task ExecuteAsync(CancellationToken StoppingToken)
+		{
+			// This background service just waits for tasks to finish and prints any exception info. The only reason to do this is to
+			// ensure we finish processing everything before shutdown.
+			using (CancellationTask StoppingTask = new CancellationTask(StoppingToken))
+			{
+				List<Task> Tasks = new List<Task>();
+				Tasks.Add(NewTaskEvent.Task);
+
+				for (; ; )
+				{
+					// Add any new tasks to be monitored
+					Task? NewTask;
+					while (NewTasks.TryDequeue(out NewTask))
+					{
+						Tasks.Add(NewTask);
+					}
+
+					// If we don't have any
+					if (Tasks.Count == 1)
+					{
+						await Task.WhenAny(NewTaskEvent.Task, StoppingTask.Task);
+						if (StoppingToken.IsCancellationRequested)
+						{
+							break;
+						}
+					}
+					else
+					{
+						// Wait for something to finish
+						Task Task = await Task.WhenAny(Tasks);
+						if (Task == NewTaskEvent.Task)
+						{
+							NewTaskEvent = new AsyncEvent();
+							Tasks[0] = NewTaskEvent.Task;
+						}
+						else
+						{
+							try
+							{
+								await Task;
+							}
+							catch (Exception Ex)
+							{
+								Logger.LogError(Ex, "Exception while executing notification");
+							}
+							Tasks.Remove(Task);
+						}
+					}
+				}
+			}
+		}
+
+		internal Task ExecuteBackgroundForTest(CancellationToken StoppingToken)
+		{
+			return ExecuteAsync(StoppingToken);
+		}
+
+		/// <summary>
+		/// Gets the <see cref="INotificationTrigger"/> for a given trigger ID, if any.
+		/// </summary>
+		/// <param name="TriggerId"></param>
+		/// <param name="bFireTrigger">If true, the trigger is fired and cannot be reused</param>
+		/// <returns></returns>
+		private async Task<INotificationTrigger?> GetNotificationTrigger(ObjectId? TriggerId, bool bFireTrigger)
+		{
+			if (TriggerId == null)
+			{
+				return null;
+			}
+
+			INotificationTrigger? Trigger = await TriggerCollection.GetAsync(TriggerId.Value);
+			if (Trigger == null)
+			{
+				return null;
+			}
+
+			return bFireTrigger ? await TriggerCollection.FireAsync(Trigger) : Trigger;
+		}
+	
+		private async Task SendJobNotificationsAsync(IJob Job, IGraph Graph)
+		{
+			using IDisposable Scope = Logger.BeginScope("Sending notifications for job {JobId}", Job.Id);
+
+			Job.GetJobState(Job.GetStepForNodeMap(), out _, out LabelOutcome Outcome);
+			JobCompleteEventRecord JobCompleteEvent = new JobCompleteEventRecord(Job.StreamId, Job.TemplateId, Outcome);
+
+			IStream? JobStream = await StreamService.GetStreamAsync(Job.StreamId);
+			if (JobStream == null)
+			{
+				Logger.LogError("Unable to get stream {StreamId}", Job.StreamId);
+				return;
+			}
+
+			List<IUser> UsersToNotify = await GetUsersToNotify(JobCompleteEvent, Job.NotificationTriggerId, true);
+			foreach (IUser UserToNotify in UsersToNotify)
+			{
+				if(Job.PreflightChange != 0)
+				{
+					if(UserToNotify.Id != Job.StartedByUserId)
+					{
+						continue;
+					}
+				}
+				EnqueueTasks(Sink => Sink.NotifyJobCompleteAsync(UserToNotify, JobStream, Job, Graph, Outcome));
+			}
+
+			if (Job.PreflightChange == 0)
+			{
+				EnqueueTasks(Sink => Sink.NotifyJobCompleteAsync(JobStream, Job, Graph, Outcome));
+			}
+
+			Logger.LogDebug("Finished sending notifications for job {JobId}", Job.Id);
+		}
+
+		private Task RecordJobCompleteMetrics(IJob Job)
+		{
+			void RecordMetric(string Type, JobStepOutcome Outcome, DateTimeOffset? StartTime, DateTimeOffset? FinishTime)
+			{
+				string OutcomeStr = Outcome switch
+				{
+					JobStepOutcome.Unspecified => "unspecified",
+					JobStepOutcome.Failure => "failure",
+					JobStepOutcome.Warnings => "warnings",
+					JobStepOutcome.Success => "success",
+					_ => "unspecified"
+				};
+				
+				string[] Tags = {"stream:" + Job.StreamId, "template:" + Job.TemplateId};
+				DogStatsd.Increment($"horde.{Type}.{OutcomeStr}.count", 1, tags: Tags);
+
+				if (StartTime == null || FinishTime == null)
+				{
+					Logger.LogDebug("Completed job or step is missing start or finish time, cannot record duration metric. Job ID={JobId}", Job.Id);
+					return;
+				}
+
+				TimeSpan Duration = FinishTime.Value - StartTime.Value;
+				DogStatsd.Timer($"horde.{Type}.{OutcomeStr}.duration", Duration.TotalSeconds, tags: Tags);
+			}
+
+			JobStepOutcome JobOutcome = Job.Batches.SelectMany(x => x.Steps).Min(x => x.Outcome);
+			DateTime? StartTime = Job.Batches.Select(x => x.StartTimeUtc).Min();
+			DateTime? FinishTime = Job.Batches.Select(x => x.FinishTimeUtc).Max();
+			RecordMetric("job", JobOutcome, StartTime, FinishTime);
+
+			// TODO: record metrics for individual steps
+			// foreach (IJobStepBatch Batch in Job.Batches)
+			// {
+			// 	foreach (IJobStep Step in Batch.Steps)
+			// 	{
+			// 	}
+			// }
+
+			return Task.CompletedTask;
+		}
+
+		private async Task<List<IUser>> GetUsersToNotify(EventRecord? Event, ObjectId? NotificationTriggerId, bool bFireTrigger)
+		{
+			List<ObjectId> UserIds = new List<ObjectId>();
+
+			// Find the notifications for all steps of this type
+			if (Event != null)
+			{
+				List<ISubscription> Subscriptions = await SubscriptionCollection.FindSubscribersAsync(Event);
+				foreach (ISubscription Subscription in Subscriptions)
+				{
+					if (Subscription.NotificationType == NotificationType.Slack)
+					{
+						UserIds.Add(Subscription.UserId);
+					}
+				}
+			}
+
+			// Find the notifications for this particular step
+			if (NotificationTriggerId != null)
+			{
+				INotificationTrigger? Trigger = await GetNotificationTrigger(NotificationTriggerId, bFireTrigger);
+				if (Trigger != null)
+				{
+					foreach (INotificationSubscription Subscription in Trigger.Subscriptions)
+					{
+						if (Subscription.Email)
+						{
+							// TODO?
+						}
+						if (Subscription.Slack)
+						{
+							UserIds.Add(Subscription.UserId);
+						}
+					}
+				}
+			}
+			return await UserCollection.FindUsersAsync(UserIds);
+		}
+
+		private async Task SendJobStepNotificationsAsync(IJob Job, SubResourceId BatchId, SubResourceId StepId)
+		{
+			using IDisposable Scope = Logger.BeginScope("Sending notifications for step {JobId}:{BatchId}:{StepId}", Job.Id, BatchId, StepId);
+
+			IJobStepBatch? Batch;
+			if(!Job.TryGetBatch(BatchId, out Batch))
+			{
+				Logger.LogError("Unable to find batch {BatchId} in job {JobId}", BatchId, Job.Id);
+				return;
+			}
+
+			IJobStep? Step;
+			if (!Batch.TryGetStep(StepId, out Step))
+			{
+				Logger.LogError("Unable to find step {StepId} in batch {JobId}:{BatchId}", StepId, Job.Id, BatchId);
+				return;
+			}
+
+			IGraph JobGraph = await GraphCollection.GetAsync(Job.GraphHash);
+			INode Node = JobGraph.GetNode(new NodeRef(Batch.GroupIdx, Step.NodeIdx));
+
+			// Find the notifications for this particular step
+			EventRecord EventRecord = new StepCompleteEventRecord(Job.StreamId, Job.TemplateId, Node.Name, Step.Outcome);
+
+			List<IUser> UsersToNotify = await GetUsersToNotify(EventRecord, Step.NotificationTriggerId, true);
+			if(UsersToNotify.Count == 0)
+			{
+				Logger.LogInformation("No users to notify for step {JobId}:{BatchId}:{StepId}", Job.Id, BatchId, StepId);
+				return;
+			}
+
+			IStream? JobStream = await StreamService.GetStreamAsync(Job.StreamId);
+			if (JobStream == null)
+			{
+				Logger.LogError("Unable to find stream {StreamId}", Job.StreamId);
+				return;
+			}
+
+			if(Step.LogId == null)
+			{
+				Logger.LogError("Step does not have a log file");
+				return;
+			}
+
+			ILogFile? LogFile = await LogFileService.GetLogFileAsync(Step.LogId.Value);
+			if(LogFile == null)
+			{
+				Logger.LogError("Step does not have a log file");
+				return;
+			}
+
+			List<ILogEvent> JobStepEvents = await LogFileService.FindLogEventsAsync(LogFile);
+			List<ILogEventData> JobStepEventData = new List<ILogEventData>();
+			foreach (ILogEvent Event in JobStepEvents)
+			{
+				ILogEventData EventData = await LogFileService.GetEventDataAsync(LogFile, Event.LineIndex, Event.LineCount);
+				JobStepEventData.Add(EventData);
+			}
+
+			foreach (IUser SlackUser in UsersToNotify)
+			{
+				if(Job.PreflightChange != 0)
+				{
+					if(SlackUser.Id != Job.StartedByUserId)
+					{
+						continue;
+					}
+				}
+				EnqueueTasks(Sink => Sink.NotifyJobStepCompleteAsync(SlackUser, JobStream, Job, Batch, Step, Node, JobStepEventData));
+			}
+			Logger.LogDebug("Finished sending notifications for step {JobId}:{BatchId}:{StepId}", Job.Id, BatchId, StepId);
+		}
+
+		private async Task SendAllLabelNotificationsAsync(IJob Job, IReadOnlyList<(LabelState State, LabelOutcome Outcome)> OldLabelStates, IReadOnlyList<(LabelState, LabelOutcome)> NewLabelStates)
+		{
+			IStream? Stream = await StreamService.GetStreamAsync(Job.StreamId);
+			if (Stream == null)
+			{
+				Logger.LogError("Unable to find stream {StreamId} for job {JobId}", Job.StreamId, Job.Id);
+				return;
+			}
+
+			IGraph? Graph = await GraphCollection.GetAsync(Job.GraphHash);
+			if (Graph == null)
+			{
+				Logger.LogError("Unable to find graph {GraphHash} for job {JobId}", Job.GraphHash, Job.Id);
+				return;
+			}
+
+			IReadOnlyDictionary<NodeRef, IJobStep> StepForNode = Job.GetStepForNodeMap();
+			for (int LabelIdx = 0; LabelIdx < Graph.Labels.Count; ++LabelIdx)
+			{
+				(LabelState State, LabelOutcome Outcome) OldLabel = OldLabelStates[LabelIdx];
+				(LabelState State, LabelOutcome Outcome) NewLabel = NewLabelStates[LabelIdx];
+				if (OldLabel != NewLabel)
+				{
+					// If the state transitioned from Unspecified to Running, don't update unless the outcome also changed.
+					if (OldLabel.State == LabelState.Unspecified && NewLabel.State == LabelState.Running && OldLabel.Outcome == NewLabel.Outcome)
+					{
+						continue;
+					}
+
+					// If the label isn't complete, don't report on outcome changing to success, this will be reported when the label state becomes complete.
+					if (NewLabel.State != LabelState.Complete && NewLabel.Outcome == LabelOutcome.Success)
+					{
+						return;
+					}
+
+					ILabel Label = Graph.Labels[LabelIdx];
+
+					EventRecord? EventId;
+					if (String.IsNullOrEmpty(Label.DashboardName))
+					{
+						EventId = null;
+					}
+					else
+					{
+						EventId = new LabelCompleteEventRecord(Job.StreamId, Job.TemplateId, Label.DashboardCategory, Label.DashboardName, NewLabel.Outcome);
+					}
+
+					ObjectId? TriggerId;
+					if (Job.LabelIdxToTriggerId.TryGetValue(LabelIdx, out ObjectId NewTriggerId))
+					{
+						TriggerId = NewTriggerId;
+					}
+					else
+					{
+						TriggerId = null;
+					}
+
+					bool bFireTrigger = NewLabel.State == LabelState.Complete;
+
+					List<IUser> UsersToNotify = await GetUsersToNotify(EventId, TriggerId, bFireTrigger);
+					if (UsersToNotify.Count > 0)
+					{
+						SendLabelUpdateNotifications(Job, Stream, Graph, StepForNode, Graph.Labels[LabelIdx], LabelIdx, NewLabel.Outcome, UsersToNotify);
+					}
+					else
+					{
+						Logger.LogDebug("No users to notify for label {DashboardName}/{UgsName} in job {JobId}", Graph.Labels[LabelIdx].DashboardName, Graph.Labels[LabelIdx].UgsName, Job.Id);
+					}
+				}
+			}
+		}
+
+		private void SendLabelUpdateNotifications(IJob Job, IStream Stream, IGraph Graph, IReadOnlyDictionary<NodeRef, IJobStep> StepForNode, ILabel Label, int LabelIdx, LabelOutcome Outcome, List<IUser> SlackUsers)
+		{
+			List<(string, JobStepOutcome, Uri)> StepData = new List<(string, JobStepOutcome, Uri)>();
+			if (Outcome != LabelOutcome.Success)
+			{
+				foreach (NodeRef IncludedNodeRef in Label.IncludedNodes)
+				{
+					INode IncludedNode = Graph.GetNode(IncludedNodeRef);
+					IJobStep IncludedStep = StepForNode[IncludedNodeRef];
+					StepData.Add((IncludedNode.Name, IncludedStep.Outcome, new Uri($"{Settings.CurrentValue.DashboardUrl}job/{Job.Id}?step={IncludedStep.Id}")));
+				}
+			}
+
+			foreach (IUser SlackUser in SlackUsers)
+			{
+				EnqueueTasks(Sink => Sink.NotifyLabelCompleteAsync(SlackUser, Job, Stream, Label, LabelIdx, Outcome, StepData));
+			}
+
+			Logger.LogDebug("Finished sending label notifications for label {DashboardName}/{UgsName} in job {JobId}", Label.DashboardName, Label.UgsName, Job.Id);
+		}
+	}
+}
