Initial drop of Horde Server/Agent for UE5/Main. Still missing some Epic-specific functionality required to build/run correctly.

diff --git a/Engine/Source/Programs/Horde/HordeServer/Notifications/Impl/SlackNotificationSink.cs b/Engine/Source/Programs/Horde/HordeServer/Notifications/Impl/SlackNotificationSink.cs
new file mode 100644
--- /dev/null
+++ b/Engine/Source/Programs/Horde/HordeServer/Notifications/Impl/SlackNotificationSink.cs
@@ -0,0 +1,1138 @@
+// Copyright Epic Games, Inc. All Rights Reserved.
+
+using EpicGames.Core;
+using HordeCommon;
+using HordeServer.Api;
+using HordeServer.Collections;
+using HordeServer.Models;
+using HordeServer.Services;
+using HordeServer.Utilities;
+using HordeServer.Utilities.BlockKit;
+using HordeServer.Utilities.Slack.BlockKit;
+using Microsoft.Extensions.Hosting;
+using Microsoft.Extensions.Logging;
+using Microsoft.Extensions.Options;
+using MongoDB.Bson;
+using MongoDB.Bson.Serialization.Attributes;
+using MongoDB.Driver;
+using Newtonsoft.Json.Linq;
+using System;
+using System.Collections.Concurrent;
+using System.Collections.Generic;
+using System.Diagnostics;
+using System.Drawing;
+using System.Globalization;
+using System.Linq;
+using System.Net.Http;
+using System.Net.WebSockets;
+using System.Security.Claims;
+using System.Text;
+using System.Text.Json;
+using System.Text.Json.Serialization;
+using System.Text.RegularExpressions;
+using System.Threading;
+using System.Threading.Tasks;
+
+using StreamId = HordeServer.Utilities.StringId<HordeServer.Models.IStream>;
+
+namespace HordeServer.Notifications.Impl
+{
+	/// <summary>
+	/// Maintains a connection to Slack, in order to receive socket-mode notifications of user interactions
+	/// </summary>
+	public class SlackNotificationSink : BackgroundService, INotificationSink
+	{
+		const string PostMessageUrl = "https://slack.com/api/chat.postMessage";
+		const string UpdateMessageUrl = "https://slack.com/api/chat.update";
+
+		class SocketResponse
+		{
+			[JsonPropertyName("ok")]
+			public bool Ok { get; set; }
+
+			[JsonPropertyName("url")]
+			public Uri? Url { get; set; }
+		}
+
+		class EventMessage
+		{
+			[JsonPropertyName("envelope_id")]
+			public string? EnvelopeId { get; set; }
+
+			[JsonPropertyName("type")]
+			public string? Type { get; set; }
+
+			[JsonPropertyName("payload")]
+			public EventPayload? Payload { get; set; }
+		}
+
+		class EventPayload
+		{
+			[JsonPropertyName("type")]
+			public string? Type { get; set; }
+
+			[JsonPropertyName("user")]
+			public UserInfo? User { get; set; }
+
+			[JsonPropertyName("response_url")]
+			public string? ResponseUrl { get; set; }
+
+			[JsonPropertyName("actions")]
+			public List<ActionInfo> Actions { get; set; } = new List<ActionInfo>();
+		}
+
+		class UserInfo
+		{
+			[JsonPropertyName("username")]
+			public string? UserName { get; set; }
+		}
+
+		class ActionInfo
+		{
+			[JsonPropertyName("value")]
+			public string? Value { get; set; }
+		}
+
+		class MessageStateDocument
+		{
+			[BsonId]
+			public ObjectId Id { get; set; }
+
+			[BsonElement("uid")]
+			public string Recipient { get; set; } = String.Empty;
+
+			[BsonElement("usr")]
+			public ObjectId? UserId { get; set; }
+
+			[BsonElement("eid")]
+			public string EventId { get; set; } = String.Empty;
+
+			[BsonElement("ch")]
+			public string Channel { get; set; } = String.Empty;
+
+			[BsonElement("ts")]
+			public string Ts { get; set; } = String.Empty;
+
+			[BsonElement("dig")]
+			public string Digest { get; set; } = String.Empty;
+		}
+
+		IIssueService IssueService;
+		IUserCollection UserCollection;
+		ILogFileService LogFileService;
+		StreamService StreamService;
+		ServerSettings Settings;
+		IMongoCollection<MessageStateDocument> MessageStates;
+		ILogger Logger;
+
+		/// <summary>
+		/// Map of email address to Slack user ID.
+		/// </summary>
+		private readonly ConcurrentDictionary<string, string> EmailAddressToSlackResponse = new ConcurrentDictionary<string, string>();
+
+		/// <summary>
+		/// Constructor
+		/// </summary>
+		/// <param name="DatabaseService"></param>
+		/// <param name="IssueService"></param>
+		/// <param name="UserCollection">The user collection</param>
+		/// <param name="LogFileService"></param>
+		/// <param name="StreamService"></param>
+		/// <param name="Settings">The current configuration settings</param>
+		/// <param name="Logger">Logging device</param>
+		public SlackNotificationSink(DatabaseService DatabaseService, IIssueService IssueService, IUserCollection UserCollection, ILogFileService LogFileService, StreamService StreamService, IOptions<ServerSettings> Settings, ILogger<SlackNotificationSink> Logger)
+		{
+			this.IssueService = IssueService;
+			this.UserCollection = UserCollection;
+			this.LogFileService = LogFileService;
+			this.StreamService = StreamService;
+			this.Settings = Settings.Value;
+			this.MessageStates = DatabaseService.Database.GetCollection<MessageStateDocument>("Slack");
+			this.Logger = Logger;
+		}
+
+		#region Message state 
+
+		async Task<MessageStateDocument?> GetMessageStateAsync(string Recipient, string EventId)
+		{
+			FilterDefinition<MessageStateDocument> Filter = Builders<MessageStateDocument>.Filter.Eq(x => x.Recipient, Recipient) & Builders<MessageStateDocument>.Filter.Eq(x => x.EventId, EventId);
+			return await MessageStates.Find(Filter).FirstOrDefaultAsync();
+		}
+
+		async Task<MessageStateDocument> SetMessageStateAsync(string Recipient, string EventId, ObjectId? UserId, string Channel, string Ts, string Digest)
+		{
+			FilterDefinition<MessageStateDocument> Filter = Builders<MessageStateDocument>.Filter.Eq(x => x.Recipient, Recipient) & Builders<MessageStateDocument>.Filter.Eq(x => x.EventId, EventId);
+			UpdateDefinition<MessageStateDocument> Update = Builders<MessageStateDocument>.Update.SetOnInsert(x => x.Id, ObjectId.GenerateNewId()).Set(x => x.UserId, UserId).Set(x => x.Channel, Channel).Set(x => x.Ts, Ts).Set(x => x.Digest, Digest);
+			return await MessageStates.FindOneAndUpdateAsync(Filter, Update, new FindOneAndUpdateOptions<MessageStateDocument> { IsUpsert = true, ReturnDocument = ReturnDocument.After });
+		}
+
+		#endregion
+
+		#region Job Complete
+
+		/// <inheritdoc/>
+		public async Task NotifyJobCompleteAsync(IStream JobStream, IJob Job, IGraph Graph, LabelOutcome Outcome)
+		{
+			if (Job.NotificationChannel != null)
+			{
+				await SendJobCompleteNotificationToChannelAsync(Job.NotificationChannel, Job.NotificationChannelFilter, JobStream, Job, Graph, Outcome);
+			}
+			if (JobStream.NotificationChannel != null)
+			{
+				await SendJobCompleteNotificationToChannelAsync(JobStream.NotificationChannel, JobStream.NotificationChannelFilter, JobStream, Job, Graph, Outcome);
+			}
+		}
+
+		Task SendJobCompleteNotificationToChannelAsync(string NotificationChannel, string? NotificationFilter, IStream JobStream, IJob Job, IGraph Graph, LabelOutcome Outcome)
+		{
+			if (NotificationFilter != null)
+			{
+				List<LabelOutcome> Outcomes = new List<LabelOutcome>();
+				foreach (string FilterOption in NotificationFilter.Split('|'))
+				{
+					LabelOutcome Result;
+					if (Enum.TryParse(FilterOption, out Result))
+					{
+						Outcomes.Add(Result);
+					}
+					else
+					{
+						Logger.LogWarning("Invalid filter option {Option} specified in job filter {NotificationChannelFilter} in job {JobId} or stream {StreamId}", FilterOption, NotificationFilter, Job.Id, Job.StreamId);
+					}
+				}
+				if (!Outcomes.Contains(Outcome))
+				{
+					return Task.CompletedTask;
+				}
+			}
+			return SendJobCompleteMessageAsync(NotificationChannel, JobStream, Job, Graph);
+		}
+
+		/// <inheritdoc/>
+		public async Task NotifyJobCompleteAsync(IUser SlackUser, IStream JobStream, IJob Job, IGraph Graph, LabelOutcome Outcome)
+		{
+			string? SlackUserId = await GetSlackUserId(SlackUser);
+			if (SlackUserId != null)
+			{
+				await SendJobCompleteMessageAsync(SlackUserId, JobStream, Job, Graph);
+			}
+		}
+
+		private Task SendJobCompleteMessageAsync(string Recipient, IStream Stream, IJob Job, IGraph Graph)
+		{
+			JobStepOutcome JobOutcome = Job.Batches.SelectMany(x => x.Steps).Min(x => x.Outcome);
+			Logger.LogInformation("Sending Slack notification for job {JobId} outcome {Outcome} to {User}", Job.Id, JobOutcome, Recipient);
+
+			Uri JobLink = new Uri($"{Settings.DashboardUrl}job/{Job.Id}");
+
+			Color OutcomeColor = JobOutcome == JobStepOutcome.Failure ? BlockKitAttachmentColors.Error : JobOutcome == JobStepOutcome.Warnings ? BlockKitAttachmentColors.Warning : BlockKitAttachmentColors.Success;
+
+			BlockKitAttachment Attachment = new BlockKitAttachment();
+			Attachment.FallbackText = $"{Stream.Name} - {GetJobChangeText(Job)} - {Job.Name} - {JobOutcome}";
+			Attachment.Color = OutcomeColor;
+			Attachment.Blocks.Add(new SectionBlock($"*<{JobLink}|{Stream.Name} - {GetJobChangeText(Job)} - {Job.Name}>*"));
+			if (JobOutcome == JobStepOutcome.Success)
+			{
+				Attachment.Blocks.Add(new SectionBlock($"*Job Succeeded*"));
+			}
+			else
+			{
+				List<string> FailedStepStrings = new List<string>();
+				List<string> WarningStepStrings = new List<string>();
+				foreach (IJobStepBatch Batch in Job.Batches)
+				{
+					foreach (IJobStep Step in Batch.Steps)
+					{
+						INode StepNode = Graph.GetNode(new NodeRef(Batch.GroupIdx, Step.NodeIdx));
+						string StepName = $"<{JobLink}?step={Step.Id}|{StepNode.Name}>";
+						if (Step.Outcome == JobStepOutcome.Failure)
+						{
+							FailedStepStrings.Add(StepName);
+						}
+						else if (Step.Outcome == JobStepOutcome.Warnings)
+						{
+							WarningStepStrings.Add(StepName);
+						}
+					}
+				}
+
+				if (FailedStepStrings.Any())
+				{
+					string Msg = $"*Errors*\n{string.Join(", ", FailedStepStrings)}";
+					Attachment.Blocks.Add(new SectionBlock(Msg.Substring(0, Math.Min(Msg.Length, 3000))));
+				}
+				else if (WarningStepStrings.Any())
+				{
+					string Msg = $"*Warnings*\n{string.Join(", ", WarningStepStrings)}";
+					Attachment.Blocks.Add(new SectionBlock(Msg.Substring(0, Math.Min(Msg.Length, 3000))));
+				}
+			}
+
+			if (Job.AutoSubmit)
+			{
+				Attachment.Blocks.Add(new DividerBlock());
+				if (Job.AutoSubmitChange != null)
+				{
+					Attachment.Blocks.Add(new SectionBlock($"Shelved files were submitted in CL {Job.AutoSubmitChange}."));
+				}
+				else
+				{
+					Attachment.Color = BlockKitAttachmentColors.Warning;
+					Attachment.Blocks.Add(new SectionBlock($"Files in CL {Job.PreflightChange} were *not submitted*. Please resolve the following issues and submit manually.\n\n```{Job.AutoSubmitMessage}```"));
+				}
+			}
+
+			return SendMessageAsync(Recipient, Attachments: new[] { Attachment });
+		}
+
+		#endregion
+
+		#region Job step complete
+
+		/// <inheritdoc/>
+		public async Task NotifyJobStepCompleteAsync(IUser SlackUser, IStream JobStream, IJob Job, IJobStepBatch Batch, IJobStep Step, INode Node, List<ILogEventData> JobStepEventData)
+		{
+			Logger.LogInformation("Sending Slack notification for job {JobId}, batch {BatchId}, step {StepId}, outcome {Outcome} to {User} ({UserId})", Job.Id, Batch.Id, Step.Id, Step.Outcome, SlackUser.Name, SlackUser.Id);
+
+			string? SlackUserId = await GetSlackUserId(SlackUser);
+			if (SlackUserId != null)
+			{
+				await SendJobStepCompleteMessageAsync(SlackUserId, JobStream, Job, Step, Node, JobStepEventData);
+			}
+		}
+
+		/// <summary>
+		/// Creates a Slack message about a completed step job.
+		/// </summary>
+		/// <param name="Recipient"></param>
+		/// <param name="Stream"></param>
+		/// <param name="Job">The job that contains the step that completed.</param>
+		/// <param name="Step">The job step that completed.</param>
+		/// <param name="Node">The node for the job step.</param>
+		/// <param name="Events">Any events that occurred during the job step.</param>
+		private Task SendJobStepCompleteMessageAsync(string Recipient, IStream Stream, IJob Job, IJobStep Step, INode Node, List<ILogEventData> Events)
+		{
+			Uri JobStepLink = new Uri($"{Settings.DashboardUrl}job/{Job.Id}?step={Step.Id}");
+			Uri JobStepLogLink = new Uri($"{Settings.DashboardUrl}log/{Step.LogId}");
+
+			Color OutcomeColor = Step.Outcome == JobStepOutcome.Failure ? BlockKitAttachmentColors.Error : Step.Outcome == JobStepOutcome.Warnings ? BlockKitAttachmentColors.Warning : BlockKitAttachmentColors.Success;
+
+			BlockKitAttachment Attachment = new BlockKitAttachment();
+			Attachment.FallbackText = $"{Stream.Name} - {GetJobChangeText(Job)} - {Job.Name} - {Node.Name} - {Step.Outcome}";
+			Attachment.Color = OutcomeColor;
+			Attachment.Blocks.Add(new SectionBlock($"*<{JobStepLink}|{Stream.Name} - {GetJobChangeText(Job)} - {Job.Name} - {Node.Name}>*"));
+			if (Step.Outcome == JobStepOutcome.Success)
+			{
+				Attachment.Blocks.Add(new SectionBlock($"*Job Step Succeeded*"));
+			}
+			else
+			{
+				List<ILogEventData> Errors = Events.Where(x => x.Severity == EventSeverity.Error).ToList();
+				List<ILogEventData> Warnings = Events.Where(x => x.Severity == EventSeverity.Warning).ToList();
+				List<string> EventStrings = new List<string>();
+				if (Errors.Any())
+				{
+					string ErrorSummary = Errors.Count > MaxJobStepEvents ? $"*Errors (First {MaxJobStepEvents} shown)*" : $"*Errors*";
+					EventStrings.Add(ErrorSummary);
+					foreach (ILogEventData Error in Errors.Take(MaxJobStepEvents))
+					{
+						EventStrings.Add($"```{Error.Message}```");
+					}
+				}
+				else if (Warnings.Any())
+				{
+					string WarningSummary = Warnings.Count > MaxJobStepEvents ? $"*Warnings (First {MaxJobStepEvents} shown)*" : $"*Warnings*";
+					EventStrings.Add(WarningSummary);
+					foreach (ILogEventData Warning in Warnings.Take(MaxJobStepEvents))
+					{
+						EventStrings.Add($"```{Warning.Message}```");
+					}
+				}
+
+				Attachment.Blocks.Add(new SectionBlock(string.Join("\n", EventStrings)));
+				Attachment.Blocks.Add(new SectionBlock($"<{JobStepLogLink}|View Job Step Log>"));
+			}
+
+			return SendMessageAsync(Recipient, Attachments: new[] { Attachment });
+		}
+
+		#endregion
+
+		#region Label complete
+
+		/// <inheritdoc/>
+		public async Task NotifyLabelCompleteAsync(IUser User, IJob Job, IStream Stream, ILabel Label, int LabelIdx, LabelOutcome Outcome, List<(string, JobStepOutcome, Uri)> StepData)
+		{
+			Logger.LogInformation("Sending Slack notification for job {JobId} outcome {Outcome} to {Name} ({UserId})", Job.Id, Outcome, User.Name, User.Id);
+
+			string? SlackUserId = await GetSlackUserId(User);
+			if (SlackUserId != null)
+			{
+				await SendLabelUpdateMessageAsync(SlackUserId, Stream, Job, Label, LabelIdx, Outcome, StepData);
+			}
+		}
+
+		Task SendLabelUpdateMessageAsync(string Recipient, IStream Stream, IJob Job, ILabel Label, int LabelIdx, LabelOutcome Outcome, List<(string, JobStepOutcome, Uri)> JobStepData)
+		{
+			Uri LabelLink = new Uri($"{Settings.DashboardUrl}job/{Job.Id}?label={LabelIdx}");
+
+			Color OutcomeColor = Outcome == LabelOutcome.Failure ? BlockKitAttachmentColors.Error : Outcome == LabelOutcome.Warnings ? BlockKitAttachmentColors.Warning : BlockKitAttachmentColors.Success;
+
+			BlockKitAttachment Attachment = new BlockKitAttachment();
+			Attachment.FallbackText = $"{Stream.Name} - {GetJobChangeText(Job)} - {Job.Name} - Label {Label.DashboardName} - {Outcome}";
+			Attachment.Color = OutcomeColor;
+			Attachment.Blocks.Add(new SectionBlock($"*<{LabelLink}|{Stream.Name} - {GetJobChangeText(Job)} - {Job.Name} - Label {Label.DashboardName}>*"));
+			if (Outcome == LabelOutcome.Success)
+			{
+				Attachment.Blocks.Add(new SectionBlock($"*Label Succeeded*"));
+			}
+			else
+			{
+				List<string> FailedStepStrings = new List<string>();
+				List<string> WarningStepStrings = new List<string>();
+				foreach ((string Name, JobStepOutcome StepOutcome, Uri Link) JobStep in JobStepData)
+				{
+					string StepString = $"<{JobStep.Link}|{JobStep.Name}>";
+					if (JobStep.StepOutcome == JobStepOutcome.Failure)
+					{
+						FailedStepStrings.Add(StepString);
+					}
+					else if (JobStep.StepOutcome == JobStepOutcome.Warnings)
+					{
+						WarningStepStrings.Add(StepString);
+					}
+				}
+
+				if (FailedStepStrings.Any())
+				{
+					Attachment.Blocks.Add(new SectionBlock($"*Errors*\n{string.Join(", ", FailedStepStrings)}"));
+				}
+				else if (WarningStepStrings.Any())
+				{
+					Attachment.Blocks.Add(new SectionBlock($"*Warnings*\n{string.Join(", ", WarningStepStrings)}"));
+				}
+			}
+
+			return SendMessageAsync(Recipient, Attachments: new[] { Attachment });
+		}
+
+		#endregion
+
+		#region Issues
+
+		/// <inheritdoc/>
+		public async Task NotifyIssueUpdatedAsync(IIssue Issue)
+		{
+			HashSet<ObjectId> UserIds = new HashSet<ObjectId>();
+			if (Issue.NotifySuspects)
+			{
+				UserIds.UnionWith(Issue.Suspects.Select(x => x.AuthorId));
+			}
+			if (Issue.OwnerId.HasValue)
+			{
+				UserIds.Add(Issue.OwnerId.Value);
+			}
+
+			HashSet<string> Channels = new HashSet<string>(StringComparer.OrdinalIgnoreCase);
+
+			List<MessageStateDocument> ExistingMessages = await MessageStates.Find(x => x.EventId == GetIssueEventId(Issue)).ToListAsync();
+			foreach (MessageStateDocument ExistingMessage in ExistingMessages)
+			{
+				if (ExistingMessage.UserId != null)
+				{
+					UserIds.Add(ExistingMessage.UserId.Value);
+				}
+				else
+				{
+					Channels.Add(ExistingMessage.Recipient);
+				}
+			}
+
+			IIssueDetails Details = await IssueService.GetIssueDetailsAsync(Issue);
+
+			if (Issue.OwnerId == null && (Issue.Suspects.Count == 0 || Issue.Suspects.All(x => x.DeclinedAt != null) || Issue.CreatedAt < DateTime.UtcNow - TimeSpan.FromHours(1.0)))
+			{
+				foreach (IIssueSpan Span in Details.Spans)
+				{
+					IStream? Stream = await StreamService.GetCachedStream(Span.StreamId);
+					if (Stream != null)
+					{
+						TemplateRef? TemplateRef;
+						if (Stream.Templates.TryGetValue(Span.TemplateRefId, out TemplateRef) && TemplateRef.TriageChannel != null)
+						{
+							Channels.Add(TemplateRef.TriageChannel);
+						}
+						else if (Stream.TriageChannel != null)
+						{
+							Channels.Add(Stream.TriageChannel);
+						}
+					}
+				}
+			}
+
+			if (UserIds.Count > 0)
+			{
+				foreach (ObjectId UserId in UserIds)
+				{
+					IUser? User = await UserCollection.GetUserAsync(UserId);
+					if (User == null)
+					{
+						Logger.LogWarning("Unable to find user {UserId}", UserId);
+					}
+					else
+					{
+						await NotifyIssueUpdatedAsync(User, Issue, Details);
+					}
+				}
+			}
+
+			if (Channels.Count > 0)
+			{
+				foreach (string Channel in Channels)
+				{
+					await SendIssueMessageAsync(Channel, Issue, Details, null);
+				}
+			}
+		}
+
+		async Task NotifyIssueUpdatedAsync(IUser User, IIssue Issue, IIssueDetails Details)
+		{
+			IUserSettings UserSettings = await UserCollection.GetSettingsAsync(User.Id);
+			if (!UserSettings.EnableIssueNotifications)
+			{
+				Logger.LogInformation("Issue notifications are disabled for user {UserId} ({UserName})", User.Id, User.Name);
+				return;
+			}
+
+			string? SlackUserId = await GetSlackUserId(User);
+			if (SlackUserId == null)
+			{
+				return;
+			}
+
+			await SendIssueMessageAsync(SlackUserId, Issue, Details, User.Id);
+		}
+
+		async Task<string> GetSampleText(IIssueDetails Details)
+		{
+			StringBuilder SampleText = new StringBuilder();
+			if (Details.Spans.Count > 0)
+			{
+				IIssueSpan Span = Details.Spans[0];
+				if (Span.FirstFailure.LogId != null)
+				{
+					ILogFile? LogFile = await LogFileService.GetCachedLogFileAsync(Span.FirstFailure.LogId.Value);
+					if (LogFile != null)
+					{
+						List<ILogEvent> Events = await LogFileService.FindEventsForSpansAsync(Details.Spans.Select(x => x.Id).Take(1), null, 0, 1);
+						foreach (ILogEvent Event in Events)
+						{
+							ILogEventData Data = await LogFileService.GetEventDataAsync(LogFile, Event.LineIndex, Event.LineCount);
+							if (SampleText.Length > 0)
+							{
+								SampleText.Append("\n\n");
+							}
+							SampleText.Append($"```{Data.Message}```");
+						}
+					}
+				}
+			}
+			return SampleText.ToString();
+		}
+
+		async Task SendIssueMessageAsync(string Recipient, IIssue Issue, IIssueDetails Details, ObjectId? UserId)
+		{
+			BlockKitAttachment Attachment = new BlockKitAttachment();
+			Attachment.Color = Color.Red;
+
+			Uri IssueUrl = Settings.DashboardUrl;
+			if (Details.Steps.Count > 0)
+			{
+				IIssueStep FirstStep = Details.Steps[0];
+				IssueUrl = new Uri(IssueUrl, $"job/{FirstStep.JobId}?step={FirstStep.StepId}&issue={Issue.Id}");
+			}
+			Attachment.Blocks.Add(new SectionBlock($"*<{IssueUrl}|Issue #{Issue.Id}: {Issue.Summary}>*"));
+
+			string StreamList = StringUtils.FormatList(Details.Spans.Select(x => $"*{x.StreamName}*").Distinct(StringComparer.OrdinalIgnoreCase).OrderBy(x => x, StringComparer.OrdinalIgnoreCase));
+			Attachment.Blocks.Add(new SectionBlock($"Occurring in {StreamList}"));
+
+			if (Issue.FixChange != null)
+			{
+				string Text;
+				if (Issue.FixChange.Value < 0)
+				{
+					Text = ":tick: Marked as a systemic issue.";
+				}
+				else
+				{
+					Text = $":tick: Marked fixed in *CL {Issue.FixChange.Value}*.";
+				}
+				Attachment.Blocks.Add(new SectionBlock(Text));
+			}
+			else if (UserId != null && Issue.OwnerId == UserId)
+			{
+				if (Issue.AcknowledgedAt.HasValue)
+				{
+					Attachment.Blocks.Add(new SectionBlock($":+1: Acknowledged at {FormatSlackTime(Issue.AcknowledgedAt.Value)}"));
+				}
+				else
+				{
+					if (Issue.NominatedById != null)
+					{
+						IUser? NominatedByUser = await UserCollection.GetUserAsync(Issue.NominatedById.Value);
+						if (NominatedByUser != null)
+						{
+							string? NominatedBySlackUserId = await GetSlackUserId(NominatedByUser);
+							string Mention = (NominatedBySlackUserId != null) ? $"<@{NominatedBySlackUserId}>" : NominatedByUser.Login ?? $"User {NominatedByUser.Id}";
+							string Text = $"You were nominated to fix this issue by {Mention} at {FormatSlackTime(Issue.NominatedAt ?? DateTime.UtcNow)}";
+							Attachment.Blocks.Add(new SectionBlock(Text));
+						}
+					}
+					else
+					{
+						List<int> Changes = Issue.Suspects.Where(x => x.AuthorId == UserId).Select(x => x.Change).OrderBy(x => x).ToList();
+						if (Changes.Count > 0)
+						{
+							string Text = $"Horde has determined that {StringUtils.FormatList(Changes.Select(x => $"CL {x}"), "or")} is the most likely cause for this issue.";
+							Attachment.Blocks.Add(new SectionBlock(Text));
+						}
+					}
+
+					ActionsBlock Actions = new ActionsBlock();
+					Actions.AddButton("Acknowledge", Value: $"issue_{Issue.Id}_ack_{UserId}", Style: ActionButton.ButtonStyle.Primary);
+					Actions.AddButton("Not Me", Value: $"issue_{Issue.Id}_decline_{UserId}", Style: ActionButton.ButtonStyle.Danger);
+					Attachment.Blocks.Add(Actions);
+				}
+			}
+			else if (Issue.OwnerId != null)
+			{
+				string OwnerMention = await FormatMentionAsync(Issue.OwnerId.Value);
+				if (Issue.AcknowledgedAt.HasValue)
+				{
+					Attachment.Blocks.Add(new SectionBlock($":+1: Acknowledged by {OwnerMention} at {FormatSlackTime(Issue.AcknowledgedAt.Value)}"));
+				}
+				else if (Issue.NominatedById == null)
+				{
+					Attachment.Blocks.Add(new SectionBlock($"Assigned to {OwnerMention}"));
+				}
+				else if (Issue.NominatedById == UserId)
+				{
+					Attachment.Blocks.Add(new SectionBlock($"You nominated {OwnerMention} to fix this issue."));
+				}
+				else
+				{
+					Attachment.Blocks.Add(new SectionBlock($"{OwnerMention} was nominated to fix this issue by {await FormatMentionAsync(Issue.NominatedById.Value)}"));
+				}
+			}
+			else if (UserId != null)
+			{
+				IIssueSuspect? Suspect = Issue.Suspects.FirstOrDefault(x => x.AuthorId == UserId);
+				if (Suspect != null)
+				{
+					if (Suspect.DeclinedAt != null)
+					{
+						Attachment.Blocks.Add(new SectionBlock($":downvote: Declined at {FormatSlackTime(Suspect.DeclinedAt.Value)}"));
+					}
+					else
+					{
+						Attachment.Blocks.Add(new SectionBlock("Please check if any of your recently submitted changes have caused this issue."));
+
+						ActionsBlock Actions = new ActionsBlock();
+						Actions.AddButton("Will Fix", Value: $"issue_{Issue.Id}_accept_{UserId}", Style: ActionButton.ButtonStyle.Primary);
+						Actions.AddButton("Not Me", Value: $"issue_{Issue.Id}_decline_{UserId}", Style: ActionButton.ButtonStyle.Danger);
+						Attachment.Blocks.Add(Actions);
+					}
+				}
+			}
+			else if (Issue.Suspects.Count > 0)
+			{
+				List<string> DeclinedLines = new List<string>();
+				foreach (IIssueSuspect Suspect in Issue.Suspects)
+				{
+					if (Suspect.DeclinedAt == null)
+					{
+						DeclinedLines.Add($":heavy_minus_sign: Ignored by {await FormatMentionAsync(Suspect.AuthorId)}");
+					}
+					else
+					{
+						DeclinedLines.Add($":downvote: Declined by {await FormatMentionAsync(Suspect.AuthorId)} at {FormatSlackTime(Suspect.DeclinedAt.Value)}");
+					}
+				}
+				Attachment.Blocks.Add(new SectionBlock(String.Join("\n", DeclinedLines)));
+			}
+
+			await SendOrUpdateMessageAsync(Recipient, GetIssueEventId(Issue), UserId, Attachments: new[] { Attachment });
+		}
+
+		static string GetIssueEventId(IIssue Issue)
+		{
+			return $"issue_{Issue.Id}";
+		}
+
+		async Task<string> FormatMentionAsync(ObjectId UserId)
+		{
+			IUser? User = await UserCollection.GetUserAsync(UserId);
+			if (User == null)
+			{
+				return $"User {UserId}";
+			}
+
+			string? SlackUserId = await GetSlackUserId(User);
+			if (SlackUserId == null)
+			{
+				return User.Login;
+			}
+
+			return $"<@{SlackUserId}>";
+		}
+
+		async Task HandleIssueResponseAsync(EventPayload Payload, ActionInfo Action)
+		{
+			string? UserName = Payload.User?.UserName;
+			if(UserName == null)
+			{
+				Logger.LogWarning("No user for message payload: {Payload}", Payload);
+				return;
+			}
+			if (Payload.ResponseUrl == null)
+			{
+				Logger.LogWarning("No response url for payload: {Payload}", Payload);
+				return;
+			}
+
+			Match Match = Regex.Match(Action.Value, @"^issue_(\d+)_([a-zA-Z]+)_([a-fA-F0-9]{24})$");
+			if (!Match.Success)
+			{
+				Logger.LogWarning("Could not match format of button action: {Action}", Action.Value);
+				return;
+			}
+
+			int IssueId = int.Parse(Match.Groups[1].Value, NumberStyles.Integer, CultureInfo.InvariantCulture);
+			string Verb = Match.Groups[2].Value;
+			ObjectId UserId = Match.Groups[3].Value.ToObjectId();
+			Logger.LogInformation("Issue {IssueId}: {Action} from {User} ({UserId})", IssueId, Verb, UserName, UserId);
+
+			if (String.Equals(Verb, "ack", StringComparison.Ordinal))
+			{
+				await IssueService.UpdateIssueAsync(IssueId, Acknowledged: true);
+			}
+			else if (String.Equals(Verb, "accept", StringComparison.Ordinal))
+			{
+				await IssueService.UpdateIssueAsync(IssueId, Owner: UserName, Acknowledged: true);
+			}
+			else if (String.Equals(Verb, "decline", StringComparison.Ordinal))
+			{
+				await IssueService.UpdateIssueAsync(IssueId, DeclinedBy: UserName);
+			}
+
+			IIssue? NewIssue = await IssueService.GetIssueAsync(IssueId);
+			if (NewIssue != null)
+			{
+				IUser? User = await UserCollection.GetUserAsync(UserId);
+				if (User != null)
+				{
+					string? Recipient = await GetSlackUserId(User);
+					if (Recipient != null)
+					{
+						IIssueDetails Details = await IssueService.GetIssueDetailsAsync(NewIssue);
+						await SendIssueMessageAsync(Recipient, NewIssue, Details, UserId);
+					}
+				}
+			}
+		}
+
+		static string FormatSlackTime(DateTimeOffset Time)
+		{
+			return $"<!date^{Time.ToUnixTimeSeconds()}^{{time}}|{Time}>";
+		}
+
+		#endregion
+
+		#region Stream updates
+
+		/// <inheritdoc/>
+		public async Task NotifyStreamUpdateFailedAsync(IStream Stream, string ErrorMessage, int Change, IUser? Author = null, string? Description = null)
+		{
+			Logger.LogInformation("Sending stream update failure notification for {Stream} ({StreamId})", Stream.Name, Stream.Id);
+
+			string? SlackUserId = null;
+			if (Author != null)
+			{
+				SlackUserId = await GetSlackUserId(Author);
+			}
+
+			if (SlackUserId != null)
+			{
+				await SendStreamUpdateFailureMessageAsync(SlackUserId, Stream, ErrorMessage, Change, SlackUserId, Description);
+			}
+			if (Settings.UpdateStreamsNotificationChannel != null)
+			{
+				await SendStreamUpdateFailureMessageAsync($"#{Settings.UpdateStreamsNotificationChannel}", Stream, ErrorMessage, Change, SlackUserId, Description);
+			}
+		}
+
+		private Task SendStreamUpdateFailureMessageAsync(string Recipient, IStream Stream, string ErrorMessage, int Change, string? Author = null, string? Description = null)
+		{
+			Color OutcomeColor = BlockKitAttachmentColors.Error;
+			BlockKitAttachment Attachment = new BlockKitAttachment();
+			Attachment.Color = OutcomeColor;
+			Attachment.FallbackText = $"{Stream.Name} ({Stream.Id}) - Update Failure";
+
+			Attachment.Blocks.Add(new HeaderBlock($"Stream Update Failure :rip:", false, true));
+
+			Attachment.Blocks.Add(new SectionBlock($"Horde was unable to update {Stream.Name} ({Stream.Id})"));
+			Attachment.Blocks.Add(new SectionBlock($"```{ErrorMessage}```"));
+			if (Author != null)
+			{
+				Attachment.Blocks.Add(new SectionBlock($"Possibly due to CL: {Change} by <@{Author}>"));
+			}
+			else
+			{
+				Attachment.Blocks.Add(new SectionBlock($"Possibly due to CL: {Change} - (Could not determine author from P4 user)"));
+			}
+			if (Description != null)
+			{
+				Attachment.Blocks.Add(new SectionBlock($"```{Description}```"));
+			}
+			else
+			{
+				Attachment.Blocks.Add(new SectionBlock($"```Could not determine description of change from P4```"));
+			}
+
+			return SendMessageAsync(Recipient, Attachments: new[] { Attachment });
+		}
+
+		#endregion
+
+		#region Stream update (file)
+
+		/// <inheritdoc/>
+		public async Task NotifyStreamUpdateFailedAsync(FileSummary File)
+		{
+			Logger.LogDebug("Sending stream update failure notification for {File}", File.DepotPath);
+			if (Settings.UpdateStreamsNotificationChannel != null)
+			{
+				await SendStreamUpdateFailureMessage($"#{Settings.UpdateStreamsNotificationChannel}", File);
+			}
+		}
+
+		/// <summary>
+		/// Creates a stream update failure message in relation to a file
+		/// </summary>
+		/// <param name="Recipient"></param>
+		/// <param name="File">The file</param>
+		/// <returns></returns>
+		Task SendStreamUpdateFailureMessage(string Recipient, FileSummary File)
+		{
+			Color OutcomeColor = BlockKitAttachmentColors.Error;
+			BlockKitAttachment Attachment = new BlockKitAttachment();
+			Attachment.Color = OutcomeColor;
+			Attachment.FallbackText = $"{File.DepotPath} - Update Failure";
+
+			Attachment.Blocks.Add(new HeaderBlock($"Stream Update Failure :rip:", false, true));
+
+			Attachment.Blocks.Add(new SectionBlock($"<!here> Horde was unable to update {File.DepotPath}"));
+			Attachment.Blocks.Add(new SectionBlock($"```{File.Error}```"));
+
+			return SendMessageAsync(Recipient, Attachments: new[] { Attachment });
+		}
+
+		#endregion
+
+		const int MaxJobStepEvents = 5;
+
+		static string GetJobChangeText(IJob Job)
+		{
+			if (Job.PreflightChange == 0)
+			{
+				return $"CL {Job.Change}";
+			}
+			else
+			{
+				return $"Preflight CL {Job.PreflightChange} against CL {Job.Change}";
+			}
+		}
+
+		private async Task<string?> GetSlackUserId(IUser User)
+		{
+			string? Email = User.Email;
+			if (Email == null)
+			{
+				Logger.LogWarning("Unable to find Slack user id for {UserId} ({Name}): No email address in user profile", User.Id, User.Name);
+				return null;
+			}
+
+			string? SlackResponse;
+			if (!EmailAddressToSlackResponse.TryGetValue(Email, out SlackResponse))
+			{
+				using HttpClient Client = new HttpClient();
+
+				using HttpRequestMessage GetUserIdRequest = new HttpRequestMessage(HttpMethod.Post, $"https://slack.com/api/users.lookupByEmail?email={Email}");
+				GetUserIdRequest.Headers.Add("Authorization", $"Bearer {Settings.SlackToken ?? ""}");
+
+				HttpResponseMessage Response = await Client.SendAsync(GetUserIdRequest);
+				SlackResponse = await Response.Content.ReadAsStringAsync();
+
+				EmailAddressToSlackResponse.TryAdd(Email, SlackResponse);
+			}
+
+			JObject ResponseObject = JObject.Parse(SlackResponse);
+			if (!ResponseObject["ok"].ToObject<bool>())
+			{
+				Logger.LogWarning("Unable to find Slack user id for {UserId} ({Name}): {Response}", User.Id, User.Name, SlackResponse);
+				return null;
+			}
+
+			return ResponseObject["user"]["id"].ToString();
+		}
+
+		class SlackResponse
+		{
+			[JsonPropertyName("ok")]
+			public bool Ok { get; set; }
+
+			[JsonPropertyName("error")]
+			public string? Error { get; set; }
+		}
+
+		class PostMessageResponse : SlackResponse
+		{
+			[JsonPropertyName("channel")]
+			public string? Channel { get; set; }
+
+			[JsonPropertyName("ts")]
+			public string? Ts { get; set; }
+		}
+
+		private async Task SendMessageAsync(string Recipient, string? Text = null, BlockBase[]? Blocks = null, BlockKitAttachment[]? Attachments = null)
+		{
+			BlockKitMessage Message = new BlockKitMessage();
+			Message.Recipient = Recipient;
+			Message.Text = Text;
+			if (Blocks != null)
+			{
+				Message.Blocks.AddRange(Blocks);
+			}
+			if (Attachments != null)
+			{
+				Message.Attachments.AddRange(Attachments);
+			}
+
+			await SendRequestAsync<PostMessageResponse>(PostMessageUrl, Message);
+		}
+
+		private async Task SendOrUpdateMessageAsync(string Recipient, string EventId, ObjectId? UserId, string? Text = null, BlockBase[]? Blocks = null, BlockKitAttachment[]? Attachments = null)
+		{
+			BlockKitMessage Message = new BlockKitMessage();
+			Message.Text = Text;
+			if (Blocks != null)
+			{
+				Message.Blocks.AddRange(Blocks);
+			}
+			if (Attachments != null)
+			{
+				Message.Attachments.AddRange(Attachments);
+			}
+
+			string RequestDigest = ContentHash.MD5(JsonSerializer.Serialize(Message, new JsonSerializerOptions { IgnoreNullValues = true })).ToString();
+
+			PostMessageResponse Response;
+
+			MessageStateDocument? CurrentState = await GetMessageStateAsync(Recipient, EventId);
+			if (CurrentState == null)
+			{
+				Logger.LogInformation("Sending new slack message to {Recipient}", Recipient);
+				Message.Recipient = Recipient;
+				Message.Ts = null;
+				Response = await SendRequestAsync<PostMessageResponse>(PostMessageUrl, Message);
+			}
+			else if (RequestDigest == CurrentState.Digest)
+			{
+				return;
+			}
+			else
+			{
+				Logger.LogInformation("Updating existing slack message for user {Recipient} ({Channel}, (MessageTs))", Recipient, CurrentState.Channel, CurrentState.Ts);
+				Message.Recipient = CurrentState.Channel;
+				Message.Ts = CurrentState.Ts;
+				Response = await SendRequestAsync<PostMessageResponse>(UpdateMessageUrl, Message);
+			}
+
+			if (Response.Ok && Response.Channel != null && Response.Ts != null)
+			{
+				await SetMessageStateAsync(Recipient, EventId, UserId, Response.Channel, Response.Ts, RequestDigest);
+			}
+		}
+
+		private async Task<TResponse> SendRequestAsync<TResponse>(string RequestUrl, object Request) where TResponse : SlackResponse
+		{
+			using (HttpClient Client = new HttpClient())
+			{
+				using (HttpRequestMessage SendMessageRequest = new HttpRequestMessage(HttpMethod.Post, RequestUrl))
+				{
+					string RequestJson = JsonSerializer.Serialize(Request, new JsonSerializerOptions { IgnoreNullValues = true });
+					using (StringContent MessageContent = new StringContent(RequestJson, Encoding.UTF8, "application/json"))
+					{
+						SendMessageRequest.Content = MessageContent;
+						SendMessageRequest.Headers.Add("Authorization", $"Bearer {Settings.SlackToken ?? ""}");
+
+						HttpResponseMessage Response = await Client.SendAsync(SendMessageRequest);
+						byte[] ResponseBytes = await Response.Content.ReadAsByteArrayAsync();
+
+						TResponse ResponseObject = JsonSerializer.Deserialize<TResponse>(ResponseBytes);
+						if(!ResponseObject.Ok)
+						{
+							Logger.LogError("Failed to send Slack message ({Error}). Request: {Request}. Response: {Response}", ResponseObject.Error, RequestJson, Encoding.UTF8.GetString(ResponseBytes));
+						}
+
+						return ResponseObject;
+					}
+				}
+			}
+		}
+
+		/// <inheritdoc/>
+		protected async override Task ExecuteAsync(CancellationToken StoppingToken)
+		{
+			if (Settings.SlackSocketToken != null)
+			{
+				while (!StoppingToken.IsCancellationRequested)
+				{
+					try
+					{
+						Uri? WebSocketUrl = await GetWebSocketUrlAsync(StoppingToken);
+						if (WebSocketUrl == null)
+						{
+							await Task.Delay(TimeSpan.FromSeconds(5.0), StoppingToken);
+							continue;
+						}
+						await HandleSocketAsync(WebSocketUrl, StoppingToken);
+					}
+					catch (OperationCanceledException)
+					{
+						break;
+					}
+					catch (Exception Ex)
+					{
+						Logger.LogError(Ex, "Exception while updating Slack socket");
+						await Task.Delay(TimeSpan.FromSeconds(5.0), StoppingToken);
+					}
+				}
+			}
+		}
+
+		/// <summary>
+		/// Get the url for opening a websocket to Slack
+		/// </summary>
+		/// <param name="StoppingToken"></param>
+		/// <returns></returns>
+		private async Task<Uri?> GetWebSocketUrlAsync(CancellationToken StoppingToken)
+		{
+			using HttpClient Client = new HttpClient();
+			Client.DefaultRequestHeaders.Add("Authorization", $"Bearer {Settings.SlackSocketToken}");
+
+			using FormUrlEncodedContent Content = new FormUrlEncodedContent(Array.Empty<KeyValuePair<string, string>>());
+			HttpResponseMessage Response = await Client.PostAsync(new Uri("https://slack.com/api/apps.connections.open"), Content, StoppingToken);
+
+			byte[] ResponseData = await Response.Content.ReadAsByteArrayAsync();
+
+			SocketResponse SocketResponse = JsonSerializer.Deserialize<SocketResponse>(ResponseData);
+			if (!SocketResponse.Ok)
+			{
+				Logger.LogWarning("Unable to get websocket url: {Response}", Encoding.UTF8.GetString(ResponseData));
+				return null;
+			}
+
+			return SocketResponse.Url;
+		}
+
+		/// <summary>
+		/// Handle the lifetime of a websocket connection
+		/// </summary>
+		/// <param name="SocketUrl"></param>
+		/// <param name="StoppingToken"></param>
+		/// <returns></returns>
+		private async Task HandleSocketAsync(Uri SocketUrl, CancellationToken StoppingToken)
+		{
+			using ClientWebSocket Socket = new ClientWebSocket();
+			await Socket.ConnectAsync(SocketUrl, StoppingToken);
+
+			byte[] Buffer = new byte[2048];
+			while (!StoppingToken.IsCancellationRequested)
+			{
+				// Read the next message
+				int Length = 0;
+				for (; ; )
+				{
+					if (Length == Buffer.Length)
+					{
+						Array.Resize(ref Buffer, Buffer.Length + 2048);
+					}
+
+					WebSocketReceiveResult Result = await Socket.ReceiveAsync(new ArraySegment<byte>(Buffer, Length, Buffer.Length - Length), StoppingToken);
+					if (Result.MessageType == WebSocketMessageType.Close)
+					{
+						return;
+					}
+					Length += Result.Count;
+
+					if (Result.EndOfMessage)
+					{
+						break;
+					}
+				}
+
+				// Get the message data
+				Logger.LogInformation("Slack event: {Message}", Encoding.UTF8.GetString(Buffer, 0, Length));
+				EventMessage EventMessage = JsonSerializer.Deserialize<EventMessage>(Buffer.AsSpan(0, Length));
+
+				// Acknowledge the message
+				if (EventMessage.EnvelopeId != null)
+				{
+					object Response = new { EventMessage.EnvelopeId };
+					await Socket.SendAsync(JsonSerializer.SerializeToUtf8Bytes(Response), WebSocketMessageType.Text, true, StoppingToken);
+				}
+
+				// Handle the message type
+				if (EventMessage.Type != null)
+				{
+					string Type = EventMessage.Type;
+					if (Type.Equals("disconnect", StringComparison.Ordinal))
+					{
+						break;
+					}
+					else if (Type.Equals("interactive", StringComparison.Ordinal))
+					{
+						await HandleInteractionMessage(EventMessage);
+					}
+					else
+					{
+						Logger.LogDebug("Unhandled event type ({Type})", Type);
+					}
+				}
+			}
+		}
+
+		/// <summary>
+		/// Handle a button being clicked
+		/// </summary>
+		/// <param name="Message">The event message</param>
+		/// <returns></returns>
+		private async Task HandleInteractionMessage(EventMessage Message)
+		{
+			if (Message.Payload != null && Message.Payload.User != null && Message.Payload.User.UserName != null && String.Equals(Message.Payload.Type, "block_actions", StringComparison.Ordinal))
+			{
+				foreach (ActionInfo Action in Message.Payload.Actions)
+				{
+					if (Action.Value != null)
+					{
+						if (Action.Value.StartsWith("issue_", StringComparison.Ordinal))
+						{
+							await HandleIssueResponseAsync(Message.Payload, Action);
+						}
+					}
+				}
+			}
+		}
+	}
+}
diff --git a/Engine/Source/Programs/Horde/HordeServer/Program.cs b/Engine/Source/Programs/Horde/HordeServer/Program.cs
new file mode 100644
--- /dev/null
+++ b/Engine/Source/Programs/Horde/HordeServer/Program.cs
@@ -0,0 +1,149 @@
+// Copyright Epic Games, Inc. All Rights Reserved.
+
+using System;
+using System.Collections.Generic;
+using System.IO;
+using System.Linq;
+using System.Net;
+using System.Reflection;
+using System.Runtime.InteropServices;
+using System.Security.Cryptography;
+using System.Security.Cryptography.X509Certificates;
+using System.Threading.Tasks;
+using EpicGames.Core;
+using HordeServer.Models;
+using Microsoft.AspNetCore.Hosting;
+using Microsoft.AspNetCore.Server.Kestrel.Core;
+using Microsoft.Extensions.Configuration;
+using Microsoft.Extensions.Hosting;
+using Microsoft.Extensions.Logging;
+using MongoDB.Bson;
+using MongoDB.Bson.Serialization;
+using MongoDB.Driver;
+using Serilog;
+using Serilog.Configuration;
+using Serilog.Events;
+using Serilog.Formatting.Json;
+using Serilog.Sinks.SystemConsole.Themes;
+
+namespace HordeServer
+{
+	static class LoggerExtensions
+	{
+		public static Serilog.LoggerConfiguration Console(this LoggerSinkConfiguration SinkConfig, ServerSettings Settings)
+		{
+			if (Settings.LogJsonToStdOut)
+			{
+				return SinkConfig.Console(new JsonFormatter(renderMessage: true));
+			}
+			else
+			{
+				ConsoleTheme Theme;
+				if (RuntimeInformation.IsOSPlatform(OSPlatform.Windows) && Environment.OSVersion.Version < new Version(10, 0))
+				{
+					Theme = SystemConsoleTheme.Literate;
+				}
+				else
+				{
+					Theme = AnsiConsoleTheme.Code;
+				}
+				return SinkConfig.Console(outputTemplate: "[{Timestamp:HH:mm:ss} {Level:w3}] {Indent}{Message:l}{NewLine}{Exception}", theme: Theme, restrictedToMinimumLevel: LogEventLevel.Debug);
+			}
+		}
+	}
+
+	class Program
+	{
+		public static DirectoryReference AppDir { get; } = new FileReference(Assembly.GetExecutingAssembly().Location).Directory;
+
+		public static void Main(string[] Args)
+		{
+			IConfiguration Config = new ConfigurationBuilder()
+				.SetBasePath(AppDir.FullName)
+				.AddJsonFile("appsettings.json", optional: false)
+				.AddJsonFile($"appsettings.{Environment.GetEnvironmentVariable("ASPNETCORE_ENVIRONMENT")}.json", optional: true)
+				.AddEnvironmentVariables()
+				.Build();
+
+			ServerSettings HordeSettings = new ServerSettings();
+			Config.GetSection("Horde").Bind(HordeSettings);
+
+			Serilog.Log.Logger = new Serilog.LoggerConfiguration()
+				.MinimumLevel.Debug()
+//				.MinimumLevel.Override("HordeServer.Services.DatabaseService", LogEventLevel.Verbose) // For MongoDB query tracing
+				.MinimumLevel.Override("Microsoft", LogEventLevel.Warning)
+				.MinimumLevel.Override("Microsoft.AspNetCore", LogEventLevel.Warning)
+				.MinimumLevel.Override("HordeServer.Authentication.OktaHandler", LogEventLevel.Warning)
+				.MinimumLevel.Override("HordeServer.Authentication.OktaViaJwtHandler", LogEventLevel.Warning)
+				.MinimumLevel.Override("HordeServer.Authentication.HordeJwtBearerHandler", LogEventLevel.Warning)
+				.MinimumLevel.Override("System.Net.Http.HttpClient", LogEventLevel.Warning)
+				.MinimumLevel.Override("Grpc", LogEventLevel.Warning)
+				.Enrich.FromLogContext()
+				.WriteTo.Console(HordeSettings)
+				.WriteTo.File(Path.Combine(AppDir.FullName, "Log.txt"), outputTemplate: "[{Timestamp:HH:mm:ss} {Level:w3}] {Indent}{Message:l}{NewLine}{Exception} [{SourceContext}]", rollingInterval: RollingInterval.Day, rollOnFileSizeLimit: true, fileSizeLimitBytes: 20 * 1024 * 1024, retainedFileCountLimit: 10)
+				.WriteTo.File(new JsonFormatter(renderMessage: true), Path.Combine(AppDir.FullName, "Log.json"), rollingInterval: RollingInterval.Day, rollOnFileSizeLimit: true, fileSizeLimitBytes: 20 * 1024 * 1024, retainedFileCountLimit: 10)
+				.CreateLogger();
+
+			try
+			{
+				using (X509Certificate2? GrpcCertificate = ReadGrpcCertificate(HordeSettings))
+				{
+					CreateHostBuilderWithCert(Args, GrpcCertificate).Build().Run();
+				}
+			}
+#pragma warning disable CA1031 // Do not catch general exception types
+			catch (Exception Ex)
+#pragma warning restore CA1031 // Do not catch general exception types
+			{
+				Serilog.Log.Logger.Error(Ex, "Unhandled exception");
+			}
+		}
+
+		// Used by WebApplicationFactory in controller tests. Uses reflection to call this exact function signature.
+		public static IHostBuilder CreateHostBuilder(string[] Args)
+		{
+			return CreateHostBuilderWithCert(Args, null);
+		}
+
+		public static IHostBuilder CreateHostBuilderWithCert(string[] Args, X509Certificate2? SslCert)
+		{
+			return Host.CreateDefaultBuilder(Args)
+				.UseSerilog()
+				.ConfigureWebHostDefaults(WebBuilder => 
+				{
+					WebBuilder.ConfigureKestrel(Options =>
+					{
+						if (SslCert != null)
+						{
+							const int HttpPort = 80;
+							Options.ListenAnyIP(HttpPort, Configure => { });
+							
+							const int Http2Port = 52103;
+							Options.ListenAnyIP(Http2Port, Configure => { Configure.Protocols = HttpProtocols.Http2; });
+
+							const int HttpsPort = 443;
+							Options.ListenAnyIP(HttpsPort, Configure => { Configure.UseHttps(SslCert); });
+						}
+					});
+					WebBuilder.UseStartup<Startup>(); 
+				});
+		}
+
+		/// <summary>
+		/// Gets the certificate to use for Grpc endpoints
+		/// </summary>
+		/// <returns>Custom certificate to use for Grpc endpoints, or null for the default.</returns>
+		static X509Certificate2? ReadGrpcCertificate(ServerSettings HordeSettings)
+		{
+			if (HordeSettings.ServerPrivateCert == null)
+			{
+				return null;
+			}
+			else
+			{
+				return new X509Certificate2(FileReference.ReadAllBytes(FileReference.Combine(AppDir, HordeSettings.ServerPrivateCert)));
+			}
+		}
+	}
+}
+
diff --git a/Engine/Source/Programs/Horde/HordeServer/Properties/AssemblyInfo.cs b/Engine/Source/Programs/Horde/HordeServer/Properties/AssemblyInfo.cs
new file mode 100644
--- /dev/null
+++ b/Engine/Source/Programs/Horde/HordeServer/Properties/AssemblyInfo.cs
@@ -0,0 +1,3 @@
+using System.Runtime.CompilerServices;
+
+[assembly: InternalsVisibleTo("HordeServerTests")]
diff --git a/Engine/Source/Programs/Horde/HordeServer/Rpc/ActionCacheService.cs b/Engine/Source/Programs/Horde/HordeServer/Rpc/ActionCacheService.cs
new file mode 100644
--- /dev/null
+++ b/Engine/Source/Programs/Horde/HordeServer/Rpc/ActionCacheService.cs
@@ -0,0 +1,88 @@
+// Copyright Epic Games, Inc. All Rights Reserved.
+
+using Build.Bazel.Remote.Execution.V2;
+using EpicGames.Core;
+using Google.Protobuf;
+using Grpc.Core;
+using HordeServer.Storage;
+using Microsoft.Extensions.Options;
+using System;
+using System.Collections.Generic;
+using System.Linq;
+using System.Threading.Tasks;
+using HordeServer.Storage.Primitives;
+using HordeServer.Utility;
+
+namespace HordeServer.Rpc
+{
+	/// <inheritdoc cref="ActionCache"/>
+	public class ActionCacheService : ActionCache.ActionCacheBase
+	{
+		/// <summary>
+		/// Storage service interface
+		/// </summary>
+		IStorageService StorageService;
+
+		/// <summary>
+		/// Constructor
+		/// </summary>
+		/// <param name="StorageService">The storage service instance</param>
+		public ActionCacheService(IStorageService StorageService)
+		{
+			this.StorageService = StorageService;
+		}
+
+		/// <inheritdoc/>
+		public override async Task<ActionResult> UpdateActionResult(UpdateActionResultRequest Request, ServerCallContext Context)
+		{
+			int Length = Request.ActionResult.CalculateSize();
+
+			byte[] Data = new byte[Length];
+			using (CodedOutputStream Stream = new CodedOutputStream(Data))
+			{
+				Request.ActionResult.WriteTo(Stream);
+			}
+
+			BlobHash ActionHash = BlobHash.Parse(Request.ActionDigest.Hash);
+			BlobHash ResultHash = await StorageService.PutBlobAsync(Data);
+			await StorageService.SetRefAsync(ActionHash, ResultHash);
+
+			return Request.ActionResult;
+		}
+
+		/// <inheritdoc/>
+		public override async Task<ActionResult> GetActionResult(GetActionResultRequest Request, ServerCallContext Context)
+		{
+			ActionResult? Result = await TryGetResult(Request.ActionDigest);
+			if (Result == null)
+			{
+				throw new RpcException(new Status(StatusCode.NotFound, $"Cache does not contain {Request.ActionDigest.Hash}"));
+			}
+			return Result;
+		}
+
+		/// <summary>
+		/// Attempts to get the cached result of an action
+		/// </summary>
+		/// <param name="ActionDigest"></param>
+		/// <returns></returns>
+		public async Task<ActionResult?> TryGetResult(Digest ActionDigest)
+		{
+			BlobHash ActionHash = ActionDigest.ToHashValue();
+
+			BlobHash? ResultHash = await StorageService.GetRefAsync(ActionHash, TimeSpan.FromDays(1.0));
+			if (ResultHash == null)
+			{
+				return null;
+			}
+
+			ReadOnlyMemory<byte>? Data = await StorageService.TryGetBlobAsync(ResultHash.Value);
+			if (Data == null)
+			{
+				return null;
+			}
+
+			return ActionResult.Parser.ParseFrom(Data.Value);
+		}
+	}
+}
diff --git a/Engine/Source/Programs/Horde/HordeServer/Rpc/ActionRpcService.cs b/Engine/Source/Programs/Horde/HordeServer/Rpc/ActionRpcService.cs
new file mode 100644
--- /dev/null
+++ b/Engine/Source/Programs/Horde/HordeServer/Rpc/ActionRpcService.cs
@@ -0,0 +1,33 @@
+using Google.Protobuf.WellKnownTypes;
+using Grpc.Core;
+using HordeCommon.Rpc;
+using HordeServer.Tasks.Impl;
+using HordeServer.Utilities;
+using System;
+using System.Collections.Generic;
+using System.Linq;
+using System.Threading.Tasks;
+
+namespace HordeServer.Rpc
+{
+	class ActionRpcService : ActionRpc.ActionRpcBase
+	{
+		ActionTaskSource ActionTaskSource;
+
+		public ActionRpcService(ActionTaskSource ActionTaskSource)
+		{
+			this.ActionTaskSource = ActionTaskSource;
+		}
+
+		/// <inheritdoc/>
+		public override Task<Empty> PostActionResult(PostActionResultRequest Request, ServerCallContext Context)
+		{
+			IActionExecuteOperation? Operation;
+			if (ActionTaskSource.TryGetOperation(Request.LeaseId.ToObjectId(), out Operation))
+			{
+				Operation.TrySetResult(Request.Result);
+			}
+			return Task.FromResult<Empty>(new Empty());
+		}
+	}
+}
diff --git a/Engine/Source/Programs/Horde/HordeServer/Rpc/ByteStreamService.cs b/Engine/Source/Programs/Horde/HordeServer/Rpc/ByteStreamService.cs
new file mode 100644
--- /dev/null
+++ b/Engine/Source/Programs/Horde/HordeServer/Rpc/ByteStreamService.cs
@@ -0,0 +1,213 @@
+// Copyright Epic Games, Inc. All Rights Reserved.
+
+using HordeServer.Storage;
+using HordeServer.Storage.Primitives;
+using EpicGames.Core;
+using Google.Bytestream;
+using Google.Protobuf;
+using Grpc.Core;
+using Microsoft.Extensions.Options;
+using System;
+using System.Collections.Concurrent;
+using System.Collections.Generic;
+using System.Linq;
+using System.Text.RegularExpressions;
+using System.Threading.Tasks;
+using System.Globalization;
+
+namespace HordeServer.Rpc
+{
+	/// <inheritdoc cref="ByteStream"/>
+	class ByteStreamService : ByteStream.ByteStreamBase
+	{
+		/// <summary>
+		/// Information about an active upload
+		/// </summary>
+		class UploadInfo
+		{
+			public long Position;
+			public BlobHash Hash { get; }
+			public byte[] Data { get; }
+
+			public UploadInfo(BlobHash Hash, long Length)
+			{
+				this.Hash = Hash;
+				this.Data = new byte[Length];
+			}
+		}
+
+		/// <summary>
+		/// Active uploads
+		/// </summary>
+		ConcurrentDictionary<string, UploadInfo> Uploads = new ConcurrentDictionary<string, UploadInfo>(StringComparer.Ordinal);
+
+		/// <summary>
+		/// Instance of the singleton storage service
+		/// </summary>
+		IStorageService StorageService;
+
+		/// <summary>
+		/// Constructor
+		/// </summary>
+		/// <param name="StorageProvider">The storage provider instance</param>
+		public ByteStreamService(IStorageService StorageProvider)
+		{
+			this.StorageService = StorageProvider;
+		}
+
+		/// <summary>
+		/// Parse the hash and size from a given resource name
+		/// </summary>
+		/// <param name="ResourceName">The resource name</param>
+		/// <returns>THe hash and size</returns>
+		static (BlobHash, long) ParseResourceName(string ResourceName)
+		{
+			// blobs/{hash}/{size}/foo/bar/baz.cc
+			Match Match = Regex.Match(ResourceName, @"^(?:[^/]+/)?(?:uploads/[^/]+/)?blobs/([0-9a-zA-Z]+)/([0-9]+)(?:/|$)");
+			if (!Match.Success)
+			{
+				throw new RpcException(new Status(StatusCode.InvalidArgument, "Invalid resource name"));
+			}
+
+			BlobHash Hash = BlobHash.Parse(Match.Groups[1].Value);
+			long Size = long.Parse(Match.Groups[2].Value, CultureInfo.InvariantCulture);
+			return (Hash, Size);
+		}
+
+		/// <inheritdoc/>
+		public override async Task Read(ReadRequest Request, IServerStreamWriter<ReadResponse> ResponseStream, ServerCallContext Context)
+		{
+			// Get the source data
+			ReadOnlyMemory<byte> SourceData;
+			if (Uploads.TryGetValue(Request.ResourceName, out UploadInfo? UploadInfo))
+			{
+				SourceData = UploadInfo.Data.AsMemory(0, (int)UploadInfo.Position);
+			}
+			else
+			{
+				ReadOnlyMemory<byte>? Data = await StorageService.TryGetBlobAsync(ParseResourceName(Request.ResourceName).Item1);
+				if (Data == null)
+				{
+					throw new RpcException(new Status(StatusCode.NotFound, $"Resource {Request.ResourceName} not found"));
+				}
+				SourceData = Data.Value;
+			}
+
+			// Get the max offset of the data
+			long EndPosition = SourceData.Length;
+			if (Request.ReadLimit > 0)
+			{
+				EndPosition = Math.Min(EndPosition, Request.ReadOffset + Request.ReadLimit);
+			}
+
+			// Return the data to the client
+			for (long Position = Request.ReadOffset; Position < EndPosition; )
+			{
+				long ChunkSize = Math.Min(EndPosition - Position, 128 * 1024);
+
+				ReadResponse Response = new ReadResponse();
+				Response.Data = ByteString.CopyFrom(SourceData.Span.Slice((int)Position, (int)ChunkSize));
+				await ResponseStream.WriteAsync(Response);
+
+				Position += ChunkSize;
+			}
+		}
+
+		/// <summary>
+		/// Finds or adds an upload info object for the given resource name
+		/// </summary>
+		/// <param name="ResourceName">The resource name</param>
+		/// <returns>Upload info object</returns>
+		UploadInfo FindOrAddUploadInfo(string ResourceName)
+		{
+			UploadInfo? UploadInfo;
+			if (!Uploads.TryGetValue(ResourceName, out UploadInfo))
+			{
+				(BlobHash Hash, long Size) = ParseResourceName(ResourceName);
+
+				UploadInfo = new UploadInfo(Hash, Size);
+				while (!Uploads.TryAdd(ResourceName, UploadInfo!))
+				{
+					UploadInfo? NewUploadInfo;
+					if (Uploads.TryGetValue(ResourceName, out NewUploadInfo))
+					{
+						UploadInfo = NewUploadInfo;
+						break;
+					}
+				}
+			}
+			return UploadInfo;
+		}
+
+		/// <inheritdoc/>
+		public override async Task<WriteResponse> Write(IAsyncStreamReader<WriteRequest> RequestStream, ServerCallContext Context)
+		{
+			WriteResponse Response = new WriteResponse();
+			if(await RequestStream.MoveNext())
+			{
+				WriteRequest Request = RequestStream.Current;
+				string ResourceName = RequestStream.Current.ResourceName;
+
+				UploadInfo UploadInfo = FindOrAddUploadInfo(ResourceName);
+				for (; ;)
+				{
+					lock (UploadInfo)
+					{
+						if (Request.WriteOffset > UploadInfo.Position)
+						{
+							throw new RpcException(new Status(StatusCode.InvalidArgument, "Invalid resource name"));
+						}
+
+						Request.Data.CopyTo(UploadInfo.Data, (int)Request.WriteOffset);
+						UploadInfo.Position = Request.WriteOffset + UploadInfo.Data.Length;
+						Response.CommittedSize = UploadInfo.Position;
+					}
+
+					if (!await RequestStream.MoveNext())
+					{
+						break;
+					}
+
+					Request = RequestStream.Current;
+
+					if(!String.IsNullOrEmpty(Request.ResourceName) && !Request.ResourceName.Equals(ResourceName, StringComparison.Ordinal))
+					{
+						throw new RpcException(new Status(StatusCode.InvalidArgument, "Resource name does not match original"));
+					}
+				}
+
+				if (Request.FinishWrite)
+				{
+					await StorageService.PutBlobAsync(UploadInfo.Hash, UploadInfo.Data);
+					Uploads.TryRemove(ResourceName, out _);
+				}
+			}
+			return Response;
+		}
+
+		/// <inheritdoc/>
+		public override async Task<QueryWriteStatusResponse> QueryWriteStatus(QueryWriteStatusRequest Request, ServerCallContext Context)
+		{
+			QueryWriteStatusResponse Response = new QueryWriteStatusResponse();
+			if (Uploads.TryGetValue(Request.ResourceName, out UploadInfo? UploadInfo))
+			{
+				lock (UploadInfo)
+				{
+					Response.CommittedSize = UploadInfo.Position;
+					Response.Complete = false;
+					return Response;
+				}
+			}
+
+			(BlobHash Hash, long Size) = ParseResourceName(Request.ResourceName);
+			if (await StorageService.ShouldPutBlobAsync(Hash))
+			{
+				throw new RpcException(new Status(StatusCode.NotFound, "Invalid resource name"));
+			}
+
+			Response.CommittedSize = Size;
+			Response.Complete = true;
+			return Response;
+		}
+	}
+}
diff --git a/Engine/Source/Programs/Horde/HordeServer/Rpc/CapabilitiesService.cs b/Engine/Source/Programs/Horde/HordeServer/Rpc/CapabilitiesService.cs
new file mode 100644
--- /dev/null
+++ b/Engine/Source/Programs/Horde/HordeServer/Rpc/CapabilitiesService.cs
@@ -0,0 +1,35 @@
+// Copyright Epic Games, Inc. All Rights Reserved.
+
+using Build.Bazel.Remote.Execution.V2;
+using Build.Bazel.Semver;
+using Grpc.Core;
+using System;
+using System.Collections.Generic;
+using System.Linq;
+using System.Threading.Tasks;
+
+namespace HordeServer.Rpc
+{
+	/// <inheritdoc cref="Capabilities"/>
+	class CapabilitiesService : Capabilities.CapabilitiesBase
+	{
+		/// <inheritdoc/>
+		public override Task<ServerCapabilities> GetCapabilities(GetCapabilitiesRequest Request, ServerCallContext Context)
+		{
+			ServerCapabilities Capabilities = new ServerCapabilities();
+
+			Capabilities.DeprecatedApiVersion = new SemVer { Major = 2 };
+			Capabilities.LowApiVersion = new SemVer { Major = 2 };
+			Capabilities.HighApiVersion = new SemVer { Major = 2 };
+
+			Capabilities.CacheCapabilities = new CacheCapabilities();
+			Capabilities.CacheCapabilities.DigestFunction.Add(DigestFunction.Types.Value.Sha256);
+
+			Capabilities.ExecutionCapabilities = new ExecutionCapabilities();
+			Capabilities.ExecutionCapabilities.DigestFunction = DigestFunction.Types.Value.Sha256;
+			Capabilities.ExecutionCapabilities.ExecEnabled = true;
+
+			return Task.FromResult(Capabilities);
+		}
+	}
+}
diff --git a/Engine/Source/Programs/Horde/HordeServer/Rpc/ContentStorageService.cs b/Engine/Source/Programs/Horde/HordeServer/Rpc/ContentStorageService.cs
new file mode 100644
--- /dev/null
+++ b/Engine/Source/Programs/Horde/HordeServer/Rpc/ContentStorageService.cs
@@ -0,0 +1,101 @@
+// Copyright Epic Games, Inc. All Rights Reserved.
+
+using Build.Bazel.Remote.Execution.V2;
+using HordeServer.Storage;
+using HordeServer.Utility;
+using EpicGames.Core;
+using Google.Protobuf;
+using Google.Rpc;
+using Grpc.Core;
+using Microsoft.Extensions.Options;
+using System;
+using System.Collections.Generic;
+using System.IO;
+using System.Linq;
+using System.Reflection;
+using System.Threading;
+using System.Threading.Tasks;
+
+using Status = Google.Rpc.Status;
+
+namespace HordeServer.Rpc
+{
+	/// <inheritdoc cref="ContentAddressableStorage"/>
+	class ContentStorageService : ContentAddressableStorage.ContentAddressableStorageBase
+	{
+		/// <summary>
+		/// Storage service singleton
+		/// </summary>
+		IStorageService StorageService;
+
+		/// <summary>
+		/// Constructor
+		/// </summary>
+		/// <param name="StorageService"></param>
+		public ContentStorageService(IStorageService StorageService)
+		{
+			this.StorageService = StorageService;
+		}
+
+		/// <inheritdoc/>
+		public override async Task<BatchReadBlobsResponse> BatchReadBlobs(BatchReadBlobsRequest Request, ServerCallContext Context)
+		{
+			BatchReadBlobsResponse Response = new BatchReadBlobsResponse();
+			foreach (Digest ItemDigest in Request.Digests)
+			{
+				ReadOnlyMemory<byte>? Data = await StorageService.TryGetBlobAsync(ItemDigest.ToHashValue());
+
+				BatchReadBlobsResponse.Types.Response Item = new BatchReadBlobsResponse.Types.Response();
+				Item.Digest = ItemDigest;
+				if (Data == null)
+				{
+					Item.Data = ByteString.Empty;
+					Item.Status = new Status { Code = (int)Code.NotFound };
+				}
+				else
+				{
+					Item.Data = ByteString.CopyFrom(Data.Value.Span);
+					Item.Status = new Status { Code = (int)Code.Ok };
+				}
+				Response.Responses.Add(Item);
+			}
+			return Response;
+		}
+
+		/// <inheritdoc/>
+		public override async Task<BatchUpdateBlobsResponse> BatchUpdateBlobs(BatchUpdateBlobsRequest Request, ServerCallContext Context)
+		{
+			BatchUpdateBlobsResponse Response = new BatchUpdateBlobsResponse();
+			foreach (BatchUpdateBlobsRequest.Types.Request ItemRequest in Request.Requests)
+			{
+				await StorageService.PutBlobAsync(ItemRequest.Digest.ToHashValue(), ItemRequest.Data.ToByteArray());
+
+				BatchUpdateBlobsResponse.Types.Response ItemResponse = new BatchUpdateBlobsResponse.Types.Response();
+				ItemResponse.Digest = ItemRequest.Digest;
+				ItemResponse.Status = new Status { Code = (int)Code.Ok };
+				Response.Responses.Add(ItemResponse);
+			}
+			return Response;
+		}
+
+		/// <inheritdoc/>
+		public override async Task<FindMissingBlobsResponse> FindMissingBlobs(FindMissingBlobsRequest Request, ServerCallContext Context)
+		{
+			FindMissingBlobsResponse Response = new FindMissingBlobsResponse();
+			foreach (Digest BlobDigest in Request.BlobDigests)
+			{
+				if (await StorageService.ShouldPutBlobAsync(BlobDigest.ToHashValue()))
+				{
+					Response.MissingBlobDigests.Add(BlobDigest);
+				}
+			}
+			return Response;
+		}
+
+		/// <inheritdoc/>
+		public override Task GetTree(GetTreeRequest request, IServerStreamWriter<GetTreeResponse> ResponseStream, ServerCallContext Context)
+		{
+			throw new NotImplementedException();
+		}
+	}
+}
diff --git a/Engine/Source/Programs/Horde/HordeServer/Rpc/ExecutionService.cs b/Engine/Source/Programs/Horde/HordeServer/Rpc/ExecutionService.cs
new file mode 100644
--- /dev/null
+++ b/Engine/Source/Programs/Horde/HordeServer/Rpc/ExecutionService.cs
@@ -0,0 +1,79 @@
+// Copyright Epic Games, Inc. All Rights Reserved.
+
+using Build.Bazel.Remote.Execution.V2;
+using EpicGames.Core;
+using Google.LongRunning;
+using Google.Protobuf;
+using Grpc.Core;
+using HordeServer.Storage;
+using Microsoft.Extensions.Options;
+using System;
+using System.Collections.Generic;
+using System.IO;
+using System.Linq;
+using System.Threading;
+using System.Threading.Tasks;
+using System.Diagnostics;
+using System.Text;
+using Google.Protobuf.WellKnownTypes;
+using Microsoft.Extensions.Hosting;
+using System.Collections.Concurrent;
+using HordeServer.Storage.Primitives;
+using HordeServer.Utility;
+using HordeCommon.Rpc.Tasks;
+using HordeServer.Models;
+using HordeServer.Services;
+using MongoDB.Bson;
+using HordeServer.Api;
+using HordeServer.Collections;
+using HordeServer.Tasks.Impl;
+using HordeServer.Utilities;
+
+using Action = Build.Bazel.Remote.Execution.V2.Action;
+using Directory = Build.Bazel.Remote.Execution.V2.Directory;
+using Status = Google.Rpc.Status;
+
+namespace HordeServer.Rpc
+{
+	class ExecutionService : Execution.ExecutionBase
+	{
+		ActionTaskSource ActionTaskSource;
+
+		public ExecutionService(ActionTaskSource ActionTaskSource)
+		{
+			this.ActionTaskSource = ActionTaskSource;
+		}
+
+		public override async Task Execute(ExecuteRequest Request, IServerStreamWriter<Operation> ResponseStream, ServerCallContext Context)
+		{
+			// Try to create the operation
+			IActionExecuteOperation? Operation = ActionTaskSource.Execute(Request);
+			if (Operation == null)
+			{
+				throw new RpcException(new Grpc.Core.Status(StatusCode.ResourceExhausted, "No agent available to execute request"));
+			}
+
+			// Return all the status updates
+			await foreach (Operation Status in Operation.ReadStatusUpdatesAsync())
+			{
+				await ResponseStream.WriteAsync(Status);
+			}
+		}
+
+		public override async Task WaitExecution(WaitExecutionRequest Request, IServerStreamWriter<Operation> ResponseStream, ServerCallContext Context)
+		{
+			// Get the operation
+			IActionExecuteOperation? Operation;
+			if (!ActionTaskSource.TryGetOperation(Request.Name.ToObjectId(), out Operation))
+			{
+				throw new RpcException(new Grpc.Core.Status(StatusCode.NotFound, "No agent available to execute request"));
+			}
+
+			// Stream all the status updates
+			await foreach (Operation Status in Operation.ReadStatusUpdatesAsync())
+			{
+				await ResponseStream.WriteAsync(Status);
+			}
+		}
+	}
+}
diff --git a/Engine/Source/Programs/Horde/HordeServer/ServerSettings.cs b/Engine/Source/Programs/Horde/HordeServer/ServerSettings.cs
new file mode 100644
--- /dev/null
+++ b/Engine/Source/Programs/Horde/HordeServer/ServerSettings.cs
@@ -0,0 +1,362 @@
+// Copyright Epic Games, Inc. All Rights Reserved.
+
+using HordeServer.Services;
+using System;
+using System.Collections.Generic;
+using System.ComponentModel.DataAnnotations;
+using System.Linq;
+using System.Threading.Tasks;
+using TimeZoneConverter;
+
+namespace HordeServer
+{
+	/// <summary>
+	/// Types of storage to use for log data
+	/// </summary>
+	public enum StorageProviderType
+	{
+		/// <summary>
+		/// AWS S3
+		/// </summary>
+		S3,
+
+		/// <summary>
+		/// Local filesystem
+		/// </summary>
+		FileSystem,
+
+		/// <summary>
+		/// In-memory only (for testing)
+		/// </summary>
+		Transient,
+
+		/// <summary>
+		/// Relay to another server (useful for testing against prod)
+		/// </summary>
+		Relay,
+	};
+
+	/// <summary>
+	/// Specifies the service to use for controlling the size of the fleet
+	/// </summary>
+	public enum FleetManagerType
+	{
+		/// <summary>
+		/// Default (empty) instance
+		/// </summary>
+		None,
+
+		/// <summary>
+		/// Use AWS EC2 instances
+		/// </summary>
+		Aws,
+	}
+
+	/// <summary>
+	/// Global settings for the application
+	/// </summary>
+	public class ServerSettings
+	{
+		/// <summary>
+		/// MongoDB connection string
+		/// </summary>
+		[Required]
+		public string DatabaseConnectionString { get; set; } = null!;
+
+		/// <summary>
+		/// MongoDB database name
+		/// </summary>
+		[Required]
+		public string DatabaseName { get; set; } = null!;
+
+		/// <summary>
+		/// The claim type for administrators
+		/// </summary>
+		public string? AdminClaimType { get; set; }
+
+		/// <summary>
+		/// Value of the claim type for administrators
+		/// </summary>
+		public string? AdminClaimValue { get; set; }
+
+		/// <summary>
+		/// Optional certificate to trust in order to access the database (eg. AWS public cert for TLS)
+		/// </summary>
+		public string? DatabasePublicCert { get; set; }
+
+		/// <summary>
+		/// Access the database in read-only mode (avoids creating indices or updating content)
+		/// Useful for debugging a local instance of HordeServer against a production database.
+		/// </summary>
+		public bool DatabaseReadOnlyMode { get; set; } = false;
+
+		/// <summary>
+		/// Optional PFX certificate to use for encryting agent SSL traffic. This can be a self-signed certificate, as long as it's trusted by agents.
+		/// </summary>
+		public string? ServerPrivateCert { get; set; }
+
+		/// <summary>
+		/// Issuer for tokens from the auth provider
+		/// </summary>
+		public string? OidcAuthority { get; set; }
+		
+		/// <summary>
+		/// Client id for the OIDC authority
+		/// </summary>
+		public string? OidcClientId { get; set; }
+
+		/// <summary>
+		/// Optional redirect url provided to OIDC login
+		/// </summary>
+		public string? OidcSigninRedirect { get; set; }		
+
+		/// <summary>
+		/// Name of the issuer in bearer tokens from the server
+		/// </summary>
+		public string? JwtIssuer { get; set; } = null!;
+
+		/// <summary>
+		/// Secret key used to sign JWTs. This setting is typically only used for development. In prod, a unique secret key will be generated and stored in the DB for each unique server instance.
+		/// </summary>
+		public string? JwtSecret { get; set; } = null!;
+
+		/// <summary>
+		/// Length of time before JWT tokens expire, in hourse
+		/// </summary>
+		public int JwtExpiryTimeHours { get; set; } = 4;
+
+		/// <summary>
+		/// Disable authentication for debugging purposes
+		/// </summary>
+		public bool DisableAuth { get; set; }
+
+		/// <summary>
+		/// Whether to enable Cors, generally for development purposes
+		/// </summary>
+		public bool CorsEnabled { get; set; } = false;
+
+		/// <summary>
+		/// Allowed Cors origin 
+		/// </summary>
+		public string CorsOrigin { get; set; } = null!;
+
+		/// <summary>
+		/// Whether to enable a schedule in test data (false by default for development builds)
+		/// </summary>
+		public bool EnableScheduleInTestData { get; set; }
+
+		/// <summary>
+		/// Interval between rebuilding the schedule queue with a DB query.
+		/// </summary>
+		public TimeSpan SchedulePollingInterval { get; set; } = TimeSpan.FromSeconds(60.0);
+
+		/// <summary>
+		/// Interval between polling for new jobs
+		/// </summary>
+		public TimeSpan NoResourceBackOffTime { get; set; } = TimeSpan.FromSeconds(30.0);
+
+		/// <summary>
+		/// Interval between attempting to assign agents to take on jobs
+		/// </summary>
+		public TimeSpan InitiateJobBackOffTime { get; set; } = TimeSpan.FromSeconds(180.0);
+
+		/// <summary>
+		/// Interval between scheduling jobs when an unknown error occurs
+		/// </summary>
+		public TimeSpan UnknownErrorBackOffTime { get; set; } = TimeSpan.FromSeconds(120.0);
+
+		/// <summary>
+		/// Config for connecting to Redis server(s).
+		/// Setting it to null will disable Redis use and connection
+		/// See format at https://stackexchange.github.io/StackExchange.Redis/Configuration.html
+		/// </summary>
+		[Required]
+		public string RedisConnectionConfig { get; set; } = null!;
+		
+		/// <summary>
+		/// Type of write cache to use in log service
+		/// Currently Supported: "InMemory" or "Redis"
+		/// </summary>
+		public string LogServiceWriteCacheType { get; set; } = "InMemory";
+
+		/// <summary>
+		/// Provider Type
+		/// Currently Supported: "S3" or "FileSystem"
+		/// </summary>
+		public StorageProviderType ExternalStorageProviderType { get; set; } = StorageProviderType.FileSystem;
+
+		/// <summary>
+		/// local log file storage directory, if using type filesystem
+		/// </summary>
+		public string LocalStorageDir { get; set; } = "D:\\HordeLogs";
+
+		/// <summary>
+		/// S3 bucket region for logfile storage
+		/// </summary>
+		public string S3BucketRegion { get; set; } = null!;
+
+		/// <summary>
+		/// Arn to assume for s3.  "Basic", "AssumeRole", "AssumeRoleWebIdentity" only
+		/// </summary>
+		public string S3CredentialType { get; set; } = null!;
+
+		/// <summary>
+		/// S3 Client username (used in Basic auth type only)
+		/// </summary>
+		public string S3ClientKeyId { get; set; } = null!;
+
+		/// <summary>
+		/// S3 client password (used in Basic auth type only)
+		/// </summary>
+		public string S3ClientSecret { get; set; } = null!;
+
+		/// <summary>
+		/// Arn to assume for s3
+		/// </summary>
+		public string S3AssumeArn { get; set; } = null!;
+
+		/// <summary>
+		/// S3 log bucket name
+		/// </summary>
+		public string S3LogBucketName { get; set; } = null!;
+
+		/// <summary>
+		/// S3 artifact bucket name
+		/// </summary>
+		public string S3ArtifactBucketName { get; set; } = null!;
+
+		/// <summary>
+		/// When using a relay storage provider, specifies the remote server to use
+		/// </summary>
+		public string? LogRelayServer { get; set; }
+
+		/// <summary>
+		/// Authentication token for using a relay server
+		/// </summary>
+		public string? LogRelayBearerToken { get; set; }
+
+		/// <summary>
+		/// Whether to log json to stdout
+		/// </summary>
+		public bool LogJsonToStdOut { get; set; } = false;
+
+		/// <summary>
+		/// Which fleet manager service to use
+		/// </summary>
+		public FleetManagerType FleetManager { get; set; } = FleetManagerType.None;
+
+		/// <summary>
+		/// Whether to run scheduled jobs. Not wanted for development.
+		/// </summary>
+		public bool DisableSchedules { get; set; } = true;
+
+		/// <summary>
+		/// Timezone for evaluating schedules
+		/// </summary>
+		public string? ScheduleTimeZone { get; set; }
+
+		/// <summary>
+		/// Address of the Perforce bridge
+		/// </summary>
+		public string? PerforceBridge { get; set; }
+
+		/// <summary>
+		/// Token for interacting with Slack
+		/// </summary>
+		public string? SlackToken { get; set; }
+
+		/// <summary>
+		/// Token for opening a socket to slack
+		/// </summary>
+		public string? SlackSocketToken { get; set; }
+
+		/// <summary>
+		/// Channel to send stream notification update failures to
+		/// </summary>
+		public string? UpdateStreamsNotificationChannel { get; set; }
+
+		/// <summary>
+		/// URI to the SmtpServer to use for sending email notifications
+		/// </summary>
+		public string? SmtpServer { get; set; }
+
+		/// <summary>
+		/// The email address to send email notifications from
+		/// </summary>
+		public string? EmailSenderAddress { get; set; }
+
+		/// <summary>
+		/// The name for the sender when sending email notifications
+		/// </summary>
+		public string? EmailSenderName { get; set; }
+
+		/// <summary>
+		/// The URl to use for generating links back to the dashboard.
+		/// </summary>
+		public Uri DashboardUrl { get; set; } = new Uri("https://localhost:3000");
+		
+		/// <summary>
+		/// The URL to Helix Swarm server
+		/// </summary>
+		public Uri? HelixSwarmServerUrl { get; set; }
+		
+		/// <summary>
+		/// Username to use when communicating with Helix Swarm
+		/// </summary>
+		public string? HelixSwarmUsername { get; set; }
+		
+		/// <summary>
+		/// Password to use when communicating with Helix Swarm
+		/// </summary>
+		public string? HelixSwarmPassword { get; set; }
+		
+		/// <summary>
+		/// Projects in Helix Swarm to monitor, including the state of reviews in each project
+		/// Multiple projects can be specified, separated by comma.
+		/// </summary>
+		public string? HelixSwarmProjects { get; set; }
+		
+		/// <summary>
+		/// Set the minimum size of the global thread pool
+		/// This value has been found in need of tweaking to avoid timeouts with the Redis client during bursts
+		/// of traffic. Default is 16 for .NET Core CLR. The correct value is dependent on the traffic the Horde Server
+		/// is receiving. For Epic's internal deployment, this is set to 40.
+		/// </summary>
+		public int? GlobalThreadPoolMinSize { get; set; }
+		
+		/// <summary>
+		/// Lazily computed timezone value
+		/// </summary>
+		public TimeZoneInfo TimeZoneInfo
+		{
+			get
+			{
+				if (CachedTimeZoneInfo == null)
+				{
+					CachedTimeZoneInfo = (ScheduleTimeZone == null) ? TimeZoneInfo.Local : TZConvert.GetTimeZoneInfo(ScheduleTimeZone);
+				}
+				return CachedTimeZoneInfo;
+			}
+		}
+
+		private TimeZoneInfo? CachedTimeZoneInfo;
+
+		/// <summary>
+		/// Check if Helix Swarm is enabled
+		/// </summary>
+		/// <returns>True if enabled</returns>
+		public bool IsSwarmEnabled()
+		{
+			return HelixSwarmServerUrl != null;
+		}
+		
+		/// <summary>
+		/// Parse the comma-separated string of Swarm projects into an actual array
+		/// </summary>
+		/// <returns>Array of Swarm projects</returns>
+		public string[] GetSwarmProjects()
+		{
+			return HelixSwarmProjects != null ? HelixSwarmProjects.Split(",") : Array.Empty<string>();
+		}
+	}
+}
