Initial drop of Horde Server/Agent for UE5/Main. Still missing some Epic-specific functionality required to build/run correctly.

diff --git a/Engine/Source/Programs/Horde/HordeServer/Services/AgentService.cs b/Engine/Source/Programs/Horde/HordeServer/Services/AgentService.cs
new file mode 100644
--- /dev/null
+++ b/Engine/Source/Programs/Horde/HordeServer/Services/AgentService.cs
@@ -0,0 +1,946 @@
+// Copyright Epic Games, Inc. All Rights Reserved.
+
+using Datadog.Trace;
+using EpicGames.Core;
+using Google.Protobuf;
+using Google.Protobuf.Reflection;
+using Google.Protobuf.WellKnownTypes;
+using HordeServer.Api;
+using HordeServer.Collections;
+using HordeCommon;
+using HordeServer.Models;
+using HordeCommon.Rpc.Tasks;
+using HordeServer.Tasks;
+using HordeServer.Utilities;
+using Microsoft.AspNetCore.Identity;
+using Microsoft.Extensions.Hosting;
+using Microsoft.Extensions.Logging;
+using Microsoft.Extensions.Options;
+using Microsoft.IdentityModel.Tokens;
+using MongoDB.Bson;
+using MongoDB.Driver;
+using System;
+using System.Collections.Generic;
+using System.Diagnostics;
+using System.IdentityModel.Tokens.Jwt;
+using System.Linq;
+using System.Linq.Expressions;
+using System.Security.Claims;
+using System.Text;
+using System.Threading;
+using System.Threading.Tasks;
+using StatsdClient;
+using PoolId = HordeServer.Utilities.StringId<HordeServer.Models.IPool>;
+using AgentSoftwareVersion = HordeServer.Utilities.StringId<HordeServer.Collections.IAgentSoftwareCollection>;
+using AgentSoftwareChannelName = HordeServer.Utilities.StringId<HordeServer.Services.AgentSoftwareChannels>;
+using System.Collections.Concurrent;
+using Microsoft.Extensions.Caching.Memory;
+
+namespace HordeServer.Services
+{
+	/// <summary>
+	/// Wraps funtionality for manipulating agents
+	/// </summary>
+	public class AgentService : TickedBackgroundService
+	{
+		/// <summary>
+		/// Maximum time between updates for an agent to be considered online
+		/// </summary>
+		public static readonly TimeSpan SessionExpiryTime = TimeSpan.FromMinutes(5);
+
+		/// <summary>
+		/// Time before a session expires that we will poll until
+		/// </summary>
+		public static readonly TimeSpan SessionLongPollTime = TimeSpan.FromSeconds(55);
+
+		/// <summary>
+		/// Time after which a session will be renewed
+		/// </summary>
+		public static readonly TimeSpan SessionRenewTime = TimeSpan.FromSeconds(50);
+
+		/// <summary>
+		/// The ACL service instance
+		/// </summary>
+		AclService AclService;
+
+		/// <summary>
+		/// The downtime service instance
+		/// </summary>
+		IDowntimeService DowntimeService;
+
+		/// <summary>
+		/// Collection of agent documents
+		/// </summary>
+		IAgentCollection Agents;
+
+		/// <summary>
+		/// Collection of lease documents
+		/// </summary>
+		ILeaseCollection Leases;
+
+		/// <summary>
+		/// Collection of session documents
+		/// </summary>
+		ISessionCollection Sessions;
+
+		/// <summary>
+		/// DogStatsD metric client
+		/// </summary>
+		IDogStatsd DogStatsd;
+
+		/// <summary>
+		/// List of task sources
+		/// </summary>
+		ITaskSource[] TaskSources;
+
+		/// <summary>
+		/// The application lifetime instance
+		/// </summary>
+		IHostApplicationLifetime ApplicationLifetime;
+
+		/// <summary>
+		/// Log output writer
+		/// </summary>
+		ILogger<AgentService> Logger;
+		
+		/// <summary>
+		/// Clock
+		/// </summary>
+		readonly IClock Clock;
+
+		/// <summary>
+		/// List of event listeners for agents connected to this server. Used to track leases that cannot be migrated to other machines (ie. REAPI leases)
+		/// </summary>
+		Dictionary<AgentId, AgentEventListener> Listeners = new Dictionary<AgentId, AgentEventListener>();
+
+		/// <summary>
+		/// Constructor
+		/// </summary>
+		/// <param name="Agents">Collection of agent documents</param>
+		/// <param name="Leases">Collection of lease documents</param>
+		/// <param name="Sessions">Collection of session documents</param>
+		/// <param name="AclService">The ACL service</param>
+		/// <param name="DowntimeService">The downtime service</param>
+		/// <param name="DogStatsd">The DogStatsd metric client</param>
+		/// <param name="TaskSources">The list of task sources</param>
+		/// <param name="ApplicationLifetime"></param>
+		/// <param name="Logger">Log output writer</param>
+		/// <param name="Clock">Clock</param>
+		public AgentService(IAgentCollection Agents, ILeaseCollection Leases, ISessionCollection Sessions, AclService AclService, IDowntimeService DowntimeService, IDogStatsd DogStatsd, IEnumerable<ITaskSource> TaskSources, IHostApplicationLifetime ApplicationLifetime, ILogger<AgentService> Logger, IClock Clock)
+			: base(TimeSpan.FromSeconds(30.0), Logger)
+		{
+			this.Agents = Agents;
+			this.Leases = Leases;
+			this.Sessions = Sessions;
+			this.AclService = AclService;
+			this.DowntimeService = DowntimeService;
+			this.DogStatsd = DogStatsd;
+			this.TaskSources = TaskSources.ToArray();
+			this.ApplicationLifetime = ApplicationLifetime;
+			this.Logger = Logger;
+			this.Clock = Clock;
+		}
+
+		/// <summary>
+		/// Issues a bearer token for the given session id
+		/// </summary>
+		/// <param name="SessionId">The session id</param>
+		/// <returns>Bearer token for the agent</returns>
+		public string IssueSessionToken(ObjectId SessionId)
+		{
+			List<AclClaim> Claims = new List<AclClaim>();
+			Claims.Add(AclService.AgentClaim);
+			Claims.Add(AclService.GetSessionClaim(SessionId));
+			return AclService.IssueBearerToken(Claims, null);
+		}
+
+		/// <summary>
+		/// Register a new agent
+		/// </summary>
+		/// <param name="Name">Name of the agent</param>
+		/// <param name="bEnabled">Whether the agent is currently enabled</param>
+		/// <param name="bEphemeral">Whether this agent is ephemeral</param>
+		/// <param name="Channel">Override for the desired software version</param>
+		/// <param name="Pools">Pools for this agent</param>
+		/// <returns>Unique id for the agent</returns>
+		public Task<IAgent> CreateAgentAsync(string Name, bool bEnabled, bool bEphemeral, AgentSoftwareChannelName? Channel, List<PoolId>? Pools)
+		{
+			return Agents.AddAsync(new AgentId(Name), bEnabled, bEphemeral, Channel, Pools);
+		}
+
+		/// <summary>
+		/// Gets an agent by ID
+		/// </summary>
+		/// <param name="AgentId">Unique id of the agent</param>
+		/// <returns>The agent document</returns>
+		public Task<IAgent?> GetAgentAsync(AgentId AgentId)
+		{
+			return Agents.GetAsync(AgentId);
+		}
+
+		/// <summary>
+		/// Finds all agents matching certain criteria
+		/// </summary>
+		/// <param name="Pool">The pool containing the agent</param>
+		/// <param name="ModifiedAfter">If set, only returns agents modified after this time</param>
+		/// <param name="Index">Index within the list of results</param>
+		/// <param name="Count">Number of results to return</param>
+		/// <returns>List of agents matching the given criteria</returns>
+		public Task<List<IAgent>> FindAgentsAsync(ObjectId? Pool, DateTime? ModifiedAfter, int? Index, int? Count)
+		{
+			return Agents.FindAsync(Pool, null, ModifiedAfter, null, Index, Count);
+		}
+
+		/// <summary>
+		/// Update an agent's settings
+		/// </summary>
+		/// <param name="Agent">Agent instance</param>
+		/// <param name="bNewEnabled">Whether the agent is enabled or not</param>
+		/// <param name="bNewRequestConform">Whether to request a conform job be run</param>
+		/// <param name="bNewRequestRestart">Whether to request the machine be restarted</param>
+		/// <param name="bNewRequestShutdown">Whether to request the machine be shutdown</param>
+		/// <param name="NewChannel">Override for the desired client version</param>
+		/// <param name="NewPools">List of pools for the agent</param>
+		/// <param name="NewAcl">New ACL for this agent</param>
+		/// <param name="NewComment">New comment</param>
+		/// <returns>Version of the software that needs to be installed on the agent. Null if the agent is running the correct version.</returns>
+		public async Task<IAgent?> UpdateAgentAsync(IAgent? Agent, bool? bNewEnabled, bool? bNewRequestConform, bool? bNewRequestRestart, bool? bNewRequestShutdown, AgentSoftwareChannelName? NewChannel, List<PoolId>? NewPools, Acl? NewAcl, string? NewComment)
+		{
+			while (Agent != null)
+			{
+				// Apply the update
+				IAgent? NewAgent = await Agents.TryUpdateSettingsAsync(Agent, bNewEnabled, bNewRequestConform, bNewRequestRestart, bNewRequestShutdown, NewChannel, NewPools, NewAcl, NewComment);
+				if (NewAgent != null)
+				{
+					Agent = NewAgent;
+
+					// If we're requesting a conform, allow the agent to break out of a long poll for work immediately to pick it up
+//					if ((bNewRequestConform ?? false) || (bNewRequestRestart ?? false) || (bNewRequestShutdown ?? false))
+//					{
+//						DispatchService.CancelLongPollForAgent(Agent.Id);
+//					}
+
+					return Agent;
+				}
+
+				// Update the agent
+				Agent = await GetAgentAsync(Agent.Id);
+			}
+			return null;
+		}
+
+		/// <summary>
+		/// Update the current workspaces for an agent.
+		/// </summary>
+		/// <param name="Agent">The agent to update</param>
+		/// <param name="Workspaces">Current list of workspaces</param>
+		/// <param name="bPendingConform">Whether the agent still needs to run another conform</param>
+		/// <returns>New agent state</returns>
+		public async Task<bool> TryUpdateWorkspacesAsync(IAgent Agent, List<AgentWorkspace> Workspaces, bool bPendingConform)
+		{
+			IAgent? NewAgent = await Agents.TryUpdateWorkspacesAsync(Agent, Workspaces, bPendingConform);
+			return NewAgent != null;
+		}
+
+		/// <summary>
+		/// Marks the agent as deleted
+		/// </summary>
+		/// <param name="Agent">The agent to delete</param>
+		/// <returns>Async task</returns>
+		public async Task DeleteAgentAsync(IAgent? Agent)
+		{
+			while (Agent != null && !Agent.Deleted)
+			{
+				IAgent? NewAgent = await Agents.TryDeleteAsync(Agent);
+				if(NewAgent != null)
+				{
+					break;
+				}
+				Agent = await GetAgentAsync(Agent.Id);
+			}
+		}
+
+		/// <summary>
+		/// Creates a new agent session
+		/// </summary>
+		/// <param name="Agent">The agent to create a session for</param>
+		/// <param name="Status">Current status of the agent</param>
+		/// <param name="Capabilities">Capabilities of the agent software</param>
+		/// <param name="Version">Version of the software that's running</param>
+		/// <returns>New agent state</returns>
+		public async Task<IAgent> CreateSessionAsync(IAgent Agent, AgentStatus Status, AgentCapabilities Capabilities, AgentSoftwareVersion? Version)
+		{
+			for (; ; )
+			{
+				// Check if there's already a session running for this agent.
+				IAgent? NewAgent;
+				if (Agent.SessionId != null)
+				{
+					// Try to terminate the current session
+					await TryTerminateSessionAsync(Agent);
+				}
+				else
+				{
+					DateTime UtcNow = Clock.UtcNow;
+
+					// Remove any outstanding leases
+					foreach (AgentLease Lease in Agent.Leases)
+					{
+						Logger.LogDebug("Removing outstanding lease {LeaseId} during session create...", Lease.Id);
+						await RemoveLeaseAsync(Agent, Lease, UtcNow, LeaseOutcome.Failed);
+					}
+
+					// Create a new session document
+					ISession NewSession = await Sessions.AddAsync(ObjectId.GenerateNewId(), Agent.Id, Clock.UtcNow, Capabilities, Version);
+					DateTime SessionExpiresAt = UtcNow + SessionExpiryTime;
+
+					// Reset the agent to use the new session
+					NewAgent = await Agents.TryStartSessionAsync(Agent, NewSession.Id, SessionExpiresAt, Status, Capabilities, Version);
+					if(NewAgent != null)
+					{
+						Agent = NewAgent;
+						break;
+					}
+
+					// Remove the session we didn't use
+					await Sessions.DeleteAsync(NewSession.Id);
+				}
+
+				// Get the current agent state
+				NewAgent = await GetAgentAsync(Agent.Id);
+				if (NewAgent == null)
+				{
+					throw new InvalidOperationException($"Invalid agent id '{Agent.Id}'");
+				}
+				Agent = NewAgent;
+			}
+			return Agent;
+		}
+
+		/// <summary>
+		/// Waits for a lease to be assigned to an agent
+		/// </summary>
+		/// <param name="Agent">The agent to assign a lease to</param>
+		/// <param name="AbortTask">Task which can be used to stop polling</param>
+		/// <returns>True if a lease was assigned, false otherwise</returns>
+		public async Task<IAgent?> WaitForLeaseAsync(IAgent? Agent, Task AbortTask)
+		{
+			while (Agent != null)
+			{
+				// Check there's an active session
+				if (!Agent.SessionExpiresAt.HasValue)
+				{
+					break;
+				}
+
+				// Check we have some time to wait
+				DateTime UtcNow = Clock.UtcNow;
+				TimeSpan MaxWaitTime = (Agent.SessionExpiresAt.Value - SessionExpiryTime + SessionLongPollTime) - UtcNow;
+				if(MaxWaitTime < TimeSpan.Zero)
+				{
+					break;
+				}
+
+				// If we're in a maintenance window, just wait for the time to expire
+				if (DowntimeService.IsDowntimeActive)
+				{
+					await Task.WhenAny(AbortTask, Task.Delay(MaxWaitTime));
+					break;
+				}
+
+				// Subscribe for any new tasks becoming available
+				List<ITaskListener> Subscriptions = new List<ITaskListener>();
+				try
+				{
+					Task WaitTask = Task.Delay(MaxWaitTime, ApplicationLifetime.ApplicationStopping);
+
+					// Create a list of tasks to wait for
+					List<Task> Tasks = new List<Task>();
+					Tasks.Add(WaitTask);
+					Tasks.Add(AbortTask);
+
+					// Try to create a waiter for each task source
+					foreach (ITaskSource TaskSource in TaskSources)
+					{
+						ITaskListener? NewSubscription = await TaskSource.SubscribeAsync(Agent);
+						if (NewSubscription != null)
+						{
+							Subscriptions.Add(NewSubscription);
+							if (!NewSubscription.LeaseTask.IsCompleted || NewSubscription.LeaseTask.Result != null)
+							{
+								Tasks.Add(NewSubscription.LeaseTask);
+							}
+							if (NewSubscription.LeaseTask.IsCompleted)
+							{
+								break;
+							}
+						}
+					}
+
+					// Wait for any lease to complete
+					Task? Result = Tasks.FirstOrDefault(x => x.IsCompleted);
+					if (Result == null)
+					{
+						try
+						{
+							Logger.LogDebug("Start long poll for agent {AgentId}", Agent.Id);
+							Result = await Task.WhenAny(Tasks);
+							Logger.LogDebug("Long poll for agent {AgentId} stopped by {Idx}", Agent.Id, Tasks.IndexOf(Result));
+						}
+						catch (OperationCanceledException)
+						{
+							Logger.LogDebug("Stopped long poll (application stopping)");
+							break;
+						}
+					}
+
+					// Log the outcome of the log poll
+					if (Result == WaitTask)
+					{
+						Logger.LogDebug("Stopped long poll (timeout)");
+						break;
+					}
+					else if (Result == AbortTask)
+					{
+						Logger.LogDebug("Stopped long poll (client abort)");
+						break;
+					}
+
+					// If we got some new work, try to add a lease for it
+					Logger.LogDebug("Stopped long poll (found work)");
+
+					ITaskListener? Subscription = Subscriptions.FirstOrDefault(x => x.LeaseTask.IsCompleted);
+					if (Subscription != null)
+					{
+						NewLeaseInfo? NewAgentLease = await Subscription.LeaseTask;
+						if (NewAgentLease != null)
+						{
+							IAgent? NewAgent = await Agents.TryAddLeaseAsync(Agent, NewAgentLease.Lease);
+							if(NewAgent != null)
+							{
+								Subscription.Accepted = true;
+								if (NewAgentLease.OnConnectionLost != null)
+								{
+									AddListener(Agent.Id, NewAgentLease.Lease.Id, NewAgentLease.OnConnectionLost);
+								}
+								await CreateLeaseAsync(Agent, NewAgentLease.Lease);
+								return NewAgent;
+							}
+						}
+					}
+				}
+				finally
+				{
+					foreach (ITaskListener Waiter in Subscriptions)
+					{
+						await Waiter.DisposeAsync();
+					}
+				}
+
+				// Update the agent
+				Agent = await GetAgentAsync(Agent.Id);
+			}
+			return Agent;
+		}
+
+		/// <summary>
+		/// Updates the state of the current agent session
+		/// </summary>
+		/// <param name="InAgent">The current agent state</param>
+		/// <param name="SessionId">Id of the session</param>
+		/// <param name="Status">New status for the agent</param>
+		/// <param name="Capabilities">New agent capabilities</param>
+		/// <param name="NewLeases">New list of leases for this session</param>
+		/// <returns>Updated agent state</returns>
+		public async Task<IAgent?> UpdateSessionAsync(IAgent InAgent, ObjectId SessionId, AgentStatus Status, AgentCapabilities? Capabilities, IList<HordeCommon.Rpc.Messages.Lease> NewLeases)
+		{
+			AgentEventListener? Listener = null;
+			try
+			{
+				lock (Listeners)
+				{
+					if (Listeners.TryGetValue(InAgent.Id, out Listener))
+					{
+						Listener.BeginUpdate(InAgent.Leases.Select(x => x.Id));
+					}
+				}
+
+				DateTime UtcNow = Clock.UtcNow;
+
+				Stopwatch Timer = Stopwatch.StartNew();
+
+				IAgent? Agent = InAgent;
+				while (Agent != null)
+				{
+					// If the agent is stopping and doesn't have any leases, we can terminate the current session.
+					if (Status == AgentStatus.Stopping && NewLeases.Count == 0)
+					{
+						// If we've already decided to terminate the session, this update is redundant but harmless
+						if (Agent.SessionId != SessionId)
+						{
+							break;
+						}
+
+						// If the session is valid, we can terminate once the agent leases are also empty
+						if (Agent.Leases.Count == 0)
+						{
+							if (!await TryTerminateSessionAsync(Agent))
+							{
+								Agent = await GetAgentAsync(Agent.Id);
+								continue;
+							}
+							break;
+						}
+					}
+
+					// Check the session id is correct.
+					if (Agent.SessionId != SessionId)
+					{
+						throw new InvalidOperationException($"Invalid agent session {SessionId}");
+					}
+
+					// Check the session hasn't expired
+					if (!Agent.IsSessionValid(UtcNow))
+					{
+						throw new InvalidOperationException("Session has already expired");
+					}
+
+					// Extend the current session time if we're within a time period of the current time expiring. This reduces
+					// unnecessary DB writes a little, but it also allows us to skip the update and jump into a long poll state
+					// if there's still time on the session left.
+					DateTime? SessionExpiresAt = null;
+					if (!Agent.SessionExpiresAt.HasValue || UtcNow > (Agent.SessionExpiresAt - SessionExpiryTime) + SessionRenewTime)
+					{
+						SessionExpiresAt = UtcNow + SessionExpiryTime;
+					}
+
+					// Flag for whether the leases array should be updated
+					bool bUpdateLeases = false;
+					List<AgentLease> Leases = new List<AgentLease>(Agent.Leases);
+
+					// Remove any completed leases from the agent
+					Dictionary<ObjectId, HordeCommon.Rpc.Messages.Lease> LeaseIdToNewState = NewLeases.ToDictionary(x => x.Id.ToObjectId(), x => x);
+					for (int Idx = 0; Idx < Leases.Count; Idx++)
+					{
+						AgentLease Lease = Leases[Idx];
+						if (Lease.State == LeaseState.Cancelled)
+						{
+							HordeCommon.Rpc.Messages.Lease? NewLease;
+							if (!LeaseIdToNewState.TryGetValue(Lease.Id, out NewLease) || NewLease.State == LeaseState.Cancelled || NewLease.State == LeaseState.Completed)
+							{
+								await RemoveLeaseAsync(Agent, Lease, UtcNow, LeaseOutcome.Cancelled);
+								Leases.RemoveAt(Idx--);
+								bUpdateLeases = true;
+							}
+						}
+						else
+						{
+							HordeCommon.Rpc.Messages.Lease? NewLease;
+							if (LeaseIdToNewState.TryGetValue(Lease.Id, out NewLease) && NewLease.State != Lease.State)
+							{
+								if (NewLease.State == LeaseState.Cancelled || NewLease.State == LeaseState.Completed)
+								{
+									await RemoveLeaseAsync(Agent, Lease, UtcNow, NewLease.Outcome);
+									Leases.RemoveAt(Idx--);
+								}
+								else if (NewLease.State == LeaseState.Active && Lease.State == LeaseState.Pending)
+								{
+									Lease.State = LeaseState.Active;
+								}
+								bUpdateLeases = true;
+							}
+						}
+					}
+
+					// If the agent is stopping, cancel all the leases. Clear out the current session once it's complete.
+					if (Status == AgentStatus.Stopping)
+					{
+						foreach (AgentLease Lease in Leases)
+						{
+							if (Lease.State != LeaseState.Cancelled)
+							{
+								Lease.State = LeaseState.Cancelled;
+								bUpdateLeases = true;
+							}
+						}
+					}
+
+					// Update the agent, and try to create new lease documents if we succeed
+					IAgent? NewAgent = await Agents.TryUpdateSessionAsync(Agent, Status, SessionExpiresAt, Capabilities, bUpdateLeases ? Leases : null);
+					if (NewAgent != null)
+					{
+						break;
+					}
+
+					// Fetch the agent again
+					Agent = await GetAgentAsync(Agent.Id);
+				}
+				return Agent;
+			}
+			finally
+			{
+				Listener?.EndUpdate();
+			}
+		}
+
+		void AddListener(AgentId AgentId, ObjectId LeaseId, Action Callback)
+		{
+			lock (Listeners)
+			{
+				AgentEventListener? Listener;
+				if (!Listeners.TryGetValue(AgentId, out Listener) || !Listener.TryAdd(LeaseId, Callback))
+				{
+					Listener = new AgentEventListener(this, AgentId, LeaseId, Callback);
+					Listeners[AgentId] = Listener;
+				}
+			}
+		}
+
+		void RemoveListener(AgentEventListener Listener)
+		{
+			lock (Listeners)
+			{
+				AgentEventListener? CurrentListener;
+				if (Listeners.TryGetValue(Listener.AgentId, out CurrentListener) && CurrentListener == Listener)
+				{
+					Listeners.Remove(Listener.AgentId);
+				}
+			}
+			Listener.Dispose();
+		}
+
+		class AgentEventListener : IDisposable
+		{
+			object LockObject = new object();
+			AgentService AgentService;
+			public AgentId AgentId;
+			Timer? Timer;
+			List<(ObjectId, Action)> LeaseCallbacks = new List<(ObjectId, Action)>();
+
+			public AgentEventListener(AgentService AgentService, AgentId AgentId, ObjectId LeaseId, Action Callback)
+			{
+				this.AgentService = AgentService;
+				this.AgentId = AgentId;
+
+				Timer = new Timer(x => AgentService.RemoveListener(this));
+				LeaseCallbacks.Add((LeaseId, Callback));
+			}
+
+			public void Dispose()
+			{
+				if (Timer != null)
+				{
+					BeginUpdate(Array.Empty<ObjectId>());
+					Timer.Dispose();
+					Timer = null;
+				}
+			}
+
+			public bool TryAdd(ObjectId LeaseId, Action Callback)
+			{
+				lock (LockObject)
+				{
+					if (LeaseCallbacks.Count > 0)
+					{
+						LeaseCallbacks.Add((LeaseId, Callback));
+						return true;
+					}
+				}
+				return false;
+			}
+
+			public void BeginUpdate(IEnumerable<ObjectId> LeaseIds)
+			{
+				lock (LockObject)
+				{
+					if (LeaseCallbacks.Count > 0)
+					{
+						// Execute any callbacks for leases that have completed
+						for (int Idx = LeaseCallbacks.Count - 1; Idx >= 0; Idx--)
+						{
+							if (!LeaseIds.Contains(LeaseCallbacks[Idx].Item1))
+							{
+								LeaseCallbacks[Idx].Item2();
+								LeaseCallbacks.RemoveAt(Idx);
+							}
+						}
+
+						// Return false when we don't have any remaining callbacks
+						if (LeaseCallbacks.Count == 0)
+						{
+							AgentService.RemoveListener(this);
+						}
+						else
+						{
+							Timer?.Change(Timeout.Infinite, Timeout.Infinite);
+						}
+					}
+				}
+			}
+
+			public void EndUpdate()
+			{
+				lock (LockObject)
+				{
+					Timer?.Change(20 * 1000, Timeout.Infinite);
+				}
+			}
+		}
+
+		class AgentUpdateScope : IDisposable
+		{
+			AgentEventListener Listener;
+
+			public AgentUpdateScope(AgentEventListener Listener)
+			{
+				this.Listener = Listener;
+			}
+
+			public void Dispose()
+			{
+				Listener.EndUpdate();
+			}
+		}
+
+		/// <summary>
+		/// Terminates an existing session. Does not update the agent itself, if it's currently 
+		/// </summary>
+		/// <param name="Agent">The agent whose current session should be terminated</param>
+		/// <returns>True if the session was terminated</returns>
+		private async Task<bool> TryTerminateSessionAsync(IAgent Agent)
+		{
+			// Make sure the agent has a valid session id
+			if (Agent.SessionId == null)
+			{
+				return true;
+			}
+
+			// Get the time that the session finishes at
+			DateTime FinishTime = Clock.UtcNow;
+			if (Agent.SessionExpiresAt.HasValue && Agent.SessionExpiresAt.Value < FinishTime)
+			{
+				FinishTime = Agent.SessionExpiresAt.Value;
+			}
+
+			// Save off the session id and current leases
+			ObjectId SessionId = Agent.SessionId.Value;
+			List<AgentLease> Leases = new List<AgentLease>(Agent.Leases);
+
+			// Clear the current session
+			IAgent? NewAgent = await Agents.TryTerminateSessionAsync(Agent);
+			if (NewAgent != null)
+			{
+				Agent = NewAgent;
+
+				// Remove any outstanding leases
+				foreach(AgentLease Lease in Leases)
+				{
+					Logger.LogDebug("Removing lease {LeaseId} during session terminate...", Lease.Id);
+					await RemoveLeaseAsync(Agent, Lease, FinishTime, LeaseOutcome.Failed);
+				}
+
+				// Update the session document
+				await Sessions.UpdateAsync(SessionId, FinishTime, Agent.Capabilities);
+				return true;
+			}
+			return false;
+		}
+
+		/// <summary>
+		/// Creates a new lease document
+		/// </summary>
+		/// <param name="Agent">Agent that will be executing the lease</param>
+		/// <param name="AgentLease">The new agent lease</param>
+		/// <returns>New lease document</returns>
+		private Task<ILease> CreateLeaseAsync(IAgent Agent, AgentLease AgentLease)
+		{
+			return Leases.AddAsync(AgentLease.Id, AgentLease.Name, Agent.Id, Agent.SessionId!.Value, AgentLease.StreamId, AgentLease.PoolId, AgentLease.LogId, AgentLease.StartTime, AgentLease.Payload!);
+		}
+
+		/// <summary>
+		/// Finds all leases matching a set of criteria
+		/// </summary>
+		/// <param name="AgentId">Unqiue id of the agent executing this lease</param>
+		/// <param name="SessionId">Unique id of the agent session</param>
+		/// <param name="StartTime">Start of the search window to return results for</param>
+		/// <param name="FinishTime">End of the search window to return results for</param>
+		/// <param name="Index">Index of the first result to return</param>
+		/// <param name="Count">Number of results to return</param>
+		/// <returns>List of leases matching the given criteria</returns>
+		public Task<List<ILease>> FindLeasesAsync(AgentId? AgentId, ObjectId? SessionId, DateTime? StartTime, DateTime? FinishTime, int Index, int Count)
+		{
+			return Leases.FindLeasesAsync(AgentId, SessionId, StartTime, FinishTime, Index, Count);
+		}
+
+		/// <summary>
+		/// Gets a specific lease
+		/// </summary>
+		/// <param name="LeaseId">Unique id of the lease</param>
+		/// <returns>The lease that was found, or null if it does not exist</returns>
+		public Task<ILease?> GetLeaseAsync(ObjectId LeaseId)
+		{
+			return Leases.GetAsync(LeaseId);
+		}
+
+		/// <summary>
+		/// Removes a lease with the given id. Updates the lease state in the database, and removes the item from the agent's leases array.
+		/// </summary>
+		/// <param name="Agent">The agent to remove a lease from</param>
+		/// <param name="Lease">The lease to cancel</param>
+		/// <param name="UtcNow">The current time</param>
+		/// <param name="Outcome">Final status of the lease</param>
+		/// <returns>Async task</returns>
+		private async Task RemoveLeaseAsync(IAgent Agent, AgentLease Lease, DateTime UtcNow, LeaseOutcome Outcome)
+		{
+			// Make sure the lease is terminated correctly
+			if (Lease.Payload == null)
+			{
+				Logger.LogWarning("Removing lease {LeaseId} (no payload)", Lease.Id);
+			}
+			else
+			{
+				Any Any = Any.Parser.ParseFrom(Lease.Payload);
+				Logger.LogInformation("Removing lease {LeaseId} ({LeaseType})", Lease.Id, Any.TypeUrl);
+
+				foreach (ITaskSource TaskSource in TaskSources)
+				{
+					if (Any.Is(TaskSource.Descriptor))
+					{
+						await TaskSource.AbortTaskAsync(Agent, Lease.Id, Any);
+						break;
+					}
+				}
+			}
+
+			// Figure out what time the lease finished
+			DateTime FinishTime = UtcNow;
+			if (Agent.SessionExpiresAt.HasValue && Agent.SessionExpiresAt.Value < FinishTime)
+			{
+				FinishTime = Agent.SessionExpiresAt.Value;
+			}
+			if (Lease.ExpiryTime.HasValue && Lease.ExpiryTime.Value < FinishTime)
+			{
+				FinishTime = Lease.ExpiryTime.Value;
+			}
+
+			// Update the lease
+			await Leases.TrySetOutcomeAsync(Lease.Id, FinishTime, Outcome);
+		}
+
+		/// <summary>
+		/// Gets information about a particular session
+		/// </summary>
+		/// <param name="SessionId">The unique session id</param>
+		/// <returns>The session information</returns>
+		public Task<ISession?> GetSessionAsync(ObjectId SessionId)
+		{
+			return Sessions.GetAsync(SessionId);
+		}
+
+		/// <summary>
+		/// Find sessions for the given agent
+		/// </summary>
+		/// <param name="AgentId">The unique agent id</param>
+		/// <param name="StartTime">Start time to include in the search</param>
+		/// <param name="FinishTime">Finish time to include in the search</param>
+		/// <param name="Index">Index of the first result to return</param>
+		/// <param name="Count">Number of results to return</param>
+		/// <returns>List of sessions matching the given criteria</returns>
+		public Task<List<ISession>> FindSessionsAsync(AgentId AgentId, DateTime? StartTime, DateTime? FinishTime, int Index, int Count)
+		{
+			return Sessions.FindAsync(AgentId, StartTime, FinishTime, Index, Count);
+		}
+
+		/// <summary>
+		/// Terminate any sessions for agents that are offline
+		/// </summary>
+		/// <param name="StoppingToken">Token indicating the service is shutting down</param>
+		/// <returns>Async task</returns>
+		protected override async Task TickAsync(CancellationToken StoppingToken)
+		{
+			while (!StoppingToken.IsCancellationRequested)
+			{
+				// Find all the agents which are ready to be expired
+				const int MaxAgents = 100;
+				DateTime UtcNow = Clock.UtcNow;
+				List<IAgent> ExpiredAgents = await Agents.FindExpiredAsync(UtcNow, MaxAgents);
+
+				// Transition each agent to being offline
+				foreach (IAgent ExpiredAgent in ExpiredAgents)
+				{
+					StoppingToken.ThrowIfCancellationRequested();
+					Logger.LogDebug("Terminating session {SessionId} for agent {Agent}", ExpiredAgent.SessionId, ExpiredAgent.Id);
+					await TryTerminateSessionAsync(ExpiredAgent);
+				}
+
+				// Try again if we didn't fetch everything
+				if(ExpiredAgents.Count < MaxAgents)
+				{
+					break;
+				}
+			}
+
+			await CollectMetrics();
+		}
+
+		private async Task CollectMetrics()
+		{
+			List<IAgent> AgentList = await Agents.FindAsync();
+			int NumAgentsTotal = AgentList.Count;
+			int NumAgentsTotalEnabled = AgentList.Count(a => a.Enabled);
+			int NumAgentsTotalDisabled = AgentList.Count(a => !a.Enabled);
+			int NumAgentsTotalOk = AgentList.Count(a => a.Enabled && a.Status == AgentStatus.Ok);
+			int NumAgentsTotalStopping = AgentList.Count(a => a.Enabled && a.Status == AgentStatus.Stopping);
+			int NumAgentsTotalUnhealthy = AgentList.Count(a => a.Enabled && a.Status == AgentStatus.Unhealthy);
+			int NumAgentsTotalUnspecified = AgentList.Count(a => a.Enabled && a.Status == AgentStatus.Unspecified);
+			
+			// TODO: utilize tags argument in a smarter way below
+			DogStatsd.Gauge("agents.total.count", NumAgentsTotal);
+			DogStatsd.Gauge("agents.total.enabled.count", NumAgentsTotalEnabled);
+			DogStatsd.Gauge("agents.total.disabled.count", NumAgentsTotalDisabled);
+			DogStatsd.Gauge("agents.total.ok.count", NumAgentsTotalOk);
+			DogStatsd.Gauge("agents.total.stopping.count", NumAgentsTotalStopping);
+			DogStatsd.Gauge("agents.total.unhealthy.count", NumAgentsTotalUnhealthy);
+			DogStatsd.Gauge("agents.total.unspecified.count", NumAgentsTotalUnspecified);
+		}
+
+		/// <summary>
+		/// Determines if the user is authorized to perform an action on a particular agent
+		/// </summary>
+		/// <param name="Agent">The agent to check</param>
+		/// <param name="Action">The action being performed</param>
+		/// <param name="User">The principal to authorize</param>
+		/// <param name="Cache">The permissions cache</param>
+		/// <returns>True if the action is authorized</returns>
+		public async Task<bool> AuthorizeAsync(IAgent Agent, AclAction Action, ClaimsPrincipal User, GlobalPermissionsCache? Cache)
+		{
+			bool? Result = Agent.Acl?.Authorize(Action, User);
+			if (Result == null)
+			{
+				return await AclService.AuthorizeAsync(Action, User, Cache);
+			}
+			else
+			{
+				return Result.Value;
+			}
+		}
+
+		/// <summary>
+		/// Determines if the user is authorized to perform an action on a particular agent
+		/// </summary>
+		/// <param name="Agent">The agent to check</param>
+		/// <param name="User">The principal to authorize</param>
+		/// <returns>True if the action is authorized</returns>
+		public bool AuthorizeSession(IAgent Agent, ClaimsPrincipal User)
+		{
+			if (Agent.SessionId != null && User.HasSessionClaim(Agent.SessionId.Value) && Agent.IsSessionValid(Clock.UtcNow))
+			{
+				return true;
+			}
+			else
+			{
+				return false;
+			}
+		}
+	}
+}
diff --git a/Engine/Source/Programs/Horde/HordeServer/Services/AgentSoftwareService.cs b/Engine/Source/Programs/Horde/HordeServer/Services/AgentSoftwareService.cs
new file mode 100644
--- /dev/null
+++ b/Engine/Source/Programs/Horde/HordeServer/Services/AgentSoftwareService.cs
@@ -0,0 +1,262 @@
+// Copyright Epic Games, Inc. All Rights Reserved.
+
+using HordeServer.Collections;
+using HordeServer.Models;
+using HordeServer.Utilities;
+using Microsoft.IdentityModel.Tokens;
+using MongoDB.Bson;
+using MongoDB.Bson.Serialization.Attributes;
+using MongoDB.Driver;
+using System;
+using System.Collections.Generic;
+using System.Diagnostics;
+using System.IO;
+using System.Linq;
+using System.Threading.Channels;
+using System.Threading.Tasks;
+
+using AgentSoftwareVersion = HordeServer.Utilities.StringId<HordeServer.Collections.IAgentSoftwareCollection>;
+using AgentSoftwareChannelName = HordeServer.Utilities.StringId<HordeServer.Services.AgentSoftwareChannels>;
+
+namespace HordeServer.Services
+{
+	/// <summary>
+	/// Information about a channel
+	/// </summary>
+	public class AgentSoftwareChannel : IAgentSoftwareChannel
+	{
+		/// <summary>
+		/// The channel id
+		/// </summary>
+		public AgentSoftwareChannelName Name { get; set; }
+
+		/// <summary>
+		/// Name of the user that made the last modification
+		/// </summary>
+		public string? ModifiedBy { get; set; }
+
+		/// <summary>
+		/// Last modification time
+		/// </summary>
+		public DateTime ModifiedTime { get; set; }
+
+		/// <summary>
+		/// The software revision number
+		/// </summary>
+		public AgentSoftwareVersion Version { get; set; }
+
+		/// <summary>
+		/// Private constructor for serialization
+		/// </summary>
+		[BsonConstructor]
+		private AgentSoftwareChannel()
+		{
+		}
+
+		/// <summary>
+		/// Constructor
+		/// </summary>
+		/// <param name="Name">Name of the channel</param>
+		public AgentSoftwareChannel(AgentSoftwareChannelName Name)
+		{
+			this.Name = Name;
+		}
+	}
+
+	/// <summary>
+	/// Singleton document used to track different versions of the agent software
+	/// </summary>
+	[SingletonDocument("5f455039d97900f2b6c735a9")]
+	public class AgentSoftwareChannels : SingletonBase
+	{
+		/// <summary>
+		/// The next version number
+		/// </summary>
+		public int NextVersion { get; set; } = 1;
+
+		/// <summary>
+		/// List of channels
+		/// </summary>
+		public List<AgentSoftwareChannel> Channels { get; set; } = new List<AgentSoftwareChannel>();
+	}
+
+	/// <summary>
+	/// Wrapper for a collection of software revisions. 
+	/// </summary>
+	public class AgentSoftwareService
+	{
+		/// <summary>
+		/// Name of the default channel
+		/// </summary>
+		public static AgentSoftwareChannelName DefaultChannelName { get; } = new AgentSoftwareChannelName("default");
+
+		/// <summary>
+		/// Collection of software documents
+		/// </summary>
+		IAgentSoftwareCollection Collection;
+
+		/// <summary>
+		/// Channels singleton
+		/// </summary>
+		ISingletonDocument<AgentSoftwareChannels> Singleton;
+
+		/// <summary>
+		/// Cached copy of the channels singleton
+		/// </summary>
+		LazyCachedValue<Task<AgentSoftwareChannels>> ChannelsDocument;
+
+		/// <summary>
+		/// Constructor
+		/// </summary>
+		/// <param name="Collection">The software collection</param>
+		/// <param name="Singleton">The channels singleton</param>
+		public AgentSoftwareService(IAgentSoftwareCollection Collection, ISingletonDocument<AgentSoftwareChannels> Singleton)
+		{
+			this.Collection = Collection;
+			this.Singleton = Singleton;
+
+			ChannelsDocument = new LazyCachedValue<Task<AgentSoftwareChannels>>(() => Singleton.GetAsync(), TimeSpan.FromSeconds(20.0));
+		}
+
+		/// <summary>
+		/// Finds all software archives matching a set of criteria
+		/// </summary>
+		public async Task<List<IAgentSoftwareChannel>> FindChannelsAsync()
+		{
+			AgentSoftwareChannels Current = await ChannelsDocument.GetLatest();
+			return Current.Channels.ConvertAll<IAgentSoftwareChannel>(x => x);
+		}
+
+		/// <summary>
+		/// Gets a single software archive
+		/// </summary>
+		/// <param name="Name">Unique id of the channel to find software for</param>
+		public async Task<IAgentSoftwareChannel?> GetChannelAsync(AgentSoftwareChannelName Name)
+		{
+			AgentSoftwareChannels Current = await ChannelsDocument.GetLatest();
+			return Current.Channels.FirstOrDefault(x => x.Name == Name);
+		}
+
+		/// <summary>
+		/// Gets a cached channel name
+		/// </summary>
+		/// <param name="Name">The cached channel name</param>
+		/// <returns></returns>
+		public async Task<IAgentSoftwareChannel?> GetCachedChannelAsync(AgentSoftwareChannelName Name)
+		{
+			AgentSoftwareChannels Document = await ChannelsDocument.GetCached();
+			return Document.Channels.FirstOrDefault(x => x.Name == Name);
+		}
+
+		/// <summary>
+		/// Removes a channel
+		/// </summary>
+		/// <param name="Name">The channel id</param>
+		/// <returns>Async task</returns>
+		public async Task DeleteChannelAsync(AgentSoftwareChannelName Name)
+		{
+			for (; ; )
+			{
+				AgentSoftwareChannels Current = await Singleton.GetAsync();
+
+				int ChannelIdx = Current.Channels.FindIndex(x => x.Name == Name);
+				if (ChannelIdx == -1)
+				{
+					break;
+				}
+
+				AgentSoftwareVersion Version = Current.Channels[ChannelIdx].Version;
+				Current.Channels.RemoveAt(ChannelIdx);
+
+				if (await Singleton.TryUpdateAsync(Current))
+				{
+					if (!Current.Channels.Any(x => x.Version == Version))
+					{
+						await Collection.RemoveAsync(Version);
+					}
+					break;
+				}
+			}
+		}
+
+		/// <summary>
+		/// Updates a new software revision
+		/// </summary>
+		/// <param name="Name">Name of the channel</param>
+		/// <param name="Author">Name of the user uploading this file</param>
+		/// <param name="Data">The input data stream. This should be a zip archive containing the HordeAgent executable.</param>
+		/// <returns>Unique id for the file</returns>
+		public async Task<AgentSoftwareVersion> SetArchiveAsync(AgentSoftwareChannelName Name, string? Author, byte[] Data)
+		{
+			// Allocate a new version number
+			int VersionNumber;
+			for(; ;)
+			{
+				AgentSoftwareChannels Instance = await Singleton.GetAsync();
+				VersionNumber = Instance.NextVersion++;
+
+				if(await Singleton.TryUpdateAsync(Instance))
+				{
+					break;
+				}
+			}
+
+			// Upload the software
+			AgentSoftwareVersion Version = new AgentSoftwareVersion($"r{VersionNumber}");
+			await Collection.AddAsync(Version, Data);
+
+			// Update the channel
+			for(; ;)
+			{
+				AgentSoftwareChannels Instance = await Singleton.GetAsync();
+
+				AgentSoftwareChannel? Channel = Instance.Channels.FirstOrDefault(x => x.Name == Name);
+				if (Channel == null)
+				{
+					Channel = new AgentSoftwareChannel(Name);
+					Instance.Channels.Add(Channel);
+				}
+
+				Channel.ModifiedBy = Author;
+				Channel.ModifiedTime = DateTime.UtcNow;
+				Channel.Version = Version;
+
+				if (await Singleton.TryUpdateAsync(Instance))
+				{
+					break;
+				}
+			}
+			return Version;
+		}
+
+		/// <summary>
+		/// Gets the zip file with a given version number
+		/// </summary>
+		/// <param name="Name">The channel name</param>
+		/// <returns>Data for the given archive</returns>
+		public async Task<byte[]?> GetArchiveAsync(AgentSoftwareChannelName Name)
+		{
+			AgentSoftwareChannels Instance = await Singleton.GetAsync();
+
+			AgentSoftwareChannel? Channel = Instance.Channels.FirstOrDefault(x => x.Name == Name);
+			if (Channel == null)
+			{
+				return null;
+			}
+			else
+			{
+				return await Collection.GetAsync(Channel.Version);
+			}
+		}
+
+		/// <summary>
+		/// Gets the zip file with a given version number
+		/// </summary>
+		/// <param name="Version">The version</param>
+		/// <returns>Data for the given archive</returns>
+		public Task<byte[]?> GetArchiveAsync(AgentSoftwareVersion Version)
+		{
+			return Collection.GetAsync(Version);
+		}
+	}
+}
diff --git a/Engine/Source/Programs/Horde/HordeServer/Services/ArtifactService.cs b/Engine/Source/Programs/Horde/HordeServer/Services/ArtifactService.cs
new file mode 100644
--- /dev/null
+++ b/Engine/Source/Programs/Horde/HordeServer/Services/ArtifactService.cs
@@ -0,0 +1,538 @@
+using Amazon.Extensions.NETCore.Setup;
+using Amazon.Runtime;
+using Amazon.S3;
+using Amazon.S3.Model;
+using HordeServer.Api;
+using HordeServer.Models;
+using HordeServer.Utilities;
+using Microsoft.Extensions.Configuration;
+using Microsoft.Extensions.Hosting;
+using Microsoft.Extensions.Logging;
+using Microsoft.Extensions.Options;
+using MongoDB.Bson;
+using MongoDB.Driver;
+using System;
+using System.Collections.Generic;
+using System.IO;
+using System.Linq;
+using System.Security.Claims;
+using System.Text;
+using System.Threading;
+using System.Threading.Tasks;
+
+namespace HordeServer.Services
+{
+	/// <summary>
+	/// Interface for a artifact storage implementation
+	/// </summary>
+	public interface IArtifactStorage
+	{
+		/// <summary>
+		/// Gets an artifact from a storage device
+		/// </summary>
+		/// <param name="JobId">the job id</param>
+		/// <param name="StepId">the step id</param>
+		/// <param name="Name">The artifact name</param>
+		/// <returns>Stream to an artifact if found, null otherwise</returns>
+		System.IO.Stream OpenReadStream(ObjectId JobId, SubResourceId? StepId, string Name);
+
+		/// <summary>
+		/// Writes artifact data
+		/// </summary>
+		/// <param name="JobId">the job id</param>
+		/// <param name="StepId">the step id</param>
+		/// <param name="Name">The artifact name</param>
+		/// <param name="Data">The data to write</param>
+		/// <returns>Whether or not the update completed</returns>
+		Task WriteAsync(ObjectId JobId, SubResourceId? StepId, string Name, System.IO.Stream Data);
+	}
+
+	/// <summary>
+	/// FileStorage implementation using the file system
+	/// </summary>
+	public sealed class FSExternalArtifactStorage : IArtifactStorage
+	{
+		/// <summary>
+		/// Base directory for artifacts
+		/// </summary>
+		private readonly DirectoryInfo BaseDir;
+
+		/// <summary>
+		/// Information Logger
+		/// </summary>
+		private readonly ILogger<ArtifactService> Logger;
+
+		/// <summary>
+		/// Constructor
+		/// </summary>
+		/// <param name="CurrentSettings">Current Horde Settings</param>
+		/// <param name="Logger">Logger interface</param>
+		public FSExternalArtifactStorage(ServerSettings CurrentSettings, ILogger<ArtifactService> Logger)
+		{
+			this.Logger = Logger;
+			this.BaseDir = new DirectoryInfo(CurrentSettings.LocalStorageDir);
+
+			Directory.CreateDirectory(this.BaseDir.FullName);
+		}
+
+		/// <inheritdoc/>
+		public System.IO.Stream OpenReadStream(ObjectId JobId, SubResourceId? StepId, string Name)
+		{
+			string FilePath = Path.Combine(BaseDir.ToString(), JobId.ToString(), StepId.ToString() ?? string.Empty, Name);
+			FileInfo ArtifactFile = new FileInfo(FilePath);
+			try
+			{
+				return new FileStream(ArtifactFile.FullName, FileMode.Open, FileAccess.Read, FileShare.ReadWrite, 4096, true);
+			}
+			catch (FileNotFoundException)
+			{
+				Logger.LogError("Tried to find artifact that didn't exist at {0}!", ArtifactFile.FullName);
+				throw;
+			}
+		}
+
+		/// <inheritdoc/>
+		public async Task WriteAsync(ObjectId JobId, SubResourceId? StepId, string Name, System.IO.Stream Data)
+		{
+			string FilePath = Path.Combine(BaseDir.FullName, JobId.ToString(), StepId.ToString() ?? string.Empty, Name);
+			FileInfo ArtifactFile = new FileInfo(FilePath);
+			Directory.CreateDirectory(ArtifactFile.Directory.FullName);
+			using (FileStream FileWriter = File.Open(ArtifactFile.FullName, FileMode.Create, FileAccess.Write, FileShare.Read))
+			{
+				await Data.CopyToAsync(FileWriter);
+			}
+		}
+	}
+
+	/// <summary>
+	/// FileStorage implementation using an s3 bucket
+	/// </summary>
+	public sealed class S3ExternalArtifactStorage : IArtifactStorage
+	{
+		/// <summary>
+		/// S3 Client
+		/// </summary>
+		private readonly IAmazonS3 Client;
+
+		/// <summary>
+		/// S3 bucket name
+		/// </summary>
+		private readonly string BucketName;
+
+		/// <summary>
+		/// Information Logger
+		/// </summary>
+		private readonly ILogger<ArtifactService> Logger;
+
+		/// <summary>
+		/// Constructor
+		/// </summary>
+		/// <param name="CurrentSettings">Current Horde Settings</param>
+		/// <param name="Configuration">Configuration</param>
+		/// <param name="Logger">Logger interface</param>
+		public S3ExternalArtifactStorage(ServerSettings CurrentSettings, IConfiguration Configuration, ILogger<ArtifactService> Logger)
+		{
+			this.Logger = Logger;
+			AWSOptions Options = Configuration.GetAWSOptions();
+			Options.Credentials = GetCredentials(CurrentSettings);
+
+			this.Client = Options.CreateServiceClient<IAmazonS3>();
+			this.BucketName = CurrentSettings.S3ArtifactBucketName;
+		}
+
+		enum AWSCredentialsType
+		{
+			Basic,
+			AssumeRole,
+			AssumeRoleWebIdentity
+		}
+
+		/// <summary>
+		/// Gets auth credentials based on type
+		/// </summary>
+		/// <param name="Settings">Horde settings</param>
+		/// <returns></returns>
+		public AWSCredentials GetCredentials(ServerSettings Settings)
+		{
+			AWSCredentialsType AuthType;
+			if (!Enum.TryParse(Settings.S3CredentialType, true, out AuthType))
+			{
+				Logger.LogError("Could not determine auth type from appsettings. Should be Basic, AssumeRole, or AssumeRoleWebIdentity!");
+				return null!;
+			}
+			switch (AuthType)
+			{
+				case AWSCredentialsType.Basic:
+					return new BasicAWSCredentials(Settings.S3ClientKeyId, Settings.S3ClientSecret);
+				case AWSCredentialsType.AssumeRole:
+					return new AssumeRoleAWSCredentials(FallbackCredentialsFactory.GetCredentials(), Settings.S3AssumeArn, "Horde");
+				case AWSCredentialsType.AssumeRoleWebIdentity:
+					return AssumeRoleWithWebIdentityCredentials.FromEnvironmentVariables();
+				default:
+					return null!;
+			}
+		}
+
+		/// <inheritdoc/>
+		public System.IO.Stream OpenReadStream(ObjectId JobId, SubResourceId? StepId, string Name)
+		{
+			string BucketPath = Path.Combine(JobId.ToString(), StepId.ToString() ?? string.Empty, Name);
+			try
+			{
+				GetObjectRequest NewGetRequest = new GetObjectRequest
+				{
+					BucketName = this.BucketName,
+					Key = BucketPath
+				};
+
+				return this.Client.GetObjectAsync(NewGetRequest).Result.ResponseStream;
+			}
+			catch (AmazonS3Exception Ex)
+			{
+				if (Ex.ErrorCode == "NoSuchKey")
+				{
+					Logger.LogError("Tried to find artifact that didn't exist at {0}!", BucketPath);
+				}
+				throw;
+			}
+		}
+
+		/// <inheritdoc/>
+		public async Task WriteAsync(ObjectId JobId, SubResourceId? StepId, string Name, System.IO.Stream Stream)
+		{
+			const int MinPartSize = 5 * 1024 * 1024;
+
+			string Key = Path.Combine(JobId.ToString(), StepId.ToString() ?? string.Empty, Name);
+
+			long StreamLen = Stream.Length;
+			if (StreamLen < MinPartSize)
+			{
+				// Read the data into memory
+				byte[] Buffer = new byte[StreamLen];
+				await ReadExactLengthAsync(Stream, Buffer, (int)StreamLen);
+
+				// Upload it to S3
+				using (MemoryStream InputStream = new MemoryStream(Buffer))
+				{
+					PutObjectRequest UploadRequest = new PutObjectRequest();
+					UploadRequest.BucketName = BucketName;
+					UploadRequest.Key = Key;
+					UploadRequest.InputStream = InputStream;
+					await Client.PutObjectAsync(UploadRequest);
+				}
+			}
+			else
+			{
+				// Initiate a multi-part upload
+				InitiateMultipartUploadRequest InitiateRequest = new InitiateMultipartUploadRequest();
+				InitiateRequest.BucketName = BucketName;
+				InitiateRequest.Key = Key;
+
+				InitiateMultipartUploadResponse InitiateResponse = await Client.InitiateMultipartUploadAsync(InitiateRequest);
+				try
+				{
+					// Buffer for reading the data
+					byte[] Buffer = new byte[MinPartSize];
+
+					// Upload all the parts
+					List<PartETag> PartTags = new List<PartETag>();
+					for(long StreamPos = 0; StreamPos < StreamLen; )
+					{
+						// Read the next chunk of data into the buffer
+						int BufferLen = (int)Math.Min((long)MinPartSize, StreamLen - StreamPos);
+						await ReadExactLengthAsync(Stream, Buffer, BufferLen);
+						StreamPos += BufferLen;
+
+						// Upload the part
+						using (MemoryStream InputStream = new MemoryStream(Buffer, 0, BufferLen))
+						{
+							UploadPartRequest PartRequest = new UploadPartRequest();
+							PartRequest.BucketName = BucketName;
+							PartRequest.Key = Key;
+							PartRequest.UploadId = InitiateResponse.UploadId;
+							PartRequest.InputStream = InputStream;
+							PartRequest.PartSize = BufferLen;
+							PartRequest.PartNumber = PartTags.Count + 1;
+							PartRequest.IsLastPart = (StreamPos == StreamLen);
+
+							UploadPartResponse PartResponse = await Client.UploadPartAsync(PartRequest);
+							PartTags.Add(new PartETag(PartResponse.PartNumber, PartResponse.ETag));
+						}
+					}
+
+					// Mark the upload as complete
+					CompleteMultipartUploadRequest CompleteRequest = new CompleteMultipartUploadRequest();
+					CompleteRequest.BucketName = BucketName;
+					CompleteRequest.Key = Key;
+					CompleteRequest.UploadId = InitiateResponse.UploadId;
+					CompleteRequest.PartETags = PartTags;
+					await Client.CompleteMultipartUploadAsync(CompleteRequest);
+				}
+				catch
+				{
+					// Abort the upload
+					AbortMultipartUploadRequest AbortRequest = new AbortMultipartUploadRequest();
+					AbortRequest.BucketName = BucketName;
+					AbortRequest.Key = Key;
+					AbortRequest.UploadId = InitiateResponse.UploadId;
+					await Client.AbortMultipartUploadAsync(AbortRequest);
+
+					throw;
+				}
+			}
+		}
+
+		/// <summary>
+		/// Reads data of an exact length into a stream
+		/// </summary>
+		/// <param name="Stream">The stream to read from</param>
+		/// <param name="Buffer">The buffer to read into</param>
+		/// <param name="Length">Length of the data to read</param>
+		/// <returns>Async task</returns>
+		static async Task ReadExactLengthAsync(System.IO.Stream Stream, byte[] Buffer, int Length)
+		{
+			int BufferPos = 0;
+			while (BufferPos < Length)
+			{
+				int BytesRead = await Stream.ReadAsync(Buffer, BufferPos, Length - BufferPos);
+				if (BytesRead == 0)
+				{
+					throw new InvalidOperationException("Unexpected end of stream");
+				}
+				BufferPos += BytesRead;
+			}
+		}
+	}
+
+	/// <summary>
+	/// Wraps functionality for manipulating artifacts
+	/// </summary>
+	public class ArtifactService
+	{
+		private enum ProviderTypes
+		{
+			S3,
+			FileSystem
+		};
+
+		/// <summary>
+		/// Backend storage container
+		/// </summary>
+		private readonly IArtifactStorage ExternalStorageProvider;
+
+		/// <summary>
+		/// Information Logger
+		/// </summary>
+		private readonly ILogger<ArtifactService> Logger;
+
+		/// <summary>
+		/// Collection of log documents
+		/// </summary>
+		private readonly IMongoCollection<Artifact> Artifacts;
+
+		/// <summary>
+		/// Constructor
+		/// </summary>
+		/// <param name="DatabaseService">The database service</param>
+		/// <param name="Configuration">The configuration</param>
+		/// <param name="Settings">Connection settings</param>
+		/// <param name="Logger">Log interface</param>
+		public ArtifactService(DatabaseService DatabaseService, IOptionsMonitor<ServerSettings> Settings, IConfiguration Configuration, ILogger<ArtifactService> Logger)
+		{
+			this.Logger = Logger;
+
+			// Initialize Cache provider
+			ServerSettings CurrentSettings = Settings.CurrentValue;
+			switch (CurrentSettings.ExternalStorageProviderType)
+			{
+				case StorageProviderType.FileSystem:
+					this.ExternalStorageProvider = new FSExternalArtifactStorage(CurrentSettings, Logger);
+					break;
+				case StorageProviderType.S3:
+					this.ExternalStorageProvider = new S3ExternalArtifactStorage(CurrentSettings, Configuration, Logger);
+					break;
+				default:
+					// shouldn't be possible to get here, but default to file system worst case.
+					this.ExternalStorageProvider = new FSExternalArtifactStorage(CurrentSettings, Logger);
+					break;
+			}
+
+			// Initialize Artifacts table
+			Artifacts = DatabaseService.Artifacts;
+			if (!DatabaseService.ReadOnlyMode)
+			{
+				Artifacts.Indexes.CreateOne(new CreateIndexModel<Artifact>(Builders<Artifact>.IndexKeys.Ascending(x => x.JobId)));
+			}
+		}
+
+		/// <summary>
+		/// Creates a new artifact
+		/// </summary>
+		/// <param name="JobId">Unique id of the job that owns this artifact</param>
+		/// <param name="StepId">Optional Step id</param>
+		/// <param name="Name">Name of artifact</param>
+		/// <param name="MimeType">Type of artifact</param>
+		/// <param name="Data">The data to write</param>
+		/// <returns>The new log file document</returns>
+		public async Task<Artifact> CreateArtifactAsync(ObjectId JobId, SubResourceId? StepId, string Name, string MimeType, System.IO.Stream Data)
+		{
+			// upload first
+			string ArtifactName = ValidateName(Name);
+			await ExternalStorageProvider.WriteAsync(JobId, StepId, ArtifactName, Data);
+
+			// then create entry
+			Artifact NewArtifact = new Artifact(JobId, StepId, ArtifactName, Data.Length, MimeType);
+			await Artifacts.InsertOneAsync(NewArtifact);
+			return NewArtifact;
+		}
+
+		/// <summary>
+		/// Gets all the available artifacts for a job
+		/// </summary>
+		/// <param name="JobId">Unique id of the job to query</param>
+		/// <param name="StepId">Unique id of the Step to query</param>
+		/// <param name="Name">Name of the artifact</param>
+		/// <returns>List of artifact documents</returns>
+		public async Task<List<Artifact>> GetArtifactsAsync(ObjectId? JobId, SubResourceId? StepId, string? Name)
+		{
+			FilterDefinitionBuilder<Artifact> Builder = Builders<Artifact>.Filter;
+
+			FilterDefinition<Artifact> Filter = FilterDefinition<Artifact>.Empty;
+			if (JobId != null)
+			{
+				Filter &= Builder.Eq(x => x.JobId, JobId.Value);
+			}
+			if (StepId != null)
+			{
+				Filter &= Builder.Eq(x => x.StepId, StepId.Value);
+			}
+			if (Name != null)
+			{
+				Filter &= Builder.Eq(x => x.Name, Name);
+			}
+
+			return await Artifacts.Find(Filter).ToListAsync();
+		}
+
+		/// <summary>
+		/// Gets a specific list of artifacts based on id
+		/// </summary>
+		/// <param name="ArtifactIds">The list of artifact Ids</param>
+		/// <returns>List of artifact documents</returns>
+		public async Task<List<Artifact>> GetArtifactsAsync(IEnumerable<ObjectId> ArtifactIds)
+		{
+			FilterDefinitionBuilder<Artifact> Builder = Builders<Artifact>.Filter;
+
+			FilterDefinition<Artifact> Filter = FilterDefinition<Artifact>.Empty;
+			Filter &= Builder.In(x => x.Id, ArtifactIds);
+			
+			return await Artifacts.Find(Filter).ToListAsync();
+		}
+
+		/// <summary>
+		/// Gets an artifact by ID
+		/// </summary>
+		/// <param name="ArtifactId">Unique id of the artifact</param>
+		/// <returns>The artifact document</returns>
+		public async Task<Artifact?> GetArtifactAsync(ObjectId ArtifactId)
+		{
+			return await Artifacts.Find<Artifact>(x => x.Id == ArtifactId).FirstOrDefaultAsync();
+		}
+
+		/// <summary>
+		/// Attempts to update a single artifact document, if the update counter is valid
+		/// </summary>
+		/// <param name="Current">The current artifact document</param>
+		/// <param name="Update">The update definition</param>
+		/// <returns>True if the document was updated, false if another writer updated the document first</returns>
+		private async Task<bool> TryUpdateArtifactAsync(Artifact Current, UpdateDefinition<Artifact> Update)
+		{
+			UpdateResult Result = await Artifacts.UpdateOneAsync<Artifact>(x => x.Id == Current.Id && x.UpdateIndex == Current.UpdateIndex, Update.Set(x => x.UpdateIndex, Current.UpdateIndex + 1));
+			return Result.ModifiedCount == 1;
+		}
+
+		/// <summary>
+		/// Updates an artifact
+		/// </summary>
+		/// <param name="Artifact">The artifact</param>
+		/// <param name="NewMimeType">New mime type</param>
+		/// <param name="NewData">New data</param>
+		/// <returns>Async task</returns>
+		public async Task<bool> UpdateArtifactAsync(Artifact? Artifact, string NewMimeType, System.IO.Stream NewData)
+		{
+			while (Artifact != null)
+			{
+				UpdateDefinitionBuilder<Artifact> UpdateBuilder = Builders<Artifact>.Update;
+
+				List<UpdateDefinition<Artifact>> Updates = new List<UpdateDefinition<Artifact>>();
+				Updates.Add(UpdateBuilder.Set(x => x.MimeType, NewMimeType));
+				Updates.Add(UpdateBuilder.Set(x => x.Length, NewData.Length));
+
+				// re-upload the data to external
+				string ArtifactName = ValidateName(Artifact.Name);
+				await ExternalStorageProvider.WriteAsync(Artifact.JobId, Artifact.StepId, ArtifactName, NewData);
+
+				if (await TryUpdateArtifactAsync(Artifact, UpdateBuilder.Combine(Updates)))
+				{
+					return true;
+				}
+
+				Artifact = await GetArtifactAsync(Artifact.Id);
+			}
+			return false;
+		}
+
+		/// <summary>
+		/// gets artifact data
+		/// </summary>
+		/// <param name="Artifact">The artifact</param>
+		/// <returns>The chunk data</returns>
+		public System.IO.Stream OpenArtifactReadStream(Artifact Artifact)
+		{
+			return ExternalStorageProvider.OpenReadStream(Artifact.JobId, Artifact.StepId, Artifact.Name);
+		}
+
+		/// <summary>
+		/// Checks that the name given is valid for an artifact
+		/// </summary>
+		/// <param name="Name"></param>
+		private static string ValidateName(string Name)
+		{
+			string NewName = Name.Replace('\\', '/');
+			if (NewName.Length == 0)
+			{
+				throw new ArgumentException("Artifact name is empty");
+			}
+			else if (NewName.StartsWith('/'))
+			{
+				throw new ArgumentException($"Artifact has an absolute path ({NewName})");
+			}
+			else if (NewName.EndsWith('/'))
+			{
+				throw new ArgumentException($"Artifact does not have a file name ({NewName})");
+			}
+			else if (NewName.Contains("//", StringComparison.Ordinal))
+			{
+				throw new ArgumentException($"Artifact name contains an invalid directory ({NewName})");
+			}
+
+			string InvalidChars = ":<>|\"";
+			for (int Idx = 0; Idx < NewName.Length; Idx++)
+			{
+				int MinIdx = Idx;
+				for (; Idx < NewName.Length && NewName[Idx] != '/'; Idx++)
+				{
+					if (InvalidChars.Contains(NewName[Idx], StringComparison.Ordinal))
+					{
+						throw new ArgumentException($"Invalid character in artifact name ({NewName})");
+					}
+				}
+				if ((Idx == MinIdx + 1 && NewName[MinIdx] == '.') || (Idx == MinIdx + 2 && NewName[MinIdx] == '.' && NewName[MinIdx + 1] == '.'))
+				{
+					throw new ArgumentException($"Artifact may not contain symbolic directory names ({NewName})");
+				}
+			}
+			return NewName;
+		}
+	}
+}
\ No newline at end of file
diff --git a/Engine/Source/Programs/Horde/HordeServer/Services/ConsistencyService.cs b/Engine/Source/Programs/Horde/HordeServer/Services/ConsistencyService.cs
new file mode 100644
--- /dev/null
+++ b/Engine/Source/Programs/Horde/HordeServer/Services/ConsistencyService.cs
@@ -0,0 +1,67 @@
+using HordeCommon;
+using HordeServer.Collections;
+using HordeServer.Models;
+using HordeServer.Utilities;
+using Microsoft.Extensions.Logging;
+using MongoDB.Bson;
+using System;
+using System.Collections.Generic;
+using System.Linq;
+using System.Threading;
+using System.Threading.Tasks;
+
+namespace HordeServer.Services
+{
+	/// <summary>
+	/// Service which checks the database for consistency and fixes up any errors
+	/// </summary>
+	class ConsistencyService : ElectedBackgroundService
+	{
+		ISessionCollection SessionCollection;
+		ILeaseCollection LeaseCollection;
+		ILogger<ConsistencyService> Logger;
+
+		/// <summary>
+		/// Constructor
+		/// </summary>
+		/// <param name="DatabaseService">Database service</param>
+		/// <param name="SessionCollection"></param>
+		/// <param name="LeaseCollection"></param>
+		/// <param name="Logger">Logging device</param>
+		public ConsistencyService(DatabaseService DatabaseService, ISessionCollection SessionCollection, ILeaseCollection LeaseCollection, ILogger<ConsistencyService> Logger)
+			: base(DatabaseService, new ObjectId("5fdb9d9f39b1291512821ef8"), Logger)
+		{
+			this.SessionCollection = SessionCollection;
+			this.LeaseCollection = LeaseCollection;
+			this.Logger = Logger;
+		}
+
+		/// <summary>
+		/// Poll for inconsistencies in the database
+		/// </summary>
+		/// <param name="StoppingToken">Stopping token</param>
+		/// <returns>Async task</returns>
+		protected override async Task<DateTime> TickLeaderAsync(CancellationToken StoppingToken)
+		{
+			DateTime NextTickTime = DateTime.UtcNow + TimeSpan.FromMinutes(20.0);
+
+			List<ISession> Sessions = await SessionCollection.FindActiveSessionsAsync();
+			Dictionary<ObjectId, ISession> SessionIdToInstance = Sessions.ToDictionary(x => x.Id, x => x);
+
+			// Find any leases that are still running when their session has terminated
+			List<ILease> Leases = await LeaseCollection.FindActiveLeasesAsync();
+			foreach (ILease Lease in Leases)
+			{
+				if (!SessionIdToInstance.ContainsKey(Lease.SessionId))
+				{
+					ISession? Session = await SessionCollection.GetAsync(Lease.SessionId);
+					DateTime FinishTime = Session?.FinishTime ?? DateTime.UtcNow;
+					Logger.LogWarning("Setting finish time for lease {LeaseId} to {FinishTime}", Lease.Id, FinishTime);
+					await LeaseCollection.TrySetOutcomeAsync(Lease.Id, FinishTime, LeaseOutcome.Cancelled);
+				}
+			}
+
+			return NextTickTime;
+		}
+	}
+}
diff --git a/Engine/Source/Programs/Horde/HordeServer/Services/DatabaseService.cs b/Engine/Source/Programs/Horde/HordeServer/Services/DatabaseService.cs
new file mode 100644
--- /dev/null
+++ b/Engine/Source/Programs/Horde/HordeServer/Services/DatabaseService.cs
@@ -0,0 +1,530 @@
+// Copyright Epic Games, Inc. All Rights Reserved.
+
+using EpicGames.Core;
+using HordeServer.Api;
+using HordeCommon;
+using HordeServer.Models;
+using HordeServer.Utilities;
+using Microsoft.Extensions.Logging;
+using Microsoft.Extensions.Options;
+using Microsoft.IdentityModel.Tokens;
+using MongoDB.Bson;
+using MongoDB.Bson.Serialization;
+using MongoDB.Driver;
+using System;
+using System.Collections.Concurrent;
+using System.Collections.Generic;
+using System.IO;
+using System.Linq;
+using System.Linq.Expressions;
+using System.Net;
+using System.Net.Security;
+using System.Runtime.InteropServices;
+using System.Security.Cryptography.X509Certificates;
+using System.Text;
+using System.Threading.Tasks;
+using IStream = HordeServer.Models.IStream;
+
+using PoolId = HordeServer.Utilities.StringId<HordeServer.Models.IPool>;
+using ProjectId = HordeServer.Utilities.StringId<HordeServer.Models.IProject>;
+using TemplateRefId = HordeServer.Utilities.StringId<HordeServer.Models.TemplateRef>;
+using System.Net.Sockets;
+using MongoDB.Driver.Core.Events;
+using System.Threading;
+using Datadog.Trace;
+using HordeServer.Collections;
+
+namespace HordeServer.Services
+{
+	/// <summary>
+	/// Singleton for accessing the database
+	/// </summary>
+	public class DatabaseService
+	{
+		/// <summary>
+		/// The database instance
+		/// </summary>
+		public IMongoDatabase Database { get; private set; }
+
+		/// <summary>
+		/// Collection of credential documents
+		/// </summary>
+		public IMongoCollection<Credential> Credentials { get; }
+
+		/// <summary>
+		/// Collection of artifact documents
+		/// </summary>
+		public IMongoCollection<Artifact> Artifacts { get; }
+
+		/// <summary>
+		/// Collection of singleton documents
+		/// </summary>
+		public IMongoCollection<BsonDocument> Singletons { get; }
+
+		/// <summary>
+		/// Settings for the application
+		/// </summary>
+		IOptionsMonitor<ServerSettings> Settings;
+
+		/// <summary>
+		/// Name of the JWT issuer
+		/// </summary>
+		public string JwtIssuer { get; }
+
+		/// <summary>
+		/// The signing key for this instance
+		/// </summary>
+		public SymmetricSecurityKey JwtSigningKey { get; }
+
+		/// <summary>
+		/// Logger for this instance
+		/// </summary>
+		ILogger<DatabaseService> Logger;
+		
+		/// <summary>
+		/// Access the database in a read-only mode (don't create indices or modify content)
+		/// </summary>
+		public bool ReadOnlyMode { get; }
+
+		/// <summary>
+		/// Trace MongoDB commands in Datadog for observability
+		/// The built-in version from Datadog tracing library doesn't work for us (missing queries).
+		/// </summary>
+		class MongoDatadogTracer : IEventSubscriber
+		{
+			readonly ConcurrentDictionary<int, Scope> Scopes = new ConcurrentDictionary<int, Scope>();
+			private readonly ReflectionEventSubscriber _subscriber;
+
+			public MongoDatadogTracer()
+			{
+				_subscriber = new ReflectionEventSubscriber(this);
+			}
+
+			public bool TryGetEventHandler<TEvent>(out Action<TEvent> handler)
+			{
+				return _subscriber.TryGetEventHandler(out handler);
+			}
+			
+			public void Handle(CommandStartedEvent e)
+			{
+				Scope Scope = Tracer.Instance.StartActive("mongodb.command");
+				string? DbName = null;
+				string? CollectionName = null;
+				string? Query = null;
+
+				string? GetCollectionName(string CommandName, BsonDocument Command)
+				{
+					return Command.TryGetValue(CommandName, out BsonValue Value) ? Value.AsString : null;
+				}
+				
+				switch (e.CommandName)
+				{
+					case "find":
+						CollectionName = GetCollectionName("find", e.Command);
+						Query = e.Command.TryGetValue("filter", out BsonValue Value) ? Value.ToJson() : null;
+						break;
+					case "insert":
+						CollectionName = GetCollectionName("insert", e.Command);
+						break;
+					case "update":
+						CollectionName = GetCollectionName("update", e.Command);
+						break;
+				}
+
+				Scope.Span.Type = "db";
+				Scope.Span.SetTag("Command", e.CommandName);
+				Scope.Span.SetTag("DbName", DbName);
+				Scope.Span.SetTag("Collection", CollectionName);
+				Scope.Span.SetTag("Query", Query);
+				Scopes[e.RequestId] = Scope;
+			}
+
+			public void Handle(CommandSucceededEvent e)
+			{
+				if (Scopes.TryGetValue(e.RequestId, out Scope? Scope))
+				{
+					Scope.Dispose();
+					Scopes.Remove(e.RequestId, out _);
+				}
+			}
+			
+			public void Handle(CommandFailedEvent e)
+			{
+				if (Scopes.TryGetValue(e.RequestId, out Scope? Scope))
+				{
+					Scope.Span.Error = true;
+					Scope.Dispose();
+					Scopes.Remove(e.RequestId, out _);
+				}
+			}
+		}
+
+		/// <summary>
+		/// Constructor
+		/// </summary>
+		/// <param name="Settings">The settings instance</param>
+		/// <param name="Logger">Instance of the logger for this service</param>
+		public DatabaseService(IOptionsMonitor<ServerSettings> Settings, ILogger<DatabaseService> Logger)
+		{
+			this.Settings = Settings;
+			this.Logger = Logger;
+
+			try
+			{
+				ServerSettings CurrentSettings = Settings.CurrentValue;
+				ReadOnlyMode = CurrentSettings.DatabaseReadOnlyMode;
+				if (CurrentSettings.DatabasePublicCert != null)
+				{
+					X509Store LocalTrustStore = new X509Store(StoreName.Root);
+					try
+					{
+						LocalTrustStore.Open(OpenFlags.ReadWrite);
+
+						X509Certificate2Collection Collection = ImportCertificateBundle(CurrentSettings.DatabasePublicCert);
+						foreach (X509Certificate2 Certificate in Collection)
+						{
+							Logger.LogInformation("Importing certificate for {Subject}", Certificate.Subject);
+						}
+
+						LocalTrustStore.AddRange(Collection);
+					}
+					finally
+					{
+						LocalTrustStore.Close();
+					}
+				}
+
+				MongoClientSettings MongoSettings = MongoClientSettings.FromConnectionString(CurrentSettings.DatabaseConnectionString);
+				MongoSettings.ClusterConfigurator = ClusterBuilder =>
+				{
+					if (Logger.IsEnabled(LogLevel.Trace))
+					{
+						ClusterBuilder.Subscribe<CommandStartedEvent>(Event => TraceMongoCommand(Event.Command));
+					}
+					ClusterBuilder.Subscribe(new MongoDatadogTracer());
+				};
+				
+				MongoSettings.SslSettings = new SslSettings();
+				MongoSettings.SslSettings.ServerCertificateValidationCallback = CertificateValidationCallBack;
+				MongoSettings.MaxConnectionPoolSize = 300; // Default is 100
+
+				//TestSslConnection(MongoSettings.Server.Host, MongoSettings.Server.Port, Logger);
+
+				MongoClient Client = new MongoClient(MongoSettings);
+				Database = Client.GetDatabase(CurrentSettings.DatabaseName);
+
+				Singletons = GetCollection<BsonDocument>("Singletons");
+
+				Globals Globals = GetGlobalsAsync().Result;
+				if (Globals.ForceReset)
+				{
+					Logger.LogInformation("Performing DB reset...");
+
+					// Drop all the collections
+					List<string> CollectionNames = Database.ListCollectionNames().ToList();
+					foreach (string CollectionName in CollectionNames)
+					{
+						Logger.LogInformation("Dropping {CollectionName}", CollectionName);
+						Database.DropCollection(CollectionName);
+					}
+
+					// Create the new singletons collection
+					Singletons = GetCollection<BsonDocument>("Singletons");
+
+					// Create the new globals object
+					Globals NewGlobals = new Globals();
+					NewGlobals.JwtSigningKey = Globals.JwtSigningKey; // Ensure agents can still register
+					Singletons.InsertOne(NewGlobals.ToBsonDocument());
+					Globals = NewGlobals;
+				}
+
+				while (Globals.JwtSigningKey == null)
+				{
+					Globals.RotateSigningKey();
+					if (!TryUpdateSingletonAsync(Globals).Result)
+					{
+						Globals = GetGlobalsAsync().Result;
+					}
+				}
+
+				JwtIssuer = Settings.CurrentValue.JwtIssuer ?? Dns.GetHostName();
+				if (String.IsNullOrEmpty(CurrentSettings.JwtSecret))
+				{
+					JwtSigningKey = new SymmetricSecurityKey(Globals.JwtSigningKey);
+				}
+				else
+				{
+					JwtSigningKey = new SymmetricSecurityKey(Convert.FromBase64String(CurrentSettings.JwtSecret));
+				}
+
+				Credentials = GetCollection<Credential>("Credentials");
+				Artifacts = GetCollection<Artifact>("Artifacts");
+			}
+			catch (Exception Ex)
+			{
+				Logger.LogError(Ex, "Exception while initializing DatabaseService");
+				throw;
+			}
+		}
+
+		/// <summary>
+		/// Logs a mongodb command, removing any fields we don't care about
+		/// </summary>
+		/// <param name="Command">The command document</param>
+		void TraceMongoCommand(BsonDocument Command)
+		{
+			List<string> Params = new List<string>();
+			List<string> Values = new List<string>();
+
+			foreach(BsonElement Element in Command)
+			{
+				if (Element.Value != null && !Element.Name.Equals("$db", StringComparison.Ordinal) && !Element.Name.Equals("lsid", StringComparison.Ordinal))
+				{
+					Params.Add($"{Element.Name}: {{{Values.Count}}}");
+					Values.Add(Element.Value.ToString()!);
+				}
+			}
+
+			Logger.LogTrace($"MongoDB: {String.Join(", ", Params)}", Values.ToArray());
+		}
+
+		/// <summary>
+		/// Imports one or more certificates from a single PEM file
+		/// </summary>
+		/// <param name="FileName">File to import</param>
+		/// <returns>Collection of certificates</returns>
+		static X509Certificate2Collection ImportCertificateBundle(string FileName)
+		{
+			X509Certificate2Collection Collection = new X509Certificate2Collection();
+
+			string Text = File.ReadAllText(FileName);
+			for (int Offset = 0; Offset < Text.Length;)
+			{
+				int NextOffset = Text.IndexOf("-----BEGIN CERTIFICATE-----", Offset + 1, StringComparison.Ordinal);
+				if (NextOffset == -1)
+				{
+					NextOffset = Text.Length;
+				}
+
+				string CertificateText = Text.Substring(Offset, NextOffset - Offset);
+				Collection.Add(new X509Certificate2(Encoding.UTF8.GetBytes(CertificateText)));
+
+				Offset = NextOffset;
+			}
+
+			return Collection;
+		}
+
+		/// <summary>
+		/// Tests connection to the given server
+		/// </summary>
+		/// <param name="Host">Host name</param>
+		/// <param name="Port">Port number to connect on</param>
+		/// <param name="Logger">Logger for diagnostic messages</param>
+		/// <returns>True if the connection was valid</returns>
+		static bool TestSslConnection(string Host, int Port, ILogger Logger)
+		{
+#pragma warning disable CA1031 // Do not catch general exception types
+			using (TcpClient Client = new TcpClient())
+			{
+				try
+				{
+					Client.Connect(Host, Port);
+					Logger.LogInformation("Successfully connected to {Host} on port {Port}", Host, Port);
+				}
+				catch (Exception Ex)
+				{
+					Logger.LogError(Ex, "Unable to connect to {Host} on port {Port}", Host, Port);
+					return false;
+				}
+
+				using (SslStream Stream = new SslStream(Client.GetStream()))
+				{
+					try
+					{
+						Stream.AuthenticateAsClient(Host);
+						Logger.LogInformation("Successfully authenticated host {Host}", Host);
+					}
+					catch (Exception Ex)
+					{
+						Logger.LogError(Ex, "Unable to authenticate as client");
+						return false;
+					}
+				}
+			}
+			return true;
+#pragma warning restore CA1031 // Do not catch general exception types
+		}
+
+		/// <summary>
+		/// Provides additional diagnostic information for SSL certificate validation
+		/// </summary>
+		/// <param name="Sender"></param>
+		/// <param name="Certificate"></param>
+		/// <param name="Chain"></param>
+		/// <param name="SslPolicyErrors"></param>
+		/// <returns>True if the certificate is allowed, false otherwise</returns>
+		bool CertificateValidationCallBack(object Sender, X509Certificate Certificate, X509Chain Chain, SslPolicyErrors SslPolicyErrors)
+		{
+			// If the certificate is a valid, signed certificate, return true.
+			if (SslPolicyErrors == SslPolicyErrors.None)
+			{
+				return true;
+			}
+
+			// Generate diagnostic information
+			StringBuilder Builder = new StringBuilder();
+			if (Sender != null)
+			{
+				string SenderInfo = StringUtils.Indent(Sender.ToString() ?? String.Empty, "    ");
+				Builder.Append($"\nSender:\n{SenderInfo}");
+			}
+			if (Certificate != null)
+			{
+				Builder.Append($"\nCertificate: {Certificate.Subject}");
+			}
+			if (Chain != null)
+			{
+				if (Chain.ChainStatus != null && Chain.ChainStatus.Length > 0)
+				{
+					Builder.Append("\nChain status:");
+					foreach (X509ChainStatus Status in Chain.ChainStatus)
+					{
+						Builder.Append($"\n  {Status.StatusInformation}");
+					}
+				}
+				if (Chain.ChainElements != null)
+				{
+					Builder.Append("\nChain elements:");
+					for (int Idx = 0; Idx < Chain.ChainElements.Count; Idx++)
+					{
+						X509ChainElement Element = Chain.ChainElements[Idx];
+						Builder.Append($"\n  {Idx,4} - Certificate: {Element.Certificate.Subject}");
+						if (Element.ChainElementStatus != null && Element.ChainElementStatus.Length > 0)
+						{
+							foreach (X509ChainStatus Status in Element.ChainElementStatus)
+							{
+								Builder.Append($"\n         Status: {Status.StatusInformation} ({Status.Status})");
+							}
+						}
+						if (!String.IsNullOrEmpty(Element.Information))
+						{
+							Builder.Append($"\n         Info: {Element.Information}");
+						}
+					}
+				}
+			}
+
+			// Print out additional diagnostic information
+			Logger.LogError("TLS certificate validation failed ({Errors}).{AdditionalInfo}", SslPolicyErrors, StringUtils.Indent(Builder.ToString(), "    "));
+			return false;
+		}
+
+		/// <summary>
+		/// Get a MongoDB collection from database
+		/// </summary>
+		/// <param name="Name">Name of collection</param>
+		/// <typeparam name="T">A MongoDB document</typeparam>
+		/// <returns></returns>
+		public IMongoCollection<T> GetCollection<T>(string Name)
+		{
+			return new DatadogTracingCollection<T>(Database.GetCollection<T>(Name));
+		}
+
+		/// <summary>
+		/// Gets a singleton document by id
+		/// </summary>
+		/// <returns>The globals document</returns>
+		public Task<Globals> GetGlobalsAsync()
+		{
+			return GetSingletonAsync<Globals>(Globals.StaticId, () => new Globals() { SchemaVersion = UpgradeService.LatestSchemaVersion });
+		}
+
+		/// <summary>
+		/// Gets a singleton document by id
+		/// </summary>
+		/// <param name="Id">Id of the singleton document</param>
+		/// <returns>The document</returns>
+		public async Task<T> GetSingletonAsync<T>(ObjectId Id) where T : SingletonBase, new()
+		{
+			T Singleton = await GetSingletonAsync(Id, () => new T());
+			Singleton.PostLoad();
+			return Singleton;
+		}
+
+		/// <summary>
+		/// Gets a singleton document by id
+		/// </summary>
+		/// <param name="Id">Id of the singleton document</param>
+		/// <param name="Constructor">Method to use to construct a new object</param>
+		/// <returns>The document</returns>
+		public async Task<T> GetSingletonAsync<T>(ObjectId Id, Func<T> Constructor) where T : SingletonBase, new()
+		{
+			FilterDefinition<BsonDocument> Filter = new BsonDocument(new BsonElement("_id", Id));
+			for (; ; )
+			{
+				BsonDocument? Object = await Singletons.Find<BsonDocument>(Filter).FirstOrDefaultAsync();
+				if (Object != null)
+				{
+					T Item = BsonSerializer.Deserialize<T>(Object);
+					Item.PostLoad();
+					return Item;
+				}
+
+				T NewItem = Constructor();
+				NewItem.Id = Id;
+				await Singletons.InsertOneAsync(NewItem.ToBsonDocument());
+			}
+		}
+
+		/// <summary>
+		/// Updates a singleton
+		/// </summary>
+		/// <typeparam name="T"></typeparam>
+		/// <param name="Id"></param>
+		/// <param name="Updater"></param>
+		/// <returns></returns>
+		public async Task UpdateSingletonAsync<T>(ObjectId Id, Action<T> Updater) where T : SingletonBase, new()
+		{
+			for (; ; )
+			{
+				T Object = await GetSingletonAsync(Id, () => new T());
+				Updater(Object);
+				
+				if (await TryUpdateSingletonAsync(Object))
+				{
+					break;
+				}
+			}
+		}
+
+		/// <summary>
+		/// Attempts to update a singleton object
+		/// </summary>
+		/// <param name="SingletonObject">The singleton object</param>
+		/// <returns>True if the singleton document was updated</returns>
+		public async Task<bool> TryUpdateSingletonAsync<T>(T SingletonObject) where T : SingletonBase
+		{
+			int PrevRevision = SingletonObject.Revision++;
+
+			BsonDocument Filter = new BsonDocument { new BsonElement("_id", SingletonObject.Id), new BsonElement(nameof(SingletonBase.Revision), PrevRevision) };
+			try
+			{
+				ReplaceOneResult Result = await Singletons.ReplaceOneAsync(Filter, SingletonObject.ToBsonDocument(), new ReplaceOptions { IsUpsert = true });
+				return Result.MatchedCount > 0;
+			}
+			catch (MongoWriteException Ex)
+			{
+				// Duplicate key error occurs if filter fails to match because revision is not the same.
+				if (Ex.WriteError != null && Ex.WriteError.Category == ServerErrorCategory.DuplicateKey)
+				{
+					return false;
+				}
+				else
+				{
+					throw;
+				}
+			}
+		}
+	}
+}
diff --git a/Engine/Source/Programs/Horde/HordeServer/Services/DowntimeService.cs b/Engine/Source/Programs/Horde/HordeServer/Services/DowntimeService.cs
new file mode 100644
--- /dev/null
+++ b/Engine/Source/Programs/Horde/HordeServer/Services/DowntimeService.cs
@@ -0,0 +1,92 @@
+using Amazon.Runtime.Internal.Util;
+using HordeServer.Models;
+using HordeServer.Utilities;
+using Microsoft.Extensions.Logging;
+using Microsoft.Extensions.Options;
+using System;
+using System.Collections.Generic;
+using System.Linq;
+using System.Security.Cryptography.X509Certificates;
+using System.Threading;
+using System.Threading.Tasks;
+
+namespace HordeServer.Services
+{
+	/// <summary>
+	/// Interface for a service which keeps track of whether we're during downtime
+	/// </summary>
+	public interface IDowntimeService
+	{
+		/// <summary>
+		/// Returns true if downtime is currently active
+		/// </summary>
+		public bool IsDowntimeActive
+		{
+			get;
+		}
+	}
+
+	/// <summary>
+	/// Service which manages the downtime schedule
+	/// </summary>
+	public class DowntimeService : TickedBackgroundService, IDowntimeService
+	{
+		/// <summary>
+		/// Whether the server is currently in downtime
+		/// </summary>
+		public bool IsDowntimeActive
+		{
+			get;
+			private set;
+		}
+
+		DatabaseService DatabaseService;
+		IOptionsMonitor<ServerSettings> Settings;
+		ILogger<DowntimeService> Logger;
+
+		/// <summary>
+		/// Constructor
+		/// </summary>
+		/// <param name="DatabaseService">The database service instance</param>
+		/// <param name="Settings">The server settings</param>
+		/// <param name="Logger">Logger instance</param>
+		public DowntimeService(DatabaseService DatabaseService, IOptionsMonitor<ServerSettings> Settings, ILogger<DowntimeService> Logger)
+			: base(TimeSpan.FromMinutes(1.0), Logger)
+		{
+			this.DatabaseService = DatabaseService;
+			this.Settings = Settings;
+			this.Logger = Logger;
+
+			// Ensure the initial value to be correct
+			TickAsync(CancellationToken.None).Wait();
+		}
+
+		/// <summary>
+		/// Periodically called tick function
+		/// </summary>
+		/// <param name="StoppingToken">Token indicating that the service should stop</param>
+		/// <returns>Async task</returns>
+		protected override async Task TickAsync(CancellationToken StoppingToken)
+		{
+			Globals Globals = await DatabaseService.GetGlobalsAsync();
+
+			DateTimeOffset Now = TimeZoneInfo.ConvertTime(DateTimeOffset.Now, Settings.CurrentValue.TimeZoneInfo);
+			bool bIsActive = Globals.ScheduledDowntime.Any(x => x.IsActive(Now));
+
+			Logger.LogDebug("Server time: {Time}. Downtime: {Downtime}. Schedule: {@Schedule}.", Now, bIsActive, Globals.ScheduledDowntime);
+
+			if (bIsActive != IsDowntimeActive)
+			{
+				IsDowntimeActive = bIsActive;
+				if (IsDowntimeActive)
+				{
+					Logger.LogInformation("Entering downtime");
+				}
+				else
+				{
+					Logger.LogInformation("Leaving downtime");
+				}
+			}
+		}
+	}
+}
diff --git a/Engine/Source/Programs/Horde/HordeServer/Services/IIssueService.cs b/Engine/Source/Programs/Horde/HordeServer/Services/IIssueService.cs
new file mode 100644
--- /dev/null
+++ b/Engine/Source/Programs/Horde/HordeServer/Services/IIssueService.cs
@@ -0,0 +1,139 @@
+using HordeServer.Models;
+using Microsoft.Extensions.Hosting;
+using MongoDB.Bson;
+using System;
+using System.Collections.Generic;
+using System.Linq;
+using System.Threading.Tasks;
+
+using StreamId = HordeServer.Utilities.StringId<HordeServer.Models.IStream>;
+
+namespace HordeServer.Services
+{
+	/// <summary>
+	/// Service which manages build health issues
+	/// </summary>
+	public interface IIssueService : IHostedService
+	{
+		/// <summary>
+		/// Delegate for issue creation events
+		/// </summary>
+		public delegate void IssueUpdatedEvent(IIssue Issue);
+
+		/// <summary>
+		/// Event triggered when a label state changes
+		/// </summary>
+		event IssueUpdatedEvent? OnIssueUpdated;
+
+		/// <summary>
+		/// Gets a list of cached open issues
+		/// </summary>
+		/// <returns>List of cached issues</returns>
+		IReadOnlyList<IIssueDetails> CachedOpenIssues { get; }
+
+		/// <summary>
+		/// Find a cached issue, or update it
+		/// </summary>
+		/// <param name="IssueId">The issue id</param>
+		/// <returns>The cached issue data</returns>
+		Task<IIssueDetails?> GetCachedIssueDetailsAsync(int IssueId);
+
+		/// <summary>
+		/// Gets the details for a particular issue
+		/// </summary>
+		/// <param name="Issue">The issue id</param>
+		/// <returns>The issue details</returns>
+		Task<IIssueDetails> GetIssueDetailsAsync(IIssue Issue);
+
+		/// <summary>
+		/// Determines whether to show alerts for a given issue
+		/// </summary>
+		/// <param name="Issue">The issue to query</param>
+		/// <param name="Spans">Spans for the issue</param>
+		/// <returns>True if we should show alerts for the given issue</returns>
+		bool ShowDesktopAlertsForIssue(IIssue Issue, IReadOnlyList<IIssueSpan> Spans);
+
+		/// <summary>
+		/// Gets an issue with the given id
+		/// </summary>
+		/// <param name="IssueId"></param>
+		/// <returns></returns>
+		Task<IIssue?> GetIssueAsync(int IssueId);
+
+		/// <summary>
+		/// Gets the spans for the given issue
+		/// </summary>
+		/// <param name="Issue">The issue id</param>
+		/// <returns>List of spans for the issue</returns>
+		Task<List<IIssueSpan>> GetIssueSpansAsync(IIssue Issue);
+
+		/// <summary>
+		/// Gets the spans for the given issue
+		/// </summary>
+		/// <param name="Span">The issue span</param>
+		/// <returns>List of steps for the issue</returns>
+		Task<List<IIssueStep>> GetIssueStepsAsync(IIssueSpan Span);
+
+		/// <summary>
+		/// Searches for open issues
+		/// </summary>
+		/// <param name="Ids">List of issue ids to include</param>
+		/// <param name="StreamId">The stream to return issues for</param>
+		/// <param name="MinChange">The minimum changelist number to include</param>
+		/// <param name="MaxChange">The maximum changelist number to include</param>
+		/// <param name="Resolved">Whether to include results that are resolved</param>
+		/// <param name="Index">Index within the results to return</param>
+		/// <param name="Count">Number of results</param>
+		/// <returns>List of streams open in the given stream at the given changelist</returns>
+		Task<List<IIssue>> FindIssuesAsync(IEnumerable<int>? Ids = null, StreamId? StreamId = null, int? MinChange = null, int? MaxChange = null, bool? Resolved = null, int? Index = null, int? Count = null);
+
+		/// <summary>
+		/// Searches for open issues affecting a job
+		/// </summary>
+		/// <param name="Ids"></param>
+		/// <param name="Job"></param>
+		/// <param name="Graph"></param>
+		/// <param name="StepId"></param>
+		/// <param name="BatchId"></param>
+		/// <param name="LabelIdx"></param>
+		/// <param name="Resolved">Whether to include results that are resolved</param>
+		/// <param name="Index">Index within the results to return</param>
+		/// <param name="Count">Number of results</param>
+		/// <returns></returns>
+		Task<List<IIssue>> FindIssuesForJobAsync(int[]? Ids, IJob Job, IGraph Graph, SubResourceId? StepId, SubResourceId? BatchId, int? LabelIdx, bool? Resolved, int Index, int Count);
+
+		/// <summary>
+		/// Find the events for a particular issue
+		/// </summary>
+		/// <param name="IssueId">The issue id</param>
+		/// <param name="LogIds">Log ids to include</param>
+		/// <param name="Index">Index of the first event to return</param>
+		/// <param name="Count">Number of issues to return</param>
+		/// <returns>List of log events</returns>
+		Task<List<ILogEvent>> FindEventsForIssueAsync(int IssueId, ObjectId[]? LogIds = null, int Index = 0, int Count = 10);
+
+		/// <summary>
+		/// Updates the state of an issue
+		/// </summary>
+		/// <param name="Id">The current issue id</param>
+		/// <param name="Summary">New summary for the issue</param>
+		/// <param name="Owner">New owner of the issue</param>
+		/// <param name="NominatedBy">Person that nominated the new owner</param>
+		/// <param name="Acknowledged">Whether the issue has been acknowledged</param>
+		/// <param name="DeclinedBy">Name of the user that has declined this issue</param>
+		/// <param name="FixChange">Fix changelist for the issue</param>
+		/// <param name="Resolved">Whether the issue has been resolved</param>
+		/// <returns>True if the issue was updated</returns>
+		Task<bool> UpdateIssueAsync(int Id, string? Summary = null, string? Owner = null, string? NominatedBy = null, bool? Acknowledged = null, string? DeclinedBy = null, int? FixChange = null, bool? Resolved = null);
+
+		/// <summary>
+		/// Marks a step as complete
+		/// </summary>
+		/// <param name="Job">The job to update</param>
+		/// <param name="Graph">Graph for the job</param>
+		/// <param name="BatchId">Unique id of the batch</param>
+		/// <param name="StepId">Unique id of the step</param>
+		/// <returns>Async task</returns>
+		Task UpdateCompleteStep(IJob Job, IGraph Graph, SubResourceId BatchId, SubResourceId StepId);
+	}
+}
