Initial drop of Horde Server/Agent for UE5/Main. Still missing some Epic-specific functionality required to build/run correctly.

diff --git a/Engine/Source/Programs/Horde/HordeServer/Services/LogFileService.cs b/Engine/Source/Programs/Horde/HordeServer/Services/LogFileService.cs
new file mode 100644
--- /dev/null
+++ b/Engine/Source/Programs/Horde/HordeServer/Services/LogFileService.cs
@@ -0,0 +1,1597 @@
+// Copyright Epic Games, Inc. All Rights Reserved.
+
+using EpicGames.Core;
+using HordeCommon;
+using HordeServer.Api;
+using HordeServer.Collections;
+using HordeServer.Logs;
+using HordeServer.Models;
+using HordeServer.Storage;
+using HordeServer.Utilities;
+using Microsoft.Extensions.Caching.Memory;
+using Microsoft.Extensions.Logging;
+using MongoDB.Bson;
+using MongoDB.Driver;
+using System;
+using System.Collections.Concurrent;
+using System.Collections.Generic;
+using System.Collections.Immutable;
+using System.Diagnostics;
+using System.Diagnostics.CodeAnalysis;
+using System.Globalization;
+using System.IO;
+using System.Linq;
+using System.Runtime.CompilerServices;
+using System.Security.Claims;
+using System.Text.Json;
+using System.Text;
+using System.Threading.Tasks;
+using System.Threading;
+using Datadog.Trace;
+
+using ILogger = Microsoft.Extensions.Logging.ILogger;
+using Stream = System.IO.Stream;
+
+[assembly: InternalsVisibleTo("HordeServerTests")]
+
+namespace HordeServer.Services
+{
+	/// <summary>
+	/// Metadata about a log file
+	/// </summary>
+	public class LogMetadata
+	{
+		/// <summary>
+		/// Length of the log file
+		/// </summary>
+		public long Length { get; set; }
+
+		/// <summary>
+		/// Number of lines in the log file
+		/// </summary>
+		public int MaxLineIndex { get; set; }
+	}
+
+	/// <summary>
+	/// Interface for the log file service
+	/// </summary>
+	public interface ILogFileService
+	{
+		/// <summary>
+		/// Creates a new log
+		/// </summary>
+		/// <param name="JobId">Unique id of the job that owns this log file</param>
+		/// <param name="SessionId">Agent session allowed to update the log</param>
+		/// <param name="Type">Type of events to be stored in the log</param>
+		/// <returns>The new log file document</returns>
+		Task<ILogFile> CreateLogFileAsync(ObjectId JobId, ObjectId? SessionId, LogType Type);
+
+		/// <summary>
+		/// Gets a logfile by ID
+		/// </summary>
+		/// <param name="LogFileId">Unique id of the log file</param>
+		/// <returns>The logfile document</returns>
+		Task<ILogFile?> GetLogFileAsync(ObjectId LogFileId);
+
+		/// <summary>
+		/// Gets a logfile by ID, returning a cached copy if available. This should only be used to retrieve constant properties set at creation, such as the session or job it's associated with.
+		/// </summary>
+		/// <param name="LogFileId">Unique id of the log file</param>
+		/// <returns>The logfile document</returns>
+		Task<ILogFile?> GetCachedLogFileAsync(ObjectId LogFileId);
+
+		/// <summary>
+		/// Returns a list of log files
+		/// </summary>
+		/// <param name="Index">Index of the first result to return</param>
+		/// <param name="Count">Number of results to return</param>
+		/// <returns>List of logfile documents</returns>
+		Task<List<ILogFile>> GetLogFilesAsync(int? Index = null, int? Count = null);
+
+		/// <summary>
+		/// Writes out chunk data and assigns to a file
+		/// </summary>
+		/// <param name="LogFile">The log file</param>
+		/// <param name="Offset">Offset within the file of data</param>
+		/// <param name="LineIndex">Current line index of the data (need not be the starting of the line)</param>
+		/// <param name="Data">the data to add</param>
+		/// <param name="Flush">Whether the current chunk is complete and should be flushed</param>
+		/// <param name="MaxChunkLength">The maximum chunk length. Defaults to 128kb.</param>
+		/// <param name="MaxSubChunkLineCount">Maximum number of lines in each sub-chunk.</param>
+		/// <returns></returns>
+		Task<ILogFile?> WriteLogDataAsync(ILogFile LogFile, long Offset, int LineIndex, ReadOnlyMemory<byte> Data, bool Flush, int MaxChunkLength = 256 * 1024, int MaxSubChunkLineCount = 128);
+
+		/// <summary>
+		/// Gets metadata about the log file
+		/// </summary>
+		/// <param name="LogFile">The log file to query</param>
+		/// <returns>Metadata about the log file</returns>
+		Task<LogMetadata> GetMetadataAsync(ILogFile LogFile);
+
+		/// <summary>
+		/// Creates new log events
+		/// </summary>
+		/// <param name="NewEvents">List of events</param>
+		/// <returns>Async task</returns>
+		Task CreateEventsAsync(List<NewLogEventData> NewEvents);
+
+		/// <summary>
+		/// Find events for a particular log file
+		/// </summary>
+		/// <param name="LogFile">The log file instance</param>
+		/// <param name="Index">Index of the first event to retrieve</param>
+		/// <param name="Count">Number of events to retrieve</param>
+		/// <returns>List of log events</returns>
+		Task<List<ILogEvent>> FindLogEventsAsync(ILogFile LogFile, int? Index = null, int? Count = null);
+
+		/// <summary>
+		/// Adds events to a log span
+		/// </summary>
+		/// <param name="Events">The events to add</param>
+		/// <param name="SpanId">The span id</param>
+		/// <returns>Async task</returns>
+		Task AddSpanToEventsAsync(IEnumerable<ILogEvent> Events, ObjectId SpanId);
+
+		/// <summary>
+		/// Find events for an issue
+		/// </summary>
+		/// <param name="SpanIds">The span ids</param>
+		/// <param name="LogIds">Log ids to include</param>
+		/// <param name="Index">Index within the events for results to return</param>
+		/// <param name="Count">Number of results to return</param>
+		/// <returns>Async task</returns>
+		Task<List<ILogEvent>> FindEventsForSpansAsync(IEnumerable<ObjectId> SpanIds, ObjectId[]? LogIds, int Index, int Count);
+
+		/// <summary>
+		/// Gets the data for an event
+		/// </summary>
+		/// <param name="LogFile">The log file instance</param>
+		/// <param name="LineIndex">Index of the line in the file</param>
+		/// <param name="LineCount">Number of lines in the event</param>
+		/// <returns>New event data instance</returns>
+		Task<ILogEventData> GetEventDataAsync(ILogFile LogFile, int LineIndex, int LineCount);
+
+		/// <summary>
+		/// Gets lines from the given log 
+		/// </summary>
+		/// <param name="LogFile">The log file</param>
+		/// <param name="Offset">Offset of the data to return</param>
+		/// <param name="Length">Length of the data to return</param>
+		/// <returns>Data for the requested range</returns>
+		Task<Stream> OpenRawStreamAsync(ILogFile LogFile, long Offset, long Length);
+
+		/// <summary>
+		/// Gets the offset of the given line number
+		/// </summary>
+		/// <param name="LogFile">The log file to search</param>
+		/// <param name="LineIdx">The line index to retrieve the offset for</param>
+		/// <returns>The actual clamped line number and offset</returns>
+		Task<(int, long)> GetLineOffsetAsync(ILogFile LogFile, int LineIdx);
+
+		/// <summary>
+		/// Search for the specified text in a log file
+		/// </summary>
+		/// <param name="LogFile">The log file to search</param>
+		/// <param name="Text">Text to search for</param>
+		/// <param name="FirstLine">Line to start search from</param>
+		/// <param name="Count">Number of results to return</param>
+		/// <param name="Stats">Receives stats for the search</param>
+		/// <returns>List of line numbers containing the given term</returns>
+		Task<List<int>> SearchLogDataAsync(ILogFile LogFile, string Text, int FirstLine, int Count, LogSearchStats Stats);
+	}
+
+	/// <summary>
+	/// Extension methods for dealing with log files
+	/// </summary>
+	public static class LogFileServiceExtensions
+	{
+		/// <summary>
+		/// Parses a stream of json text and outputs plain text
+		/// </summary>
+		/// <param name="LogFileService">The log file service</param>
+		/// <param name="LogFile">The log file to query</param>
+		/// <param name="Offset">Offset within the log file to copy</param>
+		/// <param name="Length">Length of the data to copy</param>
+		/// <param name="OutputStream">Output stream to receive the text data</param>
+		/// <returns>Async text</returns>
+		public static async Task CopyRawStreamAsync(this ILogFileService LogFileService, ILogFile LogFile, long Offset, long Length, Stream OutputStream)
+		{
+			using (Stream Stream = await LogFileService.OpenRawStreamAsync(LogFile, Offset, Length))
+			{
+				await Stream.CopyToAsync(OutputStream);
+			}
+		}
+
+		/// <summary>
+		/// Parses a stream of json text and outputs plain text
+		/// </summary>
+		/// <param name="LogFileService">The log file service</param>
+		/// <param name="LogFile">The log file to query</param>
+		/// <param name="Offset">Offset within the data to copy from</param>
+		/// <param name="Length">Length of the data to copy</param>
+		/// <param name="OutputStream">Output stream to receive the text data</param>
+		/// <returns>Async text</returns>
+		public static async Task CopyPlainTextStreamAsync(this ILogFileService LogFileService, ILogFile LogFile, long Offset, long Length, Stream OutputStream)
+		{
+			using (Stream Stream = await LogFileService.OpenRawStreamAsync(LogFile, 0, long.MaxValue))
+			{
+				byte[] ReadBuffer = new byte[4096];
+				int ReadBufferLength = 0;
+
+				byte[] WriteBuffer = new byte[4096];
+				int WriteBufferLength = 0;
+
+				while(Length > 0)
+				{
+					// Add more data to the buffer
+					int ReadBytes = await Stream.ReadAsync(ReadBuffer, ReadBufferLength, ReadBuffer.Length - ReadBufferLength);
+					ReadBufferLength += ReadBytes;
+
+					// Copy as many lines as possible to the output
+					int ConvertedBytes = 0;
+					for(int EndIdx = 1; EndIdx < ReadBufferLength; EndIdx++)
+					{
+						if (ReadBuffer[EndIdx] == '\n')
+						{
+							WriteBufferLength = LogText.ConvertToPlainText(ReadBuffer.AsSpan(ConvertedBytes, EndIdx - ConvertedBytes), WriteBuffer, WriteBufferLength);
+							ConvertedBytes = EndIdx + 1;
+						}
+					}
+
+					// If there's anything in the write buffer, write it out
+					if (WriteBufferLength > 0)
+					{
+						if (Offset < WriteBufferLength)
+						{
+							int WriteLength = (int)Math.Min((long)WriteBufferLength - Offset, Length);
+							await OutputStream.WriteAsync(WriteBuffer, (int)Offset, WriteLength);
+							Length -= WriteLength;
+						}
+						Offset = Math.Max(Offset - WriteBufferLength, 0);
+						WriteBufferLength = 0;
+					}
+
+					// If we were able to read something, shuffle down the rest of the buffer. Otherwise expand the read buffer.
+					if (ConvertedBytes > 0)
+					{
+						Buffer.BlockCopy(ReadBuffer, ConvertedBytes, ReadBuffer, 0, ReadBufferLength - ConvertedBytes);
+						ReadBufferLength -= ConvertedBytes;
+					}
+					else if(ReadBufferLength > 0)
+					{
+						Array.Resize(ref ReadBuffer, ReadBuffer.Length + 128);
+						WriteBuffer = new byte[ReadBuffer.Length];
+					}
+
+					// Exit if we didn't read anything in this iteration
+					if(ReadBytes == 0)
+					{
+						break;
+					}
+				}
+			}
+		}
+	}
+	
+	/// <summary>
+	/// Wraps functionality for manipulating logs
+	/// </summary>
+	public class LogFileService : TickedBackgroundService, ILogFileService
+	{
+		/// <summary>
+		/// Information Logger
+		/// </summary>
+		private readonly ILogger<LogFileService> Logger;
+
+		/// <summary>
+		/// Collection of log documents
+		/// </summary>
+		private readonly ILogFileCollection LogFiles;
+
+		/// <summary>
+		/// Collection of log events
+		/// </summary>
+		private readonly ILogEventCollection LogEvents;
+
+		/// <summary>
+		/// Interface for the log reader
+		/// </summary>
+		private readonly ILogStorage Storage;
+
+		/// <summary>
+		/// Interface for the log builder
+		/// </summary>
+		private readonly ILogBuilder Builder;
+
+		/// <summary>
+		/// Lock object for the <see cref="WriteTasks"/> and <see cref="WriteChunks"/> members
+		/// </summary>
+		private object WriteLock = new object();
+
+		/// <summary>
+		/// List of active write tasks
+		/// </summary>
+		private List<Task> WriteTasks = new List<Task>();
+
+		/// <summary>
+		/// Set of chunks which is currently being written to
+		/// </summary>
+		private HashSet<(ObjectId, long)> WriteChunks = new HashSet<(ObjectId, long)>();
+
+		/// <summary>
+		/// Cache of session ids for each log
+		/// </summary>
+		private IMemoryCache LogFileCache;
+
+		/// <summary>
+		/// Streams log data to a caller
+		/// </summary>
+		class ResponseStream : Stream
+		{
+			/// <summary>
+			/// The log file service that created this stream
+			/// </summary>
+			LogFileService LogFileService;
+
+			/// <summary>
+			/// The log file being read
+			/// </summary>
+			ILogFile LogFile;
+
+			/// <summary>
+			/// Starting offset within the file of the data to return 
+			/// </summary>
+			long ResponseOffset;
+
+			/// <summary>
+			/// Length of data to return
+			/// </summary>
+			long ResponseLength;
+
+			/// <summary>
+			/// Current offset within the stream
+			/// </summary>
+			long CurrentOffset;
+
+			/// <summary>
+			/// The current chunk index
+			/// </summary>
+			int ChunkIdx;
+
+			/// <summary>
+			/// Buffer containing a message for missing data
+			/// </summary>
+			ReadOnlyMemory<byte> SourceBuffer;
+
+			/// <summary>
+			/// Offset within the source buffer
+			/// </summary>
+			int SourcePos;
+
+			/// <summary>
+			/// Length of the source buffer being copied from
+			/// </summary>
+			int SourceEnd;
+
+			/// <summary>
+			/// Constructor
+			/// </summary>
+			/// <param name="LogFileService">The log file service, for q</param>
+			/// <param name="LogFile"></param>
+			/// <param name="Offset"></param>
+			/// <param name="Length"></param>
+			public ResponseStream(LogFileService LogFileService, ILogFile LogFile, long Offset, long Length)
+			{
+				this.LogFileService = LogFileService;
+				this.LogFile = LogFile;
+
+				this.ResponseOffset = Offset;
+				this.ResponseLength = Length;
+
+				this.CurrentOffset = Offset;
+
+				this.ChunkIdx = LogFile.Chunks.GetChunkForOffset(Offset);
+				this.SourceBuffer = null!;
+			}
+
+			/// <inheritdoc/>
+			public override bool CanRead => true;
+
+			/// <inheritdoc/>
+			public override bool CanSeek => false;
+
+			/// <inheritdoc/>
+			public override bool CanWrite => false;
+
+			/// <inheritdoc/>
+			public override long Length
+			{
+				get { return ResponseLength; }
+			}
+
+			/// <inheritdoc/>
+			public override long Position
+			{
+				get { return CurrentOffset - ResponseOffset; }
+				set => throw new NotImplementedException();
+			}
+
+			/// <inheritdoc/>
+			public override void Flush()
+			{
+			}
+
+			/// <inheritdoc/>
+			public override int Read(byte[] Buffer, int Offset, int Count)
+			{
+				return ReadAsync(Buffer, Offset, Count, CancellationToken.None).Result;
+			}
+
+			/// <inheritdoc/>
+			public override async Task<int> ReadAsync(byte[] Buffer, int Offset, int Length, CancellationToken CancellationToken)
+			{
+				int ReadBytes = 0;
+				while (ReadBytes < Length)
+				{
+					if (SourcePos < SourceEnd)
+					{
+						// Try to copy from the current buffer
+						int BlockSize = Math.Min(SourceEnd - SourcePos, Length - ReadBytes);
+						SourceBuffer.Slice(SourcePos, BlockSize).Span.CopyTo(Buffer.AsSpan(Offset + ReadBytes));
+						CurrentOffset += BlockSize;
+						ReadBytes += BlockSize;
+						SourcePos += BlockSize;
+					}
+					else if (CurrentOffset < ResponseOffset + ResponseLength)
+					{
+						// Move to the right chunk
+						while (ChunkIdx + 1 < LogFile.Chunks.Count && CurrentOffset >= LogFile.Chunks[ChunkIdx + 1].Offset)
+						{
+							ChunkIdx++;
+						}
+
+						// Get the end of this chunk
+						long NextOffset = ResponseOffset + ResponseLength;
+						if (ChunkIdx + 1 < LogFile.Chunks.Count)
+						{
+							ILogChunk NextChunk = LogFile.Chunks[ChunkIdx + 1];
+							NextOffset = Math.Min(NextOffset, NextChunk.Offset);
+						}
+
+						// Get the chunk data
+						ILogChunk Chunk = LogFile.Chunks[ChunkIdx];
+						LogChunkData ChunkData = await LogFileService.ReadChunkAsync(LogFile, ChunkIdx);
+
+						// Figure out which sub-chunk to use
+						int SubChunkIdx = ChunkData.GetSubChunkForOffsetWithinChunk((int)(CurrentOffset - Chunk.Offset));
+						LogSubChunkData SubChunkData = ChunkData.SubChunks[SubChunkIdx];
+
+						// Get the source data
+						long SubChunkOffset = Chunk.Offset + ChunkData.SubChunkOffset[SubChunkIdx];
+						SourceBuffer = SubChunkData.InflateText().Data;
+						SourcePos = (int)(CurrentOffset - SubChunkOffset);
+						SourceEnd = (int)Math.Min(SourceBuffer.Length, (ResponseOffset + ResponseLength) - SubChunkOffset);
+					}
+					else
+					{
+						// End of the log
+						break;
+					}
+				}
+				return ReadBytes;
+			}
+
+			/// <inheritdoc/>
+			public override long Seek(long offset, SeekOrigin origin) => throw new NotImplementedException();
+
+			/// <inheritdoc/>
+			public override void SetLength(long value) => throw new NotImplementedException();
+
+			/// <inheritdoc/>
+			public override void Write(byte[] buffer, int offset, int count) => throw new NotImplementedException();
+		}
+
+		/// <summary>
+		/// Constructor
+		/// </summary>
+		/// <param name="LogFiles">The log file collection</param>
+		/// <param name="LogEvents">The log events collection</param>
+		/// <param name="Builder">The log builder</param>
+		/// <param name="Storage">THe log storage hierarchy</param>
+		/// <param name="Logger">Log interface</param>
+		public LogFileService(ILogFileCollection LogFiles, ILogEventCollection LogEvents, ILogBuilder Builder, ILogStorage Storage, ILogger<LogFileService> Logger)
+			: base(TimeSpan.FromSeconds(30.0), Logger)
+		{
+			this.LogFiles = LogFiles;
+			this.LogEvents = LogEvents;
+			this.Logger = Logger;
+			this.LogFileCache = new MemoryCache(new MemoryCacheOptions());
+			this.Builder = Builder;
+			this.Storage = Storage;
+		}
+
+		/// <inheritdoc/>
+		public Task<ILogFile> CreateLogFileAsync(ObjectId JobId, ObjectId? SessionId, LogType Type)
+		{
+			return LogFiles.CreateLogFileAsync(JobId, SessionId, Type);
+		}
+
+		/// <inheritdoc/>
+		public async Task<ILogFile?> GetLogFileAsync(ObjectId LogFileId)
+		{
+			ILogFile? LogFile = await LogFiles.GetLogFileAsync(LogFileId);
+			if(LogFile != null)
+			{
+				AddCachedLogFile(LogFile);
+			}
+			return LogFile;
+		}
+
+		/// <summary>
+		/// Adds a log file to the cache
+		/// </summary>
+		/// <param name="LogFile">The log file to cache</param>
+		void AddCachedLogFile(ILogFile LogFile)
+		{
+			MemoryCacheEntryOptions Options = new MemoryCacheEntryOptions().SetSlidingExpiration(TimeSpan.FromSeconds(30));
+			LogFileCache.Set(LogFile.Id, LogFile, Options);
+		}
+
+		/// <summary>
+		/// Gets a cached log file by id
+		/// </summary>
+		/// <param name="LogFileId">The log file id</param>
+		/// <returns>New log file, or null if not found</returns>
+		public async Task<ILogFile?> GetCachedLogFileAsync(ObjectId LogFileId)
+		{
+			object? LogFile;
+			if (!LogFileCache.TryGetValue(LogFileId, out LogFile))
+			{
+				LogFile = await GetLogFileAsync(LogFileId);
+			}
+			return (ILogFile?)LogFile;
+		}
+
+		/// <inheritdoc/>
+		public Task<List<ILogFile>> GetLogFilesAsync(int? Index = null, int? Count = null)
+		{
+			return LogFiles.GetLogFilesAsync(Index, Count);
+		}
+
+		class WriteState
+		{
+			public long Offset;
+			public int LineIndex;
+			public ReadOnlyMemory<byte> Memory;
+
+			public WriteState(long Offset, int LineIndex, ReadOnlyMemory<byte> Memory)
+			{
+				this.Offset = Offset;
+				this.LineIndex = LineIndex;
+				this.Memory = Memory;
+			}
+		}
+
+		/// <inheritdoc/>
+		public async Task<ILogFile?> WriteLogDataAsync(ILogFile LogFile, long Offset, int LineIndex, ReadOnlyMemory<byte> Data, bool Flush, int MaxChunkLength, int MaxSubChunkLineCount)
+		{
+			using Scope Scope = Tracer.Instance.StartActive("WriteLogDataAsync");
+			Scope.Span.SetTag("LogId", LogFile.Id.ToString());
+			Scope.Span.SetTag("Offset", Offset.ToString(CultureInfo.InvariantCulture));
+			Scope.Span.SetTag("Length", Data.Length.ToString(CultureInfo.InvariantCulture));
+			Scope.Span.SetTag("LineIndex", LineIndex.ToString(CultureInfo.InvariantCulture));
+
+			// Make sure the data ends in a newline
+			if (Data.Length > 0 && Data.Span[Data.Length - 1] != '\n')
+			{
+				throw new ArgumentException("Log data must consist of a whole number of lines", nameof(Data));
+			}
+
+			// Make sure the line count is a power of two
+			if ((MaxSubChunkLineCount & (MaxSubChunkLineCount - 1)) != 0)
+			{
+				throw new ArgumentException("Maximum line count per sub-chunk must be a power of two", nameof(MaxSubChunkLineCount));
+			}
+
+			// List of the flushed chunks
+			List<long> CompleteOffsets = new List<long>();
+
+			// Add the data to new chunks
+			WriteState State = new WriteState(Offset, LineIndex, Data);
+			while (State.Memory.Length > 0)
+			{
+				// Find an existing chunk to append to
+				int ChunkIdx = LogFile.Chunks.GetChunkForOffset(State.Offset);
+				if (ChunkIdx >= 0)
+				{
+					ILogChunk Chunk = LogFile.Chunks[ChunkIdx];
+					if (await WriteLogChunkDataAsync(LogFile, Chunk, State, CompleteOffsets, MaxChunkLength, MaxSubChunkLineCount))
+					{
+						continue;
+					}
+				}
+
+				// Create a new chunk. Ensure that there's a chunk at the start of the file, even if the current write is beyond it.
+				ILogFile? NewLogFile;
+				if (LogFile.Chunks.Count == 0)
+				{
+					NewLogFile = await LogFiles.TryAddChunkAsync(LogFile, 0, 0);
+				}
+				else
+				{
+					NewLogFile = await LogFiles.TryAddChunkAsync(LogFile, State.Offset, State.LineIndex);
+				}
+
+				// Try to add a new chunk at the new location
+				if (NewLogFile == null)
+				{
+					NewLogFile = await LogFiles.GetLogFileAsync(LogFile.Id);
+					if (NewLogFile == null)
+					{
+						Logger.LogError("Unable to update log file {LogId}", LogFile.Id);
+						return null;
+					}
+					LogFile = NewLogFile;
+				}
+				else
+				{
+					Logger.LogDebug("Added new chunk at offset {Offset} to log {LogId}", State.Offset, LogFile.Id);
+					LogFile = NewLogFile;
+				}
+			}
+
+			// Flush any pending chunks on this log file
+			if (Flush)
+			{
+				foreach(ILogChunk Chunk in LogFile.Chunks)
+				{
+					if (Chunk.Length == 0 && !CompleteOffsets.Contains(Chunk.Offset))
+					{
+						await Builder.CompleteChunkAsync(LogFile.Id, Chunk.Offset);
+						CompleteOffsets.Add(Chunk.Offset);
+					}
+				}
+			}
+
+			// Write all the chunks
+			if (CompleteOffsets.Count > 0 || Flush)
+			{
+				ILogFile? NewLogFile = await WriteCompleteChunksForLogAsync(LogFile, CompleteOffsets, Flush);
+				if (NewLogFile == null)
+				{
+					return null;
+				}
+				LogFile = NewLogFile;
+			}
+			return LogFile;
+		}
+
+		/// <summary>
+		/// Append data to an existing chunk.
+		/// </summary>
+		/// <param name="LogFile">The log file to append to</param>
+		/// <param name="Chunk">Chunk within the log file to update</param>
+		/// <param name="State">Data remaining to be written</param>
+		/// <param name="CompleteOffsets">List of complete chunks</param>
+		/// <param name="MaxChunkLength">Maximum length of each chunk</param>
+		/// <param name="MaxSubChunkLineCount">Maximum number of lines in each subchunk</param>
+		/// <returns>True if data was appended to </returns>
+		private async Task<bool> WriteLogChunkDataAsync(ILogFile LogFile, ILogChunk Chunk, WriteState State, List<long> CompleteOffsets, int MaxChunkLength, int MaxSubChunkLineCount)
+		{
+			// Don't allow data to be appended if the chunk is complete
+			if(Chunk.Length > 0)
+			{
+				return false;
+			}
+
+			// Otherwise keep appending subchunks
+			bool bResult = false;
+			for (; ; )
+			{
+				// Flush the current sub-chunk if we're on a boundary
+				if (State.LineIndex > 0 && (State.LineIndex & (MaxSubChunkLineCount - 1)) == 0)
+				{
+					Logger.LogDebug("Completing log {LogId} chunk offset {Offset} sub-chunk at line {LineIndex}", LogFile.Id, Chunk.Offset, State.LineIndex);
+					await Builder.CompleteSubChunkAsync(LogFile.Id, Chunk.Offset);
+				}
+
+				// Figure out the max length to write to the current chunk
+				int MaxLength = Math.Min((int)((Chunk.Offset + MaxChunkLength) - State.Offset), State.Memory.Length);
+
+				// Figure out the maximum line index for the current sub chunk
+				int MinLineIndex = State.LineIndex;
+				int MaxLineIndex = (MinLineIndex & ~(MaxSubChunkLineCount - 1)) + MaxSubChunkLineCount;
+
+				// Append this data
+				(int Length, int LineCount) = GetWriteLength(State.Memory.Span, MaxLength, MaxLineIndex - MinLineIndex, State.Offset == Chunk.Offset);
+				if (Length > 0)
+				{
+					// Append this data
+					ReadOnlyMemory<byte> AppendData = State.Memory.Slice(0, Length);
+					if (!await Builder.AppendAsync(LogFile.Id, Chunk.Offset, State.Offset, State.LineIndex, LineCount, AppendData, LogFile.Type))
+					{
+						break;
+					}
+
+					// Update the state
+					//Logger.LogDebug("Append to log {LogId} chunk offset {Offset} (LineIndex={LineIndex}, LineCount={LineCount}, Offset={WriteOffset}, Length={WriteLength})", LogFile.Id, Chunk.Offset, State.LineIndex, LineCount, State.Offset, Length);
+					State.Offset += Length;
+					State.LineIndex += LineCount;
+					State.Memory = State.Memory.Slice(Length);
+					bResult = true;
+
+					// If this is the end of the data, bail out
+					if(State.Memory.Length == 0)
+					{
+						break;
+					}
+				}
+
+				// Flush the sub-chunk if it's full
+				if (State.LineIndex < MaxLineIndex)
+				{
+					Logger.LogDebug("Completing chunk for log {LogId} at offset {Offset}", LogFile.Id, Chunk.Offset);
+					await Builder.CompleteChunkAsync(LogFile.Id, Chunk.Offset);
+					CompleteOffsets.Add(Chunk.Offset);
+					break;
+				}
+			}
+			return bResult;
+		}
+
+		/// <summary>
+		/// Get the amount of data to write from the given span
+		/// </summary>
+		/// <param name="Span">Data to write</param>
+		/// <param name="MaxLength">Maximum length of the data to write</param>
+		/// <param name="MaxLineCount">Maximum number of lines to write</param>
+		/// <param name="bIsEmptyChunk">Whether the current chunk is empty</param>
+		/// <returns>A tuple consisting of the amount of data to write and number of lines in it</returns>
+		private static (int, int) GetWriteLength(ReadOnlySpan<byte> Span, int MaxLength, int MaxLineCount, bool bIsEmptyChunk)
+		{
+			int Length = 0;
+			int LineCount = 0;
+			for (int Idx = 0; Idx < MaxLength || bIsEmptyChunk; Idx++)
+			{
+				if (Span[Idx] == '\n')
+				{
+					Length = Idx + 1;
+					LineCount++;
+					bIsEmptyChunk = false;
+
+					if (LineCount >= MaxLineCount)
+					{
+						break;
+					}
+				}
+			}
+			return (Length, LineCount);
+		}
+
+		/// <inheritdoc/>
+		public async Task<LogMetadata> GetMetadataAsync(ILogFile LogFile)
+		{
+			LogMetadata Metadata = new LogMetadata();
+			if (LogFile.Chunks.Count > 0)
+			{
+				ILogChunk Chunk = LogFile.Chunks[LogFile.Chunks.Count - 1];
+				if (LogFile.MaxLineIndex == null || Chunk.Length == 0)
+				{
+					LogChunkData ChunkData = await ReadChunkAsync(LogFile, LogFile.Chunks.Count - 1);
+					Metadata.Length = Chunk.Offset + ChunkData.Length;
+					Metadata.MaxLineIndex = Chunk.LineIndex + ChunkData.LineCount;
+				}
+				else
+				{
+					Metadata.Length = Chunk.Offset + Chunk.Length;
+					Metadata.MaxLineIndex = LogFile.MaxLineIndex.Value;
+				}
+			}
+			return Metadata;
+		}
+
+		/// <inheritdoc/>
+		public Task CreateEventsAsync(List<NewLogEventData> NewEvents)
+		{
+			return LogEvents.AddManyAsync(NewEvents);
+		}
+
+		/// <inheritdoc/>
+		public Task<List<ILogEvent>> FindLogEventsAsync(ILogFile LogFile, int? Index = null, int? Count = null)
+		{
+			return LogEvents.FindAsync(LogFile.Id, Index, Count);
+		}
+
+		class LogEventLine : ILogEventLine
+		{
+			LogLevel Level;
+			public EventId? EventId { get; }
+			public string Message { get; }
+			public JsonElement Data { get; }
+
+			LogLevel ILogEventLine.Level => Level;
+
+			public LogEventLine(ReadOnlySpan<byte> Data)
+				: this(JsonSerializer.Deserialize<JsonElement>(Data))
+			{
+			}
+
+			[SuppressMessage("Design", "CA1031:Do not catch general exception types")]
+			public LogEventLine(JsonElement Data)
+			{
+				this.Data = Data;
+
+				JsonElement LevelElement;
+				if (!Data.TryGetProperty("level", out LevelElement) || !Enum.TryParse(LevelElement.GetString(), out Level))
+				{
+					Level = LogLevel.Information;
+				}
+
+				JsonElement IdElement;
+				if (Data.TryGetProperty("id", out IdElement))
+				{
+					int IdValue;
+					if (IdElement.TryGetInt32(out IdValue))
+					{
+						EventId = IdValue;
+					}
+				}
+
+				JsonElement MessageElement;
+				if (Data.TryGetProperty("renderedMessage", out MessageElement) || Data.TryGetProperty("message", out MessageElement))
+				{
+					Message = MessageElement.GetString();
+				}
+				else
+				{
+					Message = "(Missing message or renderedMessage field)";
+				}
+			}
+		}
+
+		class LogEventData : ILogEventData
+		{
+			public IReadOnlyList<ILogEventLine> Lines { get; }
+
+			EventId? ILogEventData.EventId => (Lines.Count > 0) ? Lines[0].EventId : null;
+			EventSeverity ILogEventData.Severity => (Lines.Count == 0) ? EventSeverity.Information : (Lines[0].Level == LogLevel.Warning) ? EventSeverity.Warning : EventSeverity.Error;
+			string ILogEventData.Message => String.Join("\n", Lines.Select(x => x.Message));
+
+			public LogEventData(IReadOnlyList<ILogEventLine> Lines)
+			{
+				this.Lines = Lines;
+			}
+		}
+
+		/// <inheritdoc/>
+		public Task AddSpanToEventsAsync(IEnumerable<ILogEvent> Events, ObjectId SpanId)
+		{
+			return LogEvents.AddSpanToEventsAsync(Events, SpanId);
+		}
+
+		/// <inheritdoc/>
+		public Task<List<ILogEvent>> FindEventsForSpansAsync(IEnumerable<ObjectId> SpanIds, ObjectId[]? LogIds, int Index, int Count)
+		{
+			return LogEvents.FindEventsForSpansAsync(SpanIds, LogIds, Index, Count);
+		}
+
+		/// <inheritdoc/>
+		public async Task<ILogEventData> GetEventDataAsync(ILogFile LogFile, int LineIndex, int LineCount)
+		{
+			using Scope Scope = Tracer.Instance.StartActive("GetEventDataAsync");
+			Scope.Span.SetTag("LogId", LogFile.Id.ToString());
+			Scope.Span.SetTag("LineIndex", LineIndex.ToString(CultureInfo.InvariantCulture));
+			Scope.Span.SetTag("LineCount", LineCount.ToString(CultureInfo.InvariantCulture));
+
+			(_, long MinOffset) = await GetLineOffsetAsync(LogFile, LineIndex);
+			(_, long MaxOffset) = await GetLineOffsetAsync(LogFile, LineIndex + LineCount);
+
+			byte[] Data = new byte[MaxOffset - MinOffset];
+			using (Stream Stream = await OpenRawStreamAsync(LogFile, MinOffset, MaxOffset - MinOffset))
+			{
+				int Length = await Stream.ReadAsync(Data.AsMemory());
+				if(Length != Data.Length)
+				{
+					Logger.LogWarning("Read less than expected from log stream (Expected {Expected}, Got {Got})", Data.Length, Length);
+				}
+				return ParseEventData(Data.AsSpan(0, Length));
+			}
+		}
+
+		/// <summary>
+		/// Parses event data from a buffer
+		/// </summary>
+		/// <param name="Data">Data to parse</param>
+		/// <returns>Parsed event data</returns>
+		private ILogEventData ParseEventData(ReadOnlySpan<byte> Data)
+		{
+			List<LogEventLine> Lines = new List<LogEventLine>();
+
+			ReadOnlySpan<byte> RemainingData = Data;
+			while(RemainingData.Length > 0)
+			{
+				int EndOfLine = RemainingData.IndexOf((byte)'\n');
+				if(EndOfLine == -1)
+				{
+					break;
+				}
+
+				ReadOnlySpan<byte> LineData = RemainingData.Slice(0, EndOfLine);
+				try
+				{
+					Lines.Add(new LogEventLine(LineData));
+				}
+				catch(JsonException Ex)
+				{
+					Logger.LogWarning(Ex, "Unable to parse line from log file: {Line}", Encoding.UTF8.GetString(LineData));
+				}
+				RemainingData = RemainingData.Slice(EndOfLine + 1);
+			}
+
+			return new LogEventData(Lines);
+		}
+
+		/// <inheritdoc/>
+		public async Task<Stream> OpenRawStreamAsync(ILogFile LogFile, long Offset, long Length)
+		{
+			if (LogFile.Chunks.Count == 0)
+			{
+				return new MemoryStream(Array.Empty<byte>(), false);
+			}
+			else
+			{
+				int LastChunkIdx = LogFile.Chunks.Count - 1;
+
+				// Clamp the length of the request
+				ILogChunk LastChunk = LogFile.Chunks[LastChunkIdx];
+				if (Length > LastChunk.Offset)
+				{
+					long LastChunkLength = LastChunk.Length;
+					if (LastChunkLength <= 0)
+					{
+						LogChunkData LastChunkData = await ReadChunkAsync(LogFile, LastChunkIdx);
+						LastChunkLength = LastChunkData.Length;
+					}
+					Length = Math.Min(Length, (LastChunk.Offset + LastChunkLength) - Offset);
+				}
+
+				// Create the new stream
+				return new ResponseStream(this, LogFile, Offset, Length);
+			}
+		}
+
+		/// <inheritdoc/>
+		public async Task<(int, long)> GetLineOffsetAsync(ILogFile LogFile, int LineIdx)
+		{
+			int ChunkIdx = LogFile.Chunks.GetChunkForLine(LineIdx);
+
+			ILogChunk Chunk = LogFile.Chunks[ChunkIdx];
+			LogChunkData ChunkData = await ReadChunkAsync(LogFile, ChunkIdx);
+
+			if (LineIdx < Chunk.LineIndex)
+			{
+				LineIdx = Chunk.LineIndex;
+			}
+
+			int MaxLineIndex = Chunk.LineIndex + ChunkData.LineCount;
+			if (LineIdx >= MaxLineIndex)
+			{
+				LineIdx = MaxLineIndex;
+			}
+
+			long Offset = Chunk.Offset + ChunkData.GetLineOffsetWithinChunk(LineIdx - Chunk.LineIndex);
+			return (LineIdx, Offset);
+		}
+
+		/// <summary>
+		/// Executes a background task
+		/// </summary>
+		/// <param name="StoppingToken">Cancellation token</param>
+		protected override async Task TickAsync(CancellationToken StoppingToken)
+		{
+			lock (WriteLock)
+			{
+				WriteTasks.RemoveCompleteTasks();
+			}
+			await IncrementalFlush();
+		}
+		
+		/// <summary>
+		/// Executes a background task. Publicly exposed for tests.
+		/// </summary>
+		public async Task TickOnlyForTestingAsync()
+		{
+			await TickAsync(new CancellationToken());
+		}
+		
+		/// <summary>
+		/// Flushes complete chunks to the storage provider
+		/// </summary>
+		/// <returns>Async task</returns>
+		private async Task IncrementalFlush()
+		{
+			// Get all the chunks older than 20 minutes
+			List<(ObjectId, long)> FlushChunks = await Builder.TouchChunksAsync(TimeSpan.FromMinutes(10.0));
+			Logger.LogDebug("Performing incremental flush of log builder ({NumChunks} chunks)", FlushChunks.Count);
+
+			// Mark them all as complete
+			foreach ((ObjectId LogId, long Offset) in FlushChunks)
+			{
+				await Builder.CompleteChunkAsync(LogId, Offset);
+			}
+
+			// Add tasks for flushing all the chunks
+			WriteCompleteChunks(FlushChunks, true);
+		}
+
+		/// <summary>
+		/// Called when the service is shutting down
+		/// </summary>
+		/// <param name="CancellationToken"></param>
+		/// <returns></returns>
+		[SuppressMessage("Design", "CA1031:Do not catch general exception types")]
+		public override async Task StopAsync(CancellationToken CancellationToken)
+		{
+			Logger.LogInformation("Stopping log file service");
+			if (Builder.FlushOnShutdown)
+			{
+				await FlushAsync();
+			}
+			Logger.LogInformation("Log service stopped");
+		}
+
+		/// <summary>
+		/// Flushes the write cache
+		/// </summary>
+		/// <returns>Async task</returns>
+		public async Task FlushAsync()
+		{
+			Logger.LogInformation("Forcing flush of pending log chunks...");
+
+			// Mark everything in the cache as complete
+			List<(ObjectId, long)> WriteChunks = await Builder.TouchChunksAsync(TimeSpan.Zero);
+			WriteCompleteChunks(WriteChunks, true);
+
+			// Wait for everything to flush
+			await FlushPendingWritesAsync();
+		}
+
+		/// <summary>
+		/// Flush any writes in progress
+		/// </summary>
+		/// <returns>Async task</returns>
+		public async Task FlushPendingWritesAsync()
+		{
+			for(; ;)
+			{
+				// Capture the current contents of the WriteTasks list
+				List<Task> Tasks;
+				lock (WriteLock)
+				{
+					WriteTasks.RemoveCompleteTasks();
+					Tasks = new List<Task>(WriteTasks);
+				}
+				if (Tasks.Count == 0)
+				{
+					break;
+				}
+
+				// Also add a delay so we'll periodically refresh the list
+				Tasks.Add(Task.Delay(TimeSpan.FromSeconds(5.0)));
+				await Task.WhenAny(Tasks);
+			}
+		}
+
+		/// <summary>
+		/// Adds tasks for writing a list of complete chunks
+		/// </summary>
+		/// <param name="ChunksToWrite">List of chunks to write</param>
+		/// <param name="bCreateIndex">Create an index for the log</param>
+		private void WriteCompleteChunks(List<(ObjectId, long)> ChunksToWrite, bool bCreateIndex)
+		{
+			foreach (IGrouping<ObjectId, long> Group in ChunksToWrite.GroupBy(x => x.Item1, x => x.Item2))
+			{
+				ObjectId LogId = Group.Key;
+
+				// Find offsets of new chunks to write
+				List<long> Offsets = new List<long>();
+				lock (WriteLock)
+				{
+					foreach (long Offset in Group.OrderBy(x => x))
+					{
+						if (WriteChunks.Add((LogId, Offset)))
+						{
+							Offsets.Add(Offset);
+						}
+					}
+				}
+
+				// Create the write task
+				if (Offsets.Count > 0)
+				{
+					Task Task = Task.Run(() => WriteCompleteChunksForLogAsync(LogId, Offsets, bCreateIndex));
+					lock (WriteLock)
+					{
+						WriteTasks.Add(Task);
+					}
+				}
+			}
+		}
+
+		/// <summary>
+		/// Writes a set of chunks to the database
+		/// </summary>
+		/// <param name="LogId">Log file to update</param>
+		/// <param name="Offsets">Chunks to write</param>
+		/// <param name="bCreateIndex">Whether to create the index for this log</param>
+		/// <returns>Async task</returns>
+		private async Task<ILogFile?> WriteCompleteChunksForLogAsync(ObjectId LogId, List<long> Offsets, bool bCreateIndex)
+		{
+			ILogFile? LogFile = await LogFiles.GetLogFileAsync(LogId);
+			if(LogFile != null)
+			{
+				LogFile = await WriteCompleteChunksForLogAsync(LogFile, Offsets, bCreateIndex);
+			}
+			return LogFile;
+		}
+
+		/// <summary>
+		/// Writes a set of chunks to the database
+		/// </summary>
+		/// <param name="LogFileInterface">Log file to update</param>
+		/// <param name="Offsets">Chunks to write</param>
+		/// <param name="bCreateIndex">Whether to create the index for this log</param>
+		/// <returns>Async task</returns>
+		[SuppressMessage("Design", "CA1031:Do not catch general exception types", Justification = "<Pending>")]
+		private async Task<ILogFile?> WriteCompleteChunksForLogAsync(ILogFile LogFileInterface, List<long> Offsets, bool bCreateIndex)
+		{
+			// Write the data to the storage provider
+			List<Task<LogChunkData?>> ChunkWriteTasks = new List<Task<LogChunkData?>>();
+			foreach (long Offset in Offsets)
+			{
+				int ChunkIdx = LogFileInterface.Chunks.BinarySearch(x => x.Offset, Offset);
+				if (ChunkIdx >= 0)
+				{
+					Logger.LogDebug("Queuing write of log {LogId} chunk {ChunkIdx} offset {Offset}", LogFileInterface.Id, ChunkIdx, Offset);
+					int LineIndex = LogFileInterface.Chunks[ChunkIdx].LineIndex;
+					ChunkWriteTasks.Add(Task.Run(() => WriteChunkAsync(LogFileInterface.Id, Offset, LineIndex)));
+				}
+			}
+
+			// Wait for the tasks to complete, periodically updating the log file object
+			ILogFile? LogFile = LogFileInterface;
+			while (ChunkWriteTasks.Count > 0)
+			{
+				// Wait for all tasks to be complete OR (any task has completed AND 30 seconds has elapsed)
+				Task AllCompleteTask = Task.WhenAll(ChunkWriteTasks);
+				Task AnyCompleteTask = Task.WhenAny(ChunkWriteTasks);
+				await Task.WhenAny(AllCompleteTask, Task.WhenAll(AnyCompleteTask, Task.Delay(TimeSpan.FromSeconds(30.0))));
+
+				// Update the log file with the written chunks
+				List<LogChunkData?> WrittenChunks = ChunkWriteTasks.RemoveCompleteTasks();
+				while (LogFile != null)
+				{
+					// Update the length of any complete chunks
+					List<CompleteLogChunkUpdate> Updates = new List<CompleteLogChunkUpdate>();
+					foreach (LogChunkData? ChunkData in WrittenChunks)
+					{
+						if (ChunkData != null)
+						{
+							int ChunkIdx = LogFile.Chunks.GetChunkForOffset(ChunkData.Offset);
+							if (ChunkIdx >= 0)
+							{
+								ILogChunk Chunk = LogFile.Chunks[ChunkIdx];
+								if (Chunk.Offset == ChunkData.Offset)
+								{
+									CompleteLogChunkUpdate Update = new CompleteLogChunkUpdate(ChunkIdx, ChunkData.Length, ChunkData.LineCount);
+									Updates.Add(Update);
+								}
+							}
+						}
+					}
+
+					// Try to apply the updates
+					ILogFile? NewLogFile = await LogFiles.TryCompleteChunksAsync(LogFile, Updates);
+					if (NewLogFile != null)
+					{
+						LogFile = NewLogFile;
+						break;
+					}
+
+					// Update the log file
+					LogFile = await GetLogFileAsync(LogFile.Id);
+				}
+			}
+
+			// Create the index if necessary
+			if (bCreateIndex && LogFile != null)
+			{
+				try
+				{
+					LogFile = await CreateIndexAsync(LogFile);
+				}
+				catch(Exception Ex)
+				{
+					Logger.LogError(Ex, "Failed to create index for log {LogId}", LogFileInterface.Id);
+				}
+			}
+
+			return LogFile;
+		}
+
+		/// <summary>
+		/// Creates an index for the given log file
+		/// </summary>
+		/// <param name="LogFile">The log file object</param>
+		/// <returns>Updated log file</returns>
+		private async Task<ILogFile?> CreateIndexAsync(ILogFile LogFile)
+		{
+			if(LogFile.Chunks.Count == 0)
+			{
+				return LogFile;
+			}
+
+			// Get the new length of the log, and early out if it won't be any longer
+			ILogChunk LastChunk = LogFile.Chunks[LogFile.Chunks.Count - 1];
+			if(LastChunk.Offset + LastChunk.Length <= (LogFile.IndexLength ?? 0))
+			{
+				return LogFile;
+			}
+
+			// Save stats for the index creation
+			using Scope Scope = Tracer.Instance.StartActive("CreateIndexAsync");
+			Scope.Span.SetTag("LogId", LogFile.Id.ToString());
+			Scope.Span.SetTag("Length", (LastChunk.Offset + LastChunk.Length).ToString(CultureInfo.InvariantCulture));
+
+			long NewLength = 0;
+			int NewLineCount = 0;
+
+			// Read the existing index if there is one
+			List<LogIndexData> Indexes = new List<LogIndexData>();
+			if (LogFile.IndexLength != null)
+			{
+				LogIndexData? ExistingIndex = await ReadIndexAsync(LogFile, LogFile.IndexLength.Value);
+				if(ExistingIndex != null)
+				{
+					Indexes.Add(ExistingIndex);
+					NewLineCount = ExistingIndex.LineCount;
+				}
+			}
+
+			// Add all the new chunks
+			int ChunkIdx = LogFile.Chunks.GetChunkForLine(NewLineCount);
+			if (ChunkIdx < 0)
+			{
+				int FirstLine = (LogFile.Chunks.Count > 0) ? LogFile.Chunks[0].LineIndex : -1;
+				throw new Exception($"Invalid chunk index {ChunkIdx}. Index.LineCount={NewLineCount}, Chunks={LogFile.Chunks.Count}, First line={FirstLine}");
+			}
+
+			for (; ChunkIdx < LogFile.Chunks.Count; ChunkIdx++)
+			{
+				ILogChunk Chunk = LogFile.Chunks[ChunkIdx];
+				LogChunkData ChunkData = await ReadChunkAsync(LogFile, ChunkIdx);
+
+				int SubChunkIdx = ChunkData.GetSubChunkForLine(Math.Max(NewLineCount - Chunk.LineIndex, 0));
+				if(SubChunkIdx < 0)
+				{
+					throw new Exception($"Invalid subchunk index {SubChunkIdx}. Chunk {ChunkIdx}/{LogFile.Chunks.Count}. Index.LineCount={NewLineCount}, Chunk.LineIndex={Chunk.LineIndex}, First subchunk {ChunkData.SubChunkLineIndex[0]}");
+				}
+
+				for (; SubChunkIdx < ChunkData.SubChunks.Count; SubChunkIdx++)
+				{
+					LogSubChunkData SubChunkData = ChunkData.SubChunks[SubChunkIdx];
+					if (SubChunkData.LineIndex >= NewLineCount)
+					{
+						try
+						{
+							Indexes.Add(SubChunkData.BuildIndex());
+						}
+						catch (Exception Ex)
+						{
+							throw new Exception($"Failed to create index block - log {LogFile.Id}, chunk {ChunkIdx} ({LogFile.Chunks.Count}), subchunk {SubChunkIdx} ({ChunkData.SubChunks.Count}), index lines: {NewLineCount}, chunk index: {Chunk.LineIndex}, subchunk index: {Chunk.LineIndex + ChunkData.SubChunkLineIndex[SubChunkIdx]}, subchunk count: {SubChunkData.LineCount}", Ex);
+						}
+
+						NewLength = SubChunkData.Offset + SubChunkData.Length;
+						NewLineCount = SubChunkData.LineIndex + SubChunkData.LineCount;
+					}
+				}
+			}
+
+			// Try to update the log file
+			ILogFile? NewLogFile = LogFile;
+			if (NewLength > (LogFile.IndexLength ?? 0))
+			{
+				LogIndexData Index = LogIndexData.Merge(Indexes);
+				Logger.LogDebug("Writing index for log {LogId} covering {Length} (index length {IndexLength})", LogFile.Id, NewLength, Index.GetSerializedSize());
+
+				await WriteIndexAsync(LogFile.Id, NewLength, Index);
+
+				while(NewLogFile != null && NewLength > (NewLogFile.IndexLength ?? 0))
+				{
+					NewLogFile = await LogFiles.TryUpdateIndexAsync(NewLogFile, NewLength);
+					if(NewLogFile != null)
+					{
+						break;
+					}
+					NewLogFile = await LogFiles.GetLogFileAsync(LogFile.Id);
+				}
+			}
+			return NewLogFile;
+		}
+
+		/// <summary>
+		/// Reads a chunk from storage
+		/// </summary>
+		/// <param name="LogFile">Log file to read from</param>
+		/// <param name="ChunkIdx">The chunk to read</param>
+		/// <returns>Chunk data</returns>
+		[SuppressMessage("Design", "CA1031:Do not catch general exception types", Justification = "<Pending>")]
+		private async Task<LogChunkData> ReadChunkAsync(ILogFile LogFile, int ChunkIdx)
+		{
+			ILogChunk Chunk = LogFile.Chunks[ChunkIdx];
+
+			// Try to read the chunk data from storage
+			LogChunkData? ChunkData = null;
+			try
+			{
+				// If the chunk is not yet complete, query the log builder
+				if (Chunk.Length == 0)
+				{
+					ChunkData = await Builder.GetChunkAsync(LogFile.Id, Chunk.Offset, Chunk.LineIndex);
+				}
+
+				// Otherwise go directly to the log storage
+				if (ChunkData == null)
+				{
+					ChunkData = await Storage.ReadChunkAsync(LogFile.Id, Chunk.Offset, Chunk.LineIndex);
+				}
+			}
+			catch (Exception Ex)
+			{
+				Logger.LogError(Ex, "Unable to read log {LogId} at offset {Offset}", LogFile.Id, Chunk.Offset);
+			}
+
+			// Get the minimum length and line count for the chunk
+			if (ChunkIdx + 1 < LogFile.Chunks.Count)
+			{
+				ILogChunk NextChunk = LogFile.Chunks[ChunkIdx + 1];
+				ChunkData = RepairChunkData(LogFile, ChunkIdx, ChunkData, (int)(NextChunk.Offset - Chunk.Offset), NextChunk.LineIndex - Chunk.LineIndex);
+			}
+			else
+			{
+				if (LogFile.MaxLineIndex != null && Chunk.Length != 0)
+				{
+					ChunkData = RepairChunkData(LogFile, ChunkIdx, ChunkData, Chunk.Length, LogFile.MaxLineIndex.Value - Chunk.LineIndex);
+				}
+				else if(ChunkData == null)
+				{
+					ChunkData = RepairChunkData(LogFile, ChunkIdx, ChunkData, 1024, 1);
+				}
+			}
+
+			return ChunkData;
+		}
+
+		/// <summary>
+		/// Validates the given chunk data, and fix it up if necessary
+		/// </summary>
+		/// <param name="LogFile">The log file instance</param>
+		/// <param name="ChunkIdx">Index of the chunk within the logfile</param>
+		/// <param name="ChunkData">The chunk data that was read</param>
+		/// <param name="Length">Expected length of the data</param>
+		/// <param name="LineCount">Expected number of lines in the data</param>
+		/// <returns>Repaired chunk data</returns>
+		LogChunkData RepairChunkData(ILogFile LogFile, int ChunkIdx, LogChunkData? ChunkData, int Length, int LineCount)
+		{
+			int CurrentLength = 0;
+			int CurrentLineCount = 0;
+			if(ChunkData != null)
+			{
+				CurrentLength = ChunkData.Length;
+				CurrentLineCount = ChunkData.LineCount;
+			}
+
+			if (ChunkData == null || CurrentLength < Length || CurrentLineCount < LineCount)
+			{
+				Logger.LogWarning("Creating placeholder subchunk for log {LogId} chunk {ChunkIdx} (length {Length} vs expected {ExpLength}, lines {LineCount} vs expected {ExpLineCount})", LogFile.Id, ChunkIdx, CurrentLength, Length, CurrentLineCount, LineCount);
+
+				List<LogSubChunkData> SubChunks = new List<LogSubChunkData>();
+				if (ChunkData != null && ChunkData.Length < Length && ChunkData.LineCount < LineCount)
+				{
+					SubChunks.AddRange(ChunkData.SubChunks);
+				}
+
+				LogText Text = new LogText();
+				Text.AppendMissingDataInfo(ChunkIdx, LogFile.Chunks[ChunkIdx].Server, Length - CurrentLength, LineCount - CurrentLineCount);
+				SubChunks.Add(new LogSubChunkData(LogFile.Type, CurrentLength, CurrentLineCount, Text));
+
+				ILogChunk Chunk = LogFile.Chunks[ChunkIdx];
+				ChunkData = new LogChunkData(Chunk.Offset, Chunk.LineIndex, SubChunks);
+			}
+			return ChunkData;
+		}
+
+		/// <summary>
+		/// Writes a set of chunks to the database
+		/// </summary>
+		/// <param name="LogFileId">Unique id of the log file</param>
+		/// <param name="Offset">Offset of the chunk to write</param>
+		/// <param name="LineIndex">First line index of the chunk</param>
+		/// <returns>Chunk daata</returns>
+		[SuppressMessage("Design", "CA1031:Do not catch general exception types", Justification = "<Pending>")]
+		private async Task<LogChunkData?> WriteChunkAsync(ObjectId LogFileId, long Offset, int LineIndex)
+		{
+			// Write the chunk to storage
+			LogChunkData? ChunkData = await Builder.GetChunkAsync(LogFileId, Offset, LineIndex);
+			if (ChunkData == null)
+			{
+				Logger.LogDebug("Log {LogId} offset {Offset} not found in log builder", LogFileId, Offset);
+			}
+			else
+			{
+				try
+				{
+					await Storage.WriteChunkAsync(LogFileId, ChunkData.Offset, ChunkData);
+				}
+				catch (Exception Ex)
+				{
+					Logger.LogError(Ex, "Unable to write log {LogId} at offset {Offset}", LogFileId, ChunkData.Offset);
+				}
+			}
+
+			// Remove it from the log builder
+			try
+			{
+				await Builder.RemoveChunkAsync(LogFileId, Offset);
+			}
+			catch (Exception Ex)
+			{
+				Logger.LogError(Ex, "Unable to remove log {LogId} at offset {Offset} from log builder", LogFileId, Offset);
+			}
+
+			return ChunkData;
+		}
+
+		/// <summary>
+		/// Reads a chunk from storage
+		/// </summary>
+		/// <param name="LogFile">Log file to read from</param>
+		/// <param name="Length">Length of the log covered by the index</param>
+		/// <returns>Chunk data</returns>
+		[SuppressMessage("Design", "CA1031:Do not catch general exception types", Justification = "<Pending>")]
+		private async Task<LogIndexData?> ReadIndexAsync(ILogFile LogFile, long Length)
+		{
+			try
+			{
+				LogIndexData? Index = await Storage.ReadIndexAsync(LogFile.Id, Length);
+				return Index;
+			}
+			catch (Exception Ex)
+			{
+				Logger.LogError(Ex, "Unable to read log {LogId} index at length {Length}", LogFile.Id, Length);
+				return null;
+			}
+		}
+
+		/// <summary>
+		/// Writes an index to the database
+		/// </summary>
+		/// <param name="LogFileId">Unique id of the log file</param>
+		/// <param name="Length">Length of the data covered by the index</param>
+		/// <param name="Index">Index to write</param>
+		/// <returns>Async task</returns>
+		[SuppressMessage("Design", "CA1031:Do not catch general exception types", Justification = "<Pending>")]
+		private async Task WriteIndexAsync(ObjectId LogFileId, long Length, LogIndexData Index)
+		{
+			try
+			{
+				await Storage.WriteIndexAsync(LogFileId, Length, Index);
+			}
+			catch(Exception Ex)
+			{
+				Logger.LogError(Ex, "Unable to write index for log {LogId}", LogFileId);
+			}
+		}
+
+		/// <summary>
+		/// Determines if the user is authorized to perform an action on a particular template
+		/// </summary>
+		/// <param name="LogFile">The template to check</param>
+		/// <param name="User">The principal to authorize</param>
+		/// <returns>True if the action is authorized</returns>
+		public static bool AuthorizeForSession(ILogFile LogFile, ClaimsPrincipal User)
+		{
+			if(LogFile.SessionId != null)
+			{
+				return User.HasClaim(HordeClaimTypes.AgentSessionId, LogFile.SessionId.Value.ToString());
+			}
+			else
+			{
+				return false;
+			}
+		}
+
+		/// <inheritdoc/>
+		public async Task<List<int>> SearchLogDataAsync(ILogFile LogFile, string Text, int FirstLine, int Count, LogSearchStats SearchStats)
+		{
+			Stopwatch Timer = Stopwatch.StartNew();
+
+			using Scope Scope = Tracer.Instance.StartActive("SearchLogDataAsync");
+			Scope.Span.SetTag("LogId", LogFile.Id.ToString());
+			Scope.Span.SetTag("Text", Text);
+			Scope.Span.SetTag("Count", Count.ToString(CultureInfo.InvariantCulture));
+
+			List<int> Results = new List<int>();
+			if (Count > 0)
+			{
+				IAsyncEnumerator<int> Enumerator = SearchLogDataInternalAsync(LogFile, Text, FirstLine, SearchStats).GetAsyncEnumerator();
+				while (await Enumerator.MoveNextAsync() && Results.Count < Count)
+				{
+					Results.Add(Enumerator.Current);
+				}
+			}
+
+			Logger.LogDebug("Search for \"{SearchText}\" in log {LogId} found {NumResults}/{MaxResults} results, took {Time}ms ({@Stats})", Text, LogFile.Id, Results.Count, Count, Timer.ElapsedMilliseconds, SearchStats);
+			return Results;
+		}
+
+		/// <inheritdoc/>
+		public async IAsyncEnumerable<int> SearchLogDataInternalAsync(ILogFile LogFile, string Text, int FirstLine, LogSearchStats SearchStats)
+		{
+			SearchText SearchText = new SearchText(Text);
+
+			// Read the index for this log file
+			if (LogFile.IndexLength != null)
+			{
+				LogIndexData? IndexData = await ReadIndexAsync(LogFile, LogFile.IndexLength.Value);
+				if(IndexData != null && FirstLine < IndexData.LineCount)
+				{
+					using Scope IndexScope = Tracer.Instance.StartActive("Indexed");
+					IndexScope.Span.SetTag("LineCount", IndexData.LineCount.ToString(CultureInfo.InvariantCulture));
+
+					foreach(int LineIndex in IndexData.Search(FirstLine, SearchText, SearchStats))
+					{
+						yield return LineIndex;
+					}
+
+					FirstLine = IndexData.LineCount;
+				}
+			}
+
+			// Manually search through the rest of the log
+			int ChunkIdx = LogFile.Chunks.GetChunkForLine(FirstLine);
+			for (; ChunkIdx < LogFile.Chunks.Count; ChunkIdx++)
+			{
+				ILogChunk Chunk = LogFile.Chunks[ChunkIdx];
+
+				// Read the chunk data
+				LogChunkData ChunkData = await ReadChunkAsync(LogFile, ChunkIdx);
+				if (FirstLine < ChunkData.LineIndex + ChunkData.LineCount)
+				{
+					// Find the first sub-chunk we're looking for
+					int SubChunkIdx = 0;
+					if (FirstLine > Chunk.LineIndex)
+					{
+						SubChunkIdx = ChunkData.GetSubChunkForLine(FirstLine - Chunk.LineIndex);
+					}
+
+					// Search through the sub-chunks
+					for (; SubChunkIdx < ChunkData.SubChunks.Count; SubChunkIdx++)
+					{
+						LogSubChunkData SubChunkData = ChunkData.SubChunks[SubChunkIdx];
+						if (FirstLine < SubChunkData.LineIndex + SubChunkData.LineCount)
+						{
+							// Create an index containing just this sub-chunk
+							LogIndexData Index = SubChunkData.BuildIndex();
+							foreach (int LineIndex in Index.Search(FirstLine, SearchText, SearchStats))
+							{
+								yield return LineIndex;
+							}
+						}
+					}
+				}
+			}
+		}
+	}
+}
+
diff --git a/Engine/Source/Programs/Horde/HordeServer/Services/MetricService.cs b/Engine/Source/Programs/Horde/HordeServer/Services/MetricService.cs
new file mode 100644
--- /dev/null
+++ b/Engine/Source/Programs/Horde/HordeServer/Services/MetricService.cs
@@ -0,0 +1,76 @@
+using HordeServer.Utilities;
+using Microsoft.Extensions.Logging;
+using System;
+using System.Threading;
+using System.Threading.Tasks;
+using StatsdClient;
+using PoolId = HordeServer.Utilities.StringId<HordeServer.Models.IPool>;
+
+namespace HordeServer.Services
+{
+	/// <summary>
+	/// Periodically send metrics for the CLR and other services that cannot be collected on a per-request basis
+	/// </summary>
+	public class MetricService : TickedBackgroundService
+	{
+		/// <summary>
+		/// DogStatsd instance
+		/// </summary>
+		private readonly IDogStatsd DogStatsd;
+
+		/// <summary>
+		/// Constructor
+		/// </summary>
+		/// <param name="DogStatsd"></param>
+		/// <param name="Logger">Logger instance</param>
+		public MetricService(IDogStatsd DogStatsd, ILogger<MetricService> Logger) : base(TimeSpan.FromSeconds(20.0), Logger)
+		{
+			this.DogStatsd = DogStatsd;
+		}
+
+		/// <summary>
+		/// Execute the background task
+		/// </summary>
+		/// <param name="StoppingToken">Cancellation token for the async task</param>
+		/// <returns>Async task</returns>
+		protected override Task TickAsync(CancellationToken StoppingToken)
+		{
+			ReportGcMetrics();
+			ReportThreadMetrics();
+			return Task.CompletedTask;
+		}
+
+		private void ReportThreadMetrics()
+		{
+			ThreadPool.GetMaxThreads(out int MaxWorkerThreads, out int MaxIoThreads);
+			ThreadPool.GetAvailableThreads(out int FreeWorkerThreads, out int FreeIoThreads);
+			ThreadPool.GetMinThreads(out int MinWorkerThreads, out int MinIoThreads);
+
+			int BusyIoThreads = MaxIoThreads - FreeIoThreads;
+			int BusyWorkerThreads = MaxWorkerThreads - FreeWorkerThreads;
+			
+			DogStatsd.Gauge("horde.clr.threadpool.io.max", MaxIoThreads);
+			DogStatsd.Gauge("horde.clr.threadpool.io.min", MinIoThreads);
+			DogStatsd.Gauge("horde.clr.threadpool.io.free", FreeIoThreads);
+			DogStatsd.Gauge("horde.clr.threadpool.io.busy", BusyIoThreads);
+			
+			DogStatsd.Gauge("horde.clr.threadpool.worker.max", MaxWorkerThreads);
+			DogStatsd.Gauge("horde.clr.threadpool.worker.min", MinWorkerThreads);
+			DogStatsd.Gauge("horde.clr.threadpool.worker.free", FreeWorkerThreads);
+			DogStatsd.Gauge("horde.clr.threadpool.worker.busy", BusyWorkerThreads);
+		}
+
+		private void ReportGcMetrics()
+		{
+			GCMemoryInfo GcMemoryInfo = GC.GetGCMemoryInfo();
+
+			DogStatsd.Gauge("horde.clr.gc.totalMemory", GC.GetTotalMemory(false));
+			DogStatsd.Gauge("horde.clr.gc.totalAllocated", GC.GetTotalAllocatedBytes());
+			DogStatsd.Gauge("horde.clr.gc.heapSize", GcMemoryInfo.HeapSizeBytes);
+			DogStatsd.Gauge("horde.clr.gc.fragmented", GcMemoryInfo.FragmentedBytes);
+			DogStatsd.Gauge("horde.clr.gc.memoryLoad", GcMemoryInfo.MemoryLoadBytes);
+			DogStatsd.Gauge("horde.clr.gc.totalAvailableMemory", GcMemoryInfo.TotalAvailableMemoryBytes);
+			DogStatsd.Gauge("horde.clr.gc.highMemoryLoadThreshold", GcMemoryInfo.HighMemoryLoadThresholdBytes);
+		}
+	}
+}
diff --git a/Engine/Source/Programs/Horde/HordeServer/Services/PerforceLoadBalancer.cs b/Engine/Source/Programs/Horde/HordeServer/Services/PerforceLoadBalancer.cs
new file mode 100644
--- /dev/null
+++ b/Engine/Source/Programs/Horde/HordeServer/Services/PerforceLoadBalancer.cs
@@ -0,0 +1,432 @@
+// Copyright Epic Games, Inc. All Rights Reserved.
+
+using Google.Protobuf.WellKnownTypes;
+using HordeCommon.Rpc.Tasks;
+using HordeServer.Collections;
+using HordeServer.Models;
+using HordeServer.Utilities;
+using Microsoft.Extensions.Logging;
+using MongoDB.Bson;
+using MongoDB.Bson.Serialization.Attributes;
+using MongoDB.Driver;
+using System;
+using System.Collections.Concurrent;
+using System.Collections.Generic;
+using System.Globalization;
+using System.Linq;
+using System.Net;
+using System.Net.Http;
+using System.Text.Json;
+using System.Threading;
+using System.Threading.Tasks;
+
+namespace HordeServer.Services
+{
+	/// <summary>
+	/// Health of a particular server
+	/// </summary>
+	public enum PerforceServerStatus
+	{
+		/// <summary>
+		/// Server could not be reached
+		/// </summary>
+		Unknown,
+
+		/// <summary>
+		/// Bad. Do not use.
+		/// </summary>
+		Unhealthy,
+
+		/// <summary>
+		/// Degraded, but functioning
+		/// </summary>
+		Degraded,
+
+		/// <summary>
+		/// Good
+		/// </summary>
+		Healthy,
+	}
+
+	/// <summary>
+	/// Perforce server infomation
+	/// </summary>
+	public interface IPerforceServer
+	{
+		/// <summary>
+		/// The server and port
+		/// </summary>
+		public string ServerAndPort { get; }
+
+		/// <summary>
+		/// The server and port before doing DNS resolution
+		/// </summary>
+		public string BaseServerAndPort { get; }
+
+		/// <summary>
+		/// The cluster this server belongs to
+		/// </summary>
+		public string Cluster { get; }
+
+		/// <summary>
+		/// Current status
+		/// </summary>
+		public PerforceServerStatus Status { get; }
+
+		/// <summary>
+		/// Number of leases using this server
+		/// </summary>
+		public int NumLeases { get; }
+
+		/// <summary>
+		/// Error message related to this server
+		/// </summary>
+		public string? Detail { get; }
+	}
+
+	/// <summary>
+	/// Load balancer for Perforce edge servers
+	/// </summary>
+	public class PerforceLoadBalancer : ElectedBackgroundService
+	{
+		/// <summary>
+		/// Information about a resolved Perforce server
+		/// </summary>
+		class PerforceServerEntry : IPerforceServer
+		{
+			public string ServerAndPort { get; set; } = String.Empty;
+			public string BaseServerAndPort { get; set; } = String.Empty;
+			public string Cluster { get; set; } = String.Empty;
+			public PerforceServerStatus Status { get; set; }
+			public string? Detail { get; set; }
+			public int NumLeases { get; set; }
+
+			[BsonConstructor]
+			private PerforceServerEntry()
+			{
+			}
+
+			public PerforceServerEntry(string ServerAndPort, string BaseServerAndPort, string Cluster, PerforceServerStatus Status, string? Detail)
+			{
+				this.ServerAndPort = ServerAndPort;
+				this.BaseServerAndPort = BaseServerAndPort;
+				this.Cluster = Cluster;
+				this.Status = Status;
+				this.Detail = Detail;
+			}
+		}
+
+		/// <summary>
+		/// 
+		/// </summary>
+		[SingletonDocument("6046aec374a9283100967ee7")]
+		class PerforceServerList : SingletonBase
+		{
+			public List<PerforceServerEntry> Servers { get; set; } = new List<PerforceServerEntry>();
+		}
+
+		DatabaseService DatabaseService;
+		ILeaseCollection LeaseCollection;
+		SingletonDocument<PerforceServerList> ServerListSingleton;
+		Random Random = new Random();
+		ILogger Logger;
+
+		/// <summary>
+		/// Constructor
+		/// </summary>
+		/// <param name="DatabaseService"></param>
+		/// <param name="LeaseCollection"></param>
+		/// <param name="Logger"></param>
+		public PerforceLoadBalancer(DatabaseService DatabaseService, ILeaseCollection LeaseCollection, ILogger<PerforceLoadBalancer> Logger)
+			: base(DatabaseService, new ObjectId("603fb20536e999974fd6ff6b"), TimeSpan.FromMinutes(1.0), TimeSpan.FromMinutes(2.0), Logger)
+		{
+			this.DatabaseService = DatabaseService;
+			this.LeaseCollection = LeaseCollection;
+			this.ServerListSingleton = new SingletonDocument<PerforceServerList>(DatabaseService);
+			this.Logger = Logger;
+		}
+
+		/// <summary>
+		/// Get the perforce servers
+		/// </summary>
+		/// <returns></returns>
+		public async Task<List<IPerforceServer>> GetServersAsync()
+		{
+			PerforceServerList ServerList = await ServerListSingleton.GetAsync();
+			return ServerList.Servers.ConvertAll<IPerforceServer>(x => x);
+		}
+
+		/// <summary>
+		/// Allocates a server for use by a lease
+		/// </summary>
+		/// <returns>The server to use. Null if there is no healthy server available.</returns>
+		public async Task<IPerforceServer?> GetServer(string Cluster)
+		{
+			PerforceServerList ServerList = await ServerListSingleton.GetAsync();
+
+			List<PerforceServerEntry> Candidates = ServerList.Servers.Where(x => x.Cluster == Cluster && x.Status >= PerforceServerStatus.Healthy).ToList();
+			if(Candidates.Count == 0)
+			{
+				int Idx = Random.Next(0, Candidates.Count);
+				return Candidates[Idx];
+			}
+			return null;
+		}
+
+		/// <summary>
+		/// Select a Perforce server to use
+		/// </summary>
+		/// <param name="Cluster"></param>
+		/// <param name="Agent"></param>
+		/// <returns></returns>
+		public async Task<IPerforceServer?> SelectServerAsync(PerforceCluster Cluster, IAgent Agent)
+		{
+			// Get the set of properties to use for checking server validity
+			HashSet<string>? Properties = Agent.Capabilities.Devices[0].Properties;
+
+			// Find all the valid servers for this agent
+			HashSet<string> ValidServers = new HashSet<string>(StringComparer.OrdinalIgnoreCase);
+			if (Properties != null)
+			{
+				ValidServers.UnionWith(Cluster.Servers.Where(x => x.Properties != null && x.Properties.Count > 0 && x.Properties.All(y => Properties.Contains(y))).Select(x => x.ServerAndPort));
+			}
+			if (ValidServers.Count == 0)
+			{
+				ValidServers.UnionWith(Cluster.Servers.Where(x => x.Properties == null || x.Properties.Count == 0).Select(x => x.ServerAndPort));
+			}
+
+			// Find all the matching servers
+			PerforceServerList ServerList = await ServerListSingleton.GetAsync();
+			List<PerforceServerEntry> Candidates = ServerList.Servers.Where(x => x.Cluster == Cluster.Name && x.Status == PerforceServerStatus.Healthy && ValidServers.Contains(x.BaseServerAndPort)).ToList();
+
+			if (Candidates.Count == 0)
+			{
+				Logger.LogDebug("Unable to find Perforce server from cluster {ClusterName} matching {ValidServers}", Cluster.Name, String.Join(", ", ValidServers));
+				return null;
+			}
+
+			// Select which server to use with a weighted average of the number of active leases
+			int Index = 0;
+			if (Candidates.Count > 1)
+			{
+				const int BaseWeight = 20;
+
+				int TotalLeases = Candidates.Sum(x => x.NumLeases);
+				int TotalWeight = Candidates.Sum(x => (TotalLeases + BaseWeight) - x.NumLeases);
+
+				int Weight = Random.Next(TotalWeight);
+				for (; Index + 1 < Candidates.Count; Index++)
+				{
+					Weight -= (TotalLeases + BaseWeight) - Candidates[Index].NumLeases;
+					if(Weight < 0)
+					{
+						break;
+					}
+				}
+			}
+
+			return Candidates[Index];
+		}
+
+		/// <inheritdoc/>
+		protected override async Task<DateTime> TickLeaderAsync(CancellationToken StoppingToken)
+		{
+			Globals Globals = await DatabaseService.GetGlobalsAsync();
+
+			// Set of new server entries
+			ConcurrentBag<PerforceServerEntry> NewServers = new ConcurrentBag<PerforceServerEntry>();
+
+			// Update the state of all the valid servers
+			foreach (PerforceCluster Cluster in Globals.PerforceClusters)
+			{
+				foreach (PerforceServer Server in Cluster.Servers)
+				{
+					try
+					{
+						await UpdateServerAsync(Cluster, Server, NewServers);
+					}
+					catch (Exception Ex)
+					{
+						Logger.LogError(Ex, "Exception while updating Perforce server status");
+					}
+				}
+			}
+
+			// Update the number of leases for each entry
+			List<PerforceServerEntry> NewServersList = NewServers.OrderBy(x => x.Cluster).ThenBy(x => x.BaseServerAndPort).ThenBy(x => x.ServerAndPort).ToList();
+			await UpdateLeaseCounts(NewServersList);
+
+			// Mark any other servers as deleted
+			PerforceServerList NewServerList = await ServerListSingleton.GetAsync();
+			NewServerList.Servers = NewServers.OrderBy(x => x.Cluster).ThenBy(x => x.BaseServerAndPort).ThenBy(x => x.ServerAndPort).ToList();
+			await ServerListSingleton.TryUpdateAsync(NewServerList);
+
+			return DateTime.UtcNow.AddMinutes(1.0);
+		}
+
+		async Task UpdateLeaseCounts(IEnumerable<PerforceServerEntry> NewServerEntries)
+		{
+			Dictionary<string, Dictionary<string, PerforceServerEntry>> NewServerLookup = new Dictionary<string, Dictionary<string, PerforceServerEntry>>();
+			foreach (PerforceServerEntry NewServerEntry in NewServerEntries)
+			{
+				if (NewServerEntry.Cluster != null)
+				{
+					Dictionary<string, PerforceServerEntry>? ClusterServers;
+					if (!NewServerLookup.TryGetValue(NewServerEntry.Cluster, out ClusterServers))
+					{
+						ClusterServers = new Dictionary<string, PerforceServerEntry>();
+						NewServerLookup.Add(NewServerEntry.Cluster, ClusterServers);
+					}
+					ClusterServers[NewServerEntry.ServerAndPort] = NewServerEntry;
+				}
+			}
+
+			List<ILease> Leases = await LeaseCollection.FindActiveLeasesAsync();
+			foreach (ILease Lease in Leases)
+			{
+				Any Any = Any.Parser.ParseFrom(Lease.Payload.ToArray());
+
+				HashSet<(string, string)> Servers = new HashSet<(string, string)>();
+				if (Any.TryUnpack(out ConformTask ConformTask))
+				{
+					foreach (HordeCommon.Rpc.Messages.AgentWorkspace Workspace in ConformTask.Workspaces)
+					{
+						Servers.Add((Workspace.Cluster, Workspace.ServerAndPort));
+					}
+				}
+				if (Any.TryUnpack(out ExecuteJobTask ExecuteJobTask))
+				{
+					if (ExecuteJobTask.AutoSdkWorkspace != null)
+					{
+						Servers.Add((ExecuteJobTask.AutoSdkWorkspace.Cluster, ExecuteJobTask.AutoSdkWorkspace.ServerAndPort));
+					}
+					if (ExecuteJobTask.Workspace != null)
+					{
+						Servers.Add((ExecuteJobTask.Workspace.Cluster, ExecuteJobTask.Workspace.ServerAndPort));
+					}
+				}
+
+				foreach ((string Cluster, string ServerAndPort) in Servers)
+				{
+					IncrementLeaseCount(Cluster, ServerAndPort, NewServerLookup);
+				}
+			}
+		}
+
+		static void IncrementLeaseCount(string Cluster, string ServerAndPort, Dictionary<string, Dictionary<string, PerforceServerEntry>> NewServers)
+		{
+			Dictionary<string, PerforceServerEntry>? ClusterServers;
+			if (Cluster != null && NewServers.TryGetValue(Cluster, out ClusterServers))
+			{
+				PerforceServerEntry? Entry;
+				if (ClusterServers.TryGetValue(ServerAndPort, out Entry))
+				{
+					Entry.NumLeases++;
+				}
+			}
+		}
+
+		async Task UpdateServerAsync(PerforceCluster Cluster, PerforceServer Server, ConcurrentBag<PerforceServerEntry> NewServers)
+		{
+			string InitialHostName = Server.ServerAndPort;
+			int Port = 1666;
+
+			int PortIdx = InitialHostName.LastIndexOf(':');
+			if (PortIdx != -1)
+			{
+				Port = int.Parse(InitialHostName.Substring(PortIdx + 1), NumberStyles.Integer, CultureInfo.InvariantCulture);
+				InitialHostName = InitialHostName.Substring(0, PortIdx);
+			}
+
+			List<string> HostNames = new List<string>();
+			if (!Server.ResolveDns || IPAddress.TryParse(InitialHostName, out IPAddress HostAddress))
+			{
+				HostNames.Add(InitialHostName);
+			}
+			else
+			{
+				await ResolveServersAsync(InitialHostName, HostNames);
+			}
+
+			List<Task> Tasks = new List<Task>();
+			foreach (string HostName in HostNames)
+			{
+				Tasks.Add(Task.Run(() => UpdateHostAsync(Cluster, Server, HostName, Port, NewServers)));
+			}
+			await Task.WhenAll(Tasks);
+		}
+
+		async Task ResolveServersAsync(string HostName, List<string> HostNames)
+		{
+			// Find all the addresses of the hosts
+			IPHostEntry Entry = await Dns.GetHostEntryAsync(HostName);
+			foreach (IPAddress Address in Entry.AddressList)
+			{
+				try
+				{
+					HostNames.Add(Address.ToString());
+				}
+				catch (Exception Ex)
+				{
+					Logger.LogDebug(Ex, "Unable to resolve host name for Perforce server {Address}", Address);
+				}
+			}
+		}
+
+		static async Task UpdateHostAsync(PerforceCluster Cluster, PerforceServer Server, string HostName, int Port, ConcurrentBag<PerforceServerEntry> NewServers)
+		{
+			string Detail = "Health check disabled";
+
+			// Get the health of the server
+			PerforceServerStatus Health = PerforceServerStatus.Healthy;
+			if (Server.HealthCheck)
+			{
+				Uri HealthCheckUrl = new Uri($"http://{HostName}:5000/healthcheck");
+				try
+				{
+					(Health, Detail) = await GetServerHealthAsync(HealthCheckUrl);
+				}
+				catch
+				{
+					(Health, Detail) = (PerforceServerStatus.Unhealthy, $"Failed to query status at {HealthCheckUrl}");
+				}
+			}
+
+			// Update the server record
+			NewServers.Add(new PerforceServerEntry($"{HostName}:{Port}", Server.ServerAndPort, Cluster.Name, Health, Detail));
+		}
+
+		static async Task<(PerforceServerStatus, string)> GetServerHealthAsync(Uri HealthCheckUrl)
+		{
+			using HttpClient Client = new HttpClient();
+			HttpResponseMessage Response = await Client.GetAsync(HealthCheckUrl);
+
+			byte[] Data = await Response.Content.ReadAsByteArrayAsync();
+			JsonDocument Document = JsonDocument.Parse(Data);
+
+			foreach (JsonElement Element in Document.RootElement.GetProperty("results").EnumerateArray())
+			{
+				if (Element.TryGetProperty("checker", out JsonElement Checker) && Checker.ValueEquals("edge_traffic_lights"))
+				{
+					if (Element.TryGetProperty("output", out JsonElement Output))
+					{
+						string Status = Output.GetString();
+						switch(Status)
+						{
+							case "green":
+								return (PerforceServerStatus.Healthy, "Server is healthy");
+							case "yellow":
+								return (PerforceServerStatus.Degraded, "Degraded service");
+							case "red":
+								return (PerforceServerStatus.Unhealthy, "Server is being drained");
+							default:
+								return (PerforceServerStatus.Unknown, $"Expected state for health check ({Status})");
+						}
+					}
+				}
+			}
+
+			return (PerforceServerStatus.Unknown, "Unable to parse health check output");
+		}
+	}
+}
