Initial drop of Horde Server/Agent for UE5/Main. Still missing some Epic-specific functionality required to build/run correctly.

diff --git a/Engine/Source/Programs/Horde/HordeServer/Services/PerforceService.cs b/Engine/Source/Programs/Horde/HordeServer/Services/PerforceService.cs
new file mode 100644
--- /dev/null
+++ b/Engine/Source/Programs/Horde/HordeServer/Services/PerforceService.cs
@@ -0,0 +1,1124 @@
+using EpicGames.Core;
+using EpicGames.Perforce;
+using HordeServer.Models;
+using Microsoft.Extensions.Logging;
+using Polly;
+using Polly.Extensions.Http;
+using System;
+using System.Collections.Generic;
+using System.Diagnostics.CodeAnalysis;
+using System.Globalization;
+using System.IO;
+using System.Linq;
+using System.Net.Http;
+using System.Text;
+using System.Text.Json;
+using System.Text.RegularExpressions;
+using System.Threading;
+using System.Threading.Tasks;
+
+namespace HordeServer.Services
+{
+	/// <summary>
+	/// Information about a Perforce user
+	/// </summary>
+	public class PerforceUserInfo
+	{
+		/// <summary>
+		/// Full name of the user
+		/// </summary>
+		public string? Name { get; set; }
+
+		/// <summary>
+		/// User's email address
+		/// </summary>
+		public string? Email { get; set; }
+	}
+
+	/// <summary>
+	/// Wrapper around Perforce functionality. Can use a local p4.exe client for development purposes, or a separate HordePerforceBridge instance over REST for deployments.
+	/// </summary>
+	public interface IPerforceService
+	{
+		/// <summary>
+		/// Create a new changelist by submitting the given file
+		/// </summary>
+		/// <param name="StreamName">The stream to query</param>
+		/// <param name="Path">Path for the file to submit</param>
+		/// <returns>New changelist number</returns>
+		public Task<int> CreateNewChangeAsync(string StreamName, string Path);
+
+		/// <summary>
+		/// Gets the code change corresponding to an actual change submitted to a stream
+		/// </summary>
+		/// <param name="StreamName"></param>
+		/// <param name="Change">The changelist number to query</param>
+		/// <returns>Code change for the latest change</returns>
+		public Task<int> GetCodeChangeAsync(string StreamName, int Change);
+
+		/// <summary>
+		/// Gets information about the given user
+		/// </summary>
+		/// <param name="UserName">The user name</param>
+		/// <returns>User information</returns>
+		public Task<PerforceUserInfo?> GetUserInfoAsync(string UserName);
+
+		/// <summary>
+		/// Gets the latest change for a particular stream
+		/// </summary>
+		/// <param name="StreamName">The stream to query</param>
+		/// <param name="MinChange">The minimum changelist number</param>
+		/// <param name="MaxChange">The maximum changelist number</param>
+		/// <param name="Results">Number of results to return</param>
+		/// <param name="ImpersonateUser">Name of the user to impersonate</param>
+		/// <returns>Latest changelist number</returns>
+		public Task<List<ChangeSummary>> GetChangesAsync(string StreamName, int? MinChange, int? MaxChange, int Results, string? ImpersonateUser);
+
+		/// <summary>
+		/// Gets the latest change for a particular stream
+		/// </summary>
+		/// <param name="StreamName">The stream to query</param>
+		/// <param name="ChangeNumbers">Change numbers to query</param>
+		/// <param name="ImpersonateUser">Name of the user to impersonate</param>
+		/// <returns>Commit details</returns>
+		public Task<List<ChangeDetails>> GetChangeDetailsAsync(string StreamName, IReadOnlyList<int> ChangeNumbers, string? ImpersonateUser);
+		
+		/// <summary>
+		/// Create a ticket as the specified user
+		/// </summary>
+		/// <param name="ImpersonateUser">Name of the user to impersonate</param>
+		/// <returns>Commit details</returns>
+		public Task<string> CreateTicket(string ImpersonateUser);
+
+		/// <summary>
+		/// Gets the latest changes for a set of depot paths
+		/// </summary>
+		/// <param name="Paths"></param>
+		/// <returns></returns>
+		public Task<List<FileSummary>> FindFilesAsync(IEnumerable<string> Paths);
+
+		/// <summary>
+		/// Gets the contents of a file in the depot
+		/// </summary>
+		/// <param name="Path">Path to read</param>
+		/// <returns>Data for the file</returns>
+		public Task<byte[]> PrintAsync(string Path);
+
+		/// <summary>
+		/// Duplicates a shelved changelist
+		/// </summary>
+		/// <param name="ShelvedChange">The shelved changelist</param>
+		/// <returns>The duplicated changelist</returns>
+		public Task<int> DuplicateShelvedChangeAsync(int ShelvedChange);
+
+		/// <summary>
+		/// Submit a shelved changelist
+		/// </summary>
+		/// <param name="ShelvedChange">The shelved changelist number, created by <see cref="DuplicateShelvedChangeAsync(int)"/></param>
+		/// <param name="OriginalChange">The original changelist number</param>
+		/// <returns>Tuple consisting of the submitted changelist number and message</returns>
+		public Task<(int? Change, string Message)> SubmitShelvedChangeAsync(int ShelvedChange, int OriginalChange);
+
+		/// <summary>
+		/// Deletes a changelist containing shelved files.
+		/// </summary>
+		/// <param name="ShelvedChange">The changelist containing shelved files</param>
+		/// <returns>Async tasy</returns>
+		public Task DeleteShelvedChangeAsync(int ShelvedChange);
+
+		/// <summary>
+		/// Updates a changelist description
+		/// </summary>
+		/// <param name="Change">The change to update</param>
+		/// <param name="Description">The new description</param>
+		/// <returns>Async task</returns>
+		public Task UpdateChangelistDescription(int Change, string Description);
+	}
+
+	/// <summary>
+	/// Extension methods for IPerforceService implementations
+	/// </summary>
+	public static class PerforceServiceExtensions
+	{
+		/// <summary>
+		/// Creates a new change for a template
+		/// </summary>
+		/// <param name="Perforce">The Perforce service instance</param>
+		/// <param name="Stream">Stream containing the template</param>
+		/// <param name="Template">The template being built</param>
+		/// <returns>New changelist number</returns>
+		public static Task<int> CreateNewChangeForTemplateAsync(this IPerforceService Perforce, IStream Stream, ITemplate Template)
+		{
+			Match Match = Regex.Match(Template.SubmitNewChange!, @"^(//[^/]+/[^/]+)/(.+)$");
+			if (Match.Success)
+			{
+				return Perforce.CreateNewChangeAsync(Match.Groups[1].Value, Match.Groups[2].Value);
+			}
+			else
+			{
+				return Perforce.CreateNewChangeAsync(Stream.Name, Template.SubmitNewChange!);
+			}
+		}
+
+		/// <summary>
+		/// Gets the latest change for a particular stream
+		/// </summary>
+		/// <param name="Perforce">The perforce implementation</param>
+		/// <param name="StreamName">The stream to query</param>
+		/// <param name="ChangeNumber">Change number to query</param>
+		/// <param name="ImpersonateUser">Name of the user to impersonate</param>
+		/// <returns>Commit details</returns>
+		public static async Task<ChangeDetails> GetChangeDetailsAsync(this IPerforceService Perforce, string StreamName, int ChangeNumber, string? ImpersonateUser)
+		{
+			List<ChangeDetails> Results = await Perforce.GetChangeDetailsAsync(StreamName, new[] { ChangeNumber }, ImpersonateUser);
+			return Results[0];
+		}
+
+		/// <summary>
+		/// Gets the latest change for a particular stream
+		/// </summary>
+		/// <param name="Perforce">The perforce implementation</param>
+		/// <param name="StreamName">The stream to query</param>
+		/// <param name="MinChange">The minimum changelist number</param>
+		/// <param name="MaxChange">The maximum changelist number</param>
+		/// <param name="Results">Number of results to return</param>
+		/// <param name="ImpersonateUser">Name of the user to impersonate</param>
+		/// <returns>Commit details</returns>
+		public static async Task<List<ChangeDetails>> GetChangeDetailsAsync(this IPerforceService Perforce, string StreamName, int? MinChange, int? MaxChange, int Results, string? ImpersonateUser)
+		{
+			List<ChangeSummary> Changes = await Perforce.GetChangesAsync(StreamName, MinChange, MaxChange, Results, ImpersonateUser);
+			return await Perforce.GetChangeDetailsAsync(StreamName, Changes.ConvertAll(x => x.Number), ImpersonateUser);
+		}
+
+		/// <summary>
+		/// Get the latest submitted change to the stream
+		/// </summary>
+		/// <param name="Perforce">The perforce implementation</param>
+		/// <param name="StreamName">The stream to query</param>
+		/// <param name="ImpersonateUser">Name of the user to impersonate</param>
+		/// <returns>Latest changelist number</returns>
+		public static async Task<int> GetLatestChangeAsync(this IPerforceService Perforce, string StreamName, string? ImpersonateUser)
+		{
+			List<ChangeSummary> Changes = await Perforce.GetChangesAsync(StreamName, null, null, 1, ImpersonateUser);
+			if (Changes.Count == 0)
+			{
+				throw new Exception($"No changes have been submitted to stream {StreamName}");
+			}
+			return Changes[0].Number;
+		}
+	}
+
+	/// <summary>
+	/// Base class for IPerforceService implementations
+	/// </summary>
+	public abstract class PerforceServiceBase : IPerforceService
+	{
+		/// <inheritdoc/>
+		public abstract Task<PerforceUserInfo?> GetUserInfoAsync(string UserName);
+
+		/// <inheritdoc/>
+		public abstract Task<int> CreateNewChangeAsync(string StreamName, string Path);
+
+		/// <inheritdoc/>
+		public abstract Task<List<ChangeSummary>> GetChangesAsync(string StreamName, int? MinChange, int? MaxChange, int Results, string? ImpersonateUser);
+		
+		/// <inheritdoc/>
+		public abstract Task<List<ChangeDetails>> GetChangeDetailsAsync(string StreamName, IReadOnlyList<int> ChangeNumbers, string? ImpersonateUser);
+
+		/// <inheritdoc/>
+		public abstract Task<string> CreateTicket(string ImpersonateUser);
+
+		/// <inheritdoc/>
+		public abstract Task<List<FileSummary>> FindFilesAsync(IEnumerable<string> Paths);
+
+		/// <inheritdoc/>
+		public abstract Task<byte[]> PrintAsync(string Path);
+
+		/// <inheritdoc/>
+		public abstract Task<int> DuplicateShelvedChangeAsync(int ShelvedChange);
+
+		/// <inheritdoc/>
+		public abstract Task<(int? Change, string Message)> SubmitShelvedChangeAsync(int Change, int OriginalChange);
+
+		/// <inheritdoc/>
+		public abstract Task DeleteShelvedChangeAsync(int ShelvedChange);
+
+		/// <inheritdoc/>
+		public abstract Task UpdateChangelistDescription(int Change, string Description);
+
+
+		/// <inheritdoc/>
+		public virtual async Task<int> GetCodeChangeAsync(string StreamName, int Change)
+		{
+			int MaxChange = Change;
+			for(; ;)
+			{
+				// Query for the changes before this point
+				List<ChangeSummary> Changes = await GetChangesAsync(StreamName, null, MaxChange, 10, null);
+				Serilog.Log.Logger.Information("Finding last code change in {Stream} before {MaxChange}: {NumResults}", StreamName, MaxChange, Changes.Count);
+				if (Changes.Count == 0)
+				{
+					return 0;
+				}
+
+				// Get the details for them
+				List<ChangeDetails> DetailsList = await GetChangeDetailsAsync(StreamName, Changes.ConvertAll(x => x.Number), null);
+				foreach (ChangeDetails Details in DetailsList.OrderByDescending(x => x.Number))
+				{
+					ChangeContentFlags ContentFlags = Details.GetContentFlags();
+					Serilog.Log.Logger.Information("Change {Change} = {Flags}", Details.Number, ContentFlags.ToString());
+					if ((Details.GetContentFlags() & ChangeContentFlags.ContainsCode) != 0)
+					{
+						return Details.Number;
+					}
+				}
+
+				// Loop round again, adjusting our maximum changelist number
+				MaxChange = Changes.Min(x => x.Number) - 1;
+			}
+		}
+
+		/// <summary>
+		/// Gets the wildcard filter for a particular range of changes
+		/// </summary>
+		/// <param name="BasePath">Base path to find files for</param>
+		/// <param name="MinChange">Minimum changelist number to query</param>
+		/// <param name="MaxChange">Maximum changelist number to query</param>
+		/// <returns>Filter string</returns>
+		public static string GetFilter(string BasePath, int? MinChange, int? MaxChange)
+		{
+			StringBuilder Filter = new StringBuilder(BasePath);
+			if (MinChange != null && MaxChange != null)
+			{
+				Filter.Append($"@{MinChange},{MaxChange}");
+			}
+			else if (MinChange != null)
+			{
+				Filter.Append($"@>={MinChange}");
+			}
+			else if (MaxChange != null)
+			{
+				Filter.Append($"@<={MaxChange}");
+			}
+			return Filter.ToString();
+		}
+
+		/// <summary>
+		/// Gets a stream-relative path from a depot path
+		/// </summary>
+		/// <param name="DepotFile">The depot file to check</param>
+		/// <param name="StreamName">Name of the stream</param>
+		/// <param name="RelativePath">Receives the relative path to the file</param>
+		/// <returns>True if the stream-relative path was returned</returns>
+		public static bool TryGetStreamRelativePath(string DepotFile, string StreamName, [NotNullWhen(true)] out string? RelativePath)
+		{
+			if (DepotFile.StartsWith(StreamName, StringComparison.OrdinalIgnoreCase) && DepotFile.Length > StreamName.Length && DepotFile[StreamName.Length] == '/')
+			{
+				RelativePath = DepotFile.Substring(StreamName.Length);
+				return true;
+			}
+
+			Match Match = Regex.Match(DepotFile, "^//[^/]+/[^/]+(/.*)$");
+			if (Match.Success)
+			{
+				RelativePath = Match.Groups[1].Value;
+				return true;
+			}
+
+			RelativePath = null;
+			return false;
+		}
+	}
+
+	/// <summary>
+	/// Empty implementation of IPerforceService
+	/// </summary>
+	public class EmptyPerforceService : IPerforceService
+	{
+		/// <inheritdoc/>
+		public Task<PerforceUserInfo?> GetUserInfoAsync(string UserName)
+		{
+			throw new NotImplementedException();
+		}
+
+		/// <inheritdoc/>
+		public Task<int> CreateNewChangeAsync(string StreamName, string Path)
+		{
+			throw new NotImplementedException();
+		}
+
+		/// <inheritdoc/>
+		public Task<List<ChangeSummary>> GetChangesAsync(string StreamName, int? MinChange, int? MaxChange, int Results, string? ImpersonateUser)
+		{
+			return Task.FromResult<List<ChangeSummary>>(new List<ChangeSummary>());
+		}
+
+		/// <inheritdoc/>
+		public Task<List<ChangeDetails>> GetChangeDetailsAsync(string StreamName, IReadOnlyList<int> ChangeNumbers, string? ImpersonateUser)
+		{
+			throw new NotImplementedException();
+		}
+
+		/// <inheritdoc />
+		public Task<string> CreateTicket(string ImpersonateUser)
+		{
+			throw new NotImplementedException();
+		}
+
+		/// <inheritdoc/>
+		public Task<int> GetCodeChangeAsync(string StreamName, int Change)
+		{
+			throw new NotImplementedException();
+		}
+
+		/// <inheritdoc/>
+		public Task<List<FileSummary>> FindFilesAsync(IEnumerable<string> Paths)
+		{
+			throw new NotImplementedException();
+		}
+
+		/// <inheritdoc/>
+		public Task<byte[]> PrintAsync(string Path)
+		{
+			throw new NotImplementedException();
+		}
+
+		/// <inheritdoc/>
+		public Task<int> DuplicateShelvedChangeAsync(int ShelvedChange)
+		{
+			throw new NotImplementedException();
+		}
+
+		/// <inheritdoc/>
+		public Task<(int? Change, string Message)> SubmitShelvedChangeAsync(int ShelvedChange, int OriginalChange)
+		{
+			throw new NotImplementedException();
+		}
+
+		/// <inheritdoc/>
+		public Task DeleteShelvedChangeAsync(int ShelvedChange)
+		{
+			throw new NotImplementedException();
+		}
+
+		/// <inheritdoc/>
+		public Task UpdateChangelistDescription(int Change, string Description)
+		{
+			throw new NotImplementedException();
+		}
+	}
+
+	/// <summary>
+	/// Local implementation of the Perforce service
+	/// </summary>
+	public class LocalPerforceService : PerforceServiceBase
+	{
+		/// <summary>
+		/// The Perforce connection
+		/// </summary>
+		PerforceConnection Perforce;
+
+		/// <summary>
+		/// Cached client connections
+		/// </summary>
+		Dictionary<string, Task<PerforceConnection>> Clients = new Dictionary<string, Task<PerforceConnection>>(StringComparer.OrdinalIgnoreCase);
+
+		/// <summary>
+		/// Information about the Perforce server
+		/// </summary>
+		Task<InfoRecord> InfoTask;
+
+		/// <summary>
+		/// Constructor
+		/// </summary>
+		/// <param name="Logger">Logger</param>
+		public LocalPerforceService(ILogger<LocalPerforceService> Logger)
+		{
+			Perforce = new PerforceConnection(null, null, Logger);
+			InfoTask = Perforce.GetInfoAsync(InfoOptions.None, CancellationToken.None);
+		}
+
+		/// <inheritdoc/>
+		public async override Task<PerforceUserInfo?> GetUserInfoAsync(string UserName)
+		{
+			PerforceResponseList<UserRecord> UserResponse = await Perforce.TryGetUserAsync(UserName, CancellationToken.None);
+			if (UserResponse.Succeeded)
+			{
+				UserRecord? UserRecord = UserResponse.Data.FirstOrDefault();
+				if (UserRecord != null)
+				{
+					PerforceUserInfo UserInfo = new PerforceUserInfo();
+					UserInfo.Name = UserRecord.FullName;
+					UserInfo.Email = UserRecord.Email;
+					return UserInfo;
+				}
+			}
+			return null;
+		}
+
+		/// <inheritdoc/>
+		public async override Task<int> CreateNewChangeAsync(string StreamName, string FilePath)
+		{
+			string ClientName = $"Horde-NewChange-{Guid.NewGuid().ToString("N")}";
+			DirectoryReference ClientDir = DirectoryReference.Combine(new DirectoryReference(Path.GetTempPath()), ClientName);
+
+			InfoRecord Info = await InfoTask;
+
+			ClientRecord Client = new ClientRecord(ClientName, Info.UserName!, ClientDir.FullName);
+			Client.Stream = StreamName;
+			await Perforce.CreateClientAsync(Client, CancellationToken.None);
+
+			PerforceConnection PerforceClient = new PerforceConnection(Perforce) { ClientName = ClientName };
+			try
+			{
+				ChangeRecord ChangeRecord = new ChangeRecord();
+				ChangeRecord.Client = ClientName;
+				ChangeRecord.Description = "New change for Horde job";
+				ChangeRecord = await PerforceClient.CreateChangeAsync(ChangeRecord, CancellationToken.None);
+
+				string ClientFilePath = $"//{ClientName}/{FilePath.TrimStart('/')}";
+				for (; ; )
+				{
+					await PerforceClient.TrySyncAsync(SyncOptions.Force, -1, new[] { ClientFilePath }, CancellationToken.None);
+					await PerforceClient.TryEditAsync(ChangeRecord.Number, new[] { ClientFilePath }, CancellationToken.None);
+
+					PerforceResponse<SubmitRecord> Response = await PerforceClient.TrySubmitAsync(ChangeRecord.Number, SubmitOptions.None, CancellationToken.None);
+					if (Response.Info == null || !Response.Info.Data.Contains(" - must resolve", StringComparison.Ordinal))
+					{
+						return Response.Data.ChangeNumber;
+					}
+
+					await PerforceClient.RevertAsync(-1, null, RevertOptions.None, new[] { "//..." }, CancellationToken.None);
+				}
+			}
+			finally
+			{
+				await PerforceClient.TryRevertAsync(-1, null, RevertOptions.None, new[] { "//..." }, CancellationToken.None);
+				await Perforce.TryDeleteClientAsync(DeleteClientOptions.None, ClientName, CancellationToken.None);
+			}
+		}
+
+		/// <inheritdoc/>
+		public override async Task<List<ChangeSummary>> GetChangesAsync(string StreamName, int? MinChange, int? MaxChange, int Results, string? ImpersonateUser)
+		{
+			PerforceConnection Client = await CreateReadOnlyClientAsync(StreamName, ImpersonateUser);
+
+			List<ChangeSummary> Changes = new List<ChangeSummary>();
+
+			string Filter = GetFilter($"//{Client.ClientName}/...", MinChange, MaxChange);
+
+			List<ChangesRecord> ChangeRecords = await Client.GetChangesAsync(ChangesOptions.LongOutput, Results, ChangeStatus.Submitted, new[] { Filter }, CancellationToken.None);
+			foreach (ChangesRecord ChangeRecord in ChangeRecords)
+			{
+				Changes.Add(new ChangeSummary(ChangeRecord.Number, ChangeRecord.User ?? String.Empty, ChangeRecord.Description ?? String.Empty));
+			}
+
+			return Changes;
+		}
+
+		/// <inheritdoc/>
+		public override async Task<List<ChangeDetails>> GetChangeDetailsAsync(string StreamName, IReadOnlyList<int> ChangeNumbers, string? ImpersonateUser)
+		{
+			// Create a client for this stream, in case it's a task stream
+			PerforceConnection Client = await CreateReadOnlyClientAsync(StreamName, ImpersonateUser);
+
+			// Get the Perforce description
+			List<DescribeRecord> DescribeRecords = await Client.DescribeAsync(ChangeNumbers.ToArray(), CancellationToken.None);
+
+			// Convert it to ChangeDetails objects
+			List<ChangeDetails> Results = new List<ChangeDetails>();
+			foreach (DescribeRecord DescribeRecord in DescribeRecords)
+			{
+				List<string> Files = new List<string>();
+				foreach (DescribeFileRecord DescribeFile in DescribeRecord.Files)
+				{
+					string? RelativePath;
+					if (TryGetStreamRelativePath(DescribeFile.DepotFile, StreamName, out RelativePath))
+					{
+						Files.Add(RelativePath);
+					}
+				}
+				Results.Add(new ChangeDetails(DescribeRecord.Number, DescribeRecord.User, DescribeRecord.Description, Files));
+			}
+			return Results;
+		}
+
+		/// <inheritdoc />
+		public override Task<string> CreateTicket(string ImpersonateUser)
+		{
+			throw new NotImplementedException();
+		}
+
+		/// <inheritdoc/>
+		public override async Task<List<FileSummary>> FindFilesAsync(IEnumerable<string> Paths)
+		{
+			List<FileSummary> Results = new List<FileSummary>();
+
+			PerforceResponseList<FStatRecord> Responses = await Perforce.TryFStatAsync(FStatOptions.ShortenOutput, Paths.ToArray(), CancellationToken.None);
+			foreach (PerforceResponse<FStatRecord> Response in Responses)
+			{
+				if (Response.Succeeded)
+				{
+					FStatRecord Record = Response.Data;
+					if (Record.DepotFile != null)
+					{
+						Results.Add(new FileSummary(Record.DepotFile, Record.HeadAction != FileAction.Delete && Record.HeadAction != FileAction.MoveDelete, Record.HeadChange));
+					}
+				}
+			}
+
+			return Results;
+		}
+
+		/// <inheritdoc/>
+		public override async Task<byte[]> PrintAsync(string DepotPath)
+		{
+			string TempFile = Path.GetTempFileName();
+			try
+			{
+				await Perforce.PrintAsync(TempFile, DepotPath, CancellationToken.None);
+				return await File.ReadAllBytesAsync(TempFile);
+			}
+			finally
+			{
+				if (File.Exists(TempFile))
+				{
+					File.SetAttributes(TempFile, FileAttributes.Normal);
+					File.Delete(TempFile);
+				}
+			}
+		}
+
+		/// <inheritdoc/>
+		public override async Task<(int? Change, string Message)> SubmitShelvedChangeAsync(int Change, int OriginalChange)
+		{
+			PerforceResponse<SubmitRecord> Response = await Perforce.TrySubmitShelvedAsync(Change, CancellationToken.None);
+			if (Response.Succeeded)
+			{
+				return (Response.Data.ChangeNumber, $"Submitted in CL {Response.Data.ChangeNumber}");
+			}
+			else
+			{
+				return (null, Response.Error!.Data);
+			}
+		}
+
+		/// <summary>
+		/// Creates a read-only client for the given stream
+		/// </summary>
+		/// <param name="StreamName">The stream to query</param>
+		/// <param name="ImpersonateUser">Name of the user to impersonate</param>
+		/// <returns>Client connection</returns>
+		async Task<PerforceConnection> CreateReadOnlyClientAsync(string StreamName, string? ImpersonateUser)
+		{
+			InfoRecord Info = await InfoTask;
+			string UserName = ImpersonateUser ?? Info.UserName!;
+			string HostName = Info.ClientHost!;
+			string ClientName = $"Horde-{UserName}-{HostName}-{StreamName.Replace('/', '+').Trim('+')}";
+
+			Task<PerforceConnection>? Result;
+			lock (Clients)
+			{
+				if (!Clients.TryGetValue(ClientName, out Result))
+				{
+					Result = CreateNewReadOnlyClientAsync(StreamName, UserName, ClientName);
+					Clients.Add(ClientName, Result);
+				}
+			}
+			return await Result;
+		}
+
+		/// <inheritdoc/>
+		public override Task<int> DuplicateShelvedChangeAsync(int ShelvedChange)
+		{
+			throw new NotImplementedException();
+		}
+
+		/// <summary>
+		/// Creates a new read-only client for the given stream
+		/// </summary>
+		/// <param name="StreamName">The stream to query</param>
+		/// <param name="UserName">Name of the user to own the client</param>
+		/// <param name="ClientName">Name of the client</param>
+		/// <returns>Client connection</returns>
+		async Task<PerforceConnection> CreateNewReadOnlyClientAsync(string StreamName, string UserName, string ClientName)
+		{
+			await Task.Yield();
+
+			PerforceConnection PerforceClient = new PerforceConnection(Perforce) { UserName = UserName };
+
+			ClientRecord Client = new ClientRecord(ClientName, UserName, "C:\\");
+			Client.Stream = StreamName;
+			await PerforceClient.CreateClientAsync(Client, CancellationToken.None);
+
+			PerforceClient.ClientName = ClientName;
+			return PerforceClient;
+		}
+
+		/// <inheritdoc/>
+		public override Task DeleteShelvedChangeAsync(int ShelvedChange)
+		{
+			throw new NotImplementedException();
+		}
+
+		/// <inheritdoc/>
+		public override Task UpdateChangelistDescription(int Change, string Description)
+		{
+			throw new NotImplementedException();
+		}
+	}
+
+	/// <summary>
+	/// Implementation of IPerforceService which connects to a HordePerforceBridge instance
+	/// </summary>
+	public sealed class BridgePerforceService : PerforceServiceBase, IDisposable
+	{
+#pragma warning disable CA1812
+		class BridgeChange
+		{
+			public string Change { get; set; } = String.Empty;
+			public string Description { get; set; } = String.Empty;
+			public string User { get; set; } = String.Empty;
+		}
+
+		class BridgeChanges
+		{
+			public List<BridgeChange>? Summaries { get; set; }
+
+			public List<string>? Errors { get; set; } = null;
+		}
+
+		class BridgeChangeFile
+		{
+			public string DepotFile { get; set; } = String.Empty;
+		}
+		
+		class BridgeUserData
+		{
+			public string? Name { get; set; } = String.Empty;
+			public string? Email { get; set; } = String.Empty;
+		}
+
+		class BridgeChangeDetail
+		{
+			public string Change { get; set; } = String.Empty;
+			public string Description { get; set; } = String.Empty;
+			public string User { get; set; } = String.Empty;
+			public List<BridgeChangeFile> Files { get; set; } = new List<BridgeChangeFile>();
+		}
+
+		class BridgeChangeDetails
+		{
+			public List<BridgeChangeDetail> Changes { get; set; } = new List<BridgeChangeDetail>();
+
+			public List<string>? Errors { get; set; } = null;
+		}
+
+		class BridgeTicket
+		{
+			public string Ticket { get; set; } = String.Empty;
+			public List<string>? Errors { get; set; } = null;
+		}
+
+		class BridgeUser
+		{
+			public BridgeUserData? User { get; set; } = null;
+			public List<string>? Errors { get; set; } = null;
+		}
+
+		class BridgeChangelistNumber
+		{
+			public string Change { get; set; } = String.Empty;
+			public string Message { get; set; } = string.Empty;
+			public List<string>? Errors { get; set; } = null;
+		}
+
+		class BridgeFStat
+		{
+			public string DepotFile { get; set; } = String.Empty;
+			public bool Exists { get; set; } = false;
+			public string HeadChange { get; set; } = String.Empty;
+			public string Errors { get; set; } = String.Empty;
+		}
+
+		class BridgePrint
+		{
+			public string Data { get; set; } = String.Empty;
+			public List<string>? Errors { get; set; } = null;
+		}
+#pragma warning restore CA1812
+
+		HttpClient Client;
+		AsyncPolicy<HttpResponseMessage> RetryPolicy;
+		Uri Server;
+		ILogger Logger;
+
+		/// <summary>
+		/// Constructor
+		/// </summary>
+		public BridgePerforceService(string Address, ILogger Logger)
+		{
+			this.Logger = Logger;
+
+			Client = new HttpClient();
+			Client.Timeout = TimeSpan.FromSeconds(300.0);
+			Server = new Uri(Address);
+			RetryPolicy = HttpPolicyExtensions.HandleTransientHttpError().WaitAndRetryAsync(3, Attempt => TimeSpan.FromSeconds(Math.Pow(2.0, Attempt)));
+		}
+
+		/// <inheritdoc/>
+		public void Dispose()
+		{
+			Client.Dispose();
+		}
+
+		/// <inheritdoc/>
+		public override async Task<PerforceUserInfo?> GetUserInfoAsync(string UserName)
+		{
+			Uri Uri = new Uri(Server, $"/api/v1/user?username={UserName}");
+
+			HttpResponseMessage Response = await RetryPolicy.ExecuteAsync(() => Client.GetAsync(Uri));
+			if (!Response.IsSuccessStatusCode)
+			{
+				Logger.LogError("GET to {Uri} returned {Code} ({Response})", Uri, Response.StatusCode, await Response.Content.ReadAsStringAsync());
+				throw new Exception("CreateTicket call to PerforceBridge failed");
+			}
+
+			byte[] Data = await Response.Content.ReadAsByteArrayAsync();
+			BridgeUser Result = JsonSerializer.Deserialize<BridgeUser>(Data.AsSpan());
+
+			if (Result.Errors != null)
+			{
+				Logger.LogError("{Uri} reported errors: {Errors}", Uri, string.Join("\n", Result.Errors));
+				throw new Exception($"{Uri} reported errors: {string.Join("\n", Result.Errors)}");
+			}
+
+			if(Result.User == null)
+			{
+				Logger.LogWarning("User info query returned not found for username {UserName}", UserName);
+				return null;
+			}
+
+			return new PerforceUserInfo { Name = Result.User.Name, Email = Result.User.Email };
+		}
+
+		/// <inheritdoc/>
+		public override async Task<int> CreateNewChangeAsync(string StreamName, string Path)
+		{
+			List<string> Arguments = new List<string>();
+			Arguments.Add($"stream={StreamName}");
+			Arguments.Add($"path={Path}");
+
+			Uri Uri = new Uri(Server, $"/api/v1/createnewchange?{String.Join("&", Arguments)}");
+
+			HttpResponseMessage Response = await RetryPolicy.ExecuteAsync(() => Client.GetAsync(Uri));
+			if (!Response.IsSuccessStatusCode)
+			{
+				Logger.LogError("GET to {Uri} returned {Code} ({Response})", Uri, Response.StatusCode, await Response.Content.ReadAsStringAsync());
+				throw new Exception("Query to PerforceBridge failed");
+			}
+
+			byte[] Data = await Response.Content.ReadAsByteArrayAsync();
+
+			BridgeChangelistNumber Result = JsonSerializer.Deserialize<BridgeChangelistNumber>(Data);
+			if (Result.Errors != null)
+			{
+				Logger.LogError("{Uri} reported errors: {Errors}", Uri, string.Join("\n", Result.Errors));
+				throw new Exception($"{Uri} reported errors: {string.Join("\n", Result.Errors)}");
+			}
+
+			int ChangeNumber;
+			if(!int.TryParse(Result.Change, out ChangeNumber))
+			{
+				throw new Exception($"Missing changelist number from response. Data: {Encoding.UTF8.GetString(Data)}");
+			}
+
+			return ChangeNumber;
+		}
+
+		/// <inheritdoc/>
+		public override async Task<List<ChangeSummary>> GetChangesAsync(string StreamName, int? MinChange, int? MaxChange, int Results, string? ImpersonateUser)
+		{
+			List<string> Arguments = new List<string>();
+			Arguments.Add($"stream={StreamName}");
+			if (MinChange != null)
+			{
+				Arguments.Add($"minchange={MinChange.Value}");
+			}
+			if (MaxChange != null)
+			{
+				Arguments.Add($"maxchange={MaxChange.Value}");
+			}
+			if (ImpersonateUser != null)
+			{
+				Arguments.Add($"p4iuser={ImpersonateUser}");
+			}
+			Arguments.Add($"maxresults={Results}");
+
+			Uri Uri = new Uri(Server, $"/api/v1/changes?{String.Join("&", Arguments)}");
+
+			HttpResponseMessage Response = await RetryPolicy.ExecuteAsync(() => Client.GetAsync(Uri));
+			if (!Response.IsSuccessStatusCode)
+			{
+				Logger.LogError("GET to {Uri} returned {Code} ({Response})", Uri, Response.StatusCode, await Response.Content.ReadAsStringAsync());
+				throw new Exception("Query to PerforceBridge failed");
+			}
+
+			byte[] Data = await Response.Content.ReadAsByteArrayAsync();
+
+			BridgeChanges Result = JsonSerializer.Deserialize<BridgeChanges>(Data.AsSpan());
+			if(Result.Errors != null)
+			{
+				Logger.LogError("{Uri} reported errors: {Errors}", Uri, string.Join("\n", Result.Errors));
+				throw new Exception($"{Uri} reported errors: {string.Join("\n", Result.Errors)}");
+			}
+
+			try
+			{
+				if (Result.Summaries == null)
+				{
+					return new List<ChangeSummary>();
+				}
+				else
+				{
+					return Result.Summaries.ConvertAll(x => new ChangeSummary(int.Parse(x.Change, CultureInfo.InvariantCulture), x.User, x.Description));
+				}
+			}
+			catch(Exception Ex)
+			{
+				throw new Exception($"Unable to parse changelist summaries from response: {Encoding.UTF8.GetString(Data)}", Ex); 
+			}
+		}
+
+		/// <inheritdoc/>
+		public override async Task<List<ChangeDetails>> GetChangeDetailsAsync(string StreamName, IReadOnlyList<int> ChangeNumbers, string? ImpersonateUser)
+		{
+			List<ChangeDetails> Results = new List<ChangeDetails>();
+
+			List<string> Arguments = new List<string>();
+			Arguments.Add($"change={string.Join(';', ChangeNumbers)}");
+			if(ImpersonateUser != null)
+			{
+				Arguments.Add($"iuser={ImpersonateUser}");
+			}
+
+			Uri Uri = new Uri(Server, $"/api/v1/describe?{String.Join("&", Arguments)}");
+
+			HttpResponseMessage Response = await RetryPolicy.ExecuteAsync(() => Client.GetAsync(Uri));
+			if (!Response.IsSuccessStatusCode)
+			{
+				Logger.LogError("GET to {Uri} returned {Code} ({Response})", Uri, Response.StatusCode, await Response.Content.ReadAsStringAsync());
+				throw new Exception("Query to PerforceBridge failed");
+			}
+
+			byte[] Data = await Response.Content.ReadAsByteArrayAsync();
+
+			BridgeChangeDetails Result = JsonSerializer.Deserialize<BridgeChangeDetails>(Data.AsSpan());
+
+			if (Result.Errors != null)
+			{
+				Logger.LogError("{Uri} reported errors: {Errors}", Uri, string.Join("\n", Result.Errors));
+				throw new Exception($"{Uri} reported errors: {string.Join("\n", Result.Errors)}");
+			}
+
+			foreach (BridgeChangeDetail Details in Result.Changes)
+			{ 
+				List<string> Files = new List<string>();
+				foreach (BridgeChangeFile File in Details.Files)
+				{
+					string? RelativePath;
+					if (TryGetStreamRelativePath(File.DepotFile, StreamName, out RelativePath))
+					{
+						Files.Add(RelativePath);
+					}
+				}
+				Results.Add(new ChangeDetails(int.Parse(Details.Change, CultureInfo.InvariantCulture), Details.User, Details.Description, Files));				
+			}
+			return Results;
+		}
+
+		/// <inheritdoc />
+		public override async Task<string> CreateTicket(string ImpersonateUser)
+		{
+			Uri Uri = new Uri(Server, $"/api/v1/ticket?p4iuser={ImpersonateUser}");
+
+			HttpResponseMessage Response = await RetryPolicy.ExecuteAsync(() => Client.GetAsync(Uri));
+			if (!Response.IsSuccessStatusCode)
+			{
+				Logger.LogError("GET to {Uri} returned {Code} ({Response})", Uri, Response.StatusCode, await Response.Content.ReadAsStringAsync());
+				throw new Exception("CreateTicket call to PerforceBridge failed");
+			}
+
+			byte[] Data = await Response.Content.ReadAsByteArrayAsync();
+			BridgeTicket Result = JsonSerializer.Deserialize<BridgeTicket>(Data.AsSpan());
+
+			if (Result.Errors != null)
+			{
+				Logger.LogError("{Uri} reported errors: {Errors}", Uri, string.Join("\n", Result.Errors));
+				throw new Exception($"{Uri} reported errors: {string.Join("\n", Result.Errors)}");
+			}
+
+			return Result.Ticket;
+		}
+
+		/// <inheritdoc/>
+		public override async Task<List<FileSummary>> FindFilesAsync(IEnumerable<string> Paths)
+		{
+			List<FileSummary> Results = new List<FileSummary>();
+			Uri Uri = new Uri(Server, $"/api/v1/findfiles?paths={string.Join(';', Paths)}");
+			HttpResponseMessage Response = await RetryPolicy.ExecuteAsync(() => Client.GetAsync(Uri));
+			if (!Response.IsSuccessStatusCode)
+			{
+				Logger.LogError("GET to {Uri} returned {Code} ({Response})", Uri, Response.StatusCode, await Response.Content.ReadAsStringAsync());
+				throw new Exception("Query to PerforceBridge failed");
+			}
+
+			byte[] Data = await Response.Content.ReadAsByteArrayAsync();
+			List<BridgeFStat> FStatResults = JsonSerializer.Deserialize<List<BridgeFStat>>(Data.AsSpan());
+			foreach(BridgeFStat Result in FStatResults)
+			{
+				if(Result.DepotFile != null)
+				{
+					Results.Add(new FileSummary(Result.DepotFile, Result.Exists, int.Parse(Result.HeadChange, CultureInfo.InvariantCulture), Result.Errors));
+				}
+			}
+			return Results;
+		}
+
+		/// <inheritdoc/>
+		public override async Task<byte[]> PrintAsync(string Path)
+		{
+			List<FileSummary> Results = new List<FileSummary>();
+			Uri Uri = new Uri(Server, $"/api/v1/print?file={Path}");
+			HttpResponseMessage Response = await RetryPolicy.ExecuteAsync(() => Client.GetAsync(Uri));
+			if (!Response.IsSuccessStatusCode)
+			{
+				Logger.LogError("GET to {Uri} returned {Code} ({Response})", Uri, Response.StatusCode, await Response.Content.ReadAsStringAsync());
+				throw new Exception("Query to PerforceBridge failed");
+			}
+
+			byte[] Data = await Response.Content.ReadAsByteArrayAsync();
+			BridgePrint Result = JsonSerializer.Deserialize<BridgePrint>(Data.AsSpan());
+
+			if (Result.Errors != null)
+			{
+				Logger.LogError("{Uri} reported errors: {Errors}", Uri, string.Join("\n", Result.Errors));
+				throw new Exception($"{Uri} reported errors: {string.Join("\n", Result.Errors)}");
+			}
+
+			if(Result.Data == null)
+			{
+				Logger.LogError("{Uri} resulted in no data being returned. Was {Path} deleted?", Uri, Path);
+				throw new Exception($"{Uri} resulted in no data being returned. Was {Path} deleted?");
+			}
+
+			return Encoding.Default.GetBytes(Result.Data);
+		}
+
+		/// <inheritdoc/>
+		public override async Task<int> DuplicateShelvedChangeAsync(int ShelvedChange)
+		{
+			List<string> Arguments = new List<string>();
+			Arguments.Add($"change={ShelvedChange}");
+
+			Uri Uri = new Uri(Server, $"/api/v2/duplicateshelvedchange?{String.Join("&", Arguments)}");
+
+			HttpResponseMessage Response = await RetryPolicy.ExecuteAsync(() => Client.GetAsync(Uri));
+			if (!Response.IsSuccessStatusCode)
+			{
+				Logger.LogError("GET to {Uri} returned {Code} ({Response})", Uri, Response.StatusCode, await Response.Content.ReadAsStringAsync());
+				throw new Exception("Query to PerforceBridge failed");
+			}
+
+			byte[] Data = await Response.Content.ReadAsByteArrayAsync();
+			BridgeChangelistNumber Result = JsonSerializer.Deserialize<BridgeChangelistNumber>(Data.AsSpan());
+
+			if (Result.Errors != null)
+			{
+				Logger.LogError("{Uri} reported errors: {Errors}", Uri, string.Join("\n", Result.Errors));
+				throw new Exception($"{Uri} reported errors: {string.Join("\n", Result.Errors)}");
+			}
+
+			return int.Parse(Result.Change, CultureInfo.InvariantCulture);
+
+		}
+		/// <inheritdoc/>
+		public override async Task<(int? Change, string Message)> SubmitShelvedChangeAsync(int Change, int OriginalChange)
+		{
+			List<string> Arguments = new List<string>();
+			Arguments.Add($"change={Change}");
+			Arguments.Add($"originalchange={OriginalChange}");
+
+			Uri Uri = new Uri(Server, $"/api/v2/submitshelvedchange?{String.Join("&", Arguments)}");
+
+			HttpResponseMessage Response = await RetryPolicy.ExecuteAsync(() => Client.GetAsync(Uri));
+			if (!Response.IsSuccessStatusCode)
+			{
+				Logger.LogError("GET to {Uri} returned {Code} ({Response})", Uri, Response.StatusCode, await Response.Content.ReadAsStringAsync());
+				return (null, "Query to PerforceBridge failed");
+			}
+
+			byte[] Data = await Response.Content.ReadAsByteArrayAsync();
+			BridgeChangelistNumber Result = JsonSerializer.Deserialize<BridgeChangelistNumber>(Data.AsSpan());
+
+			if (Result.Errors != null)
+			{
+				Logger.LogError("{Uri} reported errors: {Errors}", Uri, string.Join("\n", Result.Errors));
+				return (null, string.Join("\n", Result.Errors));
+				//throw new Exception($"{Uri} reported errors: {string.Join("\n", Result.Errors)}");
+			}
+
+			return (int.Parse(Result.Change, CultureInfo.InvariantCulture), Result.Message);
+		}
+
+		/// <inheritdoc/>
+		public override async Task DeleteShelvedChangeAsync(int ShelvedChange)
+		{
+			List<string> Arguments = new List<string>();
+			Arguments.Add($"change={ShelvedChange}");
+
+			Uri Uri = new Uri(Server, $"/api/v2/deletechange?{String.Join("&", Arguments)}");
+
+			HttpResponseMessage Response = await RetryPolicy.ExecuteAsync(() => Client.GetAsync(Uri));
+			if (!Response.IsSuccessStatusCode)
+			{
+				Logger.LogError("GET to {Uri} returned {Code} ({Response})", Uri, Response.StatusCode, await Response.Content.ReadAsStringAsync());
+				throw new Exception("Query to PerforceBridge failed");
+			}
+
+			byte[] Data = await Response.Content.ReadAsByteArrayAsync();
+			BridgeChangelistNumber Result = JsonSerializer.Deserialize<BridgeChangelistNumber>(Data.AsSpan());
+
+			if (Result.Errors != null)
+			{
+				Logger.LogError("{Uri} reported errors: {Errors}", Uri, string.Join("\n", Result.Errors));
+				throw new Exception($"{Uri} reported errors: {string.Join("\n", Result.Errors)}");
+			}
+		}
+
+		/// <inheritdoc/>
+		public override async Task UpdateChangelistDescription(int Change, string Description)
+		{
+			Uri Uri = new Uri(Server, $"/api/v2/updatechangedescription");
+
+			using (HttpContent Payload = new StringContent($"{{ \"Change\": {Change}, \"Description\": \"{JsonEncodedText.Encode(Description)}\" }}", Encoding.UTF8, "application/json"))
+			{
+				HttpResponseMessage Response = await RetryPolicy.ExecuteAsync(() => Client.PutAsync(Uri, Payload));
+				if (!Response.IsSuccessStatusCode)
+				{
+					Logger.LogError("GET to {Uri} returned {Code} ({Response})", Uri, Response.StatusCode, await Response.Content.ReadAsStringAsync());
+					throw new Exception("Query to PerforceBridge failed");
+				}
+
+				byte[] Data = await Response.Content.ReadAsByteArrayAsync();
+				BridgeChangelistNumber Result = JsonSerializer.Deserialize<BridgeChangelistNumber>(Data.AsSpan());
+
+				if (Result.Errors != null)
+				{
+					Logger.LogError("{Uri} reported errors: {Errors}", Uri, string.Join("\n", Result.Errors));
+					throw new Exception($"{Uri} reported errors: {string.Join("\n", Result.Errors)}");
+				}
+			}
+		}
+	}
+}
diff --git a/Engine/Source/Programs/Horde/HordeServer/Services/PoolService.cs b/Engine/Source/Programs/Horde/HordeServer/Services/PoolService.cs
new file mode 100644
--- /dev/null
+++ b/Engine/Source/Programs/Horde/HordeServer/Services/PoolService.cs
@@ -0,0 +1,181 @@
+// Copyright Epic Games, Inc. All Rights Reserved.
+
+using HordeCommon;
+using HordeServer.Api;
+using HordeServer.Collections;
+using HordeServer.Models;
+using HordeServer.Utilities;
+using Microsoft.Extensions.Logging;
+using Microsoft.Extensions.Options;
+using MongoDB.Bson;
+using MongoDB.Driver;
+using System;
+using System.Collections.Generic;
+using System.Linq;
+using System.Threading;
+using System.Threading.Tasks;
+
+using PoolId = HordeServer.Utilities.StringId<HordeServer.Models.IPool>;
+
+namespace HordeServer.Services
+{
+	/// <summary>
+	/// Wraps functionality for manipulating pools
+	/// </summary>
+	public class PoolService
+	{
+		/// <summary>
+		/// Collection of pool documents
+		/// </summary>
+		IPoolCollection Pools;
+
+		/// <summary>
+		/// Returns the current time
+		/// </summary>
+		IClock Clock;
+
+		/// <summary>
+		/// Cached set of pools, along with the timestamp that it was obtained
+		/// </summary>
+		Tuple<DateTime, Dictionary<PoolId, IPool>>? CachedPoolLookup;
+
+		/// <summary>
+		/// Constructor
+		/// </summary>
+		/// <param name="Pools">Collection of pool documents</param>
+		/// <param name="Clock"></param>
+		public PoolService(IPoolCollection Pools, IClock Clock)
+		{
+			this.Pools = Pools;
+			this.Clock = Clock;
+		}
+
+		/// <summary>
+		/// Creates a new pool
+		/// </summary>
+		/// <param name="Name">Name of the new pool</param>
+		/// <param name="Requirements">Requirements for agents to add to this pool</param>
+		/// <param name="EnableAutoscaling">Whether to enable autoscaling for this pool</param>
+		/// <param name="MinAgents">Minimum number of agents in the pool</param>
+		/// <param name="NumReserveAgents">Minimum number of idle agents to maintain</param>
+		/// <param name="Properties">Properties for the new pool</param>
+		/// <returns>The new pool document</returns>
+		public Task<IPool> CreatePoolAsync(string Name, AgentRequirements? Requirements = null, bool? EnableAutoscaling = null, int? MinAgents = null, int? NumReserveAgents = null, Dictionary<string, string>? Properties = null)
+		{
+			return Pools.AddAsync(PoolId.Sanitize(Name), Name, Requirements, EnableAutoscaling, MinAgents, NumReserveAgents, Properties);
+		}
+
+		/// <summary>
+		/// Deletes a pool
+		/// </summary>
+		/// <param name="PoolId">Unique id of the pool</param>
+		/// <returns>Async task object</returns>
+		public Task<bool> DeletePoolAsync(PoolId PoolId)
+		{
+			return Pools.DeleteAsync(PoolId);
+		}
+
+		/// <summary>
+		/// Updates an existing pool
+		/// </summary>
+		/// <param name="Pool">The pool to update</param>
+		/// <param name="NewName">The new name for the pool</param>
+		/// <param name="NewRequirements">New requirements for the pool</param>
+		/// <param name="NewEnableAutoscaling">Whether to enable autoscaling</param>
+		/// <param name="NewMinAgents">Minimum number of agents in the pool</param>
+		/// <param name="NewNumReserveAgents">Minimum number of idle agents to maintain</param>
+		/// <param name="NewProperties">Properties on the pool to update. Any properties with a value of null will be removed.</param>
+		/// <returns>Async task object</returns>
+		public async Task<IPool?> UpdatePoolAsync(IPool? Pool, string? NewName = null, AgentRequirements? NewRequirements = null, bool? NewEnableAutoscaling = null, int? NewMinAgents = null, int? NewNumReserveAgents = null, Dictionary<string, string?>? NewProperties = null)
+		{
+			for (; Pool != null; Pool = await Pools.GetAsync(Pool.Id))
+			{
+				IPool? NewPool = await Pools.TryUpdateAsync(Pool, NewName, NewRequirements, NewEnableAutoscaling, NewMinAgents, NewNumReserveAgents, null, NewProperties, null, null);
+				if (NewPool != null)
+				{
+					return NewPool;
+				}
+			}
+			return Pool;
+		}
+
+		/// <summary>
+		/// Gets all the available pools
+		/// </summary>
+		/// <returns>List of pool documents</returns>
+		public Task<List<IPool>> GetPoolsAsync()
+		{
+			return Pools.GetAsync();
+		}
+
+		/// <summary>
+		/// Gets a pool by ID
+		/// </summary>
+		/// <param name="PoolId">Unique id of the pool</param>
+		/// <returns>The pool document</returns>
+		public Task<IPool?> GetPoolAsync(PoolId PoolId)
+		{
+			return Pools.GetAsync(PoolId);
+		}
+
+		/// <summary>
+		/// Get a list of workspaces for the given agent
+		/// </summary>
+		/// <param name="Agent">The agent to return workspaces for</param>
+		/// <param name="ValidAtTime">Absolute time at which we expect the results to be valid. Values may be cached as long as they are after this time.</param>
+		/// <returns>List of workspaces</returns>
+		public async Task<HashSet<AgentWorkspace>> GetWorkspacesAsync(IAgent Agent, DateTime ValidAtTime)
+		{
+			HashSet<AgentWorkspace> Workspaces = new HashSet<AgentWorkspace>();
+			if (Agent.Capabilities.Devices.Count > 0)
+			{
+				DeviceCapabilities PrimaryDevice = Agent.Capabilities.Devices[0];
+				if (PrimaryDevice.Properties != null)
+				{
+					if (PrimaryDevice.Properties.Contains("OSFamily=Windows"))
+					{
+						Workspaces.Add(new AgentWorkspace(null, null, AgentWorkspace.AutoSdkIdentifier, "//UE4/Private-AutoSDK-Windows", null, true));
+					}
+					else if (PrimaryDevice.Properties.Contains("OSVersion=MacOS"))
+					{
+						Workspaces.Add(new AgentWorkspace(null, null, AgentWorkspace.AutoSdkIdentifier, "//UE4/Private-AutoSDK-Mac", null, true));
+					}
+				}
+			}
+
+			Dictionary<PoolId, IPool> PoolMapping = await GetPoolLookupAsync(ValidAtTime);
+			Workspaces.UnionWith(Agent.GetPools(PoolMapping.Values).SelectMany(x => x.Workspaces));
+			return Workspaces;
+		}
+
+		/// <summary>
+		/// Gets a mapping from pool identifiers to definitions
+		/// </summary>
+		/// <param name="ValidAtTime">Absolute time at which we expect the results to be valid. Values may be cached as long as they are after this time.</param>
+		/// <returns>Map of pool ids to pool documents</returns>
+		private async Task<Dictionary<PoolId, IPool>> GetPoolLookupAsync(DateTime ValidAtTime)
+		{
+			Tuple<DateTime, Dictionary<PoolId, IPool>>? CachedPoolLookupCopy = CachedPoolLookup;
+			if (CachedPoolLookupCopy == null || CachedPoolLookupCopy.Item1 < ValidAtTime)
+			{
+				// Get a new list of cached pools
+				DateTime NewCacheTime = Clock.UtcNow;
+				List<IPool> NewPools = await Pools.GetAsync();
+				Tuple<DateTime, Dictionary<PoolId, IPool>> NewCachedPoolLookup = Tuple.Create(NewCacheTime, NewPools.ToDictionary(x => x.Id, x => x));
+
+				// Try to swap it with the current version
+				while (CachedPoolLookupCopy == null || CachedPoolLookupCopy.Item1 < NewCacheTime)
+				{
+					Tuple<DateTime, Dictionary<PoolId, IPool>>? OriginalValue = Interlocked.CompareExchange(ref CachedPoolLookup, NewCachedPoolLookup, CachedPoolLookupCopy);
+					if (OriginalValue == CachedPoolLookupCopy)
+					{
+						CachedPoolLookupCopy = NewCachedPoolLookup;
+						break;
+					}
+					CachedPoolLookupCopy = OriginalValue;
+				}
+			}
+			return CachedPoolLookupCopy.Item2;
+		}
+	}
+}
diff --git a/Engine/Source/Programs/Horde/HordeServer/Services/PoolUpdateService.cs b/Engine/Source/Programs/Horde/HordeServer/Services/PoolUpdateService.cs
new file mode 100644
--- /dev/null
+++ b/Engine/Source/Programs/Horde/HordeServer/Services/PoolUpdateService.cs
@@ -0,0 +1,127 @@
+using HordeServer.Collections;
+using HordeServer.Models;
+using HordeServer.Utilities;
+using Microsoft.Extensions.Logging;
+using System;
+using System.Collections.Generic;
+using System.Linq;
+using System.Threading;
+using System.Threading.Tasks;
+
+using PoolId = HordeServer.Utilities.StringId<HordeServer.Models.IPool>;
+
+namespace HordeServer.Services
+{
+	/// <summary>
+	/// Periodically updates pool documents to contain the correct workspaces
+	/// </summary>
+	public class PoolUpdateService : TickedBackgroundService
+	{
+		/// <summary>
+		/// Collection of pool documents
+		/// </summary>
+		IPoolCollection Pools;
+
+		/// <summary>
+		/// Collection of stream documents
+		/// </summary>
+		IStreamCollection Streams;
+
+		/// <summary>
+		/// The logger instance for this service
+		/// </summary>
+		ILogger<PoolService> Logger;
+
+		/// <summary>
+		/// Constructor
+		/// </summary>
+		/// <param name="Pools">Collection of pool documents</param>
+		/// <param name="Streams">Collection of stream documents</param>
+		/// <param name="Logger">Logger instance</param>
+		public PoolUpdateService(IPoolCollection Pools, IStreamCollection Streams, ILogger<PoolService> Logger)
+			: base(TimeSpan.FromSeconds(30.0), Logger)
+		{
+			this.Pools = Pools;
+			this.Streams = Streams;
+			this.Logger = Logger;
+		}
+
+		/// <summary>
+		/// Execute the background task
+		/// </summary>
+		/// <param name="StoppingToken">Cancellation token for the async task</param>
+		/// <returns>Async task</returns>
+		protected override async Task TickAsync(CancellationToken StoppingToken)
+		{
+			// Capture the start time for this operation. We use this to attempt to sequence updates to agents, and prevent overriding another server's updates.
+			DateTime StartTime = DateTime.UtcNow;
+
+			// Update the list
+			bool bRetryUpdate = true;
+			while (bRetryUpdate && !StoppingToken.IsCancellationRequested)
+			{
+				Logger.LogDebug("Updating pool->workspace map");
+
+				// Assume this will be the last iteration
+				bRetryUpdate = false;
+
+				// Capture the list of pools at the start of this update
+				List<IPool> CurrentPools = await Pools.GetAsync();
+
+				// Lookup table of pool id to workspaces
+				Dictionary<PoolId, List<AgentWorkspace>> PoolToAgentWorkspaces = new Dictionary<PoolId, List<AgentWorkspace>>();
+
+				// Populate the workspace list from the current stream
+				List<IStream> ActiveStreams = await Streams.FindAllAsync();
+				foreach (IStream ActiveStream in ActiveStreams)
+				{
+					foreach (KeyValuePair<string, AgentType> AgentTypePair in ActiveStream.AgentTypes)
+					{
+						// Create the new agent workspace
+						AgentWorkspace? AgentWorkspace;
+						if (ActiveStream.TryGetAgentWorkspace(AgentTypePair.Value, out AgentWorkspace))
+						{
+							AgentType AgentType = AgentTypePair.Value;
+
+							// Find or add a list of workspaces for this pool
+							List<AgentWorkspace>? AgentWorkspaces;
+							if (!PoolToAgentWorkspaces.TryGetValue(AgentType.Pool, out AgentWorkspaces))
+							{
+								AgentWorkspaces = new List<AgentWorkspace>();
+								PoolToAgentWorkspaces.Add(AgentType.Pool, AgentWorkspaces);
+							}
+
+							// Add it to the list
+							if (!AgentWorkspaces.Contains(AgentWorkspace))
+							{
+								AgentWorkspaces.Add(AgentWorkspace);
+							}
+						}
+					}
+				}
+
+				// Update the list of workspaces for each pool
+				foreach (IPool CurrentPool in CurrentPools)
+				{
+					// Get the new list of workspaces for this pool
+					List<AgentWorkspace>? NewWorkspaces;
+					if (!PoolToAgentWorkspaces.TryGetValue(CurrentPool.Id, out NewWorkspaces))
+					{
+						NewWorkspaces = new List<AgentWorkspace>();
+					}
+
+					// Update the pools document
+					if (!AgentWorkspace.SetEquals(CurrentPool.Workspaces, NewWorkspaces) || CurrentPool.Workspaces.Count != NewWorkspaces.Count)
+					{
+						Logger.LogInformation("New workspaces for pool {Pool}:{Workspaces}", CurrentPool.Id, String.Join("", NewWorkspaces.Select(x => $"\n  Identifier=\"{x.Identifier}\", Stream={x.Stream}")));
+						if (await Pools.TryUpdateAsync(CurrentPool, NewWorkspaces: NewWorkspaces) == null)
+						{
+							Logger.LogInformation("Pool modified; will retry");
+							bRetryUpdate = true;
+						}
+					}
+				}
+			}
+		}
+	}
+}
diff --git a/Engine/Source/Programs/Horde/HordeServer/Services/ProjectService.cs b/Engine/Source/Programs/Horde/HordeServer/Services/ProjectService.cs
new file mode 100644
--- /dev/null
+++ b/Engine/Source/Programs/Horde/HordeServer/Services/ProjectService.cs
@@ -0,0 +1,184 @@
+// Copyright Epic Games, Inc. All Rights Reserved.
+
+using HordeServer.Api;
+using HordeServer.Collections;
+using HordeServer.Models;
+using HordeServer.Utilities;
+using Microsoft.Extensions.Options;
+using MongoDB.Bson;
+using MongoDB.Driver;
+using System;
+using System.Collections.Generic;
+using System.Linq;
+using System.Security.Claims;
+using System.Threading.Tasks;
+
+using ProjectId = HordeServer.Utilities.StringId<HordeServer.Models.IProject>;
+
+namespace HordeServer.Services
+{
+	/// <summary>
+	/// Cache of information about job ACLs
+	/// </summary>
+	public class ProjectPermissionsCache : GlobalPermissionsCache
+	{
+		/// <summary>
+		/// Map of project id to permissions for that project
+		/// </summary>
+		public Dictionary<ProjectId, IProjectPermissions?> Projects { get; } = new Dictionary<ProjectId, IProjectPermissions?>();
+	}
+
+	/// <summary>
+	/// Wraps functionality for manipulating projects
+	/// </summary>
+	public class ProjectService
+	{
+		/// <summary>
+		/// The ACL service
+		/// </summary>
+		AclService AclService;
+
+		/// <summary>
+		/// Collection of project documents
+		/// </summary>
+		IProjectCollection Projects;
+
+		/// <summary>
+		/// Constructor
+		/// </summary>
+		/// <param name="AclService">The ACL service</param>
+		/// <param name="Projects">Collection of project documents</param>
+		public ProjectService(AclService AclService, IProjectCollection Projects)
+		{
+			this.AclService = AclService;
+			this.Projects = Projects;
+		}
+
+		/// <summary>
+		/// Attempts to create a new project
+		/// </summary>
+		/// <param name="Id">Identifier for the new project</param>
+		/// <param name="Name">Name of the new project</param>
+		/// <param name="Order">Order to display this project on the dashboard</param>
+		/// <param name="Categories">Categories for the new project</param>
+		/// <param name="Properties">Properties for the new project</param>
+		/// <returns>The new project document</returns>
+		public Task<IProject?> TryCreateProjectAsync(ProjectId Id, string Name, int? Order, List<StreamCategory>? Categories, Dictionary<string, string>? Properties)
+		{
+			return Projects.TryAddAsync(Id, Name, Order, Categories, Properties);
+		}
+
+		/// <summary>
+		/// Updates an existing project
+		/// </summary>
+		/// <param name="ProjectId">Unique id of the project</param>
+		/// <param name="NewName">The new name for the project</param>
+		/// <param name="NewOrder">The new order for this project on the dashboard</param>
+		/// <param name="NewCategories">List of categories for the project</param>
+		/// <param name="NewProperties">Properties on the project to update. Any properties with a value of null will be removed.</param>
+		/// <param name="NewAcl">New acl for the project</param>
+		/// <returns>Async task object</returns>
+		public Task UpdateProjectAsync(ProjectId ProjectId, string? NewName, int? NewOrder, List<StreamCategory>? NewCategories, Dictionary<string, string>? NewProperties, Acl? NewAcl)
+		{
+			return Projects.UpdateAsync(ProjectId, NewName, NewOrder, NewCategories, NewProperties, NewAcl);
+		}
+
+		/// <summary>
+		/// Gets all the available projects
+		/// </summary>
+		/// <returns>List of project documents</returns>
+		public Task<List<IProject>> GetProjectsAsync()
+		{
+			return Projects.FindAllAsync();
+		}
+
+		/// <summary>
+		/// Gets a project by ID
+		/// </summary>
+		/// <param name="ProjectId">Unique id of the project</param>
+		/// <returns>The project document</returns>
+		public Task<IProject?> GetProjectAsync(ProjectId ProjectId)
+		{
+			return Projects.GetAsync(ProjectId);
+		}
+
+		/// <summary>
+		/// Gets a project's permissions info by ID
+		/// </summary>
+		/// <param name="ProjectId">Unique id of the project</param>
+		/// <returns>The project document</returns>
+		public Task<IProjectPermissions?> GetProjectPermissionsAsync(ProjectId ProjectId)
+		{
+			return Projects.GetPermissionsAsync(ProjectId);
+		}
+
+		/// <summary>
+		/// Deletes a project by id
+		/// </summary>
+		/// <param name="ProjectId">Unique id of the project</param>
+		/// <param name="StreamService">The stream service singleton</param>
+		/// <param name="JobService">The job service singleton</param>
+		public async Task DeleteProjectAsync(ProjectId ProjectId, StreamService StreamService, JobService JobService)
+		{
+			await StreamService.DeleteStreamsForProjectAsync(ProjectId, JobService);
+			await Projects.DeleteAsync(ProjectId);
+		}
+
+		/// <summary>
+		/// Determines if the user is authorized to perform an action on a particular project
+		/// </summary>
+		/// <param name="Acl">Acl for the project to check</param>
+		/// <param name="Action">The action being performed</param>
+		/// <param name="User">The principal to authorize</param>
+		/// <param name="Cache">Cache for the scope table</param>
+		/// <returns>True if the action is authorized</returns>
+		private Task<bool> AuthorizeAsync(Acl? Acl, AclAction Action, ClaimsPrincipal User, GlobalPermissionsCache? Cache)
+		{
+			bool? Result = Acl?.Authorize(Action, User);
+			if (Result == null)
+			{
+				return AclService.AuthorizeAsync(Action, User, Cache);
+			}
+			else
+			{
+				return Task.FromResult(Result.Value);
+			}
+		}
+
+		/// <summary>
+		/// Determines if the user is authorized to perform an action on a particular project
+		/// </summary>
+		/// <param name="Project">The project to check</param>
+		/// <param name="Action">The action being performed</param>
+		/// <param name="User">The principal to authorize</param>
+		/// <param name="Cache">Cache for the scope table</param>
+		/// <returns>True if the action is authorized</returns>
+		public Task<bool> AuthorizeAsync(IProject Project, AclAction Action, ClaimsPrincipal User, GlobalPermissionsCache? Cache)
+		{
+			return AuthorizeAsync(Project.Acl, Action, User, Cache);
+		}
+
+		/// <summary>
+		/// Determines if the user is authorized to perform an action on a particular project
+		/// </summary>
+		/// <param name="ProjectId">The project id to check</param>
+		/// <param name="Action">The action being performed</param>
+		/// <param name="User">The principal to authorize</param>
+		/// <param name="Cache">Cache for project permissions</param>
+		/// <returns>True if the action is authorized</returns>
+		public async Task<bool> AuthorizeAsync(ProjectId ProjectId, AclAction Action, ClaimsPrincipal User, ProjectPermissionsCache? Cache)
+		{
+			IProjectPermissions? Permissions;
+			if (Cache == null)
+			{
+				Permissions = await GetProjectPermissionsAsync(ProjectId);
+			}
+			else if (!Cache.Projects.TryGetValue(ProjectId, out Permissions))
+			{
+				Permissions = await GetProjectPermissionsAsync(ProjectId);
+				Cache.Projects.Add(ProjectId, Permissions);
+			}
+			return await AuthorizeAsync(Permissions?.Acl, Action, User, Cache);
+		}
+	}
+}
diff --git a/Engine/Source/Programs/Horde/HordeServer/Services/RequestTrackerService.cs b/Engine/Source/Programs/Horde/HordeServer/Services/RequestTrackerService.cs
new file mode 100644
--- /dev/null
+++ b/Engine/Source/Programs/Horde/HordeServer/Services/RequestTrackerService.cs
@@ -0,0 +1,171 @@
+﻿using System;
+using System.Collections.Concurrent;
+using System.Collections.Generic;
+using System.Linq;
+using System.Text;
+using System.Threading.Tasks;
+using Microsoft.AspNetCore.Http;
+using Microsoft.Extensions.Logging;
+
+namespace HordeServer.Utilities
+{
+	/// <summary>
+	/// Tracks open HTTP requests to ASP.NET
+	/// Will block a pending shutdown until all requests in progress are finished (or graceful timeout is reached)
+	/// This avoids interrupting long running requests such as artifact uploads.
+	/// </summary>
+	public class RequestTrackerService
+	{
+		/// <summary>
+		/// Writer for log output
+		/// </summary>
+		private readonly ILogger<RequestTrackerService> Logger;
+
+		readonly ConcurrentDictionary<string, TrackedRequest> RequestsInProgress = new ConcurrentDictionary<string, TrackedRequest>();
+
+		/// <summary>
+		/// Constructor
+		/// </summary>
+		/// <param name="Logger">Logger</param>
+		public RequestTrackerService(ILogger<RequestTrackerService> Logger)
+		{
+			this.Logger = Logger;
+		}
+
+		/// <summary>
+		/// Called by the middleware when a request is started
+		/// </summary>
+		/// <param name="Context">HTTP Context</param>
+		public void RequestStarted(HttpContext Context)
+		{
+			RequestsInProgress[Context.TraceIdentifier] = new TrackedRequest(Context.Request);
+		}
+		
+		/// <summary>
+		/// Called by the middleware when a request is finished (no matter if an exception occurred or not)
+		/// </summary>
+		/// <param name="Context">HTTP Context</param>
+		public void RequestFinished(HttpContext Context)
+		{
+			RequestsInProgress.Remove(Context.TraceIdentifier, out _);
+		}
+
+		/// <summary>
+		/// Get current requests in progress
+		/// </summary>
+		/// <returns>The requests in progress</returns>
+		public IReadOnlyDictionary<string, TrackedRequest> GetRequestsInProgress()
+		{
+			return RequestsInProgress;
+		}
+
+		private string GetRequestsInProgressAsString()
+		{
+			List<KeyValuePair<string, TrackedRequest>> Requests = GetRequestsInProgress().ToList();
+			Requests.Sort((A, B) => A.Value.StartedAt.CompareTo(B.Value.StartedAt));
+			StringBuilder Content = new StringBuilder();
+			foreach (KeyValuePair<string,TrackedRequest> Pair in Requests)
+			{
+				int AgeInMs = Pair.Value.GetTimeSinceStartInMs();
+				string Path = Pair.Value.Request.Path;
+				Content.AppendLine($"{AgeInMs,9}  {Path}");
+			}
+
+			return Content.ToString();
+		}
+
+		/// <summary>
+		/// Log all requests currently in progress
+		/// </summary>
+		public void LogRequestsInProgress()
+		{
+			if (GetRequestsInProgress().Count == 0)
+			{
+				Logger.LogInformation("There are no requests in progress!");
+			}
+			else
+			{
+				Logger.LogInformation("Current open requests are:\n{RequestsInProgress}", GetRequestsInProgressAsString());
+			}
+		}
+	}
+	
+	/// <summary>
+	/// Value object for tracked requests
+	/// </summary>
+	public class TrackedRequest
+	{
+		/// <summary>
+		/// When the request was received
+		/// </summary>
+		public DateTime StartedAt { get; }
+			
+		/// <summary>
+		/// The HTTP request being tracked
+		/// </summary>
+		public HttpRequest Request { get; }
+
+		/// <summary>
+		/// Constructor
+		/// </summary>
+		/// <param name="Request">HTTP request being tracked</param>
+		public TrackedRequest(HttpRequest Request)
+		{
+			StartedAt = DateTime.UtcNow;
+			this.Request = Request;
+		}
+
+		/// <summary>
+		/// How long the request has been running
+		/// </summary>
+		/// <returns>Time elapsed in milliseconds since request started</returns>
+		public int GetTimeSinceStartInMs()
+		{
+			return (int) (DateTime.UtcNow - StartedAt).TotalMilliseconds;
+		}
+	}
+	
+	/// <summary>
+	/// ASP.NET Middleware to track open requests
+	/// </summary>
+	public class RequestTrackerMiddleware
+	{
+		private readonly RequestDelegate Next;
+	
+		/// <summary>
+		/// Constructor
+		/// </summary>
+		/// <param name="Next">Next middleware to call</param>
+		public RequestTrackerMiddleware(RequestDelegate Next)
+		{
+			this.Next = Next;
+		}
+	
+		/// <summary>
+		/// Invoked by ASP.NET framework itself
+		/// </summary>
+		/// <param name="Context">HTTP Context</param>
+		/// <param name="Service">The RequestTrackerService singleton</param>
+		/// <returns></returns>
+		public async Task Invoke(HttpContext Context, RequestTrackerService Service)
+		{
+			if (!Context.Request.Path.StartsWithSegments("/health", StringComparison.Ordinal))
+			{
+				try
+				{
+					Service.RequestStarted(Context);
+					await Next(Context);
+				}
+				finally
+				{
+					Service.RequestFinished(Context);
+				}
+			}
+			else
+			{
+				// Ignore requests to /health/*
+				await Next(Context);
+			}
+		}
+	}
+}
\ No newline at end of file
