Initial drop of Horde Server/Agent for UE5/Main. Still missing some Epic-specific functionality required to build/run correctly.

diff --git a/Engine/Source/Programs/Horde/HordeServer/Services/RpcService.cs b/Engine/Source/Programs/Horde/HordeServer/Services/RpcService.cs
new file mode 100644
--- /dev/null
+++ b/Engine/Source/Programs/Horde/HordeServer/Services/RpcService.cs
@@ -0,0 +1,1053 @@
+using Google.Protobuf.WellKnownTypes;
+using Grpc.Core;
+using HordeCommon;
+using HordeCommon.Rpc;
+using HordeServer.Collections;
+using HordeServer.Models;
+using HordeServer.Utilities;
+using Microsoft.AspNetCore.Authorization;
+using Microsoft.Extensions.Logging;
+using MongoDB.Bson;
+using MongoDB.Bson.Serialization;
+using System;
+using System.Collections.Generic;
+using System.Globalization;
+using System.Linq;
+using System.Net;
+using System.Runtime.CompilerServices;
+using System.Security.Claims;
+using System.Text;
+using System.Threading.Tasks;
+using AgentSoftwareChannelName = HordeServer.Utilities.StringId<HordeServer.Services.AgentSoftwareChannels>;
+using AgentSoftwareVersion = HordeServer.Utilities.StringId<HordeServer.Collections.IAgentSoftwareCollection>;
+using IStream = HordeServer.Models.IStream;
+using StreamId = HordeServer.Utilities.StringId<HordeServer.Models.IStream>;
+using HordeServer.Tasks.Impl;
+
+[assembly: InternalsVisibleTo("HordeAgentTests")]
+
+namespace HordeServer.Services
+{
+	/// <summary>
+	/// Implements the Horde gRPC service for bots updating their status and dequeing work
+	/// </summary>
+	[Authorize]
+	public class RpcService : HordeRpc.HordeRpcBase
+	{
+		/// <summary>
+		/// Timeout before closing a long-polling request (client will retry again) 
+		/// </summary>
+		internal TimeSpan LongPollTimeout = TimeSpan.FromMinutes(9);
+
+		/// <summary>
+		/// Instance of the DatabaseService singleton
+		/// </summary>
+		DatabaseService DatabaseService;
+
+		/// <summary>
+		/// Instance of the AclService singleton
+		/// </summary>
+		AclService AclService;
+
+		/// <summary>
+		/// Instance of the AgentService singleton
+		/// </summary>
+		AgentService AgentService;
+
+		/// <summary>
+		/// The stream service instance
+		/// </summary>
+		StreamService StreamService;
+
+		/// <summary>
+		/// The job service instance
+		/// </summary>
+		JobService JobService;
+
+		/// <summary>
+		/// The software service instance
+		/// </summary>
+		AgentSoftwareService AgentSoftwareService;
+
+		/// <summary>
+		/// The artifact service instance
+		/// </summary>
+		ArtifactService ArtifactService;
+
+		/// <summary>
+		/// The log file service instance
+		/// </summary>
+		ILogFileService LogFileService;
+
+		/// <summary>
+		/// The credential service instance
+		/// </summary>
+		CredentialService CredentialService;
+
+		/// <summary>
+		/// The pool service instance
+		/// </summary>
+		PoolService PoolService;
+
+		/// <summary>
+		/// The application lifetime interface
+		/// </summary>
+		LifetimeService LifetimeService;
+
+		/// <summary>
+		/// Collection of graph documents
+		/// </summary>
+		IGraphCollection Graphs;
+
+		/// <summary>
+		/// Collection of testdata documents
+		/// </summary>
+		ITestDataCollection TestData;
+
+		/// <summary>
+		/// The RPC action source
+		/// </summary>
+		ActionTaskSource ActionTaskSource;
+
+		/// <summary>
+		/// The conform task source
+		/// </summary>
+		ConformTaskSource ConformTaskSource;
+
+		/// <summary>
+		/// Writer for log output
+		/// </summary>
+		ILogger<RpcService> Logger;
+
+		/// <summary>
+		/// Constructor
+		/// </summary>
+		/// <param name="DatabaseService">Instance of the DatabaseService singleton</param>
+		/// <param name="AclService">Instance of the AclService singleton</param>
+		/// <param name="AgentService">Instance of the AgentService singleton</param>
+		/// <param name="JobService">Instance of the JobService singleton</param>
+		/// <param name="StreamService">Instance of the StreamService singleton</param>
+		/// <param name="LogFileService">Instance of the LogFileService singleton</param>
+		/// <param name="AgentSoftwareService">Instance of the JobService singleton</param>
+		/// <param name="ArtifactService">Instance of the ArtifactService singleton</param>
+		/// <param name="CredentialService">Instance of the CredentialsService singleton</param>
+		/// <param name="PoolService">Instance of the PoolService singleton</param>
+		/// <param name="LifetimeService">The application lifetime</param>
+		/// <param name="Graphs">Collection of graph documents</param>
+		/// <param name="TestData">Collection of testdata</param>
+		/// <param name="ActionTaskSource"></param>
+		/// <param name="ConformTaskSource"></param>
+		/// <param name="Logger">Log writer</param>
+		public RpcService(DatabaseService DatabaseService, AclService AclService, AgentService AgentService, StreamService StreamService, JobService JobService, AgentSoftwareService AgentSoftwareService, ArtifactService ArtifactService, ILogFileService LogFileService, CredentialService CredentialService, PoolService PoolService, LifetimeService LifetimeService, IGraphCollection Graphs, ITestDataCollection TestData, ActionTaskSource ActionTaskSource, ConformTaskSource ConformTaskSource, ILogger<RpcService> Logger)
+		{
+			this.DatabaseService = DatabaseService;
+			this.AclService = AclService;
+			this.AgentService = AgentService;
+			this.StreamService = StreamService;
+			this.JobService = JobService;
+			this.AgentSoftwareService = AgentSoftwareService;
+			this.ArtifactService = ArtifactService;
+			this.LogFileService = LogFileService;
+			this.CredentialService = CredentialService;
+			this.PoolService = PoolService;
+			this.LifetimeService = LifetimeService;
+			this.Graphs = Graphs;
+			this.TestData = TestData;
+			this.ActionTaskSource = ActionTaskSource;
+			this.ConformTaskSource = ConformTaskSource;
+			this.Logger = Logger;
+		}
+
+		/// <summary>
+		/// Waits until the server is terminating
+		/// </summary>
+		/// <param name="Reader">Request reader</param>
+		/// <param name="Writer">Response writer</param>
+		/// <param name="Context">Context for the call</param>
+		/// <returns>Response object</returns>
+		public override async Task QueryServerState(IAsyncStreamReader<QueryServerStateRequest> Reader, IServerStreamWriter<QueryServerStateResponse> Writer, ServerCallContext Context)
+		{
+			if (await Reader.MoveNext())
+			{
+				QueryServerStateRequest Request = Reader.Current;
+				Logger.LogInformation("Start server query for client {Name}", Request.Name);
+
+				// Return the current response
+				QueryServerStateResponse Response = new QueryServerStateResponse();
+				Response.Name = Dns.GetHostName();
+				await Writer.WriteAsync(Response);
+
+				// Move to the next request from the client. This should always be the end of the stream, but will not occur until the client stops requesting responses.
+				Task<bool> MoveNextTask = Reader.MoveNext();
+
+				// Wait for the client to close the stream or a shutdown to start
+				Task LongPollDelay = Task.Delay(LongPollTimeout);
+				Task WaitTask = await Task.WhenAny(MoveNextTask, LifetimeService.StoppingTask, LongPollDelay);
+
+				if (WaitTask == MoveNextTask)
+				{
+					throw new Exception("Unexpected request to QueryServerState posted from client.");
+				}
+				else if (WaitTask == LifetimeService.StoppingTask)
+				{
+					Logger.LogInformation("Notifying client {Name} of server shutdown", Request.Name);
+					await Writer.WriteAsync(Response);
+				}
+				else if (WaitTask == LongPollDelay)
+				{
+					// Send same response as server shutdown. In the agent perspective, they will be identical.
+					await Writer.WriteAsync(Response);
+				}
+			}
+		}
+
+		/// <summary>
+		/// Waits until the server is terminating
+		/// </summary>
+		/// <param name="Reader">Request reader</param>
+		/// <param name="Writer">Response writer</param>
+		/// <param name="Context">Context for the call</param>
+		/// <returns>Response object</returns>
+		public override async Task QueryServerStateV2(IAsyncStreamReader<QueryServerStateRequest> Reader, IServerStreamWriter<QueryServerStateResponse> Writer, ServerCallContext Context)
+		{
+			if (await Reader.MoveNext())
+			{
+				QueryServerStateRequest Request = Reader.Current;
+				Logger.LogInformation("Start server query for client {Name}", Request.Name);
+
+				try
+				{
+					// Return the current response
+					QueryServerStateResponse Response = new QueryServerStateResponse();
+					Response.Name = Dns.GetHostName();
+					Response.Stopping = LifetimeService.IsStopping;
+					await Writer.WriteAsync(Response);
+
+					// Move to the next request from the client. This should always be the end of the stream, but will not occur until the client stops requesting responses.
+					Task<bool> MoveNextTask = Reader.MoveNext();
+
+					// Wait for the client to close the stream or a shutdown to start
+					if (await Task.WhenAny(MoveNextTask, LifetimeService.StoppingTask) == LifetimeService.StoppingTask)
+					{
+						Response.Stopping = true;
+						await Writer.WriteAsync(Response);
+					}
+
+					// Wait until the client has finished sending
+					while (await MoveNextTask)
+					{
+						MoveNextTask = Reader.MoveNext();
+					}
+				}
+				catch (Exception Ex)
+				{
+					Logger.LogError(Ex, "Exception in QueryServerState for {Name}", Request.Name);
+					throw;
+				}
+			}
+		}
+
+		/// <summary>
+		/// Updates the workspaces synced for an agent
+		/// </summary>
+		/// <param name="Request">The request parameters</param>
+		/// <param name="Context">Context for the call</param>
+		/// <returns>Response object</returns>
+		public override async Task<UpdateAgentWorkspacesResponse> UpdateAgentWorkspaces(UpdateAgentWorkspacesRequest Request, ServerCallContext Context)
+		{
+			for (; ; )
+			{
+				// Get the current agent state
+				IAgent? Agent = await AgentService.GetAgentAsync(new AgentId(Request.AgentId));
+				if (Agent == null)
+				{
+					throw new StructuredRpcException(StatusCode.OutOfRange, "Agent {AgentId} does not exist", Request.AgentId);
+				}
+
+				// Get the new workspaces
+				List<AgentWorkspace> NewWorkspaces = Request.Workspaces.Select(x => new AgentWorkspace(x)).ToList();
+
+				// Get the set of workspaces that are currently required
+				HashSet<AgentWorkspace> ConformWorkspaces = await PoolService.GetWorkspacesAsync(Agent, DateTime.UtcNow);
+				bool bPendingConform = !ConformWorkspaces.SetEquals(NewWorkspaces);
+
+				// Update the workspaces
+				if (await AgentService.TryUpdateWorkspacesAsync(Agent, NewWorkspaces, bPendingConform))
+				{
+					UpdateAgentWorkspacesResponse Response = new UpdateAgentWorkspacesResponse();
+					if (bPendingConform)
+					{
+						Response.Retry = await ConformTaskSource.GetWorkspacesAsync(Agent, Response.PendingWorkspaces);
+					}
+					return Response;
+				}
+			}
+		}
+
+		/// <summary>
+		/// Creates a new session
+		/// </summary>
+		/// <param name="Request">Request to create a new agent</param>
+		/// <param name="Context">Context for the RPC call</param>
+		/// <returns>Information about the new agent</returns>
+		public async override Task<CreateSessionResponse> CreateSession(CreateSessionRequest Request, ServerCallContext Context)
+		{
+			if (Request.Capabilities == null)
+			{
+				throw new StructuredRpcException(StatusCode.InvalidArgument, "Capabilities may not be null");
+			}
+
+			AgentId AgentId = new AgentId(Request.Name);
+			using IDisposable Scope = Logger.BeginScope("CreateSession({AgentId})", AgentId.ToString());
+
+			// Find the agent
+			IAgent? Agent = await AgentService.GetAgentAsync(AgentId);
+			if (Agent == null)
+			{
+				if (!await AclService.AuthorizeAsync(AclAction.CreateAgent, Context.GetHttpContext().User))
+				{
+					throw new StructuredRpcException(StatusCode.PermissionDenied, "User is not authenticated to create new agents");
+				}
+
+				const bool bEnabled = true;
+				const bool bEphemeral = false;
+				Agent = await AgentService.CreateAgentAsync(Request.Name, bEnabled, bEphemeral, null, null);
+			}
+
+			// Make sure we're allowed to create sessions on this agent
+			if (!await AgentService.AuthorizeAsync(Agent, AclAction.CreateSession, Context.GetHttpContext().User, null))
+			{
+				throw new StructuredRpcException(StatusCode.PermissionDenied, "User is not authenticated to create session for {AgentId}", Request.Name);
+			}
+
+			// Create a new session
+			Agent = await AgentService.CreateSessionAsync(Agent, Request.Status, new AgentCapabilities(Request.Capabilities), AgentSoftwareVersion.FromNullable(Request.Version));
+			if (Agent == null)
+			{
+				throw new StructuredRpcException(StatusCode.NotFound, "Agent {AgentId} not found", Request.Name);
+			}
+
+			// Create the response
+			CreateSessionResponse Response = new CreateSessionResponse();
+			Response.AgentId = Agent.Id.ToString();
+			Response.SessionId = Agent.SessionId.ToString();
+			Response.ExpiryTime = Timestamp.FromDateTime(Agent.SessionExpiresAt!.Value);
+			Response.Token = AgentService.IssueSessionToken(Agent.SessionId!.Value);
+			return Response;
+		}
+
+		/// <summary>
+		/// Updates an agent session
+		/// </summary>
+		/// <param name="Reader">Request to create a new agent</param>
+		/// <param name="Writer">Writer for response objects</param>
+		/// <param name="Context">Context for the RPC call</param>
+		/// <returns>Information about the new agent</returns>
+		public async override Task UpdateSession(IAsyncStreamReader<UpdateSessionRequest> Reader, IServerStreamWriter<UpdateSessionResponse> Writer, ServerCallContext Context)
+		{
+			// Read the request object
+			Task<bool> NextRequestTask = Reader.MoveNext();
+			if (await NextRequestTask)
+			{
+				UpdateSessionRequest Request = Reader.Current;
+				using IDisposable Scope = Logger.BeginScope("UpdateSession for agent {AgentId}, session {SessionId}", Request.AgentId, Request.SessionId);
+
+				Logger.LogDebug("Updating session for {AgentId}", Request.AgentId);
+				foreach (HordeCommon.Rpc.Messages.Lease Lease in Request.Leases)
+				{
+					Logger.LogDebug("Session {SessionId}, Lease {LeaseId} - State: {LeaseState}, Outcome: {LeaseOutcome}", Request.SessionId, Lease.Id, Lease.State, Lease.Outcome);
+				}
+
+				// Get a task for moving to the next item. This will only complete once the call has closed.
+				NextRequestTask = Reader.MoveNext();
+
+				// Get the current agent state
+				IAgent? Agent = await AgentService.GetAgentAsync(new AgentId(Request.AgentId));
+				if(Agent != null)
+				{
+					// Check we're authorized to update it
+					if (!AgentService.AuthorizeSession(Agent, Context.GetHttpContext().User))
+					{
+						throw new StructuredRpcException(StatusCode.PermissionDenied, "Not authenticated for {AgentId}", Request.AgentId);
+					}
+
+					// Get the new capabilities of this agent
+					AgentCapabilities? Capabilities = null;
+					if (Request.Capabilities != null)
+					{
+						Capabilities = new AgentCapabilities(Request.Capabilities);
+					}
+
+					// Capture the current agent update index. This allows us to detect if anything has changed.
+					uint UpdateIndex = Agent.UpdateIndex;
+
+					// Update the agent session and return to the caller if anything changes
+					Agent = await AgentService.UpdateSessionAsync(Agent, Request.SessionId.ToObjectId(), Request.Status, Capabilities, Request.Leases);
+					if (Agent != null && Agent.UpdateIndex == UpdateIndex && (Agent.Leases.Count > 0 || Agent.Status != AgentStatus.Stopping))
+					{
+						Agent = await AgentService.WaitForLeaseAsync(Agent, NextRequestTask);
+					}
+				}
+
+				// Handle the invalid agent case
+				if (Agent == null)
+				{
+					throw new StructuredRpcException(StatusCode.NotFound, "Invalid agent name '{AgentId}'", Request.AgentId);
+				}
+
+				// Create the new session info
+				UpdateSessionResponse Response = new UpdateSessionResponse();
+				Response.Leases.Add(Agent.Leases.Select(x => x.ToRpcMessage()));
+				Response.ExpiryTime = (Agent.SessionExpiresAt == null) ? new Timestamp() : Timestamp.FromDateTime(Agent.SessionExpiresAt.Value);
+				await Writer.WriteAsync(Response);
+
+				// Wait for the client to close the stream
+				while (await NextRequestTask)
+				{
+					NextRequestTask = Reader.MoveNext();
+				}
+			}
+		}
+
+		/// <summary>
+		/// Gets information about a stream
+		/// </summary>
+		/// <param name="Request">Request arguments</param>
+		/// <param name="Context">Context for the RPC call</param>
+		/// <returns>Information about the new agent</returns>
+		public override async Task<GetStreamResponse> GetStream(GetStreamRequest Request, ServerCallContext Context)
+		{
+			StreamId StreamIdValue = new StreamId(Request.StreamId);
+
+			IStream? Stream = await StreamService.GetStreamAsync(StreamIdValue);
+			if (Stream == null)
+			{
+				throw new StructuredRpcException(StatusCode.NotFound, "Stream {StreamId} does not exist", Request.StreamId);
+			}
+			if (!await StreamService.AuthorizeAsync(Stream, AclAction.ViewStream, Context.GetHttpContext().User, null))
+			{
+				throw new StructuredRpcException(StatusCode.PermissionDenied, "Not authenticated to access stream {StreamId}", Request.StreamId);
+			}
+
+			return Stream.ToRpcResponse();
+		}
+
+		/// <summary>
+		/// Updates properties on a stream
+		/// </summary>
+		/// <param name="Request">Request arguments</param>
+		/// <param name="Context">Context for the RPC call</param>
+		/// <returns>Information about the new agent</returns>
+		public async override Task<Empty> UpdateStream(UpdateStreamRequest Request, ServerCallContext Context)
+		{
+			StreamId StreamIdValue = new StreamId(Request.StreamId);
+
+			IStream? Stream = await StreamService.GetStreamAsync(StreamIdValue);
+			if (Stream == null)
+			{
+				throw new StructuredRpcException(StatusCode.NotFound, "Stream {StreamId} does not exist", Request.StreamId);
+			}
+			if (!await StreamService.AuthorizeAsync(Stream, AclAction.UpdateStream, Context.GetHttpContext().User, null))
+			{
+				throw new StructuredRpcException(StatusCode.PermissionDenied, "Not authenticated to modify stream {StreamId}", Request.StreamId);
+			}
+
+			await StreamService.UpdateStreamAsync(Stream, Request.Name, NewProperties: Request.Properties.ToDictionary(x => x.Name, x => (string?)x.Value));
+			return new Empty();
+		}
+
+		/// <summary>
+		/// Gets information about a job
+		/// </summary>
+		/// <param name="Request">Request arguments</param>
+		/// <param name="Context">Context for the RPC call</param>
+		/// <returns>Information about the new agent</returns>
+		public async override Task<GetJobResponse> GetJob(GetJobRequest Request, ServerCallContext Context)
+		{
+			ObjectId JobIdValue = Request.JobId.ToObjectId();
+
+			IJob? Job = await JobService.GetJobAsync(JobIdValue);
+			if (Job == null)
+			{
+				throw new StructuredRpcException(StatusCode.NotFound, "Job {JobId} does not exist", Request.JobId);
+			}
+			if (!JobService.AuthorizeSession(Job, Context.GetHttpContext().User))
+			{
+				throw new StructuredRpcException(StatusCode.PermissionDenied, "Not authenticated to access job {JobId}", Request.JobId);
+			}
+
+			return Job.ToRpcResponse();
+		}
+
+		/// <summary>
+		/// Updates properties on a job
+		/// </summary>
+		/// <param name="Request">Request arguments</param>
+		/// <param name="Context">Context for the RPC call</param>
+		/// <returns>Information about the new agent</returns>
+		public async override Task<Empty> UpdateJob(UpdateJobRequest Request, ServerCallContext Context)
+		{
+			ObjectId JobIdValue = Request.JobId.ToObjectId();
+
+			IJob? Job = await JobService.GetJobAsync(JobIdValue);
+			if (Job == null)
+			{
+				throw new StructuredRpcException(StatusCode.NotFound, "Job {JobId} does not exist", Request.JobId);
+			}
+			if (!JobService.AuthorizeSession(Job, Context.GetHttpContext().User))
+			{
+				throw new StructuredRpcException(StatusCode.PermissionDenied, "Not authenticated to modify job {JobId}", Request.JobId);
+			}
+
+			await JobService.UpdateJobAsync(Job, Name: Request.Name);
+			return new Empty();
+		}
+
+		/// <summary>
+		/// Starts executing a batch
+		/// </summary>
+		/// <param name="Request">Request arguments</param>
+		/// <param name="Context">Context for the RPC call</param>
+		/// <returns>Information about the new agent</returns>
+		public async override Task<BeginBatchResponse> BeginBatch(BeginBatchRequest Request, ServerCallContext Context)
+		{
+			ObjectId JobId = Request.JobId.ToObjectId();
+			SubResourceId BatchId = Request.BatchId.ToSubResourceId();
+
+			IJob? Job = await JobService.GetJobAsync(Request.JobId.ToObjectId());
+			if (Job == null)
+			{
+				throw new StructuredRpcException(StatusCode.NotFound, "Job {JobId} not found", Request.JobId);
+			}
+
+			IJobStepBatch Batch = AuthorizeBatch(Job, Request.BatchId.ToSubResourceId(), Context);
+			if (!await JobService.UpdateBatchAsync(Job, BatchId, null, Api.JobStepBatchState.Running))
+			{
+				throw new StructuredRpcException(StatusCode.NotFound, "Batch {JobId}:{BatchId} not found for updating", Request.JobId, Request.BatchId);
+			}
+
+			IGraph Graph = await JobService.GetGraphAsync(Job);
+
+			BeginBatchResponse Response = new BeginBatchResponse();
+			Response.LogId = Batch.LogId.ToString();
+			Response.AgentType = Graph.Groups[Batch.GroupIdx].AgentType;
+			return Response;
+		}
+
+		/// <summary>
+		/// Finishes executing a batch
+		/// </summary>
+		/// <param name="Request">Request arguments</param>
+		/// <param name="Context">Context for the RPC call</param>
+		/// <returns>Information about the new agent</returns>
+		public async override Task<Empty> FinishBatch(FinishBatchRequest Request, ServerCallContext Context)
+		{
+			IJob? Job = await JobService.GetJobAsync(Request.JobId.ToObjectId());
+			if (Job == null)
+			{
+				throw new StructuredRpcException(StatusCode.NotFound, "Job {JobId} not found", Request.JobId);
+			}
+
+			IJobStepBatch Batch = AuthorizeBatch(Job, Request.BatchId.ToSubResourceId(), Context);
+			await JobService.UpdateBatchAsync(Job, Batch.Id, null, Api.JobStepBatchState.Complete);
+			return new Empty();
+		}
+
+		/// <summary>
+		/// Starts executing a step
+		/// </summary>
+		/// <param name="Request">Request arguments</param>
+		/// <param name="Context">Context for the RPC call</param>
+		/// <returns>Information about the new agent</returns>
+		public async override Task<BeginStepResponse> BeginStep(BeginStepRequest Request, ServerCallContext Context)
+		{
+			// Log file that's allocated for this step. We only need do this once.
+			ILogFile? Log = null;
+
+			// Get the job
+			for (; ; )
+			{
+				// Check the job exists and we can access it
+				IJob? Job = await JobService.GetJobAsync(Request.JobId.ToObjectId());
+				if (Job == null)
+				{
+					throw new StructuredRpcException(StatusCode.NotFound, "Job {JobId} not found", Request.JobId);
+				}
+
+				// Find the batch being executed
+				IJobStepBatch Batch = AuthorizeBatch(Job, Request.BatchId.ToSubResourceId(), Context);
+
+				// Figure out which step to execute next
+				IJobStep? Step;
+				for (int StepIdx = 0; ; StepIdx++)
+				{
+					// If there aren't any more steps, send a complete message
+					if (StepIdx == Batch.Steps.Count)
+					{
+						BeginStepResponse Response = new BeginStepResponse();
+						Response.State = BeginStepResponse.Types.Result.Complete;
+						return Response;
+					}
+
+					// Check if this step is ready to be executed
+					Step = Batch.Steps[StepIdx];
+					if (Step.State == JobStepState.Ready)
+					{
+						break;
+					}
+					if (Step.State == JobStepState.Waiting)
+					{
+						BeginStepResponse Response = new BeginStepResponse();
+						Response.State = BeginStepResponse.Types.Result.Waiting;
+						return Response;
+					}
+				}
+
+				// Create a log file if necessary
+				if (Log == null)
+				{
+					Log = await LogFileService.CreateLogFileAsync(Job.Id, Batch.SessionId, Api.LogType.Json);
+				}
+
+				// Get the node for this step
+				IGraph Graph = await JobService.GetGraphAsync(Job);
+				INode Node = Graph.Groups[Batch.GroupIdx].Nodes[Step.NodeIdx];
+
+				// Figure out all the credentials for it (and check we can access them)
+				Dictionary<string, string> Credentials = new Dictionary<string, string>();
+				//				if (Node.Credentials != null)
+				//				{
+				//					ClaimsPrincipal Principal = new ClaimsPrincipal(new ClaimsIdentity(Job.Claims.Select(x => new Claim(x.Type, x.Value))));
+				//					if (!await GetCredentialsForStep(Principal, Node, Credentials, Message => FailStep(Job, Batch.Id, Step.Id, Log, Message)))
+				//					{
+				//						Log = null;
+				//						continue;
+				//					}
+				//				}
+
+				// Update the step state
+				if (await JobService.TryUpdateStepAsync(Job, Batch.Id, Step.Id, JobStepState.Running, JobStepOutcome.Unspecified, null, null, Log.Id, null, null, null, null))
+				{
+					BeginStepResponse Response = new BeginStepResponse();
+					Response.State = BeginStepResponse.Types.Result.Ready;
+					Response.LogId = Log.Id.ToString();
+					Response.StepId = Step.Id.ToString();
+					Response.Name = Node.Name;
+					Response.Credentials.Add(Credentials);
+					if (Node.Properties != null)
+					{
+						Response.Properties.Add(Node.Properties);
+					}
+					Response.Warnings = Node.Warnings;
+					return Response;
+				}
+			}
+		}
+
+		/// <summary>
+		/// Gets all the required credentials for the given step
+		/// </summary>
+		/// <param name="User">The user to validate</param>
+		/// <param name="Node">The node being executed</param>
+		/// <param name="Credentials">Receives a list of credentials for the step</param>
+		/// <param name="WriteError">Delegate used to write error messages</param>
+		/// <returns>Async task</returns>
+		private async Task<bool> GetCredentialsForStep(ClaimsPrincipal User, INode Node, Dictionary<string, string> Credentials, Func<string, Task> WriteError)
+		{
+			if (Node.Credentials != null)
+			{
+				Dictionary<string, Credential> Cache = new Dictionary<string, Credential>(StringComparer.OrdinalIgnoreCase);
+				foreach (KeyValuePair<string, string> Pair in Node.Credentials)
+				{
+					// Get the credential path. We expect this in the format <CredentialName>.<PropertyName>
+					string Path = Pair.Value;
+
+					// Find the separator
+					int Idx = Path.IndexOf('.', StringComparison.Ordinal);
+					if (Idx == -1)
+					{
+						await WriteError($"Invalid credential path '{Path}'. Requested credentials should be in the form '<CredentialName>.<PropertyName>'.\n");
+						return false;
+					}
+
+					// Split the path into credential and property names
+					string CredentialName = Path.Substring(0, Idx);
+					string PropertyName = Path.Substring(Idx + 1);
+
+					// Try to get the credential with this name
+					Credential? Credential;
+					if (!Cache.TryGetValue(CredentialName, out Credential))
+					{
+						Credential = await CredentialService.GetCredentialAsync(CredentialName);
+						if (Credential == null)
+						{
+							await WriteError($"No credential called '{CredentialName}' could be found");
+							return false;
+						}
+						if (!await CredentialService.AuthorizeAsync(Credential, AclAction.ViewCredential, User, null))
+						{
+							await WriteError($"User is not allowed to view credential '{CredentialName}'");
+							return false;
+						}
+						Cache[CredentialName] = Credential;
+					}
+
+					// Get the property
+					string? PropertyValue;
+					if (!Credential!.Properties.TryGetValue(PropertyName, out PropertyValue))
+					{
+						await WriteError($"No property called '{PropertyName}' found in credential '{CredentialName}'");
+						return false;
+					}
+
+					// Add it to the output list
+					Credentials.Add(Pair.Key, PropertyValue);
+				}
+			}
+			return true;
+		}
+
+		async Task<IJob> GetJobAsync(ObjectId JobId)
+		{
+			IJob? Job = await JobService.GetJobAsync(JobId);
+			if (Job == null)
+			{
+				throw new StructuredRpcException(StatusCode.NotFound, "Job {JobId} not found", JobId);
+			}
+			return Job;
+		}
+
+		static IJobStepBatch AuthorizeBatch(IJob Job, SubResourceId BatchId, ServerCallContext Context)
+		{
+			IJobStepBatch? Batch;
+			if (!Job.TryGetBatch(BatchId, out Batch))
+			{
+				throw new StructuredRpcException(StatusCode.NotFound, "Unable to find batch {JobId}:{BatchId}", Job.Id, BatchId);
+			}
+			if (Batch.SessionId == null)
+			{
+				throw new StructuredRpcException(StatusCode.PermissionDenied, "Batch {JobId}:{BatchId} has no session id", Job.Id, BatchId);
+			}
+
+			ClaimsPrincipal Principal = Context.GetHttpContext().User;
+			if (!Principal.HasSessionClaim(Batch.SessionId.Value))
+			{
+				throw new StructuredRpcException(StatusCode.PermissionDenied, "Session id {SessionId} not valid for batch {JobId}:{BatchId}. Expected {ExpectedSessionId}.", Principal.GetSessionClaim() ?? ObjectId.Empty, Job.Id, BatchId, Batch.SessionId.Value);
+			}
+
+			return Batch;
+		}
+
+		/// <summary>
+		/// Updates the state of a jobstep
+		/// </summary>
+		/// <param name="Request">Request arguments</param>
+		/// <param name="Context">Context for the RPC call</param>
+		/// <returns>Information about the new agent</returns>
+		public async override Task<Empty> UpdateStep(UpdateStepRequest Request, ServerCallContext Context)
+		{
+			IJob? Job = await JobService.GetJobAsync(Request.JobId.ToObjectId());
+			if (Job == null)
+			{
+				throw new StructuredRpcException(StatusCode.NotFound, "Job {JobId} not found", Request.JobId);
+			}
+
+			IJobStepBatch Batch = AuthorizeBatch(Job, Request.BatchId.ToSubResourceId(), Context);
+
+			await JobService.UpdateStepAsync(Job, Batch.Id, Request.StepId.ToSubResourceId(), Request.State, Request.Outcome, null, null, null, null, null);
+			return new Empty();
+		}
+
+		/// <summary>
+		/// Get the state of a jobstep
+		/// </summary>
+		/// <param name="Request">Request arguments</param>
+		/// <param name="Context">Context for the RPC call</param>
+		/// <returns>Information about the step</returns>
+		public async override Task<GetStepResponse> GetStep(GetStepRequest Request, ServerCallContext Context)
+		{
+			IJob Job = await GetJobAsync(Request.JobId.ToObjectId());
+			IJobStepBatch Batch = AuthorizeBatch(Job, Request.BatchId.ToSubResourceId(), Context);
+
+			SubResourceId StepId = Request.StepId.ToSubResourceId();
+			if (!Batch.TryGetStep(StepId, out IJobStep? Step))
+			{
+				throw new StructuredRpcException(StatusCode.NotFound, "Unable to find step {JobId}:{BatchId}:{StepId}", Job.Id, Batch.Id, StepId);
+			}
+
+			return new GetStepResponse(Step.Outcome, Step.State, Step.AbortRequested);
+		}
+
+		/// <summary>
+		/// Updates the state of a jobstep
+		/// </summary>
+		/// <param name="Request">Request arguments</param>
+		/// <param name="Context">Context for the RPC call</param>
+		/// <returns>Information about the new agent</returns>
+		public async override Task<UpdateGraphResponse> UpdateGraph(UpdateGraphRequest Request, ServerCallContext Context)
+		{
+			List<Api.CreateGroupRequest> ApiGroups = new List<Api.CreateGroupRequest>();
+			foreach (CreateGroupRequest Group in Request.Groups)
+			{
+				List<Api.CreateNodeRequest> ApiNodes = new List<Api.CreateNodeRequest>();
+				foreach (CreateNodeRequest Node in Group.Nodes)
+				{
+					Api.CreateNodeRequest ApiNode = new Api.CreateNodeRequest(Node.Name, Node.InputDependencies.ToList(), Node.OrderDependencies.ToList(), Node.Priority, Node.AllowRetry, Node.RunEarly, Node.Warnings, new Dictionary<string, string>(Node.Credentials), new Dictionary<string, string>(Node.Properties));
+					ApiNodes.Add(ApiNode);
+				}
+				ApiGroups.Add(new Api.CreateGroupRequest(Group.AgentType, ApiNodes));
+			}
+
+			List<Api.CreateAggregateRequest> ApiAggregates = new List<Api.CreateAggregateRequest>();
+			foreach (CreateAggregateRequest Aggregate in Request.Aggregates)
+			{
+				Api.CreateAggregateRequest ApiAggregate = new Api.CreateAggregateRequest(Aggregate.Name, Aggregate.Nodes.ToList());
+				ApiAggregates.Add(ApiAggregate);
+			}
+
+			List<Api.CreateLabelRequest> ApiLabels = new List<Api.CreateLabelRequest>();
+			foreach (CreateLabelRequest Label in Request.Labels)
+			{
+				Api.CreateLabelRequest ApiLabel = new Api.CreateLabelRequest(Label.DashboardName, Label.DashboardCategory, Label.UgsName, Label.UgsProject, Label.Change, Label.RequiredNodes.ToList(), Label.IncludedNodes.ToList());
+				ApiLabels.Add(ApiLabel);
+			}
+
+			ObjectId JobIdValue = Request.JobId.ToObjectId();
+			for (; ; )
+			{
+				IJob? Job = await JobService.GetJobAsync(JobIdValue);
+				if (Job == null)
+				{
+					throw new StructuredRpcException(StatusCode.NotFound, "Resource not found");
+				}
+				if (!JobService.AuthorizeSession(Job, Context.GetHttpContext().User))
+				{
+					throw new StructuredRpcException(StatusCode.PermissionDenied, "Access denied");
+				}
+
+				IGraph Graph = await JobService.GetGraphAsync(Job);
+				Graph = await Graphs.AppendAsync(Graph, ApiGroups, ApiAggregates, ApiLabels);
+
+				if (await JobService.TryUpdateGraphAsync(Job, Graph))
+				{
+					return new UpdateGraphResponse();
+				}
+			}
+		}
+
+		/// <summary>
+		/// Creates a set of events
+		/// </summary>
+		/// <param name="Request">Request arguments</param>
+		/// <param name="Context">Context for the RPC call</param>
+		/// <returns>Information about the new agent</returns>
+		public async override Task<Empty> CreateEvents(CreateEventsRequest Request, ServerCallContext Context)
+		{
+			if (!await AclService.AuthorizeAsync(AclAction.CreateEvent, Context.GetHttpContext().User))
+			{
+				throw new StructuredRpcException(StatusCode.NotFound, "Access denied");
+			}
+
+			List<NewLogEventData> NewEvents = new List<NewLogEventData>();
+			foreach (CreateEventRequest Event in Request.Events)
+			{
+				NewLogEventData NewEvent = new NewLogEventData();
+				NewEvent.LogId = Event.LogId.ToObjectId();
+				NewEvent.Severity = Event.Severity;
+				NewEvent.LineIndex = Event.LineIndex;
+				NewEvent.LineCount = Event.LineCount;
+				NewEvents.Add(NewEvent);
+			}
+			await LogFileService.CreateEventsAsync(NewEvents);
+			return new Empty();
+		}
+
+		/// <summary>
+		/// Writes output to a log file
+		/// </summary>
+		/// <param name="Request">Request arguments</param>
+		/// <param name="Context">Context for the RPC call</param>
+		/// <returns>Information about the new agent</returns>
+		public async override Task<Empty> WriteOutput(WriteOutputRequest Request, ServerCallContext Context)
+		{
+			ILogFile? LogFile = await LogFileService.GetCachedLogFileAsync(Request.LogId.ToObjectId());
+			if (LogFile == null)
+			{
+				throw new StructuredRpcException(StatusCode.NotFound, "Resource not found");
+			}
+			if (!HordeServer.Services.LogFileService.AuthorizeForSession(LogFile, Context.GetHttpContext().User))
+			{
+				throw new StructuredRpcException(StatusCode.PermissionDenied, "Access denied");
+			}
+
+			await LogFileService.WriteLogDataAsync(LogFile, Request.Offset, Request.LineIndex, Request.Data.ToArray(), Request.Flush);
+			return new Empty();
+		}
+
+		/// <summary>
+		/// Uploads a new agent archive
+		/// </summary>
+		/// <param name="Request">Request arguments</param>
+		/// <param name="Context">Context for the RPC call</param>
+		/// <returns>Information about the new agent</returns>
+		public async override Task<UploadSoftwareResponse> UploadSoftware(UploadSoftwareRequest Request, ServerCallContext Context)
+		{
+			if (!await AclService.AuthorizeAsync(AclAction.UploadSoftware, Context.GetHttpContext().User))
+			{
+				throw new StructuredRpcException(StatusCode.PermissionDenied, "Access to software is forbidden");
+			}
+
+			AgentSoftwareVersion Version = await AgentSoftwareService.SetArchiveAsync(new AgentSoftwareChannelName(Request.Channel), null, Request.Data.ToArray());
+
+			UploadSoftwareResponse Response = new UploadSoftwareResponse();
+			Response.Version = Version.ToString();
+			return Response;
+		}
+
+		/// <summary>
+		/// Downloads a new agent archive
+		/// </summary>
+		/// <param name="Request">Request arguments</param>
+		/// <param name="ResponseStream">Writer for the output data</param>
+		/// <param name="Context">Context for the RPC call</param>
+		/// <returns>Information about the new agent</returns>
+		public async override Task DownloadSoftware(DownloadSoftwareRequest Request, IServerStreamWriter<DownloadSoftwareResponse> ResponseStream, ServerCallContext Context)
+		{
+			if (!await AclService.AuthorizeAsync(AclAction.DownloadSoftware, Context.GetHttpContext().User))
+			{
+				throw new StructuredRpcException(StatusCode.NotFound, "Access to software is forbidden");
+			}
+
+			AgentSoftwareVersion Version = new AgentSoftwareVersion(Request.Version);
+
+			byte[]? Data = await AgentSoftwareService.GetArchiveAsync(Version);
+			if (Data == null)
+			{
+				throw new StructuredRpcException(StatusCode.NotFound, "Missing version {Version}");
+			}
+
+			for (int Offset = 0; Offset < Data.Length;)
+			{
+				int NextOffset = Math.Min(Offset + 128 * 1024, Data.Length);
+
+				DownloadSoftwareResponse Response = new DownloadSoftwareResponse();
+				Response.Data = Google.Protobuf.ByteString.CopyFrom(Data.AsSpan(Offset, NextOffset - Offset));
+
+				await ResponseStream.WriteAsync(Response);
+
+				Offset = NextOffset;
+			}
+		}
+
+		/// <summary>
+		/// Uploads a new artifact
+		/// </summary>
+		/// <param name="Reader">Request arguments</param>
+		/// <param name="Context">Context for the RPC call</param>
+		/// <returns>Information about the new agent</returns>
+		public async override Task<UploadArtifactResponse> UploadArtifact(IAsyncStreamReader<UploadArtifactRequest> Reader, ServerCallContext Context)
+		{
+			// Advance to the metadata object
+			if (!await Reader.MoveNext())
+			{
+				throw new StructuredRpcException(StatusCode.DataLoss, "Missing request for artifact upload");
+			}
+
+			// Read the request object
+			UploadArtifactMetadata? Metadata = Reader.Current.Metadata;
+			if (Metadata == null)
+			{
+				throw new StructuredRpcException(StatusCode.DataLoss, "Expected metadata in first artifact request");
+			}
+
+			// Get the job and step
+			IJob Job = await GetJobAsync(Metadata.JobId.ToObjectId());
+			IJobStepBatch Batch = AuthorizeBatch(Job, Metadata.BatchId.ToSubResourceId(), Context);
+
+			IJobStep? Step;
+			if (!Job.TryGetStep(Metadata.BatchId.ToSubResourceId(), Metadata.StepId.ToSubResourceId(), out Step))
+			{
+				throw new StructuredRpcException(StatusCode.NotFound, "Unable to find step {JobId}:{BatchId}:{StepId}", Job.Id, Metadata.BatchId, Metadata.StepId);
+			}
+
+			// Upload the stream
+			using (ArtifactChunkStream InputStream = new ArtifactChunkStream(Reader, Metadata.Length))
+			{
+				Artifact Artifact = await ArtifactService.CreateArtifactAsync(Job.Id, Step.Id, Metadata.Name, Metadata.MimeType, InputStream);
+
+				UploadArtifactResponse Response = new UploadArtifactResponse();
+				Response.Id = Artifact.Id.ToString();
+				return Response;
+			}
+		}
+
+		/// <summary>
+		/// Uploads new test data
+		/// </summary>
+		/// <param name="Reader">Request arguments</param>
+		/// <param name="Context">Context for the RPC call</param>
+		/// <returns>Information about the new agent</returns>
+		public override async Task<UploadTestDataResponse> UploadTestData(IAsyncStreamReader<UploadTestDataRequest> Reader, ServerCallContext Context)
+		{
+			IJob? Job = null;
+			IJobStep? JobStep = null;
+
+			while (await Reader.MoveNext())
+			{
+				UploadTestDataRequest Request = Reader.Current;
+
+				ObjectId JobId = Request.JobId.ToObjectId();
+				if (Job == null || JobId != Job.Id)
+				{
+					Job = await JobService.GetJobAsync(JobId);
+					if (Job == null)
+					{
+						throw new StructuredRpcException(StatusCode.NotFound, "Unable to find job {JobId}", JobId);
+					}
+					JobStep = null;
+				}
+
+				SubResourceId JobStepId = Request.JobStepId.ToSubResourceId();
+				if (JobStep == null || JobStepId != JobStep.Id)
+				{
+					if (!Job.TryGetStep(JobStepId, out JobStep))
+					{
+						throw new StructuredRpcException(StatusCode.NotFound, "Unable to find step {JobStepId} on job {JobId}", JobStepId, JobId);
+					}
+				}
+
+				string Text = Encoding.UTF8.GetString(Request.Value.ToArray());
+				BsonDocument Document = BsonSerializer.Deserialize<BsonDocument>(Text);
+				await TestData.AddAsync(Job, JobStep, Request.Key, Document);
+			}
+
+			return new UploadTestDataResponse();
+		}
+
+		/// <summary>
+		/// Create a new report on a job or job step
+		/// </summary>
+		/// <param name="Request"></param>
+		/// <param name="Context"></param>
+		/// <returns></returns>
+		public override async Task<CreateReportResponse> CreateReport(CreateReportRequest Request, ServerCallContext Context)
+		{
+			IJob Job = await GetJobAsync(Request.JobId.ToObjectId());
+			IJobStepBatch Batch = AuthorizeBatch(Job, Request.BatchId.ToSubResourceId(), Context);
+
+			Report NewReport = new Report { Name = Request.Name, Placement = Request.Placement, ArtifactId = Request.ArtifactId.ToObjectId() };
+			if (Request.Scope == ReportScope.Job)
+			{
+				Logger.LogDebug("Adding report to job {JobId}: {Name} -> {ArtifactId}", Job.Id, Request.Name, Request.ArtifactId);
+				await JobService.UpdateJobAsync(Job, Reports: new List<Report> { NewReport });
+			}
+			else
+			{
+				Logger.LogDebug("Adding report to step {JobId}:{BatchId}:{StepId}: {Name} -> {ArtifactId}", Job.Id, Batch.Id, Request.StepId, Request.Name, Request.ArtifactId);
+				await JobService.UpdateStepAsync(Job, Batch.Id, Request.StepId.ToSubResourceId(), NewReports: new List<Report> { NewReport });
+			}
+
+			return new CreateReportResponse();
+		}
+	}
+}
diff --git a/Engine/Source/Programs/Horde/HordeServer/Services/ScheduleService.cs b/Engine/Source/Programs/Horde/HordeServer/Services/ScheduleService.cs
new file mode 100644
--- /dev/null
+++ b/Engine/Source/Programs/Horde/HordeServer/Services/ScheduleService.cs
@@ -0,0 +1,436 @@
+// Copyright Epic Games, Inc. All Rights Reserved.
+
+using Amazon.S3.Model.Internal.MarshallTransformations;
+using Datadog.Trace;
+using EpicGames.Core;
+using HordeCommon;
+using HordeServer.Api;
+using HordeServer.Collections;
+using HordeServer.Controllers;
+using HordeServer.Models;
+using HordeServer.Utilities;
+using Microsoft.AspNetCore.Server.Kestrel.Core.Features;
+using Microsoft.Extensions.Hosting;
+using Microsoft.Extensions.Logging;
+using Microsoft.Extensions.Options;
+using MongoDB.Bson;
+using MongoDB.Driver;
+using System;
+using System.Collections.Generic;
+using System.Diagnostics;
+using System.Linq;
+using System.Security.Claims;
+using System.Threading;
+using System.Threading.Tasks;
+using TimeZoneConverter;
+using StreamId = HordeServer.Utilities.StringId<HordeServer.Models.IStream>;
+using TemplateRefId = HordeServer.Utilities.StringId<HordeServer.Models.TemplateRef>;
+
+namespace HordeServer.Services
+{
+	/// <summary>
+	/// Manipulates schedule instances
+	/// </summary>
+	public class ScheduleService : ElectedBackgroundService
+	{
+		/// <summary>
+		/// The downtime service instance
+		/// </summary>
+		IDowntimeService DowntimeService;
+
+		/// <summary>
+		/// Collection of graph documents
+		/// </summary>
+		IGraphCollection Graphs;
+
+		/// <summary>
+		/// The perforce service singleton
+		/// </summary>
+		IPerforceService Perforce;
+
+		/// <summary>
+		/// Collection of job documents
+		/// </summary>
+		IJobCollection JobCollection;
+
+		/// <summary>
+		/// The job service
+		/// </summary>
+		JobService JobService;
+
+		/// <summary>
+		/// The stream service
+		/// </summary>
+		StreamService StreamService;
+
+		/// <summary>
+		/// The template service
+		/// </summary>
+		TemplateService TemplateService;
+
+		/// <summary>
+		/// The timezone to use for schedules
+		/// </summary>
+		TimeZoneInfo TimeZone;
+
+		/// <summary>
+		/// The clock instance
+		/// </summary>
+		IClock Clock;
+
+		/// <summary>
+		/// The logging instance
+		/// </summary>
+		ILogger Logger;
+
+		/// <summary>
+		/// Constructor
+		/// </summary>
+		/// <param name="DatabaseService">Database service instance</param>
+		/// <param name="Graphs">Collection of graph documents</param>
+		/// <param name="DowntimeService">The downtime service</param>
+		/// <param name="Perforce">The Perforce service</param>
+		/// <param name="JobCollection">Job collection singleton</param>
+		/// <param name="JobService">The job service instance</param>
+		/// <param name="StreamService">The stream service instance</param>
+		/// <param name="TemplateService">The template service instance</param>
+		/// <param name="Clock"></param>
+		/// <param name="Settings">Settings for the server</param>
+		/// <param name="Logger">Logging instance</param>
+		public ScheduleService(DatabaseService DatabaseService, IGraphCollection Graphs, IDowntimeService DowntimeService, IPerforceService Perforce, IJobCollection JobCollection, JobService JobService, StreamService StreamService, TemplateService TemplateService, IClock Clock, IOptionsMonitor<ServerSettings> Settings, ILogger<ScheduleService> Logger)
+			: base(DatabaseService, new ObjectId("6035593e6d721d80fb9efa5c"), Logger)
+		{
+			this.Graphs = Graphs;
+			this.DowntimeService = DowntimeService;
+			this.Perforce = Perforce;
+			this.JobCollection = JobCollection;
+			this.JobService = JobService;
+			this.StreamService = StreamService;
+			this.TemplateService = TemplateService;
+			this.Clock = Clock;
+
+			string? ScheduleTimeZone = Settings.CurrentValue.ScheduleTimeZone;
+			this.TimeZone = (ScheduleTimeZone == null) ? TimeZoneInfo.Local : TZConvert.GetTimeZoneInfo(ScheduleTimeZone);
+
+			this.Logger = Logger;
+		}
+
+		/// <summary>
+		/// Execute any scheduled tasks
+		/// </summary>
+		/// <param name="StoppingToken">Cancellation token for stopping the service</param>
+		/// <returns>Async task</returns>
+		protected override async Task<DateTime> TickLeaderAsync(CancellationToken StoppingToken)
+		{
+			DateTime NextTickTime = DateTime.UtcNow.AddMinutes(1.0);
+			if (!DowntimeService.IsDowntimeActive)
+			{
+				Logger.LogInformation("Updating scheduled triggers...");
+				using (Scope Scope = Tracer.Instance.StartActive("ScheduleService Tick"))
+				{
+					NextTickTime = await UpdateSchedulesAsync();
+				}
+			}
+			return NextTickTime;
+		}
+
+		/// <summary>
+		/// Update all the schedules
+		/// </summary>
+		/// <returns>Next time that checkes need to be updated</returns>
+		[System.Diagnostics.CodeAnalysis.SuppressMessage("Design", "CA1031:Do not catch general exception types", Justification = "<Pending>")]
+		async Task<DateTime> UpdateSchedulesAsync()
+		{
+			// Find the next time at which each schedule should run
+			DateTimeOffset Now = Clock.UtcNow;
+			DateTimeOffset NextUpdateTime = Now.AddMinutes(1.0);
+
+			// Find all the enabled schedules
+			List<IStream> Streams = await StreamService.GetStreamsAsync();
+			foreach (IStream Stream in Streams)
+			{
+				if (Stream.IsPaused(Clock.UtcNow))
+				{
+					Logger.LogDebug("Skipping schedule update for stream {StreamId}. It has been paused until {PausedUntil} with comment '{PauseComment}'.", Stream.Id, Stream.PausedUntil, Stream.PauseComment);
+					continue;
+				}
+				
+				foreach ((TemplateRefId TemplateRefId, TemplateRef TemplateRef) in Stream.Templates)
+				{
+					Schedule? Schedule = TemplateRef.Schedule;
+					if (Schedule != null && Schedule.Enabled)
+					{
+						using IDisposable Scope = Logger.BeginScope("Updating schedule {StreamId}/{TemplateRefId}", Stream.Id, TemplateRefId);
+
+						DateTimeOffset? NextTriggerTime = await UpdateScheduleAsync(Stream, TemplateRefId, TemplateRef, Schedule, Now);
+						if (NextTriggerTime.HasValue && NextTriggerTime.Value < NextUpdateTime)
+						{
+							NextUpdateTime = NextTriggerTime.Value;
+						}
+					}
+				}
+			}
+
+			return NextUpdateTime.UtcDateTime;
+		}
+
+		/// <summary>
+		/// Updates a schedule
+		/// </summary>
+		/// <returns>Next time to trigger the schedule</returns>
+		[System.Diagnostics.CodeAnalysis.SuppressMessage("Design", "CA1031:Do not catch general exception types", Justification = "<Pending>")]
+		private async Task<DateTimeOffset?> UpdateScheduleAsync(IStream Stream, TemplateRefId TemplateRefId, TemplateRef TemplateRef, Schedule Schedule, DateTimeOffset Now)
+		{
+			// Check if we can run the trigger
+			DateTimeOffset? NextTriggerTime = Schedule.GetNextTriggerTime(TimeZone);
+			if (!NextTriggerTime.HasValue)
+			{
+				Logger.LogInformation("Schedule {StreamId}/{TemplateRefId} is disabled", Stream.Id, TemplateRefId);
+				return NextTriggerTime;
+			}
+			if (NextTriggerTime.Value > Now)
+			{
+				Logger.LogDebug("Schedule {StreamId}/{TemplateRefId} will next trigger at {NextTriggerTime}", Stream.Id, TemplateRefId, NextTriggerTime.Value);
+				return NextTriggerTime;
+			}
+			Logger.LogInformation("Schedule {StreamId}/{TemplateRefId} is ready to trigger", Stream.Id, TemplateRefId);
+
+			// Get a list of jobs that we need to remove
+			List<ObjectId> RemoveJobIds = new List<ObjectId>();
+			foreach (ObjectId ActiveJobId in Schedule.ActiveJobs)
+			{
+				IJob? Job = await JobService.GetJobAsync(ActiveJobId);
+				if (Job == null || Job.Batches.All(x => x.State == JobStepBatchState.Complete))
+				{
+					Logger.LogInformation("Removing active job {JobId}", ActiveJobId);
+					RemoveJobIds.Add(ActiveJobId);
+				}
+			}
+
+			// Trigger this schedule
+			List<IJob> CreatedJobs = new List<IJob>();
+			try
+			{
+				await TriggerAsync(Stream, TemplateRefId, TemplateRef, Schedule, Schedule.ActiveJobs.Count - RemoveJobIds.Count, CreatedJobs);
+			}
+			catch (Exception Ex)
+			{
+				Logger.LogError(Ex, "Failed to start schedule {StreamId}/{TemplateRefId}", Stream.Id, TemplateRefId);
+			}
+
+			// Update the stream
+			if (CreatedJobs.Count > 0)
+			{
+				await StreamService.UpdateScheduleTriggerAsync(Stream, TemplateRefId, Now, CreatedJobs.Max(x => x.Change), CreatedJobs.ConvertAll(x => x.Id), RemoveJobIds);
+			}
+			else if (RemoveJobIds.Count > 0)
+			{
+				await StreamService.UpdateScheduleTriggerAsync(Stream, TemplateRefId, null, null, new List<ObjectId>(), RemoveJobIds);
+			}
+
+			// Return the next update time
+			return Schedule.GetNextTriggerTime(Now, TimeZone);
+		}
+
+		/// <summary>
+		/// Trigger a schedule to run
+		/// </summary>
+		/// <param name="Stream">Stream for the schedule</param>
+		/// <param name="TemplateRefId"></param>
+		/// <param name="TemplateRef"></param>
+		/// <param name="Schedule"></param>
+		/// <param name="NumActiveJobs"></param>
+		/// <param name="CreatedJobs"></param>
+		/// <returns>Async task</returns>
+		private async Task TriggerAsync(IStream Stream, TemplateRefId TemplateRefId, TemplateRef TemplateRef, Schedule Schedule, int NumActiveJobs, List<IJob> CreatedJobs)
+		{
+			// Check we're not already at the maximum number of allowed jobs
+			if (Schedule.MaxActive != 0 && NumActiveJobs >= Schedule.MaxActive)
+			{
+				Logger.LogInformation("Skipping trigger of {StreamName} template {TemplateRefId} - already have maximum number of jobs running ({NumJobs})", Stream.Id, TemplateRefId, Schedule.MaxActive);
+				return;
+			}
+
+			// Minimum changelist number, inclusive
+			int MinChangeNumber = Schedule.LastTriggerChange;
+			if (MinChangeNumber > 0 && !Schedule.RequireSubmittedChange)
+			{
+				MinChangeNumber--;
+			}
+
+			// Maximum changelist number, exclusive
+			int? MaxChangeNumber = null;
+
+			// Get the maximum number of changes to trigger
+			int MaxNewChanges = 1;
+			if (Schedule.MaxChanges != 0)
+			{
+				MaxNewChanges = Schedule.MaxChanges;
+			}
+			if (Schedule.MaxActive != 0)
+			{
+				MaxNewChanges = Math.Min(MaxNewChanges, Schedule.MaxActive - NumActiveJobs);
+			}
+
+			// Create a timer to limit the amount we look back through P4 history
+			Stopwatch Timer = Stopwatch.StartNew();
+			ChangeContentFlags? FilterFlags = Schedule.GetFilterFlags();
+
+			// Start as many jobs as possible
+			List<(int Change, int CodeChange)> TriggerChanges = new List<(int, int)>();
+			while (TriggerChanges.Count < MaxNewChanges)
+			{
+				// Get the next valid change
+				int Change, CodeChange;
+				if (Schedule.Gate != null)
+				{
+					(Change, CodeChange) = await GetNextChangeForGateAsync(Stream.Id, TemplateRefId, Schedule.Gate, MinChangeNumber, MaxChangeNumber);
+				}
+				else
+				{
+					(Change, CodeChange) = await GetNextChangeAsync(Stream.Name, MinChangeNumber, MaxChangeNumber);
+				}
+
+				// Quit if we didn't find anything
+				if (Change <= 0)
+				{
+					break;
+				}
+				if (Change == MinChangeNumber && (Schedule.RequireSubmittedChange || TriggerChanges.Count > 0))
+				{
+					break;
+				}
+
+				// Adjust the changelist for the desired filter
+				if(await ShouldBuildChangeAsync(Stream.Name, Change, FilterFlags))
+				{
+					TriggerChanges.Add((Change, CodeChange));
+				}
+
+				// Check we haven't exceeded the time limit
+				if (Timer.Elapsed > TimeSpan.FromMinutes(2.0))
+				{
+					Logger.LogError("Querying for changes to trigger has taken {Time}. Aborting.", Timer.Elapsed);
+					break;
+				}
+
+				// Start the next build before this change
+				if (FilterFlags != null && (FilterFlags.Value & ChangeContentFlags.ContainsContent) == 0)
+				{
+					MaxChangeNumber = Math.Min(CodeChange, Change - 1);
+				}
+				else
+				{
+					MaxChangeNumber = Change - 1;
+				}
+			}
+
+			// Early out if there's nothing to do
+			if (TriggerChanges.Count == 0)
+			{
+				Logger.LogInformation("Skipping trigger of {StreamName} template {TemplateRefId} - no candidate changes after CL {LastTriggerChange}", Stream.Id, TemplateRefId, Schedule.LastTriggerChange);
+				return;
+			}
+
+			// Get the matching template
+			ITemplate? Template = await TemplateService.GetTemplateAsync(TemplateRef.Hash);
+			if (Template == null)
+			{
+				Logger.LogWarning("Unable to find template '{TemplateHash}' for '{TemplateRefId}'", TemplateRef.Hash, TemplateRefId);
+				return;
+			}
+
+			// Register the graph for it
+			IGraph Graph = await Graphs.AddAsync(Template);
+
+			// We may need to submit a new change for any new jobs. This only makes sense if there's one change.
+			if (Template.SubmitNewChange != null)
+			{
+				int NewChange = await Perforce.CreateNewChangeForTemplateAsync(Stream, Template);
+				int NewCodeChange = await Perforce.GetCodeChangeAsync(Stream.Name, NewChange);
+				TriggerChanges = new List<(int, int)> { (NewChange, NewCodeChange) };
+			}
+
+			// Try to start all the new jobs
+			foreach ((int Change, int CodeChange) in TriggerChanges.OrderBy(x => x.Change))
+			{
+				List<string> DefaultArguments = Template.GetDefaultArguments();
+				IJob NewJob = await JobService.CreateJobAsync(null, Stream.Id, TemplateRefId, Template.Id, Graph, Template.Name, Change, CodeChange, null, null, null, null, Template.Priority, null, TemplateRef.ChainedJobs, TemplateRef.ShowUgsBadges, TemplateRef.ShowUgsAlerts, TemplateRef.NotificationChannel, TemplateRef.NotificationChannelFilter, null, Template.Counters, DefaultArguments);
+				Logger.LogInformation("Started new job for {StreamName} template {TemplateName} at CL {Change} (Code CL {CodeChange}): {JobId}", Stream.Id, TemplateRef.Name, Change, CodeChange, NewJob.Id);
+				CreatedJobs.Add(NewJob);
+			}
+		}
+
+		/// <summary>
+		/// Tests whether a schedule should build a particular change, based on its requested change filters
+		/// </summary>
+		/// <param name="StreamName"></param>
+		/// <param name="Change"></param>
+		/// <param name="FilterFlags"></param>
+		/// <returns></returns>
+		private async Task<bool> ShouldBuildChangeAsync(string StreamName, int Change, ChangeContentFlags? FilterFlags)
+		{
+			if (FilterFlags != null && FilterFlags.Value != 0)
+			{
+				ChangeDetails Details = await Perforce.GetChangeDetailsAsync(StreamName, Change, null);
+				if ((Details.GetContentFlags() & FilterFlags.Value) == 0)
+				{
+					Logger.LogDebug("Not building change {Change} ({ChangeFlags}) due to filter flags ({FilterFlags})", Change, Details.GetContentFlags().ToString(), FilterFlags.Value.ToString());
+					return false;
+				}
+			}
+			return true;
+		}
+
+		/// <summary>
+		/// Get the next change to build in this stream, between a given range
+		/// </summary>
+		/// <returns>The new change and code change to build, or 0.</returns>
+		private async Task<(int Change, int CodeChange)> GetNextChangeAsync(string StreamName, int? MinChange, int? MaxChange)
+		{
+			List<ChangeSummary> Changes = await Perforce.GetChangesAsync(StreamName, MinChange, MaxChange, 1, null);
+			if (Changes.Count > 0)
+			{
+				int Change = Changes[0].Number;
+				int CodeChange = await Perforce.GetCodeChangeAsync(StreamName, Change);
+				return (Change, CodeChange);
+			}
+			return default;
+		}
+
+		/// <summary>
+		/// Gets the next change to build for a schedule on a gate
+		/// </summary>
+		/// <returns></returns>
+		private async Task<(int Change, int CodeChange)> GetNextChangeForGateAsync(StreamId StreamId, TemplateRefId TemplateRefId, ScheduleGate Gate, int? MinChange, int? MaxChange)
+		{
+			for (; ; )
+			{
+				List<IJob> Jobs = await JobCollection.FindAsync(StreamId: StreamId, Templates: new[] { Gate.TemplateRefId }, MinChange: MinChange, MaxChange: MaxChange, Count: 1);
+				if (Jobs.Count == 0)
+				{
+					return default;
+				}
+
+				IJob Job = Jobs[0];
+
+				IGraph? Graph = await Graphs.GetAsync(Job.GraphHash);
+				if (Graph != null)
+				{
+					(JobStepState, JobStepOutcome)? State = Job.GetTargetState(Graph, Gate.Target);
+					if (State != null && State.Value.Item1 == JobStepState.Completed)
+					{
+						JobStepOutcome Outcome = State.Value.Item2;
+						if (Outcome == JobStepOutcome.Success || Outcome == JobStepOutcome.Warnings)
+						{
+							return (Job.Change, Job.CodeChange);
+						}
+						Logger.LogInformation("Skipping trigger of {StreamName} template {TemplateRefId} - last {OtherTemplateRefId} job ({JobId}) ended with errors", StreamId, TemplateRefId, Gate.TemplateRefId, Job.Id);
+					}
+				}
+
+				MaxChange = Job.Change - 1;
+			}
+		}
+	}
+}
diff --git a/Engine/Source/Programs/Horde/HordeServer/Services/SoftwareService.cs b/Engine/Source/Programs/Horde/HordeServer/Services/SoftwareService.cs
new file mode 100644
--- /dev/null
+++ b/Engine/Source/Programs/Horde/HordeServer/Services/SoftwareService.cs
@@ -0,0 +1,200 @@
+// Copyright Epic Games, Inc. All Rights Reserved.
+
+using HordeServer.Api;
+using HordeServer.Collections;
+using HordeServer.Models;
+using HordeServer.Utilities;
+using MongoDB.Bson;
+using MongoDB.Driver;
+using System;
+using System.Collections.Generic;
+using System.Diagnostics;
+using System.IO;
+using System.IO.Compression;
+using System.Linq;
+using System.Security.Claims;
+using System.Security.Cryptography;
+using System.Text;
+using System.Text.Json;
+using System.Threading.Tasks;
+
+using SoftwareChannelId = HordeServer.Utilities.StringId<HordeServer.Services.SoftwareVersionsSingleton>;
+
+namespace HordeServer.Services
+{
+	/// <summary>
+	/// Singleton document used to track
+	/// </summary>
+	[SingletonDocument("5f455039d97900f2b6c735a9")]
+	class SoftwareVersionsSingleton : SingletonBase
+	{
+		/// <summary>
+		/// The next version number
+		/// </summary>
+		public int NextVersion { get; set; }
+
+		/// <summary>
+		/// Maps from a label name to a software version index
+		/// </summary>
+		public Dictionary<SoftwareChannelId, int> Channels { get; set; } = new Dictionary<SoftwareChannelId, int>();
+
+		/// <summary>
+		/// Maps from a version nubmer to a document id
+		/// </summary>
+		public Dictionary<int, ObjectId> Versions { get; set; } = new Dictionary<int, ObjectId>();
+	}
+
+	/// <summary>
+	/// Wrapper for a collection of software revisions. 
+	/// </summary>
+	public class SoftwareService
+	{
+		/// <summary>
+		/// Cached software information
+		/// </summary>
+		class CachedSoftware
+		{
+			/// <summary>
+			/// The software id
+			/// </summary>
+			public readonly ObjectId? Id;
+
+			/// <summary>
+			/// Time at which the 
+			/// </summary>
+			readonly Stopwatch Timer;
+
+			/// <summary>
+			/// Constructor
+			/// </summary>
+			/// <param name="Id">The unique software id</param>
+			public CachedSoftware(ObjectId? Id)
+			{
+				this.Id = Id;
+				this.Timer = Stopwatch.StartNew();
+			}
+
+			/// <summary>
+			/// Whether the cached software is valid
+			/// </summary>
+			public TimeSpan ElapsedTime
+			{
+				get { return Timer.Elapsed; }
+			}
+		}
+
+		/// <summary>
+		/// Collection of software documents
+		/// </summary>
+		ISoftwareCollection SoftwareCollection;
+
+		/// <summary>
+		/// The cached software instance
+		/// </summary>
+		CachedSoftware? CachedDefaultSoftware;
+
+		/// <summary>
+		/// Constructor
+		/// </summary>
+		/// <param name="SoftwareCollection">The software collection</param>
+		public SoftwareService(ISoftwareCollection SoftwareCollection)
+		{
+			this.SoftwareCollection = SoftwareCollection;
+		}
+
+		/// <summary>
+		/// Create a new software revision
+		/// </summary>
+		/// <param name="ZipStream">The input data stream. This should be a zip archive containing the HordeAgent executable.</param>
+		/// <param name="UserName">Name of the user uploading this file</param>
+		/// <param name="bDefault">Whether this software should immediately be made the default</param>
+		/// <returns>Unique id for the file</returns>
+		public async Task<ObjectId> CreateSoftwareAsync(System.IO.Stream ZipStream, string? UserName, bool bDefault)
+		{
+			// Read the whole thing into memory
+			using MemoryStream Stream = new MemoryStream();
+			await ZipStream.CopyToAsync(Stream);
+
+			// Insert the new document
+			ISoftware App = await SoftwareCollection.AddAsync(Stream.ToArray(), UserName, bDefault);
+			return App.Id;
+		}
+
+		/// <summary>
+		/// Upload a software archive to the backend
+		/// </summary>
+		/// <param name="Id">Unique id of the software to update</param>
+		/// <param name="User">Name of the current user</param>
+		/// <param name="Default">Whether the software should be the default</param>
+		/// <returns>Async task</returns>
+		public Task<bool> UpdateSoftwareAsync(ObjectId Id, string? User, bool Default)
+		{
+			return SoftwareCollection.UpdateAsync(Id, User, Default);
+		}
+
+		/// <summary>
+		/// Finds all software archives matching a set of criteria
+		/// </summary>
+		/// <param name="UploadedByUser">The user that uploaded the software</param>
+		/// <param name="MadeDefaultByUser">The user that made the software the default</param>
+		/// <param name="MadeDefault">Whether the software was made default</param>
+		/// <param name="Offset">Index of the first result to return</param>
+		/// <param name="Count">Number of results to return</param>
+		public Task<List<ISoftware>> FindSoftwareAsync(string? UploadedByUser, string? MadeDefaultByUser, bool? MadeDefault, int Offset, int Count)
+		{
+			return SoftwareCollection.FindAsync(UploadedByUser, MadeDefaultByUser, MadeDefault, Offset, Count);
+		}
+
+		/// <summary>
+		/// Gets a single software archive
+		/// </summary>
+		/// <param name="Id">Unique id of the file to retrieve</param>
+		/// <param name="bIncludeData">Whether to include the data</param>
+		public Task<ISoftware?> GetSoftwareAsync(ObjectId Id, bool bIncludeData)
+		{
+			return SoftwareCollection.GetAsync(Id, bIncludeData);
+		}
+
+		/// <summary>
+		/// Finds the current default software
+		/// </summary>
+		/// <param name="bIncludeData">Whether to include the data</param>
+		/// <returns>The current default software</returns>
+		public async Task<ISoftware?> GetDefaultSoftwareAsync(bool bIncludeData)
+		{
+			ISoftware? DefaultSoftware = await SoftwareCollection.GetDefaultAsync(bIncludeData);
+			CachedDefaultSoftware = new CachedSoftware(DefaultSoftware?.Id);
+			return DefaultSoftware;
+		}
+
+		/// <summary>
+		/// Gets a cached value for the default software
+		/// </summary>
+		/// <returns>Unique id of the default software</returns>
+		public async Task<string?> GetCachedDefaultSoftwareIdAsync()
+		{
+			CachedSoftware? CachedDefaultSoftwareCopy = CachedDefaultSoftware;
+			if (CachedDefaultSoftwareCopy != null && CachedDefaultSoftwareCopy.ElapsedTime < TimeSpan.FromMinutes(1.0))
+			{
+				return CachedDefaultSoftwareCopy.Id?.ToString();
+			}
+
+			ISoftware? DefaultSoftware = await GetDefaultSoftwareAsync(false);
+			if (DefaultSoftware != null)
+			{
+				return DefaultSoftware.Id.ToString();
+			}
+
+			return null;
+		}
+
+		/// <summary>
+		/// Removes a software archive 
+		/// </summary>
+		/// <param name="Id">Unique id of the software to delete</param>
+		public Task<bool> DeleteSoftwareAsync(ObjectId Id)
+		{
+			return SoftwareCollection.DeleteAsync(Id);
+		}
+	}
+}
diff --git a/Engine/Source/Programs/Horde/HordeServer/Services/StreamConfigService.cs b/Engine/Source/Programs/Horde/HordeServer/Services/StreamConfigService.cs
new file mode 100644
--- /dev/null
+++ b/Engine/Source/Programs/Horde/HordeServer/Services/StreamConfigService.cs
@@ -0,0 +1,148 @@
+using HordeServer.Api;
+using HordeServer.Controllers;
+using HordeServer.Models;
+using HordeServer.Notifications;
+using HordeServer.Utilities;
+using Microsoft.Extensions.Logging;
+using MongoDB.Bson;
+using System;
+using System.Collections.Generic;
+using System.IO;
+using System.Linq;
+using System.Text;
+using System.Text.Json;
+using System.Threading;
+using System.Threading.Tasks;
+
+namespace HordeServer.Services
+{
+	/// <summary>
+	/// Polls Perforce for stream config changes
+	/// </summary>
+	public class StreamConfigService : ElectedBackgroundService
+	{
+		/// <summary>
+		/// Stream service
+		/// </summary>
+		StreamService StreamService;
+		
+		/// <summary>
+		/// Template service
+		/// </summary>
+		TemplateService TemplateService;
+
+		/// <summary>
+		/// Instance of the perforce service
+		/// </summary>
+		IPerforceService PerforceService;
+
+		/// <summary>
+		/// Instance of the notification service
+		/// </summary>
+		INotificationService NotificationService;
+
+		/// <summary>
+		/// Logging device
+		/// </summary>
+		ILogger Logger;
+
+		/// <summary>
+		/// Constructor
+		/// </summary>
+		/// <param name="DatabaseService"></param>
+		/// <param name="PerforceService"></param>
+		/// <param name="StreamService"></param>
+		/// <param name="TemplateService"></param>
+		/// <param name="NotificationService"></param>
+		/// <param name="Logger"></param>
+		public StreamConfigService(DatabaseService DatabaseService, IPerforceService PerforceService, StreamService StreamService, TemplateService TemplateService, INotificationService NotificationService, ILogger<StreamConfigService> Logger)
+			: base(DatabaseService, new ObjectId("5ff60549e7632b15e64ac2f7"), Logger)
+		{
+			this.PerforceService = PerforceService;
+			this.StreamService = StreamService;
+			this.TemplateService = TemplateService;
+			this.NotificationService = NotificationService;
+			this.Logger = Logger;
+		}
+
+		/// <inheritdoc/>
+		[System.Diagnostics.CodeAnalysis.SuppressMessage("Design", "CA1031:Do not catch general exception types", Justification = "<Pending>")]
+		protected override async Task<DateTime> TickLeaderAsync(CancellationToken StoppingToken)
+		{
+			DateTime NextTickTime = DateTime.UtcNow + TimeSpan.FromMinutes(1.0);
+			Logger.LogDebug("Updating stream configs from Perforce");
+
+			List<IStream> Streams = await StreamService.GetStreamsAsync();
+
+			List<IStream> ConfigStreams = Streams.Where(x => x.ConfigPath != null).ToList();
+			if (ConfigStreams.Count > 0)
+			{
+				List<FileSummary> Files = await PerforceService.FindFilesAsync(ConfigStreams.Select(x => x.ConfigPath!));
+				foreach(FileSummary Error in Files.Where(File => File.Error != null))
+				{
+					NotificationService.NotifyUpdateStreamFailure(Error);
+				}
+				Dictionary<string, FileSummary> DepotPathToFileSummary = Files.ToDictionary(x => x.DepotPath, x => x);
+
+				for (int Idx = 0; Idx < ConfigStreams.Count; Idx++)
+				{
+					IStream Stream = ConfigStreams[Idx];
+
+					FileSummary? Summary;
+					if(!DepotPathToFileSummary.TryGetValue(Stream.ConfigPath!, out Summary))
+					{
+						Logger.LogError("Config file for {StreamName} was not found at {ConfigPath}", Stream.Name, Stream.ConfigPath);
+					}
+					else if(Summary.Change != Stream.ConfigChange)
+					{
+						Logger.LogInformation("Updating config for {Stream} ({StreamId}) to {Path}@{Change}", Stream.Name, Stream.Id, Summary.DepotPath, Summary.Change);						
+						try
+						{
+							Logger.LogInformation("Fetching {DepotPath}@{Change}", Summary.DepotPath, Summary.Change);
+							byte[] Data = await PerforceService.PrintAsync($"{Summary.DepotPath}@{Summary.Change}");
+
+							Logger.LogInformation("Config file is {NumBytes}", Data.Length);
+
+							JsonSerializerOptions Options = new JsonSerializerOptions();
+							Startup.ConfigureJsonSerializer(Options);
+							SetStreamRequest Request = JsonSerializer.Deserialize<SetStreamRequest>(Data, Options);
+
+							while (Stream.ConfigChange == null || Stream.ConfigChange < Summary.Change)
+							{
+								IStream? NewStream = await StreamsController.SetStreamAsync(Stream.ProjectId, Stream.Id, Stream, Summary.Change, Request, StreamService, TemplateService);
+								if(NewStream != null)
+								{
+									break;
+								}
+
+								NewStream = await StreamService.GetStreamAsync(Stream.Id);
+								if(NewStream == null)
+								{
+									break;
+								}
+
+								Stream = NewStream;
+							}
+
+							Logger.LogInformation("Update complete");
+						}
+						catch (Exception Ex)
+						{
+							Logger.LogError(Ex, "Error while updating {Stream} ({StreamId})", Stream.Name, Stream.Id);
+							List<ChangeDetails> Details = await PerforceService.GetChangeDetailsAsync("", new List<int>() { Summary.Change }, null);
+							string? Author = null;
+							string? Description = null;
+							if(Details.Count == 1)
+							{
+								Author = Details[0].Author;
+								Description = Details[0].Description;
+							}
+							NotificationService.NotifyUpdateStreamFailure(Stream, Ex.Message, Summary.Change, Author, Description);
+						}
+					}
+				}
+			}
+			return NextTickTime;
+		}
+	}
+}
diff --git a/Engine/Source/Programs/Horde/HordeServer/Services/StreamService.cs b/Engine/Source/Programs/Horde/HordeServer/Services/StreamService.cs
new file mode 100644
--- /dev/null
+++ b/Engine/Source/Programs/Horde/HordeServer/Services/StreamService.cs
@@ -0,0 +1,433 @@
+// Copyright Epic Games, Inc. All Rights Reserved.
+
+using HordeServer.Api;
+using HordeCommon;
+using HordeServer.Models;
+using HordeServer.Utilities;
+using Microsoft.Extensions.Options;
+using MongoDB.Bson;
+using MongoDB.Driver;
+using System;
+using System.Collections.Generic;
+using System.Linq;
+using System.Security.Claims;
+using System.Text;
+using System.Threading.Tasks;
+using System.Linq.Expressions;
+
+using ProjectId = HordeServer.Utilities.StringId<HordeServer.Models.IProject>;
+using StreamId = HordeServer.Utilities.StringId<HordeServer.Models.IStream>;
+using TemplateRefId = HordeServer.Utilities.StringId<HordeServer.Models.TemplateRef>;
+using HordeServer.Collections;
+using Microsoft.Extensions.Caching.Memory;
+
+namespace HordeServer.Services
+{
+	/// <summary>
+	/// Cache of information about stream ACLs
+	/// </summary>
+	public class StreamPermissionsCache : ProjectPermissionsCache
+	{
+		/// <summary>
+		/// Map of stream id to permissions for that stream
+		/// </summary>
+		public Dictionary<StreamId, IStreamPermissions?> Streams { get; set; } = new Dictionary<StreamId, IStreamPermissions?>();
+	}
+
+	/// <summary>
+	/// Wraps functionality for manipulating streams
+	/// </summary>
+	public sealed class StreamService : IDisposable
+	{
+		/// <summary>
+		/// The project service instance
+		/// </summary>
+		ProjectService ProjectService;
+
+		/// <summary>
+		/// Collection of stream documents
+		/// </summary>
+		IStreamCollection Streams;
+
+		/// <summary>
+		/// Cache of stream documents
+		/// </summary>
+		MemoryCache StreamCache = new MemoryCache(new MemoryCacheOptions());
+
+		/// <summary>
+		/// Constructor
+		/// </summary>
+		/// <param name="ProjectService">The project service instance</param>
+		/// <param name="Streams">Collection of stream documents</param>
+		public StreamService(ProjectService ProjectService, IStreamCollection Streams)
+		{
+			this.ProjectService = ProjectService;
+			this.Streams = Streams;
+		}
+
+		/// <summary>
+		/// Dispose of any managed resources
+		/// </summary>
+		public void Dispose()
+		{
+			StreamCache.Dispose();
+		}
+
+		/// <summary>
+		/// Creates a new stream
+		/// </summary>
+		/// <param name="Id">Unique id for the new stream</param>
+		/// <param name="Name">Name of the new stream</param>
+		/// <param name="ProjectId">Unique id of the project</param>
+		/// <param name="ConfigPath">Path to the config file in Perforce</param>
+		/// <param name="Order">Order for this stream</param>
+		/// <param name="NotificationChannel">Notification channel for this stream</param>
+		/// <param name="NotificationChannelFilter">Notification channel filter for this stream</param>
+		/// <param name="TriageChannel"></param>
+		/// <param name="DefaultPreflight">Default template for preflights</param>
+		/// <param name="Tabs">Tabs for the new stream</param>
+		/// <param name="AgentTypes">Map of new agent types to machine attributes</param>
+		/// <param name="WorkspaceTypes">Map of new workspace types</param>
+		/// <param name="TemplateRefs">New job templates for the stream</param>
+		/// <param name="Properties">New properties for the stream</param>
+		/// <param name="Acl">Acl for the stream</param>
+		/// <returns>The new stream document</returns>
+		public Task<IStream?> TryCreateStreamAsync(StreamId Id, string Name, ProjectId ProjectId, string? ConfigPath = null, int? Order = null, string? NotificationChannel = null, string? NotificationChannelFilter = null, string? TriageChannel = null, DefaultPreflight? DefaultPreflight = null, List<StreamTab>? Tabs = null, Dictionary<string, AgentType>? AgentTypes = null, Dictionary<string, WorkspaceType>? WorkspaceTypes = null, Dictionary<TemplateRefId, TemplateRef>? TemplateRefs = null, Dictionary<string, string>? Properties = null, Acl? Acl = null)
+		{
+			return Streams.TryCreateAsync(Id, Name, ProjectId, ConfigPath, Order, NotificationChannel, NotificationChannelFilter, TriageChannel, DefaultPreflight, Tabs, AgentTypes, WorkspaceTypes, TemplateRefs, Properties, Acl);
+		}
+
+		/// <summary>
+		/// Replaces a stream
+		/// </summary>
+		/// <param name="Stream">The stream interface</param>
+		/// <param name="Name">Name of the new stream</param>
+		/// <param name="ConfigPath">Path to the config file in Perforce</param>
+		/// <param name="ConfigChange">CL number of the config file</param>
+		/// <param name="Order">Order for this stream</param>
+		/// <param name="NotificationChannel">Notification channel for this stream</param>
+		/// <param name="NotificationChannelFilter">Notification channel filter for this stream</param>
+		/// <param name="TriageChannel"></param>
+		/// <param name="DefaultPreflight">Default template for preflights</param>
+		/// <param name="Tabs">Tabs for the new stream</param>
+		/// <param name="AgentTypes">Map of new agent types to machine attributes</param>
+		/// <param name="WorkspaceTypes">Map of new workspace types</param>
+		/// <param name="TemplateRefs">New job templates for the stream</param>
+		/// <param name="Properties">New properties for the stream</param>
+		/// <param name="Acl">Acl for the stream</param>
+		/// <param name="PausedUntil">The new datetime for pausing builds</param>
+		/// <param name="PauseComment">The reason for pausing</param>
+		/// <returns>The new stream document</returns>
+		public Task<IStream?> TryReplaceStreamAsync(IStream Stream, string Name, string? ConfigPath = null, int? ConfigChange = null, int? Order = null, string? NotificationChannel = null, string? NotificationChannelFilter = null, string? TriageChannel = null, DefaultPreflight? DefaultPreflight = null, List<StreamTab>? Tabs = null, Dictionary<string, AgentType>? AgentTypes = null, Dictionary<string, WorkspaceType>? WorkspaceTypes = null, Dictionary<TemplateRefId, TemplateRef>? TemplateRefs = null, Dictionary<string, string>? Properties = null, Acl? Acl = null, DateTime? PausedUntil = null, string? PauseComment = null)
+		{
+			return Streams.TryReplaceAsync(Stream, Name, ConfigPath, ConfigChange, Order, NotificationChannel, NotificationChannelFilter, TriageChannel, DefaultPreflight, Tabs, AgentTypes, WorkspaceTypes, TemplateRefs, Properties, Acl, PausedUntil, PauseComment);
+		}
+
+		/// <summary>
+		/// Deletes an existing stream
+		/// </summary>
+		/// <param name="StreamId">Unique id of the stream</param>
+		/// <param name="JobService">Instance of the job service</param>
+		/// <returns>Async task object</returns>
+		public async Task DeleteStreamAsync(StreamId StreamId, JobService JobService)
+		{
+			await JobService.DeleteJobsForStreamAsync(StreamId);
+			await Streams.DeleteAsync(StreamId);
+		}
+
+		/// <summary>
+		/// Deletes all the streams for a particular project
+		/// </summary>
+		/// <param name="ProjectId">Unique id of the project</param>
+		/// <param name="JobService">Instance of the job service</param>
+		/// <returns>Async task object</returns>
+		public async Task DeleteStreamsForProjectAsync(ProjectId ProjectId, JobService JobService)
+		{
+			List<IStream> Streams = await GetStreamsAsync(new[] { ProjectId });
+			foreach (IStream Stream in Streams)
+			{
+				await DeleteStreamAsync(Stream.Id, JobService);
+			}
+		}
+
+		/// <summary>
+		/// Updates an existing stream
+		/// </summary>
+		/// <param name="Stream">The stream to update</param>
+		/// <param name="NewName">The new name for the stream</param>
+		/// <param name="NewOrder">The new order for the stream</param>
+		/// <param name="NewNotificationChannel">The new notification channel for the stream</param>
+		/// <param name="NewNotificationChannelFilter">The new notification channel filter for the stream</param>
+		/// <param name="NewTriageChannel"></param>
+		/// <param name="NewTabs">New tabs for the stream</param>
+		/// <param name="NewAgentTypes">Map of agent types to update. Anything with a value of null will be removed.</param>
+		/// <param name="NewWorkspaceTypes">Map of workspace types to update. Anything with a value of null will be removed.</param>
+		/// <param name="NewTemplateRefs">New template references for this stream</param>
+		/// <param name="NewProperties">Properties on the stream to update. Anything with a value of null will be removed.</param>
+		/// <param name="NewAcl">The new ACL object</param>
+		/// <param name="UpdatePauseFields">Must be set to true to update the pause fields</param>
+		/// <param name="NewPausedUntil">The new datetime for pausing builds</param>
+		/// <param name="NewPauseComment">The reason for pausing</param>
+		/// <returns>Async task object</returns>
+		public async Task<IStream?> UpdateStreamAsync(IStream? Stream, string? NewName = null, int? NewOrder = null, string? NewNotificationChannel = null, string? NewNotificationChannelFilter = null, string? NewTriageChannel = null, List<StreamTab>? NewTabs = null, Dictionary<string, AgentType?>? NewAgentTypes = null, Dictionary<string, WorkspaceType?>? NewWorkspaceTypes = null, Dictionary<TemplateRefId, TemplateRef>? NewTemplateRefs = null, Dictionary<string, string?>? NewProperties = null, Acl? NewAcl = null, bool? UpdatePauseFields = null, DateTime? NewPausedUntil = null, string? NewPauseComment = null)
+		{
+			for (; Stream != null; Stream = await GetStreamAsync(Stream.Id))
+			{
+				if (await TryUpdateStreamAsync(Stream, NewName, NewOrder, NewNotificationChannel, NewNotificationChannelFilter, NewTriageChannel, NewTabs, NewAgentTypes, NewWorkspaceTypes, NewTemplateRefs, NewProperties, NewAcl, UpdatePauseFields, NewPausedUntil, NewPauseComment))
+				{
+					break;
+				}
+			}
+			return Stream;
+		}
+
+		/// <summary>
+		/// Attempts to update the last trigger time for a schedule
+		/// </summary>
+		/// <param name="Stream">The stream to update</param>
+		/// <param name="TemplateRefId">The template ref id</param>
+		/// <param name="LastTriggerTime"></param>
+		/// <param name="LastTriggerChange"></param>
+		/// <param name="AddJobs">Jobs to add</param>
+		/// <param name="RemoveJobs">Jobs to remove</param>
+		/// <returns>True if the stream was updated</returns>
+		public async Task<IStream?> UpdateScheduleTriggerAsync(IStream Stream, TemplateRefId TemplateRefId, DateTimeOffset? LastTriggerTime, int? LastTriggerChange, List<ObjectId> AddJobs, List<ObjectId> RemoveJobs)
+		{
+			IStream? NewStream = Stream;
+			while (NewStream != null)
+			{
+				TemplateRef? TemplateRef;
+				if (!NewStream.Templates.TryGetValue(TemplateRefId, out TemplateRef))
+				{
+					break;
+				}
+				if (TemplateRef.Schedule == null)
+				{
+					break;
+				}
+
+				List<ObjectId> NewActiveJobs = TemplateRef.Schedule.ActiveJobs.Except(RemoveJobs).Union(AddJobs).ToList();
+				if (await Streams.TryUpdateScheduleTriggerAsync(NewStream, TemplateRefId, LastTriggerTime, LastTriggerChange, NewActiveJobs))
+				{
+					return NewStream;
+				}
+
+				NewStream = await Streams.GetAsync(NewStream.Id);
+			}
+			return null;
+		}
+
+		/// <summary>
+		/// Updates an existing stream
+		/// </summary>
+		/// <param name="Stream">The stream to update</param>
+		/// <param name="NewName">The new name for the stream</param>
+		/// <param name="NewOrder">New order for the stream</param>
+		/// <param name="NewNotificationChannel">New notification channel for the stream</param>
+		/// <param name="NewNotificationChannelFilter">New notification channel filter for the stream</param>
+		/// <param name="NewTriageChannel"></param>
+		/// <param name="NewTabs">New tabs for the stream</param>
+		/// <param name="NewAgentTypes">Map of agent types to update. Anything with a value of null will be removed.</param>
+		/// <param name="NewWorkspaceTypes">Map of workspace types to update. Anything with a value of null will be removed.</param>
+		/// <param name="NewTemplateRefs">New template references for this stream</param>
+		/// <param name="NewProperties">Properties on the stream to update. Anything with a value of null will be removed.</param>
+		/// <param name="NewAcl">The new ACL object</param>
+		/// <param name="UpdatePauseFields">Must be set to true to update the pause fields</param>
+		/// <param name="NewPausedUntil">The new datetime for pausing builds</param>
+		/// <param name="NewPauseComment">The reason for pausing</param>
+		/// <returns>Async task object</returns>
+		public Task<bool> TryUpdateStreamAsync(IStream Stream, string? NewName, int? NewOrder, string? NewNotificationChannel, string? NewNotificationChannelFilter, string? NewTriageChannel, List<StreamTab>? NewTabs, Dictionary<string, AgentType?>? NewAgentTypes, Dictionary<string, WorkspaceType?>? NewWorkspaceTypes, Dictionary<TemplateRefId, TemplateRef>? NewTemplateRefs, Dictionary<string, string?>? NewProperties, Acl? NewAcl, bool? UpdatePauseFields = null, DateTime? NewPausedUntil = null, string? NewPauseComment = null)
+		{
+			return Streams.TryUpdatePropertiesAsync(Stream, NewName, NewOrder, NewNotificationChannel, NewNotificationChannelFilter, NewTriageChannel, NewTabs, NewAgentTypes, NewWorkspaceTypes, NewTemplateRefs, NewProperties, NewAcl, UpdatePauseFields, NewPausedUntil, NewPauseComment);
+		}
+
+		/// <summary>
+		/// Gets all the available streams for a project
+		/// </summary>
+		/// <returns>List of stream documents</returns>
+		public Task<List<IStream>> GetStreamsAsync()
+		{
+			return Streams.FindAllAsync();
+		}
+
+		/// <summary>
+		/// Gets all the available streams for a project
+		/// </summary>
+		/// <param name="ProjectId">Unique id of the project to query</param>
+		/// <returns>List of stream documents</returns>
+		public Task<List<IStream>> GetStreamsAsync(ProjectId ProjectId)
+		{
+			return Streams.FindForProjectsAsync(new[] { ProjectId });
+		}
+
+		/// <summary>
+		/// Gets all the available streams for a project
+		/// </summary>
+		/// <param name="ProjectIds">Unique id of the project to query</param>
+		/// <returns>List of stream documents</returns>
+		public Task<List<IStream>> GetStreamsAsync(ProjectId[] ProjectIds)
+		{
+			if (ProjectIds.Length == 0)
+			{
+				return Streams.FindAllAsync();
+			}
+			else
+			{
+				return Streams.FindForProjectsAsync(ProjectIds);
+			}
+		}
+
+		/// <summary>
+		/// Gets a stream by ID
+		/// </summary>
+		/// <param name="StreamId">Unique id of the stream</param>
+		/// <returns>The stream document</returns>
+		public Task<IStream?> GetStreamAsync(StreamId StreamId)
+		{
+			return Streams.GetAsync(StreamId);
+		}
+
+		/// <summary>
+		/// Adds a cached stream interface
+		/// </summary>
+		/// <param name="Stream">The stream to add</param>
+		/// <returns>The new stream</returns>
+		private void AddCachedStream(IStream Stream)
+		{
+			MemoryCacheEntryOptions Options = new MemoryCacheEntryOptions().SetAbsoluteExpiration(TimeSpan.FromMinutes(1.0));
+			StreamCache.Set(Stream.Id, Stream, Options);
+		}
+
+		/// <summary>
+		/// Gets a cached stream interface
+		/// </summary>
+		/// <param name="StreamId">Unique id for the stream</param>
+		/// <returns>The new stream</returns>
+		public async Task<IStream?> GetCachedStream(StreamId StreamId)
+		{
+			object? Stream;
+			if (!StreamCache.TryGetValue(StreamId, out Stream))
+			{
+				Stream = await GetStreamAsync(StreamId);
+				if (Stream != null)
+				{
+					AddCachedStream((IStream)Stream);
+				}
+			}
+			return (IStream?)Stream;
+		}
+
+		/// <summary>
+		/// Gets a stream's permissions by ID
+		/// </summary>
+		/// <param name="StreamId">Unique id of the stream</param>
+		/// <returns>The stream document</returns>
+		public Task<IStreamPermissions?> GetStreamPermissionsAsync(StreamId StreamId)
+		{
+			return Streams.GetPermissionsAsync(StreamId);
+		}
+
+		/// <summary>
+		/// Updates the query time on a stream object
+		/// </summary>
+		/// <param name="Stream">The stream to update</param>
+		/// <param name="LastCommitTime">The last query time</param>
+		/// <returns>New state of the commit stream</returns>
+		public async Task<IStream?> UpdateCommitTimeAsync(IStream? Stream, DateTime LastCommitTime)
+		{
+			// Try to update the commit stream with the new commits
+			while (Stream != null && (!Stream.LastCommitTime.HasValue || LastCommitTime > Stream.LastCommitTime))
+			{
+				// Try to update the document
+				if(await Streams.TryUpdateCommitTimeAsync(Stream, LastCommitTime))
+				{
+					break;
+				}
+
+				// Fetch the new stream document
+				Stream = await GetStreamAsync(Stream.Id);
+			}
+			return Stream;
+		}
+
+		/// <summary>
+		/// Determines if the user is authorized to perform an action on a particular stream
+		/// </summary>
+		/// <param name="Acl">ACL for the stream to check</param>
+		/// <param name="ProjectId">The parent project id</param>
+		/// <param name="Action">The action being performed</param>
+		/// <param name="User">The principal to authorize</param>
+		/// <param name="Cache">Cache of project permissions</param>
+		/// <returns>True if the action is authorized</returns>
+		private Task<bool> AuthorizeAsync(Acl? Acl, ProjectId ProjectId, AclAction Action, ClaimsPrincipal User, ProjectPermissionsCache? Cache)
+		{
+			bool? Result = Acl?.Authorize(Action, User);
+			if (Result == null)
+			{
+				return ProjectService.AuthorizeAsync(ProjectId, Action, User, Cache);
+			}
+			else
+			{
+				return Task.FromResult(Result.Value);
+			}
+		}
+
+		/// <summary>
+		/// Determines if the user is authorized to perform an action on a particular stream
+		/// </summary>
+		/// <param name="Stream">The stream to check</param>
+		/// <param name="Action">The action being performed</param>
+		/// <param name="User">The principal to authorize</param>
+		/// <param name="Cache">Cache of project permissions</param>
+		/// <returns>True if the action is authorized</returns>
+		public Task<bool> AuthorizeAsync(IStream Stream, AclAction Action, ClaimsPrincipal User, ProjectPermissionsCache? Cache)
+		{
+			return AuthorizeAsync(Stream.Acl, Stream.ProjectId, Action, User, Cache);
+		}
+
+		/// <summary>
+		/// Determines if the user is authorized to perform an action on a particular stream
+		/// </summary>
+		/// <param name="Stream">The stream to check</param>
+		/// <param name="Template">Template within the stream to check</param>
+		/// <param name="Action">The action being performed</param>
+		/// <param name="User">The principal to authorize</param>
+		/// <param name="Cache">Cache of project permissions</param>
+		/// <returns>True if the action is authorized</returns>
+		public Task<bool> AuthorizeAsync(IStream Stream, TemplateRef Template, AclAction Action, ClaimsPrincipal User, ProjectPermissionsCache? Cache)
+		{
+			bool? Result = Template.Acl?.Authorize(Action, User);
+			if (Result == null)
+			{
+				return AuthorizeAsync(Stream, Action, User, Cache);
+			}
+			else
+			{
+				return Task.FromResult(Result.Value);
+			}
+		}
+
+		/// <summary>
+		/// Determines if the user is authorized to perform an action on a particular project
+		/// </summary>
+		/// <param name="StreamId">The stream id to check</param>
+		/// <param name="Action">The action being performed</param>
+		/// <param name="User">The principal to authorize</param>
+		/// <param name="Cache">Cache of stream permissions</param>
+		/// <returns>True if the action is authorized</returns>
+		public async Task<bool> AuthorizeAsync(StreamId StreamId, AclAction Action, ClaimsPrincipal User, StreamPermissionsCache? Cache)
+		{
+			IStreamPermissions? Permissions;
+			if (Cache == null)
+			{
+				Permissions = await GetStreamPermissionsAsync(StreamId);
+			}
+			else if (!Cache.Streams.TryGetValue(StreamId, out Permissions))
+			{
+				Permissions = await GetStreamPermissionsAsync(StreamId);
+				Cache.Streams.Add(StreamId, Permissions);
+			}
+			return Permissions != null && await AuthorizeAsync(Permissions.Acl, Permissions.ProjectId, Action, User, Cache);
+		}
+	}
+}
diff --git a/Engine/Source/Programs/Horde/HordeServer/Services/TelemetryService.cs b/Engine/Source/Programs/Horde/HordeServer/Services/TelemetryService.cs
new file mode 100644
--- /dev/null
+++ b/Engine/Source/Programs/Horde/HordeServer/Services/TelemetryService.cs
@@ -0,0 +1,140 @@
+using Google.Protobuf.WellKnownTypes;
+using HordeCommon.Rpc.Tasks;
+using HordeServer.Collections;
+using HordeServer.Models;
+using HordeServer.Utilities;
+using Microsoft.Extensions.Logging;
+using MongoDB.Bson;
+using System;
+using System.Collections.Generic;
+using System.Linq;
+using System.Threading;
+using System.Threading.Tasks;
+
+using PoolId = HordeServer.Utilities.StringId<HordeServer.Models.IPool>;
+
+namespace HordeServer.Services
+{
+	/// <summary>
+	/// Service which updates telemetry periodically
+	/// </summary>
+	public class TelemetryService : ElectedBackgroundService
+	{
+		/// <summary>
+		/// Collection of telemetry documents
+		/// </summary>
+		ITelemetryCollection TelemetryCollection;
+
+		/// <summary>
+		/// Collection of agent documents
+		/// </summary>
+		IAgentCollection AgentCollection;
+
+		/// <summary>
+		/// Collection of lease documents
+		/// </summary>
+		ILeaseCollection LeaseCollection;
+
+		/// <summary>
+		/// Collection of pool documents
+		/// </summary>
+		IPoolCollection PoolCollection;
+
+		/// <summary>
+		/// Logger for diagnostic messages
+		/// </summary>
+		ILogger Logger;
+
+		/// <summary>
+		/// Constructor
+		/// </summary>
+		/// <param name="DatabaseService"></param>
+		/// <param name="TelemetryCollection"></param>
+		/// <param name="AgentCollection"></param>
+		/// <param name="LeaseCollection"></param>
+		/// <param name="PoolCollection"></param>
+		/// <param name="Logger"></param>
+		public TelemetryService(DatabaseService DatabaseService, ITelemetryCollection TelemetryCollection, IAgentCollection AgentCollection, ILeaseCollection LeaseCollection, IPoolCollection PoolCollection, ILogger<TelemetryService> Logger)
+			: base(DatabaseService, new ObjectId("600082cb9b822d7afee6f0d1"), Logger)
+		{
+			this.TelemetryCollection = TelemetryCollection;
+			this.AgentCollection = AgentCollection;
+			this.LeaseCollection = LeaseCollection;
+			this.PoolCollection = PoolCollection;
+			this.Logger = Logger;
+		}
+
+		/// <inheritdoc/>
+		protected override async Task<DateTime> TickLeaderAsync(CancellationToken StoppingToken)
+		{
+			DateTime CurrentTime = DateTime.UtcNow;
+			DateTime NextTickTime = CurrentTime + TimeSpan.FromMinutes(10.0);
+
+			// Find the last time that we need telemetry for
+			DateTime MaxTime = CurrentTime.Date + TimeSpan.FromHours(CurrentTime.Hour);
+
+			// Get the latest telemetry data
+			IUtilizationTelemetry? Latest = await TelemetryCollection.GetLatestUtilizationTelemetryAsync();
+			TimeSpan Interval = TimeSpan.FromHours(1.0);
+			int Count = (Latest == null) ? (7 * 24) : (int)Math.Round((MaxTime - Latest.FinishTime) / Interval);
+			DateTime MinTime = MaxTime - Count * Interval;
+
+			// Query all the current data
+			List<IAgent> Agents = await AgentCollection.FindAsync();
+			List<IPool> Pools = await PoolCollection.GetAsync();
+			List<ILease> Leases = await LeaseCollection.FindLeasesAsync(MinTime: MinTime);
+
+			// Find all the agents
+			Dictionary<AgentId, List<IPool>> AgentToPools = Agents.ToDictionary(x => x.Id, x => x.GetPools(Pools).ToList());
+
+			// Generate all the telemetry data
+			DateTime BucketMinTime = MinTime;
+			for (int Idx = 0; Idx < Count; Idx++)
+			{
+				DateTime BucketMaxTime = BucketMinTime + Interval;
+				Logger.LogInformation("Creating telemetry for {MinTime} to {MaxTime}", BucketMinTime, BucketMaxTime);
+
+				NewUtilizationTelemetry Telemetry = new NewUtilizationTelemetry(BucketMinTime, BucketMaxTime);
+				Telemetry.NumAgents = Agents.Count;
+				foreach (IPool Pool in AgentToPools.Values.SelectMany(x => x))
+				{
+					Telemetry.FindOrAddPool(Pool.Id).NumAgents++;
+				}
+				foreach (ILease Lease in Leases)
+				{
+					if (Lease.StartTime < BucketMaxTime && Lease.FinishTime >= BucketMinTime)
+					{
+						List<IPool>? LeasePools;
+						if (AgentToPools.TryGetValue(Lease.AgentId, out LeasePools))
+						{
+							DateTime FinishTime = Lease.FinishTime ?? BucketMaxTime;
+							double Time = new TimeSpan(Math.Min(FinishTime.Ticks, BucketMaxTime.Ticks) - Math.Max(Lease.StartTime.Ticks, BucketMinTime.Ticks)).TotalHours;
+
+							foreach (IPool Pool in LeasePools)
+							{
+								NewPoolUtilizationTelemetry PoolTelemetry = Telemetry.FindOrAddPool(Pool.Id);
+								if (Lease.PoolId == null || Lease.StreamId == null)
+								{
+									PoolTelemetry.AdminTime += Time;
+								}
+								else if (Pool.Id == Lease.PoolId)
+								{
+									PoolTelemetry.FindOrAddStream(Lease.StreamId.Value).Time += Time;
+								}
+								else
+								{
+									PoolTelemetry.OtherTime += Time;
+								}
+							}
+						}
+					}
+				}
+				await TelemetryCollection.AddUtilizationTelemetryAsync(Telemetry);
+
+				BucketMinTime = BucketMaxTime;
+			}
+
+			return NextTickTime;
+		}
+	}
+}
diff --git a/Engine/Source/Programs/Horde/HordeServer/Services/UpgradeService.cs b/Engine/Source/Programs/Horde/HordeServer/Services/UpgradeService.cs
new file mode 100644
--- /dev/null
+++ b/Engine/Source/Programs/Horde/HordeServer/Services/UpgradeService.cs
@@ -0,0 +1,108 @@
+using HordeServer.Collections;
+using HordeServer.Collections.Impl;
+using HordeServer.Models;
+using Microsoft.Extensions.DependencyInjection;
+using Microsoft.Extensions.Logging;
+using MongoDB.Driver;
+using Serilog.Events;
+using System;
+using System.Collections.Generic;
+using System.Linq;
+using System.Threading.Tasks;
+
+namespace HordeServer.Services
+{
+	/// <summary>
+	/// Service to perform schema upgrades on the database
+	/// </summary>
+	public class UpgradeService
+	{
+		/// <summary>
+		/// The current schema version
+		/// </summary>
+		public const int LatestSchemaVersion = 5;
+
+		/// <summary>
+		/// The database service instance
+		/// </summary>
+		DatabaseService DatabaseService { get; set; }
+
+		/// <summary>
+		/// The DI service provider
+		/// </summary>
+		IServiceProvider ServiceProvider;
+
+		/// <summary>
+		/// Logger instance
+		/// </summary>
+		ILogger<UpgradeService> Logger;
+
+		/// <summary>
+		/// Constructor
+		/// </summary>
+		/// <param name="DatabaseService">The database service singleton</param>
+		/// <param name="ServiceProvider">The DI service provider</param>
+		/// <param name="Logger">Logger instance</param>
+		public UpgradeService(DatabaseService DatabaseService, IServiceProvider ServiceProvider, ILogger<UpgradeService> Logger)
+		{
+			this.DatabaseService = DatabaseService;
+			this.ServiceProvider = ServiceProvider;
+			this.Logger = Logger;
+		}
+
+		/// <summary>
+		/// Update the database schema
+		/// </summary>
+		/// <param name="FromVersion">The version to upgrade from. Automatically detected if this is null.</param>
+		/// <returns>Async task</returns>
+		public async Task UpgradeSchemaAsync(int? FromVersion)
+		{
+			Globals Globals = await DatabaseService.GetGlobalsAsync();
+			int SchemaVersion = FromVersion ?? Globals.SchemaVersion ?? 0;
+
+			while (SchemaVersion < LatestSchemaVersion)
+			{
+				Logger.LogInformation("Upgrading from schema version {Version}", SchemaVersion);
+
+				// Handle the current version
+				//				if (SchemaVersion == 0)
+				//				{
+				//					IssueCollection IssueCollection = new IssueCollection(DatabaseService);
+				//					IJobCollection JobCollection = ServiceProvider.GetService<IJobCollection>();
+				//					await IssueCollection.UpdateSchemaLogIdFields(JobCollection, Logger);
+				//				}
+				//				if(SchemaVersion == 2)
+				//				{
+				//					IssueCollection IssueCollection = new IssueCollection(DatabaseService);
+				//					await IssueCollection.FixResolvedIssuesAsync(Logger);
+				//				}
+
+				if (SchemaVersion == 3)
+				{
+					IJobCollection JobCollection = ServiceProvider.GetRequiredService<IJobCollection>();
+					await JobCollection.UpgradeDocumentsAsync();
+				}
+				if (SchemaVersion == 4)
+				{
+					UserCollectionV1 UserCollectionV1 = new UserCollectionV1(ServiceProvider.GetRequiredService<DatabaseService>());
+					UserCollectionV2 UserCollectionV2 = new UserCollectionV2(ServiceProvider.GetRequiredService<DatabaseService>(), ServiceProvider.GetRequiredService<ILogger<UserCollectionV2>>());
+					await UserCollectionV2.ResaveDocumentsAsync(UserCollectionV1);
+				}
+
+				// Increment the version number
+				SchemaVersion++;
+
+				// Try to update the current schema version number
+				while (Globals.SchemaVersion == null || Globals.SchemaVersion < SchemaVersion)
+				{
+					Globals.SchemaVersion = SchemaVersion;
+					if (await DatabaseService.TryUpdateSingletonAsync(Globals))
+					{
+						break;
+					}
+					Globals = await DatabaseService.GetGlobalsAsync();
+				}
+			}
+		}
+	}
+}
