Initial drop of Horde Server/Agent for UE5/Main. Still missing some Epic-specific functionality required to build/run correctly.

diff --git a/Engine/Source/Programs/Horde/HordeServer/Startup.cs b/Engine/Source/Programs/Horde/HordeServer/Startup.cs
new file mode 100644
--- /dev/null
+++ b/Engine/Source/Programs/Horde/HordeServer/Startup.cs
@@ -0,0 +1,709 @@
+// Copyright Epic Games, Inc. All Rights Reserved.
+
+using System;
+using System.Collections.Generic;
+using System.Diagnostics;
+using System.IdentityModel.Tokens.Jwt;
+using System.IO;
+using System.Linq;
+using System.Reflection;
+using System.Security.Claims;
+using System.Text.Encodings.Web;
+using System.Text.Json;
+using System.Text.Json.Serialization;
+using System.Threading.Tasks;
+using Serilog;
+using HordeServer.Models;
+using HordeServer.Services;
+using Microsoft.AspNetCore.Authentication;
+using Microsoft.AspNetCore.Authentication.Cookies;
+using Microsoft.AspNetCore.Authentication.JwtBearer;
+using Microsoft.AspNetCore.Authentication.OAuth.Claims;
+using Microsoft.AspNetCore.Authentication.OpenIdConnect;
+using Microsoft.AspNetCore.Authorization;
+using Microsoft.AspNetCore.Builder;
+using Microsoft.AspNetCore.Hosting;
+using Microsoft.Extensions.Configuration;
+using Microsoft.Extensions.DependencyInjection;
+using Microsoft.Extensions.Hosting;
+using Microsoft.Extensions.Logging;
+using Microsoft.Extensions.Options;
+using Microsoft.IdentityModel.Protocols.OpenIdConnect;
+using Microsoft.IdentityModel.Tokens;
+using Microsoft.OpenApi.Models;
+using MongoDB.Bson;
+using MongoDB.Bson.Serialization;
+using MongoDB.Bson.Serialization.Conventions;
+using System.Net.Http;
+using Microsoft.AspNetCore.HttpOverrides;
+using HordeServer.Utilities;
+using Grpc.Core.Interceptors;
+using Grpc.Core;
+using EpicGames.Core;
+using Microsoft.Extensions.Diagnostics.HealthChecks;
+using Microsoft.AspNetCore.Diagnostics.HealthChecks;
+using HordeServer.Collections;
+using Microsoft.AspNetCore.Http;
+using System.Diagnostics.CodeAnalysis;
+using System.Globalization;
+using System.Threading;
+using HordeCommon;
+using Microsoft.AspNetCore.DataProtection;
+using StackExchange.Redis;
+using HordeServer.Storage;
+using HordeServer.Storage.Impl;
+using HordeServer.Logs;
+using HordeServer.Logs.Builder;
+using HordeServer.Logs.Storage;
+using HordeServer.Logs.Readers;
+using HordeServer.Logs.Storage.Impl;
+using HordeServer.Collections.Impl;
+using HordeServer.Services.Impl;
+using Microsoft.Extensions.DependencyInjection.Extensions;
+using HordeServer.Authentication;
+using HordeServer.Rpc;
+using HordeServer.Storage.Services;
+using HordeServer.Tasks.Impl;
+using HordeServer.Tasks;
+using StatsdClient;
+using Status = Grpc.Core.Status;
+using HordeServer.Notifications.Impl;
+using HordeServer.Notifications;
+
+namespace HordeServer
+{
+	class Startup
+	{
+		class GrpcExceptionInterceptor : Interceptor
+		{
+			ILogger<GrpcExceptionInterceptor> Logger;
+
+			public GrpcExceptionInterceptor(ILogger<GrpcExceptionInterceptor> Logger)
+			{
+				this.Logger = Logger;
+			}
+
+			[SuppressMessage("Design", "CA1031:Do not catch general exception types", Justification = "<Pending>")]
+			IDisposable? RequestContext<TRequest>(TRequest Request)
+			{
+				try
+				{
+					return Logger.BeginScope("Request: {Request}", JsonSerializer.Serialize<TRequest>(Request));
+				}
+				catch
+				{
+					return null;
+				}
+			}
+
+			[SuppressMessage("Design", "CA1031:Do not catch general exception types", Justification = "<Pending>")]
+			IDisposable? ConnectionContext(ConnectionInfo Connection)
+			{
+				try
+				{
+					return Logger.BeginScope("Client: {ClientIp}", Connection.RemoteIpAddress);
+				}
+				catch
+				{
+					return null;
+				}
+			}
+
+			public override AsyncClientStreamingCall<TRequest, TResponse> AsyncClientStreamingCall<TRequest, TResponse>(ClientInterceptorContext<TRequest, TResponse> Context, AsyncClientStreamingCallContinuation<TRequest, TResponse> Continuation)
+			{
+				try
+				{
+					return base.AsyncClientStreamingCall(Context, Continuation);
+				}
+				catch (StructuredRpcException Ex)
+				{
+					Logger.LogError(Ex, Ex.Format, Ex.Args);
+					throw;
+				}
+				catch (Exception Ex)
+				{
+					Logger.LogError(Ex, "Exception in call to {Method}", Context.Method);
+					throw new RpcException(new Status(StatusCode.Internal, $"An exception was thrown on the server: {Ex}"));
+				}
+			}
+
+			public override AsyncDuplexStreamingCall<TRequest, TResponse> AsyncDuplexStreamingCall<TRequest, TResponse>(ClientInterceptorContext<TRequest, TResponse> Context, AsyncDuplexStreamingCallContinuation<TRequest, TResponse> Continuation)
+			{
+				try
+				{
+					return base.AsyncDuplexStreamingCall(Context, Continuation);
+				}
+				catch (StructuredRpcException Ex)
+				{
+					Logger.LogError(Ex, Ex.Format, Ex.Args);
+					throw;
+				}
+				catch (Exception Ex)
+				{
+					Logger.LogError(Ex, "Exception in call to {Method}", Context.Method);
+					throw new RpcException(new Status(StatusCode.Internal, $"An exception was thrown on the server: {Ex}"));
+				}
+			}
+
+			[SuppressMessage("Design", "CA1031:Do not catch general exception types", Justification = "<Pending>")]
+			public override AsyncServerStreamingCall<TResponse> AsyncServerStreamingCall<TRequest, TResponse>(TRequest Request, ClientInterceptorContext<TRequest, TResponse> Context, AsyncServerStreamingCallContinuation<TRequest, TResponse> Continuation)
+			{
+				try
+				{
+					return base.AsyncServerStreamingCall(Request, Context, Continuation);
+				}
+				catch (StructuredRpcException Ex)
+				{
+					using IDisposable? Scope = RequestContext(Request); 
+					Logger.LogError(Ex, Ex.Format, Ex.Args);
+					throw;
+				}
+				catch (Exception Ex)
+				{
+					using IDisposable? Scope = RequestContext(Request);
+					Logger.LogError(Ex, "Exception in call to {Method}", Context.Method);
+					throw new RpcException(new Status(StatusCode.Internal, $"An exception was thrown on the server: {Ex}"));
+				}
+			}
+
+			[SuppressMessage("Design", "CA1031:Do not catch general exception types", Justification = "<Pending>")]
+			public override async Task<TResponse> UnaryServerHandler<TRequest, TResponse>(TRequest Request, ServerCallContext Context, UnaryServerMethod<TRequest, TResponse> Continuation)
+			{
+				try
+				{
+					return await base.UnaryServerHandler(Request, Context, Continuation);
+				}
+				catch (StructuredRpcException Ex)
+				{
+					using IDisposable? ConnectionScope = ConnectionContext(Context.GetHttpContext().Connection);
+					using IDisposable? Scope = RequestContext(Request);
+					Logger.LogError(Ex, Ex.Format, Ex.Args);
+					throw;
+				}
+				catch (Exception Ex)
+				{
+					using IDisposable? ConnectionScope = ConnectionContext(Context.GetHttpContext().Connection);
+					using IDisposable? Scope = RequestContext(Request);
+					Logger.LogError(Ex, "Exception in call to {Method}", Context.Method);
+					throw new RpcException(new Status(StatusCode.Internal, $"An exception was thrown on the server: {Ex}"));
+				}
+			}
+		}
+
+		class BsonSerializationProvider : IBsonSerializationProvider
+		{
+			public IBsonSerializer? GetSerializer(Type Type)
+			{
+				if (Type == typeof(ContentHash))
+				{
+					return new ContentHashSerializer();
+				}
+				if (Type == typeof(DateTimeOffset))
+				{
+					return new DateTimeOffsetStringSerializer();
+				}
+				return null;
+			}
+		}
+
+		class JsonObjectIdConverter : JsonConverter<ObjectId>
+		{
+			public override ObjectId Read(ref Utf8JsonReader Reader, Type TypeToConvert, JsonSerializerOptions Options)
+			{
+				return Reader.GetString().ToObjectId();
+			}
+
+			public override void Write(Utf8JsonWriter Writer, ObjectId ObjectId, JsonSerializerOptions Options)
+			{
+				Writer.WriteStringValue(ObjectId.ToString());
+			}
+		}
+		
+		class JsonDateTimeConverter : JsonConverter<DateTime>
+		{
+			public override DateTime Read(ref Utf8JsonReader Reader, Type TypeToConvert, JsonSerializerOptions Options)
+			{
+				Debug.Assert(TypeToConvert == typeof(DateTime));
+				return DateTime.Parse(Reader.GetString(), CultureInfo.CurrentCulture);
+			}
+
+			public override void Write(Utf8JsonWriter Writer, DateTime DateTime, JsonSerializerOptions Options)
+			{
+				Writer.WriteStringValue(DateTime.ToUniversalTime().ToString("yyyy'-'MM'-'dd'T'HH':'mm':'ssZ", CultureInfo.CurrentCulture));
+			}
+		}
+
+		public Startup(IConfiguration Configuration)
+		{
+			this.Configuration = Configuration;
+		}
+
+		public IConfiguration Configuration { get; }
+
+		ServerSettings Settings { get; set; } = new ServerSettings();
+
+		private ConnectionMultiplexer? RedisConMux = null;
+
+		// This method gets called by the runtime. Use this method to add services to the container.
+		public void ConfigureServices(IServiceCollection Services)
+		{
+			IConfigurationSection ConfigSection = Configuration.GetSection("Horde");
+			Services.AddOptions<ServerSettings>().Configure(Options => ConfigSection.Bind(Options)).ValidateDataAnnotations();
+
+			ConfigSection.Bind(Settings);
+
+			if (Settings.GlobalThreadPoolMinSize != null)
+			{
+				// Min thread pool size is set to combat timeouts seen with the Redis client.
+				// See comments for <see cref="ServerSettings.GlobalThreadPoolMinSize" /> and
+				// [REDACTED_PROFILE_URL]/StackExchange.Redis/issues/1680
+				int Min = Settings.GlobalThreadPoolMinSize.Value;
+				ThreadPool.SetMinThreads(Min, Min);
+			}
+
+			if (String.IsNullOrEmpty(Settings.RedisConnectionConfig))
+			{
+				// FIXME: Debug why [Required] annotation is not working in ServerSettings.
+				// This is just a crude fix to avoid a cryptic startup message.
+				throw new Exception("Redis not configured in settings!");
+			}
+			
+			// The Connect() is sync as AddSingleton() does not support async operations.
+			// This *should* be fine as this is only established once during startup (reconnects happen async)
+			RedisConMux = ConnectionMultiplexer.Connect(Settings.RedisConnectionConfig);
+			Services.AddSingleton<IDatabase>(Sp => RedisConMux.GetDatabase());
+			Services.AddSingleton(new RedisConnectionPool(20, Settings.RedisConnectionConfig));
+			
+			Services.AddDataProtection().PersistKeysToStackExchangeRedis(RedisConMux, "aspnet-data-protection");
+
+			if (Settings.CorsEnabled)
+			{
+				Services.AddCors(Options =>
+				{
+					Options.AddPolicy("CorsPolicy",
+						Builder => Builder.WithOrigins(Settings.CorsOrigin.Split(";"))
+						.AllowAnyMethod()
+						.AllowAnyHeader()
+						.AllowCredentials());
+				});
+			}
+
+			Services.AddGrpc(Options =>
+			{
+				Options.EnableDetailedErrors = true;
+				Options.Interceptors.Add(typeof(LifetimeGrpcInterceptor));
+				Options.Interceptors.Add(typeof(GrpcExceptionInterceptor));
+			});
+
+			Services.AddSingleton<IAgentCollection, AgentCollection>();
+			Services.AddSingleton<IAgentSoftwareCollection, AgentSoftwareCollection>();
+			Services.AddSingleton<ICounterCollection, CounterCollection>();
+			Services.AddSingleton<IGraphCollection, GraphCollection>();
+			Services.AddSingleton<IIssueCollection, IssueCollection>();
+			Services.AddSingleton<IJobCollection, JobCollection>();
+			Services.AddSingleton<IJobStepRefCollection, JobStepRefCollection>();
+			Services.AddSingleton<IJobTimingCollection, JobTimingCollection>();
+			Services.AddSingleton<ILeaseCollection, LeaseCollection>();
+			Services.AddSingleton<ILogEventCollection, LogEventCollection>();
+			Services.AddSingleton<ILogFileCollection, LogFileCollection>();
+			Services.AddSingleton<INotificationTriggerCollection, NotificationTriggerCollection>();
+			Services.AddSingleton<IPoolCollection, PoolCollection>();
+			Services.AddSingleton<IProjectCollection, ProjectCollection>();
+			Services.AddSingleton<ISessionCollection, SessionCollection>();
+			Services.AddSingleton<IServiceAccountCollection, ServiceAccountCollection>();
+			Services.AddSingleton<ISubscriptionCollection, SubscriptionCollection>();
+			Services.AddSingleton<ISoftwareCollection, SoftwareCollection>();
+			Services.AddSingleton<IStreamCollection, StreamCollection>();
+			Services.AddSingleton<ITemplateCollection, TemplateCollection>();
+			Services.AddSingleton<ITestDataCollection, TestDataCollection>();
+			Services.AddSingleton<ITelemetryCollection, TelemetryCollection>();
+			Services.AddSingleton<IUgsMetadataCollection, UgsMetadataCollection>();
+			Services.AddSingleton<IUserCollection, UserCollectionV2>();
+
+			Services.AddSingleton(typeof(ISingletonDocument<>), typeof(SingletonDocument<>));
+
+			Services.AddSingleton<AutoscaleService>();
+			switch (Settings.FleetManager)
+			{
+				case FleetManagerType.Aws:
+					Services.AddSingleton<IFleetManager, AwsFleetManager>();
+					break;
+				default:
+					Services.AddSingleton<IFleetManager, DefaultFleetManager>();
+					break;
+			}
+
+			Services.AddSingleton<AclService>();
+			Services.AddSingleton<AgentService>();
+			Services.AddSingleton<AgentSoftwareService>();
+			Services.AddSingleton<ArtifactService>();
+			Services.AddSingleton<ConsistencyService>();
+			Services.AddSingleton<RequestTrackerService>();
+			Services.AddSingleton<CredentialService>();
+			Services.AddSingleton<DatabaseService>();
+			Services.AddSingleton<IDogStatsd>(ctx =>
+			{
+				string? DatadogAgentHost = Environment.GetEnvironmentVariable("DD_AGENT_HOST");
+				if (DatadogAgentHost != null)
+				{
+					// Datadog agent is configured, enable DogStatsD for metric collection
+					StatsdConfig Config = new StatsdConfig
+					{
+						StatsdServerName = DatadogAgentHost,
+						StatsdPort = 8125,
+					};
+					
+					DogStatsdService DogStatsdService = new DogStatsdService();
+					DogStatsdService.Configure(Config);
+					return DogStatsdService;
+				}
+				return new NoOpDogStatsd();
+			});
+			Services.AddSingleton<IClock, Clock>();
+			Services.AddSingleton<IDowntimeService, DowntimeService>();
+			Services.AddSingleton<IIssueService, IssueService>();
+			Services.AddSingleton<JobService>();
+			Services.AddSingleton<LifetimeService>();
+			Services.AddSingleton<ILogFileService, LogFileService>();
+			Services.AddSingleton<INotificationService, NotificationService>();			
+
+			if (String.IsNullOrEmpty(Settings.PerforceBridge))
+			{
+				Services.AddSingleton<IPerforceService, LocalPerforceService>();
+			}
+			else
+			{
+				Services.AddSingleton<IPerforceService, BridgePerforceService>(SP => new BridgePerforceService(Settings.PerforceBridge!, SP.GetRequiredService<ILogger<BridgePerforceService>>()));
+			}
+
+			Services.AddSingleton<PerforceLoadBalancer>();
+			Services.AddSingleton<PoolService>();
+			Services.AddSingleton<ProjectService>();
+			Services.AddSingleton<ScheduleService>();
+			Services.AddSingleton<SlackNotificationSink>();
+			Services.AddSingleton<INotificationSink, SlackNotificationSink>(SP => SP.GetRequiredService<SlackNotificationSink>());
+			Services.AddSingleton<SoftwareService>();
+			Services.AddSingleton<StreamService>();
+			Services.AddSingleton<TemplateService>();
+			Services.AddSingleton<UpgradeService>();
+
+			Services.AddSingleton<ActionCacheService>();
+			Services.AddSingleton<ByteStreamService>();
+			Services.AddSingleton<CapabilitiesService>();
+			Services.AddSingleton<ContentStorageService>();
+			Services.AddSingleton<ExecutionService>();
+
+			Services.AddSingleton<FileSystemStorageBackend>();
+			Services.AddSingleton<IStorageService, SimpleStorageService>(SP => new SimpleStorageService(SP.GetRequiredService<FileSystemStorageBackend>()));
+
+			ConfigureStorageProvider(Services, Settings);
+			ConfigureLogStorage(Services);
+			//			ConfigureLogFileWriteCache(Services, Settings);
+
+			if (Settings.DisableAuth)
+			{
+				var Scheme = AnonymousAuthenticationHandler.AuthenticationScheme;
+				Services.AddAuthentication(Scheme)
+					.AddScheme<AnonymousAuthenticationOptions, AnonymousAuthenticationHandler>(Scheme, Options =>
+					{
+						Options.AdminClaimType = Settings.AdminClaimType;
+						Options.AdminClaimValue = Settings.AdminClaimValue;
+					});
+			}
+			else
+			{
+				AuthenticationBuilder AuthBuilder = Services.AddAuthentication(Options =>
+					{
+						// If an authentication cookie is present, use it to get authentication information
+						Options.DefaultAuthenticateScheme = CookieAuthenticationDefaults.AuthenticationScheme;
+						Options.DefaultSignInScheme = CookieAuthenticationDefaults.AuthenticationScheme;
+
+						// If authentication is required, and no cookie is present, use OIDC to sign in
+						Options.DefaultChallengeScheme = OktaDefaults.AuthenticationScheme;
+					})
+					.AddCookie(Options =>
+					{
+						Options.Events.OnValidatePrincipal = Context =>
+						{
+							if (Context.Principal.FindFirst(HordeClaimTypes.UserId) == null)
+							{
+								Context.RejectPrincipal();
+							}
+							return Task.CompletedTask;
+						};
+
+						Options.Events.OnRedirectToAccessDenied = Context => 
+						{ 
+							Context.Response.StatusCode = StatusCodes.Status403Forbidden; 
+							return Context.Response.CompleteAsync(); 
+						};
+					});
+
+				List<string> AuthSchemes = new List<string>();
+				AuthSchemes.Add(CookieAuthenticationDefaults.AuthenticationScheme);
+
+				AuthBuilder.AddScheme<ServiceAccountAuthOptions, ServiceAccountAuthHandler>(
+					ServiceAccountAuthHandler.AuthenticationScheme, Options => { });
+				AuthSchemes.Add(ServiceAccountAuthHandler.AuthenticationScheme);
+
+				if (Settings.OidcClientId != null && Settings.OidcAuthority != null)
+				{
+					AuthBuilder.AddOkta(OktaDefaults.AuthenticationScheme, OpenIdConnectDefaults.DisplayName, 
+						Options =>
+						{
+							Options.Authority = Settings.OidcAuthority;
+							Options.ClientId = Settings.OidcClientId;
+
+							if (!String.IsNullOrEmpty(Settings.OidcSigninRedirect))
+							{
+								Options.Events = new OpenIdConnectEvents
+								{
+									OnRedirectToIdentityProvider = async RedirectContext =>
+									{
+										RedirectContext.ProtocolMessage.RedirectUri = Settings.OidcSigninRedirect;
+										await Task.CompletedTask;
+									}
+								};
+							}
+						});
+					AuthSchemes.Add(OktaDefaults.AuthenticationScheme);
+				}
+
+				AuthBuilder.AddScheme<JwtBearerOptions, HordeJwtBearerHandler>(HordeJwtBearerHandler.AuthenticationScheme, Options => { });
+				AuthSchemes.Add(HordeJwtBearerHandler.AuthenticationScheme);
+
+				Services.AddAuthorization(Options =>
+					{
+						Options.DefaultPolicy = new AuthorizationPolicyBuilder(AuthSchemes.ToArray())
+							.RequireAuthenticatedUser()							
+							.Build();
+					});
+			}
+
+			Services.AddHostedService(Provider => Provider.GetRequiredService<AgentService>());
+			Services.AddHostedService(Provider => Provider.GetRequiredService<AutoscaleService>());
+			Services.AddHostedService(Provider => Provider.GetRequiredService<ConsistencyService>());
+			Services.AddHostedService(Provider => (DowntimeService)Provider.GetRequiredService<IDowntimeService>());
+			Services.AddHostedService(Provider => Provider.GetRequiredService<JobTaskSource>());
+			Services.AddHostedService(Provider => Provider.GetRequiredService<IIssueService>());
+			Services.AddHostedService(Provider => (LogFileService)Provider.GetRequiredService<ILogFileService>());
+			Services.AddHostedService(Provider => (NotificationService)Provider.GetRequiredService<INotificationService>());
+			if (!Settings.DisableSchedules)
+			{
+				Services.AddHostedService(Provider => Provider.GetRequiredService<ScheduleService>());
+			}
+
+			Services.AddHostedService<MetricService>();
+			Services.AddHostedService(Provider => Provider.GetRequiredService<PerforceLoadBalancer>());
+			Services.AddHostedService<PoolUpdateService>();
+			Services.AddHostedService(Provider => Provider.GetRequiredService<SlackNotificationSink>());
+			Services.AddHostedService<StreamConfigService>();
+			Services.AddHostedService<TelemetryService>();
+
+			// Task sources. Order of registration is important here; it dictates the order in which sources are served.
+			Services.AddSingleton<JobTaskSource>();
+			Services.AddSingleton<ActionTaskSource>();
+			Services.AddSingleton<ConformTaskSource>();
+
+			Services.AddSingleton<ITaskSource, UpgradeTaskSource>();
+			Services.AddSingleton<ITaskSource, ShutdownTaskSource>();
+			Services.AddSingleton<ITaskSource, RestartTaskSource>();
+			Services.AddSingleton<ITaskSource, ConformTaskSource>(Provider => Provider.GetRequiredService<ConformTaskSource>());
+			Services.AddSingleton<ITaskSource, JobTaskSource>(Provider => Provider.GetRequiredService<JobTaskSource>());
+			Services.AddSingleton<ITaskSource, ActionTaskSource>(Provider => Provider.GetRequiredService<ActionTaskSource>());
+
+			// Allow longer to shutdown so we can debug missing cancellation tokens
+			Services.Configure<HostOptions>(Options =>
+			{
+				Options.ShutdownTimeout = TimeSpan.FromSeconds(30.0);
+			});
+
+			// Allow forwarded headers
+			Services.Configure<ForwardedHeadersOptions>(Options =>
+			{
+				Options.ForwardedHeaders = ForwardedHeaders.XForwardedFor | ForwardedHeaders.XForwardedProto;
+				Options.KnownProxies.Clear();
+				Options.KnownNetworks.Clear();
+			});
+
+			Services.AddMvc().AddJsonOptions(Options => ConfigureJsonSerializer(Options.JsonSerializerOptions));
+
+			Services.AddControllers();
+
+			Services.AddSwaggerGen(Config =>
+			{
+				Config.SwaggerDoc("v1", new OpenApiInfo { Title = "Horde Server API", Version = "v1" });
+				Config.IncludeXmlComments(Path.Combine(AppContext.BaseDirectory, $"{Assembly.GetExecutingAssembly().GetName().Name}.xml"));
+			});
+
+			ConfigureMongoDbClient();
+
+			OnAddHealthChecks(Services);
+		}
+
+		public static void ConfigureJsonSerializer(JsonSerializerOptions Options)
+		{
+			Options.IgnoreNullValues = true;
+			Options.PropertyNamingPolicy = JsonNamingPolicy.CamelCase;
+			Options.PropertyNameCaseInsensitive = true;
+			Options.Converters.Add(new JsonObjectIdConverter());
+			Options.Converters.Add(new JsonStringEnumConverter());
+			Options.Converters.Add(new JsonKnownTypesConverterFactory());
+			Options.Converters.Add(new JsonStringIdConverterFactory());
+			Options.Converters.Add(new JsonDateTimeConverter());
+		}
+
+		private static void ConfigureStorageProvider(IServiceCollection Services, ServerSettings Settings)
+		{
+			switch (Settings.ExternalStorageProviderType)
+			{
+				case StorageProviderType.FileSystem:
+					Services.AddSingleton<IStorageBackend, FileSystemStorageBackend>();
+					break;
+				case StorageProviderType.S3:
+					Services.AddSingleton<IStorageBackend, AwsStorageBackend>();
+					break;
+				case StorageProviderType.Transient:
+					Services.AddSingleton<IStorageBackend, TransientStorageBackend>();
+					break;
+				case StorageProviderType.Relay:
+					Services.AddSingleton<IStorageBackend, RelayStorageBackend>();
+					break;
+				default:
+					throw new NotImplementedException();
+			}
+		}
+
+		private static void ConfigureLogStorage(IServiceCollection Services)
+		{
+			Services.AddSingleton<ILogBuilder>(Provider =>
+			{
+				RedisConnectionPool? RedisConnectionPool = Provider.GetService<RedisConnectionPool>();
+				if(RedisConnectionPool == null)
+				{
+					return new LocalLogBuilder();
+				}
+				else
+				{
+					return new RedisLogBuilder(RedisConnectionPool, Provider.GetRequiredService<ILogger<RedisLogBuilder>>());
+				}
+			});
+			Services.AddSingleton<ILogStorage>(Provider =>
+			{
+				ILogStorage Storage = new PersistentLogStorage(Provider.GetRequiredService<IStorageBackend>(), Provider.GetRequiredService<ILogger<PersistentLogStorage>>());
+
+//				IDatabase? RedisDb = Provider.GetService<IDatabase>();
+//				if (RedisDb != null)
+//				{
+//					Storage = new RedisLogStorage(RedisDb, Provider.GetService<ILogger<RedisLogStorage>>(), Storage);
+//				}
+
+				Storage = new SequencedLogStorage(Storage);
+				Storage = new LocalLogStorage(50, Storage);
+				return Storage;
+			});
+		}
+/*
+		private static void ConfigureLogFileWriteCache(IServiceCollection Services, ServerSettings Settings)
+		{
+			bool RedisConfigured = !String.IsNullOrEmpty(Settings.RedisConnectionConfig);
+			string CacheType = Settings.LogServiceWriteCacheType.ToLower(CultureInfo.CurrentCulture);
+
+			if (CacheType == "inmemory")
+			{
+				Services.AddSingleton<ILogFileWriteCache2>(Sp => new InMemoryLogFileWriteCache2());
+			}
+			else if (CacheType == "redis" && RedisConfigured)
+			{
+				Services.AddSingleton<ILogFileWriteCache2>(Sp => new RedisLogFileWriteCache2(Sp.GetService<ILogger<RedisLogFileWriteCache>>(), Sp.GetService<IDatabase>()));
+			}
+			else if (CacheType == "redis" && !RedisConfigured)
+			{
+				throw new Exception("Redis must be configured to use the Redis-backed log write cache");
+			}
+			else
+			{
+				throw new Exception("Unknown value set for LogServiceWriteCacheType in config: " + Settings.LogServiceWriteCacheType);
+			}
+		}
+*/
+		public static void ConfigureMongoDbClient()
+		{
+			// Ignore extra elements on deserialized documents
+			ConventionPack ConventionPack = new ConventionPack();
+			ConventionPack.Add(new IgnoreExtraElementsConvention(true));
+			ConventionPack.Add(new EnumRepresentationConvention(BsonType.String));
+			ConventionRegistry.Register("Horde", ConventionPack, Type => true);
+
+			// Register the custom serializers
+			BsonSerializer.RegisterSerializationProvider(new BsonSerializationProvider());
+			BsonSerializer.RegisterSerializationProvider(new StringIdSerializationProvider());
+		}
+
+		private static void OnAddHealthChecks(IServiceCollection Services)
+		{
+			IHealthChecksBuilder HealthChecks = Services.AddHealthChecks().AddCheck("self", () => HealthCheckResult.Healthy(), tags: new[] { "self" });
+		}
+
+		// This method gets called by the runtime. Use this method to configure the HTTP request pipeline.
+		public void Configure(IApplicationBuilder App, IWebHostEnvironment Env)
+		{
+			App.UseForwardedHeaders();
+
+			if (Env.IsDevelopment())
+			{
+				App.UseDeveloperExceptionPage();
+			}
+
+			if (Settings.CorsEnabled)
+			{
+				App.UseCors("CorsPolicy");
+			}
+
+			// Enable middleware to serve generated Swagger as a JSON endpoint.
+			App.UseSwagger();
+
+			// Enable serilog request logging
+			App.UseSerilogRequestLogging();
+
+			// Include the source IP address with requests
+			App.Use(async (Context, Next) => {
+				using (Serilog.Context.LogContext.PushProperty("RemoteIP", Context.Connection.RemoteIpAddress))
+				{
+					await Next();
+				}
+			});
+
+			// Enable middleware to serve swagger-ui (HTML, JS, CSS, etc.),
+			// specifying the Swagger JSON endpoint.
+			App.UseSwaggerUI(c =>
+			{
+				c.SwaggerEndpoint("/swagger/v1/swagger.json", "Horde Server API");
+				c.RoutePrefix = "swagger";
+			});
+
+			if (!Env.IsDevelopment())
+			{
+				App.UseMiddleware<RequestTrackerMiddleware>();	
+			}
+			
+			App.UseRouting();
+
+			App.UseAuthentication();
+			App.UseAuthorization();
+
+			App.UseEndpoints(Endpoints =>
+			{
+				Endpoints.MapGrpcService<RpcService>();
+				Endpoints.MapGrpcService<ActionRpcService>();
+				Endpoints.MapGrpcService<ActionCacheService>();
+				Endpoints.MapGrpcService<ByteStreamService>();
+				Endpoints.MapGrpcService<CapabilitiesService>();
+				Endpoints.MapGrpcService<ContentStorageService>();
+				Endpoints.MapGrpcService<ExecutionService>();
+
+				Endpoints.MapControllers();
+			});
+		}
+	}
+}
diff --git a/Engine/Source/Programs/Horde/HordeServer/Storage/Backends/AwsStorageBackend.cs b/Engine/Source/Programs/Horde/HordeServer/Storage/Backends/AwsStorageBackend.cs
new file mode 100644
--- /dev/null
+++ b/Engine/Source/Programs/Horde/HordeServer/Storage/Backends/AwsStorageBackend.cs
@@ -0,0 +1,212 @@
+using Amazon.Extensions.NETCore.Setup;
+using Amazon.Runtime;
+using Amazon.S3;
+using Amazon.S3.Model;
+using EpicGames.Core;
+using HordeServer.Utilities;
+using Microsoft.Extensions.Configuration;
+using Microsoft.Extensions.Logging;
+using Microsoft.Extensions.Options;
+using System;
+using System.Collections.Generic;
+using System.Diagnostics.CodeAnalysis;
+using System.Globalization;
+using System.IO;
+using System.Linq;
+using System.Threading;
+using System.Threading.Tasks;
+
+namespace HordeServer.Storage.Impl
+{
+	/// <summary>
+	/// Exception wrapper for S3 requests
+	/// </summary>
+	[SuppressMessage("Design", "CA1032:Implement standard exception constructors", Justification = "<Pending>")]
+	public sealed class AwsException : Exception
+	{
+		/// <summary>
+		/// Constructor
+		/// </summary>
+		/// <param name="Message">Message for the exception</param>
+		/// <param name="InnerException">Inner exception data</param>
+		public AwsException(string? Message, Exception? InnerException)
+			: base(Message, InnerException)
+		{
+		}
+	}
+
+	/// <summary>
+	/// FileStorage implementation using an s3 bucket
+	/// </summary>
+	public sealed class AwsStorageBackend : IStorageBackend
+	{
+		/// <summary>
+		/// S3 Client
+		/// </summary>
+		private readonly IAmazonS3 Client;
+
+		/// <summary>
+		/// S3 bucket name
+		/// </summary>
+		private readonly string BucketName;
+
+		/// <summary>
+		/// Semaphore for connecting to AWS
+		/// </summary>
+		private SemaphoreSlim Semaphore;
+
+		/// <summary>
+		/// Logger interface
+		/// </summary>
+		private ILogger Logger;
+
+		/// <summary>
+		/// Constructor
+		/// </summary>
+		/// <param name="Settings">Current Horde Settings</param>
+		/// <param name="Configuration">Configuration</param>
+		/// <param name="Logger">Logger interface</param>
+		public AwsStorageBackend(IOptionsMonitor<ServerSettings> Settings, IConfiguration Configuration, ILogger<AwsStorageBackend> Logger)
+		{
+			ServerSettings CurrentSettings = Settings.CurrentValue;
+			AWSOptions Options = Configuration.GetAWSOptions();
+			Options.Credentials = GetCredentials(CurrentSettings);
+
+			this.Client = Options.CreateServiceClient<IAmazonS3>();
+			this.BucketName = CurrentSettings.S3LogBucketName;
+			this.Semaphore = new SemaphoreSlim(16);
+			this.Logger = Logger;
+		}
+
+		enum AWSCredentialsType
+		{
+			Basic,
+			AssumeRole,
+			AssumeRoleWebIdentity
+		}
+
+		/// <summary>
+		/// Gets auth credentials based on type
+		/// </summary>
+		/// <param name="Settings">Horde settings</param>
+		/// <returns></returns>
+		static AWSCredentials GetCredentials(ServerSettings Settings)
+		{
+			AWSCredentialsType AuthType;
+			if (!Enum.TryParse(Settings.S3CredentialType, true, out AuthType))
+			{
+				throw new Exception("Could not determine auth type from appsettings. Should be Basic, AssumeRole, or AssumeRoleWebIdentity!");
+			}
+
+			switch (AuthType)
+			{
+				case AWSCredentialsType.Basic:
+					return new BasicAWSCredentials(Settings.S3ClientKeyId, Settings.S3ClientSecret);
+				case AWSCredentialsType.AssumeRole:
+					return new AssumeRoleAWSCredentials(FallbackCredentialsFactory.GetCredentials(), Settings.S3AssumeArn, "Horde");
+				case AWSCredentialsType.AssumeRoleWebIdentity:
+					return AssumeRoleWithWebIdentityCredentials.FromEnvironmentVariables();
+				default:
+					throw new Exception($"Invalid authentication type: {AuthType}");
+			}
+		}
+
+		/// <inheritdoc/>
+		public void Dispose()
+		{
+			Client.Dispose();
+			Semaphore.Dispose();
+		}
+
+		/// <inheritdoc/>
+		public Task<bool> TouchAsync(string Path)
+		{
+			throw new NotImplementedException();
+		}
+
+		/// <inheritdoc/>
+		public async Task<ReadOnlyMemory<byte>?> ReadAsync(string Path)
+		{
+			using (await Semaphore.UseWaitAsync())
+			{
+				GetObjectRequest NewGetRequest = new GetObjectRequest();
+				NewGetRequest.BucketName = BucketName;
+				NewGetRequest.Key = Path;
+
+				try
+				{
+					using (GetObjectResponse GetResponse = await this.Client.GetObjectAsync(NewGetRequest))
+					{
+						using (Stream ResponseStream = GetResponse.ResponseStream)
+						{
+							using (MemoryStream OutputStream = new MemoryStream())
+							{
+								await ResponseStream.CopyToAsync(OutputStream);
+								return OutputStream.ToArray();
+							}
+						}
+					}
+				}
+				catch(Exception Ex)
+				{
+					throw new AwsException($"Unable to read from bucket {BucketName} path {Path}", Ex);
+				}
+			}
+		}
+
+		/// <inheritdoc/>
+		[SuppressMessage("Design", "CA1031:Do not catch general exception types", Justification = "<Pending>")]
+		public async Task WriteAsync(string Path, ReadOnlyMemory<byte> Data)
+		{
+			TimeSpan[] RetryTimes =
+			{
+				TimeSpan.FromSeconds(1.0),
+				TimeSpan.FromSeconds(5.0),
+				TimeSpan.FromSeconds(10.0),
+			};
+
+			for (int Attempt = 0; ; Attempt++)
+			{
+				try
+				{
+					await WriteInternalAsync(Path, Data);
+					Logger.LogDebug("Written data to {Path}", Path);
+					break;
+				}
+				catch (Exception Ex)
+				{
+					Logger.LogError(Ex, "Unable to write data to {Path} ({Attempt}/{AttemptCount})", Path, Attempt + 1, RetryTimes.Length + 1);
+					if(Attempt >= RetryTimes.Length)
+					{
+						throw new AwsException($"Unable to write to bucket {BucketName} path {Path}", Ex);
+					}
+				}
+
+				await Task.Delay(RetryTimes[Attempt]);
+			}
+		}
+
+		/// <inheritdoc/>
+		public async Task WriteInternalAsync(string Path, ReadOnlyMemory<byte> Data)
+		{
+			using (await Semaphore.UseWaitAsync())
+			{
+				using (ReadOnlyMemoryStream InputStream = new ReadOnlyMemoryStream(Data)) 
+				{
+					PutObjectRequest NewUploadRequest = new PutObjectRequest();
+					NewUploadRequest.BucketName = BucketName;
+					NewUploadRequest.Key = Path;
+					NewUploadRequest.InputStream = InputStream;
+					NewUploadRequest.Metadata.Add("bytes-written", Data.Length.ToString(CultureInfo.InvariantCulture));
+					await Client.PutObjectAsync(NewUploadRequest);
+				}
+			}
+		}
+
+		/// <inheritdoc/>
+		public Task DeleteAsync(string Path)
+		{
+			throw new NotImplementedException();
+		}
+	}
+}
diff --git a/Engine/Source/Programs/Horde/HordeServer/Storage/Backends/FileSystemStorageBackend.cs b/Engine/Source/Programs/Horde/HordeServer/Storage/Backends/FileSystemStorageBackend.cs
new file mode 100644
--- /dev/null
+++ b/Engine/Source/Programs/Horde/HordeServer/Storage/Backends/FileSystemStorageBackend.cs
@@ -0,0 +1,124 @@
+using EpicGames.Core;
+using Microsoft.Extensions.Options;
+using System;
+using System.Collections.Generic;
+using System.IO;
+using System.Linq;
+using System.Threading.Tasks;
+
+namespace HordeServer.Storage.Impl
+{
+	/// <summary>
+	/// IStorageProvider implementation which uses the file system
+	/// </summary>
+	public sealed class FileSystemStorageBackend : IStorageBackend
+	{
+		/// <summary>
+		/// Base directory for log files
+		/// </summary>
+		private readonly DirectoryReference BaseDir;
+
+		/// <summary>
+		/// Constructor
+		/// </summary>
+		/// <param name="Settings">Current Horde Settings</param>
+		public FileSystemStorageBackend(IOptionsMonitor<ServerSettings> Settings)
+		{
+			ServerSettings CurrentSettings = Settings.CurrentValue;
+			BaseDir = new DirectoryReference(CurrentSettings.LocalStorageDir);
+			DirectoryReference.CreateDirectory(BaseDir);
+		}
+
+		/// <inheritdoc/>
+		public void Dispose()
+		{
+		}
+
+		/// <summary>
+		/// Tests if a given path exists
+		/// </summary>
+		/// <param name="Path">Path to check</param>
+		/// <returns>True if the path exists</returns>
+		public bool Exists(string Path)
+		{
+			FileReference Location = FileReference.Combine(BaseDir, Path);
+			return FileReference.Exists(Location);
+		}
+
+		/// <summary>
+		/// Touches the given path
+		/// </summary>
+		/// <param name="Path"></param>
+		/// <returns></returns>
+		public Task<bool> TouchAsync(string Path)
+		{
+			FileReference Location = FileReference.Combine(BaseDir, Path);
+			if (FileReference.Exists(Location))
+			{
+				try
+				{
+					FileReference.SetLastWriteTimeUtc(Location, DateTime.UtcNow);
+					return Task.FromResult(true);
+				}
+				catch(FileNotFoundException)
+				{
+					// Does not exist
+				}
+			}
+			return Task.FromResult(false);
+		}
+
+		/// <inheritdoc/>
+		public async Task<ReadOnlyMemory<byte>?> ReadAsync(string Path)
+		{
+			FileReference Location = FileReference.Combine(BaseDir, Path);
+			if (!FileReference.Exists(Location))
+			{
+				return null;
+			}
+
+			try
+			{
+				return await FileReference.ReadAllBytesAsync(Location);
+			}
+			catch (DirectoryNotFoundException)
+			{
+				return null;
+			}
+			catch(FileNotFoundException)
+			{
+				return null;
+			}
+		}
+
+		/// <inheritdoc/>
+		public async Task WriteAsync(string Path, ReadOnlyMemory<byte> Data)
+		{
+			FileReference Location = FileReference.Combine(BaseDir, Path);
+			DirectoryReference.CreateDirectory(Location.Directory);
+
+			FileReference TempLocation = new FileReference(Location.FullName + ".tmp");
+			using (FileStream Stream = FileReference.Open(TempLocation, FileMode.Create, FileAccess.Write, FileShare.Read))
+			{
+				await Stream.WriteAsync(Data);
+			}
+
+			try
+			{
+				FileReference.Move(TempLocation, Location, true);
+			}
+			catch(IOException)
+			{
+				// Already exists
+			}
+		}
+
+		/// <inheritdoc/>
+		public Task DeleteAsync(string Path)
+		{
+			FileReference Location = FileReference.Combine(BaseDir, Path);
+			FileReference.Delete(Location);
+			return Task.CompletedTask;
+		}
+	}
+}
diff --git a/Engine/Source/Programs/Horde/HordeServer/Storage/Backends/RelayStorageBackend.cs b/Engine/Source/Programs/Horde/HordeServer/Storage/Backends/RelayStorageBackend.cs
new file mode 100644
--- /dev/null
+++ b/Engine/Source/Programs/Horde/HordeServer/Storage/Backends/RelayStorageBackend.cs
@@ -0,0 +1,105 @@
+using EpicGames.Core;
+using Microsoft.Extensions.Options;
+using System;
+using System.Collections.Generic;
+using System.IO;
+using System.Linq;
+using System.Net.Http;
+using System.Net.Http.Headers;
+using System.Threading.Tasks;
+
+namespace HordeServer.Storage.Impl
+{
+	/// <summary>
+	/// Implementation of ILogFileStorage which forwards requests to another server
+	/// </summary>
+	public sealed class RelayStorageBackend : IStorageBackend
+	{
+		/// <summary>
+		/// The client to connect with
+		/// </summary>
+		HttpClient Client;
+
+		/// <summary>
+		/// The base server URL
+		/// </summary>
+		Uri ServerUrl;
+
+		/// <summary>
+		/// Local storage provider
+		/// </summary>
+		FileSystemStorageBackend LocalStorage;
+
+		/// <summary>
+		/// Constructor
+		/// </summary>
+		/// <param name="Settings">Settings for the server instance</param>
+		public RelayStorageBackend(IOptionsMonitor<ServerSettings> Settings)
+		{
+			ServerSettings CurrentSettings = Settings.CurrentValue;
+			if (CurrentSettings.LogRelayServer == null)
+			{
+				throw new InvalidDataException("Missing LogRelayServer in server configuration");
+			}
+			if (CurrentSettings.LogRelayBearerToken == null)
+			{
+				throw new InvalidDataException("Missing LogRelayBearerToken in server configuration");
+			}
+
+			ServerUrl = new Uri(CurrentSettings.LogRelayServer);
+
+			Client = new HttpClient();
+			Client.DefaultRequestHeaders.Authorization = new AuthenticationHeaderValue("Bearer", CurrentSettings.LogRelayBearerToken);
+
+			LocalStorage = new FileSystemStorageBackend(Settings);
+		}
+
+		/// <inheritdoc/>
+		public void Dispose()
+		{
+			Client.Dispose();
+			LocalStorage.Dispose();
+		}
+
+		/// <inheritdoc/>
+		public Task<bool> TouchAsync(string Path)
+		{
+			throw new NotImplementedException();
+		}
+
+		/// <inheritdoc/>
+		public async Task<ReadOnlyMemory<byte>?> ReadAsync(string Path)
+		{
+			ReadOnlyMemory<byte>? LocalResult = await LocalStorage.ReadAsync(Path);
+			if (LocalResult != null)
+			{
+				return LocalResult;
+			}
+
+			Uri Url = new Uri(ServerUrl, $"api/v1/debug/storage?path={Path}");
+			using (HttpResponseMessage Response = await Client.GetAsync(Url))
+			{
+				if (Response.IsSuccessStatusCode)
+				{
+					byte[] ResponseData = await Response.Content.ReadAsByteArrayAsync();
+					await LocalStorage.WriteAsync(Path, ResponseData);
+					return ResponseData;
+				}
+			}
+
+			return null;
+		}
+
+		/// <inheritdoc/>
+		public Task WriteAsync(string Path, ReadOnlyMemory<byte> Data)
+		{
+			return LocalStorage.WriteAsync(Path, Data);
+		}
+
+		/// <inheritdoc/>
+		public Task DeleteAsync(string Path)
+		{
+			throw new NotImplementedException();
+		}
+	}
+}
diff --git a/Engine/Source/Programs/Horde/HordeServer/Storage/Backends/TransientStorageBackend.cs b/Engine/Source/Programs/Horde/HordeServer/Storage/Backends/TransientStorageBackend.cs
new file mode 100644
--- /dev/null
+++ b/Engine/Source/Programs/Horde/HordeServer/Storage/Backends/TransientStorageBackend.cs
@@ -0,0 +1,57 @@
+using System;
+using System.Collections.Generic;
+using System.Linq;
+using System.Threading.Tasks;
+
+namespace HordeServer.Storage.Impl
+{
+	/// <summary>
+	/// In-memory implementation of ILogFileStorage
+	/// </summary>
+	public sealed class TransientStorageBackend : IStorageBackend
+	{
+		/// <summary>
+		/// Data storage
+		/// </summary>
+		Dictionary<string, byte[]> PathToData = new Dictionary<string, byte[]>();
+
+		/// <inheritdoc/>
+		public void Dispose()
+		{
+		}
+
+		/// <inheritdoc/>
+		public Task<bool> TouchAsync(string Path)
+		{
+			return Task.FromResult(PathToData.ContainsKey(Path));
+		}
+
+		/// <inheritdoc/>
+		public Task<ReadOnlyMemory<byte>?> ReadAsync(string Path)
+		{
+			byte[]? Data;
+			if (PathToData.TryGetValue(Path, out Data))
+			{
+				return Task.FromResult<ReadOnlyMemory<byte>?>(Data);
+			}
+			else
+			{
+				return Task.FromResult<ReadOnlyMemory<byte>?>(null);
+			}
+		}
+
+		/// <inheritdoc/>
+		public Task WriteAsync(string Path, ReadOnlyMemory<byte> Data)
+		{
+			PathToData[Path] = Data.Span.ToArray();
+			return Task.CompletedTask;
+		}
+
+		/// <inheritdoc/>
+		public Task DeleteAsync(string Path)
+		{
+			PathToData.Remove(Path);
+			return Task.CompletedTask;
+		}
+	}
+}
diff --git a/Engine/Source/Programs/Horde/HordeServer/Storage/IStorageBackend.cs b/Engine/Source/Programs/Horde/HordeServer/Storage/IStorageBackend.cs
new file mode 100644
--- /dev/null
+++ b/Engine/Source/Programs/Horde/HordeServer/Storage/IStorageBackend.cs
@@ -0,0 +1,42 @@
+using System;
+using System.Collections.Generic;
+using System.Linq;
+using System.Threading.Tasks;
+
+namespace HordeServer.Storage
+{
+	/// <summary>
+	/// Interface for a log file storage implementation
+	/// </summary>
+	public interface IStorageBackend : IDisposable
+	{
+		/// <summary>
+		/// Updates the timestamp on the given file. Touching the file ensures that it will not be deleted in a "reasonable" timeframe.
+		/// </summary>
+		/// <param name="Path">Path to update</param>
+		/// <returns>True if the file exists, false otherwise</returns>
+		Task<bool> TouchAsync(string Path);
+
+		/// <summary>
+		/// Gets a chunk from a storage device
+		/// </summary>
+		/// <param name="Path">The path to read from</param>
+		/// <returns>The data to create from</returns>
+		Task<ReadOnlyMemory<byte>?> ReadAsync(string Path);
+
+		/// <summary>
+		/// Tries to update a log chunk after verifying it's the correct version
+		/// </summary>
+		/// <param name="Path">The path to write to</param>
+		/// <param name="Data">Buffer containing the data to write</param>
+		/// <returns>Async task</returns>
+		Task WriteAsync(string Path, ReadOnlyMemory<byte> Data);
+
+		/// <summary>
+		/// Deletes a file with the given path
+		/// </summary>
+		/// <param name="Path">Path of the file to delete</param>
+		/// <returns>Async task</returns>
+		Task DeleteAsync(string Path);
+	}
+}
diff --git a/Engine/Source/Programs/Horde/HordeServer/Storage/IStorageService.cs b/Engine/Source/Programs/Horde/HordeServer/Storage/IStorageService.cs
new file mode 100644
--- /dev/null
+++ b/Engine/Source/Programs/Horde/HordeServer/Storage/IStorageService.cs
@@ -0,0 +1,106 @@
+using HordeServer.Storage.Primitives;
+using System;
+using System.Collections.Generic;
+using System.Linq;
+using System.Threading.Tasks;
+
+namespace HordeServer.Storage
+{
+	/// <summary>
+	/// Combined interface for replicated content addressable storage and ref table.
+	/// </summary>
+	public interface IStorageService
+	{
+		/// <summary>
+		/// Attempts to retrieve a blob from the store. 
+		/// 
+		/// Note that this call may succeed even though the object is not in persistent storage; it may be retrieved
+		/// from a volatile cache. Call <see cref="ShouldPutBlobAsync(BlobHash)"/> to determine if a blob needs to
+		/// be put to the store.
+		/// 
+		/// This method takes a timestamp indicating the required replication consistency. The method may return null
+		/// if an object was put to the store after this time but it has not replicated down yet.
+		/// </summary>
+		/// <param name="Hash">Hash of the blob</param>
+		/// <param name="Deadline">Deadline for returning a result. If zero, the storage provider should use a default value.</param>
+		/// <returns>The blob data. May return ReadOnlyMemory<byte>.Empty if it does not exist.</byte></returns>
+		Task<ReadOnlyMemory<byte>?> TryGetBlobAsync(BlobHash Hash, DateTime Deadline = default(DateTime));
+
+		/// <summary>
+		/// Adds a blob to the store.
+		/// </summary>
+		/// <param name="Hash">Hash of the value. Used as a hint for possible sharding; the server will verify this hash and fail if it does not match.</param>
+		/// <param name="Value">The blob data</param>
+		/// <returns>Async task</returns>
+		Task PutBlobAsync(BlobHash Hash, ReadOnlyMemory<byte> Value);
+
+		/// <summary>
+		/// Determines if the blob should be written to the store. The server may return true even if it exists in a store
+		/// at its discretion (eg. perhaps the client -> store link is faster than the store -> parent store link, or a GC
+		/// operation is in progress)
+		/// </summary>
+		/// <param name="Hash">Hash of the blob to write</param>
+		/// <returns>True if the blob should be written to the store</returns>
+		Task<bool> ShouldPutBlobAsync(BlobHash Hash);
+
+		/// <summary>
+		/// Gets a reference from the storage layer
+		/// </summary>
+		/// <param name="Key">Key for the reference</param>
+		/// <param name="MaxDrift">Required replication consistency</param>
+		/// <returns>Hash of the result</returns>
+		Task<BlobHash?> GetRefAsync(BlobHash Key, TimeSpan MaxDrift = default(TimeSpan));
+
+		/// <summary>
+		/// Sets the value associated with a key. Writes are propagated as a "last writer wins" policy.
+		/// </summary>
+		/// <param name="Key">The key to set</param>
+		/// <param name="Value">Value for the key. If null, the key is deleted.</param>
+		/// <returns>Async task</returns>
+		Task SetRefAsync(BlobHash Key, BlobHash? Value);
+
+		/// <summary>
+		/// Updates the last modified timestamp of the given key. Possibly equivalent to doing a SET() on the same key/value pair, but does not 
+		/// need to force a value write to propagate up the tree. Mainly used for GC.
+		/// </summary>
+		/// <param name="Key">The key to touch</param>
+		/// <returns>Async task</returns>
+		Task TouchRefAsync(BlobHash Key);
+	}
+
+	/// <summary>
+	/// Extension methods for the storage service
+	/// </summary>
+	static class StorageServiceExtensions
+	{
+		/// <summary>
+		/// Gets a blob by hash, throwing an exception if it does not exist.
+		/// </summary>
+		/// <param name="StorageService">The storage provider</param>
+		/// <param name="Hash">Hash of the blob to retrieve</param>
+		/// <param name="Deadline">Deadline for retreving the data</param>
+		/// <returns>Memory for the blob</returns>
+		public static async Task<ReadOnlyMemory<byte>> GetBlobAsync(this IStorageService StorageService, BlobHash Hash, DateTime Deadline = default(DateTime))
+		{
+			ReadOnlyMemory<byte>? Blob = await StorageService.TryGetBlobAsync(Hash, Deadline);
+			if (Blob == null)
+			{
+				throw new KeyNotFoundException($"Unable to read key {Hash}");
+			}
+			return Blob.Value;
+		}
+
+		/// <summary>
+		/// Hashes a blob, then adds it to the storage service. If the hash is already known, prefer <see cref="IStorageService.PutBlobAsync(BlobHash, ReadOnlyMemory{byte})"/>
+		/// </summary>
+		/// <param name="StorageService">The storage service</param>
+		/// <param name="Data">The data to add</param>
+		/// <returns>Hash of the data that was added</returns>
+		public static async Task<BlobHash> PutBlobAsync(this IStorageService StorageService, ReadOnlyMemory<byte> Data)
+		{
+			BlobHash Hash = BlobHash.Compute(Data);
+			await StorageService.PutBlobAsync(Hash, Data);
+			return Hash;
+		}
+	}
+}
diff --git a/Engine/Source/Programs/Horde/HordeServer/Storage/Primitives/BlobHash.cs b/Engine/Source/Programs/Horde/HordeServer/Storage/Primitives/BlobHash.cs
new file mode 100644
--- /dev/null
+++ b/Engine/Source/Programs/Horde/HordeServer/Storage/Primitives/BlobHash.cs
@@ -0,0 +1,301 @@
+// Copyright Epic Games, Inc. All Rights Reserved.
+
+using EpicGames.Core;
+using System;
+using System.Collections;
+using System.Collections.Generic;
+using System.Diagnostics.CodeAnalysis;
+using System.Linq;
+using System.Security.Cryptography;
+using System.Threading.Tasks;
+
+namespace HordeServer.Storage.Primitives
+{
+	/// <summary>
+	/// Stores a reference to an immutable SHA256 hash value. Supports comparison and equality tests by value.
+	/// </summary>
+	public readonly struct BlobHash : IEquatable<BlobHash>, IComparable<BlobHash>
+	{
+		/// <summary>
+		/// Zero vlaue hash
+		/// </summary>
+		public static BlobHash Zero { get; } = new BlobHash(new byte[NumBytes]);
+
+		/// <summary>
+		/// Length of the hash in bytes
+		/// </summary>
+		public const int NumBytes = 32;
+
+		/// <summary>
+		/// Length of the hash in bits.
+		/// </summary>
+		public const int NumBits = NumBytes * 8;
+
+		/// <summary>
+		/// The memory backing this hash value
+		/// </summary>
+		public ReadOnlyMemory<byte> Memory { get; }
+
+		/// <summary>
+		/// Span for the memory
+		/// </summary>
+		public ReadOnlySpan<byte> Span => Memory.Span;
+
+		/// <summary>
+		/// Constructor
+		/// </summary>
+		/// <param name="Memory">Backing storage for the hash. Must be <see cref="NumBytes"/> long.</param>
+		public BlobHash(ReadOnlyMemory<byte> Memory)
+		{
+			if (Memory.Length != NumBytes)
+			{
+				throw new ArgumentException("Data is incorrect length");
+			}
+			this.Memory = Memory;
+		}
+
+		/// <summary>
+		/// Check if one hash starts with the same bits as another
+		/// </summary>
+		/// <param name="Other">The hash to test against</param>
+		/// <param name="Length">Number of bits to compare</param>
+		/// <returns>True if the hashes match, false otherwise</returns>
+		public bool StartsWith(BlobHash Other, int Length)
+		{
+			ReadOnlySpan<byte> ThisSpan = Span;
+			ReadOnlySpan<byte> OtherSpan = Other.Span;
+
+			int Idx = 0;
+			int RemainingLength = Length;
+
+			for (; RemainingLength > 8; RemainingLength -= 8, Idx++)
+			{
+				if (ThisSpan[Idx] != OtherSpan[Idx])
+				{
+					return false;
+				}
+			}
+
+			int Mask = ~((1 << (8 - RemainingLength)) - 1);
+			return (ThisSpan[Idx] & Mask) == (OtherSpan[Idx] & Mask);
+		}
+
+		/// <summary>
+		/// Mask off the left number of bits from the hash
+		/// </summary>
+		/// <param name="NumBits">Number of bits to include</param>
+		/// <returns>Masked hash value</returns>
+		public BlobHash MaskLeft(int NumBits)
+		{
+			ReadOnlySpan<byte> OldData = Span;
+			byte[] NewData = new byte[NumBytes];
+
+			int Idx = 0;
+			int RemainingBits = NumBits;
+			for (; RemainingBits > 8; RemainingBits -= 8, Idx++)
+			{
+				NewData[Idx] = OldData[Idx];
+			}
+
+			NewData[Idx] = (byte)(OldData[Idx] & ~((1 << (8 - RemainingBits)) - 1));
+			return new BlobHash(NewData);
+		}
+
+		/// <summary>
+		/// Compare two hashes for equality
+		/// </summary>
+		/// <param name="A">First hash</param>
+		/// <param name="B">Second hash</param>
+		/// <returns>True if the hashes are equal</returns>
+		public static bool operator ==(BlobHash A, BlobHash B)
+		{
+			return A.Span.SequenceEqual(B.Span);
+		}
+
+		/// <summary>
+		/// Compare two hashes for inequality
+		/// </summary>
+		/// <param name="A">First hash</param>
+		/// <param name="B">Second hash</param>
+		/// <returns>True if the hashes are not equal</returns>
+		public static bool operator !=(BlobHash A, BlobHash B)
+		{
+			return !(A == B);
+		}
+
+		/// <summary>
+		/// Tests whether A > B
+		/// </summary>
+		public static bool operator >(BlobHash A, BlobHash B)
+		{
+			return A.CompareTo(B) > 0;
+		}
+
+		/// <summary>
+		/// Tests whether A is less than B
+		/// </summary>
+		public static bool operator <(BlobHash A, BlobHash B)
+		{
+			return A.CompareTo(B) < 0;
+		}
+
+		/// <summary>
+		/// Tests whether A is greater than or equal to B
+		/// </summary>
+		public static bool operator >=(BlobHash A, BlobHash B)
+		{
+			return A.CompareTo(B) >= 0;
+		}
+
+		/// <summary>
+		/// Tests whether A is less than or equal to B
+		/// </summary>
+		public static bool operator <=(BlobHash A, BlobHash B)
+		{
+			return A.CompareTo(B) <= 0;
+		}
+
+		/// <summary>
+		/// Create a copy of this hash value, by duplicating the underlying memory
+		/// </summary>
+		/// <returns>Duplicated hash value</returns>
+		public BlobHash Clone()
+		{
+			return new BlobHash(Memory.ToArray());
+		}
+
+		/// <inheritdoc/>
+		public override bool Equals(object? Other)
+		{
+			if (Other is BlobHash)
+			{
+				return Equals((BlobHash)Other);
+			}
+			else
+			{
+				return false;
+			}
+		}
+
+		/// <inheritdoc cref="IEquatable{T}.Equals(T)"/>
+		public bool Equals(BlobHash Other)
+		{
+			return this == Other;
+		}
+
+		/// <inheritdoc cref="IComparable{T}.CompareTo(T)"/>
+		public int CompareTo(BlobHash Other)
+		{
+			ReadOnlySpan<byte> A = Span;
+			ReadOnlySpan<byte> B = Other.Span;
+
+			for (int Idx = 0; Idx < A.Length && Idx < B.Length; Idx++)
+			{
+				int Compare = A[Idx] - B[Idx];
+				if (Compare != 0)
+				{
+					return Compare;
+				}
+			}
+			return A.Length - B.Length;
+		}
+
+		/// <summary>
+		/// Parse a hash value from string
+		/// </summary>
+		/// <param name="Text"></param>
+		/// <returns></returns>
+		public static BlobHash Parse(string Text)
+		{
+			return new BlobHash(StringUtils.ParseHexString(Text));
+		}
+
+		/// <summary>
+		/// Compute a hash from a block of data
+		/// </summary>
+		/// <param name="Data"></param>
+		/// <returns></returns>
+		public static BlobHash Compute(ReadOnlyMemory<byte> Data)
+		{
+			return Compute(Data.ToArray());
+		}
+
+		/// <summary>
+		/// Compute a hash from a block of data
+		/// </summary>
+		/// <param name="Data"></param>
+		/// <returns></returns>
+		public static BlobHash Compute(byte[] Data)
+		{
+			using (SHA256 Algorithm = SHA256.Create())
+			{
+				byte[] Result = Algorithm.ComputeHash(Data);
+				return new BlobHash(Result);
+			}
+		}
+
+		/// <inheritdoc/>
+		public override int GetHashCode()
+		{
+			ReadOnlySpan<byte> LocalSpan = Span;
+
+			HashCode HashCode = new HashCode();
+			for (int Idx = 0; Idx < NumBytes; Idx++)
+			{
+				HashCode.Add(LocalSpan[Idx]);
+			}
+			return HashCode.ToHashCode();
+		}
+
+		/// <inheritdoc/>
+		public override string ToString()
+		{
+			return StringUtils.FormatHexString(Span);
+		}
+	}
+
+	/// <summary>
+	/// Read-only array of blob hashes
+	/// </summary>
+	struct ReadOnlyBlobHashArray : IReadOnlyList<BlobHash>
+	{
+		/// <summary>
+		/// Underlying storage for the array
+		/// </summary>
+		public ReadOnlyMemory<byte> Memory { get; }
+
+		/// <summary>
+		/// Constructor
+		/// </summary>
+		/// <param name="Memory">Storage for the array</param>
+		public ReadOnlyBlobHashArray(ReadOnlyMemory<byte> Memory)
+		{
+			this.Memory = Memory;
+		}
+
+		/// <summary>
+		/// Accessor for elements in the array
+		/// </summary>
+		/// <param name="Index">Index of the element to retrieve</param>
+		/// <returns>Hash at the given index</returns>
+		public BlobHash this[int Index] => new BlobHash(Memory.Slice(Index * BlobHash.NumBytes, BlobHash.NumBytes));
+
+		/// <inheritdoc/>
+		public int Count => Memory.Length / BlobHash.NumBytes;
+
+		/// <inheritdoc/>
+		public IEnumerator<BlobHash> GetEnumerator()
+		{
+			ReadOnlyMemory<byte> Remaining = Memory;
+			while (!Remaining.IsEmpty)
+			{
+				yield return new BlobHash(Remaining.Slice(0, BlobHash.NumBytes));
+				Remaining = Remaining.Slice(BlobHash.NumBytes);
+			}
+		}
+
+		/// <inheritdoc/>
+		IEnumerator IEnumerable.GetEnumerator() => GetEnumerator();
+	}
+}
+
diff --git a/Engine/Source/Programs/Horde/HordeServer/Storage/Primitives/BlobRef.cs b/Engine/Source/Programs/Horde/HordeServer/Storage/Primitives/BlobRef.cs
new file mode 100644
--- /dev/null
+++ b/Engine/Source/Programs/Horde/HordeServer/Storage/Primitives/BlobRef.cs
@@ -0,0 +1,85 @@
+using System;
+using System.Collections.Generic;
+using System.Linq;
+using System.Threading.Tasks;
+
+namespace HordeServer.Storage.Primitives
+{
+	/// <summary>
+	/// Reference to a blob object
+	/// </summary>
+	struct BlobRef
+	{
+		/// <summary>
+		/// Type of the blob
+		/// </summary>
+		public IBlobType Type { get; }
+
+		/// <summary>
+		/// The blob hash
+		/// </summary>
+		public BlobHash Hash { get; }
+
+		/// <summary>
+		/// Constructor
+		/// </summary>
+		/// <param name="Type"></param>
+		/// <param name="Hash"></param>
+		public BlobRef(IBlobType Type, BlobHash Hash)
+		{
+			this.Type = Type;
+			this.Hash = Hash;
+		}
+
+		/// <inheritdoc/>
+		public override string ToString()
+		{
+			return $"{Type.Name}: {Hash}";
+		}
+	}
+
+	/// <summary>
+	/// Typed reference to a blob object
+	/// </summary>
+	/// <typeparam name="TBlob">The blob type</typeparam>
+	struct BlobRef<TBlob> where TBlob : IBlobType, new()
+	{
+		/// <summary>
+		/// Static type accessor
+		/// </summary>
+		public static TBlob Type => new TBlob();
+
+		/// <summary>
+		/// Hash of the referenced blob
+		/// </summary>
+		public BlobHash Hash { get; }
+
+		/// <summary>
+		/// Constructor
+		/// </summary>
+		/// <param name="Hash">Hash of the blob</param>
+		public BlobRef(BlobHash Hash)
+		{
+			this.Hash = Hash;
+		}
+
+		/// <summary>
+		/// Converts from a dynamically typed blob to this blob type
+		/// </summary>
+		/// <param name="Ref">The blob reference</param>
+		public static explicit operator BlobRef<TBlob>(BlobRef Ref)
+		{
+			_ = (TBlob)Ref.Type;
+			return new BlobRef<TBlob>(Ref.Hash);
+		}
+
+		/// <summary>
+		/// Converts from this blob to a dynamically typed blob
+		/// </summary>
+		/// <param name="Ref"></param>
+		public static implicit operator BlobRef(BlobRef<TBlob> Ref)
+		{
+			return new BlobRef(BlobRef<TBlob>.Type, Ref.Hash);
+		}
+	}
+}
diff --git a/Engine/Source/Programs/Horde/HordeServer/Storage/Primitives/BlobType.cs b/Engine/Source/Programs/Horde/HordeServer/Storage/Primitives/BlobType.cs
new file mode 100644
--- /dev/null
+++ b/Engine/Source/Programs/Horde/HordeServer/Storage/Primitives/BlobType.cs
@@ -0,0 +1,40 @@
+using System;
+using System.Collections.Generic;
+using System.Linq;
+using System.Threading.Tasks;
+
+namespace HordeServer.Storage.Primitives
+{
+	/// <summary>
+	/// Interface for a particular blob type
+	/// </summary>
+	interface IBlobType
+	{
+		/// <summary>
+		/// Name of the type
+		/// </summary>
+		public string Name { get; }
+
+		/// <summary>
+		/// Gets references from the given blob
+		/// </summary>
+		/// <param name="Data">The blob data</param>
+		/// <returns>References</returns>
+		public IEnumerable<BlobRef> GetRefs(ReadOnlyMemory<byte> Data);
+	}
+
+	/// <summary>
+	/// Leaf blob type definition
+	/// </summary>
+	sealed class LeafBlob : IBlobType
+	{
+		/// <inheritdoc/>
+		public string Name => "Leaf";
+
+		/// <inheritdoc/>
+		public IEnumerable<BlobRef> GetRefs(ReadOnlyMemory<byte> Data)
+		{
+			yield break;
+		}
+	}
+}
diff --git a/Engine/Source/Programs/Horde/HordeServer/Storage/Primitives/KeyValueItem.cs b/Engine/Source/Programs/Horde/HordeServer/Storage/Primitives/KeyValueItem.cs
new file mode 100644
--- /dev/null
+++ b/Engine/Source/Programs/Horde/HordeServer/Storage/Primitives/KeyValueItem.cs
@@ -0,0 +1,101 @@
+using System;
+using System.Collections.Generic;
+using System.Linq;
+using System.Threading.Tasks;
+
+namespace HordeServer.Storage.Primitives
+{
+	/// <summary>
+	/// A key/value pair for insertion into a KeyValueTree
+	/// </summary>
+	struct KeyValueItem
+	{
+		/// <summary>
+		/// The key
+		/// </summary>
+		public BlobHash Key { get; }
+
+		/// <summary>
+		/// Block of memory containing the value for this item
+		/// </summary>
+		public ReadOnlyMemory<byte> Value { get; }
+
+		/// <summary>
+		/// Constructor
+		/// </summary>
+		/// <param name="Key"></param>
+		/// <param name="Value"></param>
+		public KeyValueItem(BlobHash Key, ReadOnlyMemory<byte> Value)
+		{
+			this.Key = Key;
+			this.Value = Value;
+		}
+	}
+
+	class KeyValueItemReader
+	{
+		ReadOnlyMemory<byte> Keys;
+		ReadOnlyMemory<byte> Values;
+		int ValueSize;
+
+		public KeyValueItemReader(KeyValueTree Tree, ReadOnlyKeyValueNode Node)
+		{
+			this.Keys = Node.Keys;
+			this.Values = Node.Values;
+			this.ValueSize = Node.IsLeafNode ? Tree.LeafValueSize : BlobHash.NumBytes;
+		}
+
+		public bool IsEmpty => Keys.Length == 0;
+
+		public void MoveNext()
+		{
+			Keys = Keys.Slice(BlobHash.NumBytes);
+			Values = Values.Slice(ValueSize);
+		}
+
+		public void CopyTo(KeyValueItemWriter Writer)
+		{
+			Writer.WriteItems(Keys, Values);
+
+			Keys = ReadOnlyMemory<byte>.Empty;
+			Values = ReadOnlyMemory<byte>.Empty;
+		}
+
+		public void CopyTo(KeyValueItemWriter Writer, int Count)
+		{
+			ReadOnlyMemory<byte> CopyKeys = Keys.Slice(0, Count * BlobHash.NumBytes);
+			ReadOnlyMemory<byte> CopyValues = Values.Slice(0, Count * ValueSize);
+
+			Writer.WriteItems(CopyKeys, CopyValues);
+
+			Keys = Keys.Slice(CopyKeys.Length);
+			Values = Values.Slice(CopyValues.Length);
+		}
+	}
+
+	class KeyValueItemWriter
+	{
+		public Memory<byte> Keys { get; private set; }
+		public Memory<byte> Values { get; private set; }
+
+		public KeyValueItemWriter(KeyValueNode Node)
+		{
+			this.Keys = Node.Keys;
+			this.Values = Node.Values;
+		}
+
+		public void WriteItem(BlobHash Key, ReadOnlyMemory<byte> Value)
+		{
+			WriteItems(Key.Memory, Value);
+		}
+
+		public void WriteItems(ReadOnlyMemory<byte> NewKeys, ReadOnlyMemory<byte> NewValues)
+		{
+			NewKeys.CopyTo(Keys);
+			Keys = Keys.Slice(NewKeys.Length);
+
+			NewValues.CopyTo(Values);
+			Values = Values.Slice(NewValues.Length);
+		}
+	}
+}
diff --git a/Engine/Source/Programs/Horde/HordeServer/Storage/Primitives/KeyValueNode.cs b/Engine/Source/Programs/Horde/HordeServer/Storage/Primitives/KeyValueNode.cs
new file mode 100644
--- /dev/null
+++ b/Engine/Source/Programs/Horde/HordeServer/Storage/Primitives/KeyValueNode.cs
@@ -0,0 +1,159 @@
+using System;
+using System.Buffers.Binary;
+using System.Collections;
+using System.Collections.Generic;
+using System.Linq;
+using System.Threading.Tasks;
+
+namespace HordeServer.Storage.Primitives
+{
+	/// <summary>
+	/// A node in a key value tree
+	/// </summary>
+	class KeyValueNode
+	{
+		/// <summary>
+		/// The underlying memory storing the node data
+		/// </summary>
+		public Memory<byte> Memory { get; }
+
+		/// <summary>
+		/// Number of child items from this node
+		/// </summary>
+		public int NumItems { get; }
+
+		/// <summary>
+		/// Number of items if this node is merged with its children. This will be equal to NumItems if this is the child of a parent-child relation.
+		/// Note that this is different from the total number of leaf nodes under this tree; it is only used for splitting/merging layers.
+		/// </summary>
+		public int NumItemsIfMerged { get; }
+
+		/// <summary>
+		/// The number of valid bits in keys for this node
+		/// </summary>
+		public int NumKeyBits { get; }
+
+		/// <summary>
+		/// The number of valid bits in the key if it is merged with its children
+		/// </summary>
+		public int NumKeyBitsIfMerged { get; }
+
+		/// <summary>
+		/// Data for the array of keys
+		/// </summary>
+		public Memory<byte> Keys { get; }
+
+		/// <summary>
+		/// Data for the array of values
+		/// </summary>
+		public Memory<byte> Values { get; }
+
+		/// <summary>
+		/// Constructor
+		/// </summary>
+		/// <param name="NumItems"></param>
+		/// <param name="NumItemsIfMerged"></param>
+		/// <param name="NumKeyBits"></param>
+		/// <param name="NumKeyBitsIfMerged"></param>
+		/// <param name="ValueSize">Size of a value in this node</param>
+		public KeyValueNode(int NumItems, int NumItemsIfMerged, int NumKeyBits, int NumKeyBitsIfMerged, int ValueSize)
+		{
+			this.NumItems = NumItems;
+			this.NumItemsIfMerged = NumItemsIfMerged;
+			this.NumKeyBits = NumKeyBits;
+			this.NumKeyBitsIfMerged = NumKeyBitsIfMerged;
+
+			int KeyDataLength = BlobHash.NumBytes * NumItems;
+			int ValueDataLength = ValueSize * NumItems;
+			Memory = new byte[(sizeof(int) * 4) + KeyDataLength + ValueDataLength];
+
+			Memory<byte> Output = Memory;
+
+			BinaryPrimitives.WriteInt32LittleEndian(Output.Span, NumItems);
+			Output = Output.Slice(sizeof(int));
+
+			BinaryPrimitives.WriteInt32LittleEndian(Output.Span, NumItemsIfMerged);
+			Output = Output.Slice(sizeof(int));
+
+			BinaryPrimitives.WriteInt32LittleEndian(Output.Span, NumKeyBits);
+			Output = Output.Slice(sizeof(int));
+
+			BinaryPrimitives.WriteInt32LittleEndian(Output.Span, NumKeyBitsIfMerged);
+			Output = Output.Slice(sizeof(int));
+
+			Keys = Output.Slice(0, KeyDataLength);
+			Output = Output.Slice(KeyDataLength);
+
+			Values = Output;
+		}
+	}
+
+	/// <summary>
+	/// Read-only version of <see cref="KeyValueNode"/>
+	/// </summary>
+	class ReadOnlyKeyValueNode
+	{
+		/// <inheritdoc cref="KeyValueNode.Memory"/>
+		public ReadOnlyMemory<byte> Memory { get; }
+
+		/// <inheritdoc cref="KeyValueNode.NumItems"/>
+		public int NumItems { get; }
+
+		/// <inheritdoc cref="KeyValueNode.NumItemsIfMerged"/>
+		public int NumItemsIfMerged { get; }
+
+		/// <inheritdoc cref="KeyValueNode.NumKeyBits"/>
+		public int NumKeyBits { get; }
+
+		/// <inheritdoc cref="KeyValueNode.NumKeyBitsIfMerged"/>
+		public int NumKeyBitsIfMerged { get; }
+
+		/// <inheritdoc cref="KeyValueNode.Keys"/>
+		public ReadOnlyMemory<byte> Keys { get; }
+
+		/// <inheritdoc cref="KeyValueNode.Values"/>
+		public ReadOnlyMemory<byte> Values { get; }
+
+		/// <summary>
+		/// Whether this node is a leaf node
+		/// </summary>
+		public bool IsLeafNode => NumKeyBits == BlobHash.NumBits;
+
+		/// <summary>
+		/// Constructor
+		/// </summary>
+		/// <param name="Memory">Memory backing this node</param>
+		public ReadOnlyKeyValueNode(ReadOnlyMemory<byte> Memory)
+		{
+			this.Memory = Memory;
+
+			ReadOnlyMemory<byte> Input = Memory;
+
+			NumItems = BinaryPrimitives.ReadInt32LittleEndian(Input.Span);
+			Input = Input.Slice(sizeof(int));
+
+			NumItemsIfMerged = BinaryPrimitives.ReadInt32LittleEndian(Input.Span);
+			Input = Input.Slice(sizeof(int));
+
+			NumKeyBits = BinaryPrimitives.ReadInt32LittleEndian(Input.Span);
+			Input = Input.Slice(sizeof(int));
+
+			NumKeyBitsIfMerged = BinaryPrimitives.ReadInt32LittleEndian(Input.Span);
+			Input = Input.Slice(sizeof(int));
+
+			Keys = Input.Slice(0, NumItems * BlobHash.NumBytes);
+			Input = Input.Slice(Keys.Length);
+
+			Values = Input;
+		}
+
+		/// <summary>
+		/// Implicit conversion operator from a KeyValueNode
+		/// </summary>
+		/// <param name="Other"></param>
+		public static implicit operator ReadOnlyKeyValueNode(KeyValueNode Other)
+		{
+			return new ReadOnlyKeyValueNode(Other.Memory);
+		}
+	}
+}
diff --git a/Engine/Source/Programs/Horde/HordeServer/Storage/Primitives/KeyValueTree.cs b/Engine/Source/Programs/Horde/HordeServer/Storage/Primitives/KeyValueTree.cs
new file mode 100644
--- /dev/null
+++ b/Engine/Source/Programs/Horde/HordeServer/Storage/Primitives/KeyValueTree.cs
@@ -0,0 +1,530 @@
+using EpicGames.Core;
+using System;
+using System.Buffers.Binary;
+using System.Collections.Generic;
+using System.Diagnostics;
+using System.Linq;
+using System.Threading.Tasks;
+
+namespace HordeServer.Storage.Primitives
+{
+	/// <summary>
+	/// Base type definition for key/value trees. Handles both internal and leaf nodes.
+	/// </summary>
+	abstract class KeyValueTree : IBlobType
+	{
+		/// <summary>
+		/// Default maximum node size
+		/// </summary>
+		const int DefaultMaxNodeSize = 1024 * 1024;
+
+		/// <summary>
+		/// The key type
+		/// </summary>
+		public IBlobType KeyType;
+
+		/// <summary>
+		/// Length of each value
+		/// </summary>
+		public int LeafValueSize { get; }
+
+		/// <summary>
+		/// Offset within each value item of blob references, and their types
+		/// </summary>
+		public (int, IBlobType)[] LeafValueRefTypes { get; }
+
+		/// <summary>
+		/// Number of items for an internal node
+		/// </summary>
+		public int MaxItemsPerInternalNode { get; }
+
+		/// <summary>
+		/// Number of items for a leaf node
+		/// </summary>
+		public int MaxItemsPerLeafNode { get; }
+
+		/// <inheritdoc/>
+		public string Name => "KeyValueTree";
+
+		/// <summary>
+		/// Constructor
+		/// </summary>
+		/// <param name="KeyType">Type of key</param>
+		/// <param name="LeafValueSize">Length of each value</param>
+		/// <param name="LeafValueRefTypes">References from each value</param>
+		public KeyValueTree(IBlobType KeyType, int LeafValueSize, params (int, IBlobType)[] LeafValueRefTypes)
+			: this(KeyType, LeafValueSize, LeafValueRefTypes, DefaultMaxNodeSize / (BlobHash.NumBytes * 2), DefaultMaxNodeSize / (BlobHash.NumBytes + LeafValueSize))
+		{
+		}
+
+		/// <summary>
+		/// Constructor
+		/// </summary>
+		/// <param name="KeyType">Type of key</param>
+		/// <param name="LeafValueSize">Length of each value</param>
+		/// <param name="LeafValueRefTypes">References from each value</param>
+		/// <param name="MaxItemsPerInternalNode">Maximum number of items to store in each internal node</param>
+		/// <param name="MaxItemsPerLeafNode">Maximum number of items to store in each leaf node</param>
+		public KeyValueTree(IBlobType KeyType, int LeafValueSize, (int, IBlobType)[] LeafValueRefTypes, int MaxItemsPerInternalNode, int MaxItemsPerLeafNode)
+		{
+			this.KeyType = KeyType;
+			this.LeafValueSize = LeafValueSize;
+			this.LeafValueRefTypes = LeafValueRefTypes;
+			this.MaxItemsPerInternalNode = MaxItemsPerInternalNode;
+			this.MaxItemsPerLeafNode = MaxItemsPerLeafNode;
+		}
+
+		/// <summary>
+		/// Get references from the given data
+		/// </summary>
+		/// <param name="Data">Blob of data</param>
+		/// <returns>Sequence of blob references</returns>
+		public IEnumerable<BlobRef> GetRefs(ReadOnlyMemory<byte> Data)
+		{
+			ReadOnlyKeyValueNode Node = new ReadOnlyKeyValueNode(Data);
+			if (Node.NumKeyBits == BlobHash.NumBits)
+			{
+				// Leaf node: keys are references to other blobs, values follow type definition
+				ReadOnlyMemory<byte> Keys = Node.Keys;
+				while (!Keys.IsEmpty)
+				{
+					yield return new BlobRef(KeyType, new BlobHash(Keys.Slice(0, BlobHash.NumBytes)));
+					Keys = Keys.Slice(BlobHash.NumBytes);
+				}
+
+				ReadOnlyMemory<byte> Values = Node.Values;
+				while (!Values.IsEmpty)
+				{
+					foreach ((int Offset, IBlobType Type) in LeafValueRefTypes)
+					{
+						yield return new BlobRef(Type, new BlobHash(Values.Slice(Offset, BlobHash.NumBytes)));
+					}
+					Values = Values.Slice(LeafValueSize);
+				}
+			}
+			else
+			{
+				// Internal node: keys are just prefixes, and not concrete references. Values are always just references to other trees.
+				ReadOnlyBlobHashArray Values = new ReadOnlyBlobHashArray(Node.Values);
+				foreach(BlobHash Value in Values)
+				{
+					yield return new BlobRef(this, Value);
+				}
+			}
+		}
+	}
+
+	/// <summary>
+	/// Base type definition for key/value trees with a static key type
+	/// </summary>
+	abstract class KeyValueTree<TKey> : KeyValueTree where TKey : IBlobType, new()
+	{
+		/// <summary>
+		/// Constructor
+		/// </summary>
+		/// <param name="ValueSize"></param>
+		/// <param name="ValueRefTypes"></param>
+		public KeyValueTree(int ValueSize, params (int, IBlobType)[] ValueRefTypes)
+			: base(BlobRef<TKey>.Type, ValueSize, ValueRefTypes)
+		{
+		}
+	}
+
+	/// <summary>
+	/// Base type definition for key/value trees with a static key type
+	/// </summary>
+	sealed class KeyValueTree<TKey, TValue> : KeyValueTree where TKey : IBlobType, new() where TValue : IBlobType, new()
+	{
+		/// <summary>
+		/// Constructor
+		/// </summary>
+		public KeyValueTree()
+			: base(BlobRef<TKey>.Type, BlobHash.NumBytes, (0, BlobRef<TValue>.Type))
+		{
+		}
+	}
+
+	/// <summary>
+	/// Extension methods for manipulating key/value trees
+	/// </summary>
+	static class KeyValueTreeExtensions
+	{
+		/// <summary>
+		/// Creates a new, empty key value tree
+		/// </summary>
+		/// <typeparam name="TTree">The tree type</typeparam>
+		/// <param name="Storage">Storage service</param>
+		/// <returns>Reference to the new tree</returns>
+		public static async Task<BlobRef<TTree>> CreateKeyValueTreeAsync<TTree>(this IStorageService Storage) where TTree : KeyValueTree, new()
+		{
+			KeyValueNode Node = new KeyValueNode(0, 0, BlobHash.NumBits, BlobHash.NumBits, 0);
+			return new BlobRef<TTree>(await Storage.PutKeyValueNodeAsync(Node));
+		}
+
+		public static async Task<BlobRef<TValue>?> SearchKeyValueTreeAsync<TKey, TValue>(this IStorageService Storage, BlobRef<KeyValueTree<TKey, TValue>> Tree, BlobHash FindHash) where TKey : IBlobType, new() where TValue : IBlobType, new()
+		{
+			ReadOnlyMemory<byte>? Result = await SearchKeyValueTreeInternalAsync(Storage, Tree, FindHash);
+			if(Result == null)
+			{
+				return null;
+			}
+			return new BlobRef<TValue>(new BlobHash(Result.Value));
+		}
+
+		public static Task<ReadOnlyMemory<byte>?> SearchKeyValueTreeAsync<TTree>(this IStorageService Storage, BlobRef<TTree> Tree, BlobHash FindHash) where TTree : KeyValueTree, new()
+		{
+			return SearchKeyValueTreeInternalAsync(Storage, Tree, FindHash);
+		}
+
+		static async Task<ReadOnlyMemory<byte>?> SearchKeyValueTreeInternalAsync<TTree>(this IStorageService Storage, BlobRef<TTree> Tree, BlobHash FindHash) where TTree : KeyValueTree, new()
+		{
+			BlobHash Hash = Tree.Hash;
+			for (; ; )
+			{
+				ReadOnlyKeyValueNode Node = await Storage.GetKeyValueNodeAsync(Hash);
+
+				BlobHash MaskHash = FindHash.MaskLeft(Node.NumKeyBits);
+
+				ReadOnlyBlobHashArray Keys = new ReadOnlyBlobHashArray(Node.Keys);
+
+				int KeyIdx = Keys.BinarySearch(MaskHash);
+				if (KeyIdx < 0)
+				{
+					return null;
+				}
+				else if (Node.NumKeyBits == BlobHash.NumBits)
+				{
+					return Node.Values.Slice(KeyIdx * BlobRef<TTree>.Type.LeafValueSize, BlobRef<TTree>.Type.LeafValueSize);
+				}
+				else
+				{
+					Hash = new ReadOnlyBlobHashArray(Node.Values)[KeyIdx];
+				}
+			}
+		}
+
+		/// <summary>
+		/// Updates a key value tree with the given entries
+		/// </summary>
+		/// <typeparam name="TTree">The tree to update</typeparam>
+		/// <param name="Storage">Storage service instance</param>
+		/// <param name="Tree">The tree to update</param>
+		/// <param name="NewItems">New items to add</param>
+		/// <returns></returns>
+		public static async Task<BlobRef<TTree>> UpdateKeyValueTreeAsync<TTree>(this IStorageService Storage, BlobRef<TTree> Tree, IEnumerable<KeyValueItem> NewItems) where TTree : KeyValueTree, new()
+		{
+			BlobHash NewRootHash = await UpdateKeyValueTreeAsync(Storage, BlobRef<TTree>.Type, Tree.Hash, NewItems);
+			return new BlobRef<TTree>(NewRootHash);
+		}
+
+		/// <summary>
+		/// Updates a key/value tree with a set of items
+		/// </summary>
+		/// <param name="Storage">The storage service</param>
+		/// <param name="Tree">The tree to update</param>
+		/// <param name="RootHash">Hash of the root node of the tree</param>
+		/// <param name="NewItems">New items to add</param>
+		/// <returns>Hash for the new tree root</returns>
+		static async Task<BlobHash> UpdateKeyValueTreeAsync(this IStorageService Storage, KeyValueTree Tree, BlobHash RootHash, IEnumerable<KeyValueItem> NewItems)
+		{
+			// Sort all the items and remove duplicates
+			ArraySegment<KeyValueItem> NewSortedItems = NewItems.OrderBy(x => x.Key).ToArray();
+			if (NewSortedItems.Count > 0)
+			{
+				int OutputCount = 1;
+				for (int InputCount = 1; InputCount < NewSortedItems.Count; InputCount++)
+				{
+					if (NewSortedItems[InputCount].Key != NewSortedItems[OutputCount - 1].Key)
+					{
+						NewSortedItems[OutputCount++] = NewSortedItems[InputCount];
+					}
+				}
+				NewSortedItems = NewSortedItems.Slice(0, OutputCount);
+			}
+
+			// Get the root node
+			ReadOnlyKeyValueNode RootNode;
+			if (RootHash == BlobHash.Zero)
+			{
+				RootNode = new KeyValueNode(0, 0, BlobHash.NumBits, BlobHash.NumBits, Tree.LeafValueSize);
+			}
+			else
+			{
+				RootNode = await Storage.GetKeyValueNodeAsync(RootHash);
+			}
+
+			// Get the new node
+			KeyValueNode NewRootNode = await UpdateKeyValueNodeAsync(Storage, Tree, 0, RootNode, NewSortedItems);
+			return await Storage.PutKeyValueNodeAsync(NewRootNode);
+		}
+
+		/// <summary>
+		/// Update a key value tree node
+		/// </summary>
+		/// <param name="Storage">Storage service</param>
+		/// <param name="Tree">The tree definition</param>
+		/// <param name="ParentKeyBits">Number of bits in the parent node</param>
+		/// <param name="Node">The current node state</param>
+		/// <param name="NewItems">The items to modify</param>
+		/// <returns>Updated node</returns>
+		static async Task<KeyValueNode> UpdateKeyValueNodeAsync(IStorageService Storage, KeyValueTree Tree, int ParentKeyBits, ReadOnlyKeyValueNode Node, ArraySegment<KeyValueItem> NewItems)
+		{
+			// Create a list wrapper for the keys, so we can perform binary searches on it
+			ReadOnlyBlobHashArray Keys = new ReadOnlyBlobHashArray(Node.Keys);
+
+			// Keep track of the total number of child items underneath this node in order to update the merge count
+			int NewNumChildItems = Node.NumItemsIfMerged - Node.NumItems;
+
+			// Find the necessary updates to the node
+			List<(int KeyIdx, KeyValueItem? NewItem)> Updates = new List<(int, KeyValueItem?)>(NewItems.Count);
+			if (Node.NumKeyBits == BlobHash.NumBits)
+			{
+				foreach (KeyValueItem NewItem in NewItems)
+				{
+					int KeyIdx = Keys.BinarySearch(NewItem.Key);
+					if (KeyIdx >= 0 || !NewItem.Value.IsEmpty)
+					{
+						Updates.Add((KeyIdx, NewItem));
+					}
+				}
+			}
+			else
+			{
+				ReadOnlyBlobHashArray Children = new ReadOnlyBlobHashArray(Node.Values);
+				for (int Idx = 0; Idx < NewItems.Count;)
+				{
+					BlobHash KeyPrefix = NewItems[Idx].Key.MaskLeft(Node.NumKeyBits);
+
+					// Find all the operations which start with the same prefix
+					int MinIdx = Idx++;
+					while (Idx < NewItems.Count && NewItems[Idx].Key.StartsWith(KeyPrefix, Node.NumKeyBits))
+					{
+						Idx++;
+					}
+
+					// Check if it already exists within the tree
+					int KeyIdx = Keys.BinarySearch(KeyPrefix);
+
+					// Get the old child node
+					ReadOnlyKeyValueNode OldChildNode;
+					if (KeyIdx >= 0)
+					{
+						OldChildNode = await Storage.GetKeyValueNodeAsync(Children[KeyIdx]);
+					}
+					else
+					{
+						OldChildNode = CreateEmptyChildNode(Node.NumKeyBits);
+					}
+
+					// Create the new item
+					KeyValueNode NewChildNode = await UpdateKeyValueNodeAsync(Storage, Tree, Node.NumKeyBits, OldChildNode, NewItems.Slice(MinIdx, Idx - MinIdx));
+					if (NewChildNode.NumItems == 0)
+					{
+						// Only remove it if it already exists
+						if (KeyIdx >= 0)
+						{
+							Updates.Add((KeyIdx, null));
+						}
+					}
+					else
+					{
+						// Add the new item to the update list
+						BlobHash NewChildHash = await Storage.PutKeyValueNodeAsync(NewChildNode);
+						KeyValueItem NewItem = new KeyValueItem(KeyPrefix, NewChildHash.Memory);
+						Updates.Add((KeyIdx, NewItem));
+
+						// Update the number of merged items
+						NewNumChildItems += NewChildNode.NumItemsIfMerged - OldChildNode.NumItemsIfMerged;
+					}
+				}
+			}
+
+			// Create a new buffer for the keys
+			int NewNumItems = Node.NumItems + Updates.Count(x => x.KeyIdx < 0) - Updates.Count(x => x.NewItem == null);
+			KeyValueNode NewNode = CreateNode(Tree, NewNumItems, NewNumItems + NewNumChildItems, Node.NumKeyBits, Node.NumKeyBitsIfMerged);
+
+			// Copy all the data across
+			int ReadIdx = 0;
+			KeyValueItemReader Reader = new KeyValueItemReader(Tree, Node);
+			KeyValueItemWriter Writer = new KeyValueItemWriter(NewNode);
+			foreach ((int KeyIdx, KeyValueItem? NewItem) in Updates)
+			{
+				// Copy all the data before the new item
+				if (KeyIdx < 0)
+				{
+					int Count = ~KeyIdx - ReadIdx;
+					Reader.CopyTo(Writer, Count);
+					ReadIdx += Count;
+				}
+				else
+				{
+					int Count = KeyIdx - ReadIdx;
+					Reader.CopyTo(Writer, Count);
+					Reader.MoveNext();
+					ReadIdx += Count + 1;
+				}
+
+				// Copy the new item
+				if (NewItem != null)
+				{
+					Writer.WriteItem(NewItem.Value.Key, NewItem.Value.Value);
+				}
+			}
+			Reader.CopyTo(Writer);
+
+			// Merge together layers of the tree
+			while (NewNode.NumKeyBitsIfMerged > NewNode.NumKeyBits && NewNode.NumItemsIfMerged < Tree.MaxItemsPerInternalNode)
+			{
+				NewNode = await MergeAsync(Storage, Tree, NewNode);
+			}
+
+			// If it's larger than a threshold, we need to split this node into two
+			while (NewNode.NumItems > ((NewNode.NumKeyBits == BlobHash.NumBits)? Tree.MaxItemsPerLeafNode : Tree.MaxItemsPerInternalNode))
+			{
+				NewNode = await SplitAsync(Storage, Tree, NewNode, ParentKeyBits);
+			}
+
+			// Return the updated node
+			return NewNode;
+		}
+
+		/// <summary>
+		/// Merges a node iwth its children
+		/// </summary>
+		/// <param name="Storage"></param>
+		/// <param name="Tree"></param>
+		/// <param name="Node"></param>
+		/// <returns></returns>
+		public static async Task<KeyValueNode> MergeAsync(IStorageService Storage, KeyValueTree Tree, ReadOnlyKeyValueNode Node)
+		{
+			KeyValueNode NewNode = CreateNode(Tree, Node.NumItemsIfMerged, Node.NumItemsIfMerged, Node.NumKeyBitsIfMerged, Node.NumKeyBitsIfMerged);
+			KeyValueItemWriter Writer = new KeyValueItemWriter(NewNode);
+
+			ReadOnlyBlobHashArray Values = new ReadOnlyBlobHashArray(Node.Values);
+			for (int Idx = 0; Idx < Values.Count; Idx++)
+			{
+				ReadOnlyKeyValueNode ChildNode = await Storage.GetKeyValueNodeAsync(Values[Idx]);
+				Writer.WriteItems(ChildNode.Keys, ChildNode.Values);
+			}
+
+			return NewNode;
+		}
+
+		/// <summary>
+		/// Splits a node in the key value tree
+		/// </summary>
+		/// <param name="Storage"></param>
+		/// <param name="Tree">The tree definition</param>
+		/// <param name="Node"></param>
+		/// <param name="ParentKeyBits"></param>
+		/// <returns></returns>
+		public static async Task<KeyValueNode> SplitAsync(IStorageService Storage, KeyValueTree Tree, ReadOnlyKeyValueNode Node, int ParentKeyBits)
+		{
+			// Given an even distribution of keys across the hash space, keeping the tree balanced requires that we keep the 
+			// number of items in the parent node approximately equal to the average number of items in each child node.
+			// 
+			// If h is the number of bits in the hash, n is the number of bits in the parent node, and m is the number of 
+			// child nodes, it follows that:
+			// 
+			//    Items in parent node = 2 ^ n
+			//    Items in each child node = (2 ^ (h - n)) / m
+			// 
+			// Thus: n = (h - log2(m)) / 2
+			int MaxChildren = Node.IsLeafNode ? Tree.MaxItemsPerLeafNode : Tree.MaxItemsPerInternalNode;
+			int NewMaxKeyBits = ParentKeyBits + Math.Max(1, (int)(((Node.NumKeyBits - ParentKeyBits) - Math.Log2(MaxChildren)) * 0.5));
+
+			List<KeyValueItem> NewItems = new List<KeyValueItem>();
+
+			KeyValueItemReader Reader = new KeyValueItemReader(Tree, Node);
+
+			ReadOnlyBlobHashArray Keys = new ReadOnlyBlobHashArray(Node.Keys);
+			for (int KeyIdx = 0; KeyIdx < Keys.Count; )
+			{
+				// Mask off the first key prefix
+				BlobHash KeyPrefix = Keys[KeyIdx].MaskLeft(NewMaxKeyBits);
+
+				// Count the number of keys with the same prefix
+				int Count = 1;
+				while (KeyIdx + Count < Keys.Count && Keys[KeyIdx + Count].StartsWith(KeyPrefix, NewMaxKeyBits))
+				{
+					Count++;
+				}
+
+				// Create the new child node
+				KeyValueNode NewChildNode = CreateNode(Tree, Count, Count, Node.NumKeyBits, Node.NumKeyBits);
+				Reader.CopyTo(new KeyValueItemWriter(NewChildNode), Count);
+				KeyIdx += Count;
+
+				// If the child node is still above the max number of children, split it again
+				while (NewChildNode.NumItems > MaxChildren)
+				{
+					NewChildNode = await SplitAsync(Storage, Tree, NewChildNode, NewMaxKeyBits);
+				}
+
+				// Finally, store the child node and add it to its new parent
+				BlobHash NewChildHash = await Storage.PutKeyValueNodeAsync(NewChildNode);
+				NewItems.Add(new KeyValueItem(KeyPrefix, NewChildHash.Memory));
+			}
+
+			KeyValueNode NewNode = new KeyValueNode(NewItems.Count, Node.NumItems, NewMaxKeyBits, Node.NumKeyBits, BlobHash.NumBytes);
+
+			KeyValueItemWriter Writer = new KeyValueItemWriter(NewNode);
+			foreach (KeyValueItem NewItem in NewItems)
+			{
+				Writer.WriteItem(NewItem.Key, NewItem.Value);
+			}
+
+			return NewNode;
+		}
+
+		public static KeyValueNode CreateNode(KeyValueTree Tree, int NumItems, int NumItemsIfMerged, int NumKeyBits, int NumKeyBitsIfMerged)
+		{
+			int NewValueSize = (NumKeyBits == BlobHash.NumBits) ? Tree.LeafValueSize : BlobHash.NumBytes;
+			return new KeyValueNode(NumItems, NumItemsIfMerged, NumKeyBits, NumKeyBitsIfMerged, NewValueSize);
+		}
+
+		public static ReadOnlyKeyValueNode CreateEmptyChildNode(int MaxKeyBits)
+		{
+			int ChildKeyBits = GetChildKeyLength(MaxKeyBits);
+			return new KeyValueNode(0, 0, ChildKeyBits, ChildKeyBits, 0);
+		}
+
+		public static int GetChildKeyLength(int KeyLength)
+		{
+			// Find the largest power of two boundary that the current key length is on
+			int KeyStep = 1;
+			while (KeyLength + KeyStep < BlobHash.NumBits)
+			{
+				int NextKeyStep = KeyStep << 1;
+				if ((KeyLength & (NextKeyStep - 1)) != 0)
+				{
+					break;
+				}
+				KeyStep = NextKeyStep;
+			}
+			return KeyLength + KeyStep;
+		}
+
+		/// <summary>
+		/// Adds a key/value tree node to a storage provider
+		/// </summary>
+		/// <param name="Provider"></param>
+		/// <param name="Node"></param>
+		/// <returns></returns>
+		public static async Task<BlobHash> PutKeyValueNodeAsync(this IStorageService Provider, ReadOnlyKeyValueNode Node)
+		{
+			return await Provider.PutBlobAsync(Node.Memory);
+		}
+
+		/// <summary>
+		/// Gets a key/value tree node from a storage provider
+		/// </summary>
+		/// <param name="Provider"></param>
+		/// <param name="Hash"></param>
+		/// <returns></returns>
+		public static async Task<ReadOnlyKeyValueNode> GetKeyValueNodeAsync(this IStorageService Provider, BlobHash Hash)
+		{
+			return new ReadOnlyKeyValueNode(await Provider.GetBlobAsync(Hash));
+		}
+	}
+}
diff --git a/Engine/Source/Programs/Horde/HordeServer/Storage/Services/SimpleStorageService.cs b/Engine/Source/Programs/Horde/HordeServer/Storage/Services/SimpleStorageService.cs
new file mode 100644
--- /dev/null
+++ b/Engine/Source/Programs/Horde/HordeServer/Storage/Services/SimpleStorageService.cs
@@ -0,0 +1,119 @@
+using EpicGames.Core;
+using HordeServer.Storage.Impl;
+using HordeServer.Storage.Primitives;
+using Microsoft.Extensions.Options;
+using System;
+using System.Collections.Generic;
+using System.IO;
+using System.Linq;
+using System.Reflection;
+using System.Threading.Tasks;
+
+namespace HordeServer.Storage.Services
+{
+	/// <summary>
+	/// Very basic implementation of <see cref="IStorageService"/>. Backed by a filesystem, and uses timestamps for GC. Data may be evicted at any time.
+	/// </summary>
+	class SimpleStorageService : IStorageService
+	{
+		/// <summary>
+		/// Path for storing blobs
+		/// </summary>
+		const string BlobsPath = "SimpleStorage/Blobs";
+
+		/// <summary>
+		/// Path for storing refs
+		/// </summary>
+		const string RefsPath = "SimpleStorage/Refs";
+
+		/// <summary>
+		/// The storage backend
+		/// </summary>
+		IStorageBackend Backend;
+
+		/// <summary>
+		/// Constructor
+		/// </summary>
+		/// <param name="Backend">The storage backend</param>
+		public SimpleStorageService(IStorageBackend Backend)
+		{
+			this.Backend = Backend;
+		}
+
+		/// <summary>
+		/// Gets the for a file with the given hash
+		/// </summary>
+		/// <param name="BaseDir">Base directory for storage</param>
+		/// <param name="Hash">Hash of the file</param>
+		/// <returns>Path to the file for this blob</returns>
+		static string GetItemPath(string BaseDir, BlobHash Hash)
+		{
+			return $"{BaseDir}/{StringUtils.FormatHexString(Hash.Span.Slice(0, 1))}/{StringUtils.FormatHexString(Hash.Span.Slice(1, 1))}/{Hash}.dat";
+		}
+
+		/// <inheritdoc/>
+		public async Task<ReadOnlyMemory<byte>?> TryGetBlobAsync(BlobHash Hash, DateTime Deadline = default)
+		{
+			string Path = GetItemPath(BlobsPath, Hash);
+			return await Backend.ReadAsync(Path);
+		}
+
+		/// <inheritdoc/>
+		public async Task PutBlobAsync(BlobHash Hash, ReadOnlyMemory<byte> Value)
+		{
+			BlobHash RealHash = BlobHash.Compute(Value);
+			if(Hash != RealHash)
+			{
+				throw new InvalidDataException($"Hash for blob does not match. Expected {RealHash}, got {Hash}");
+			}
+
+			string Path = GetItemPath(BlobsPath, Hash);
+			if (!await Backend.TouchAsync(Path))
+			{
+				await Backend.WriteAsync(Path, Value);
+			}
+		}
+
+		/// <inheritdoc/>
+		public Task<bool> ShouldPutBlobAsync(BlobHash Hash)
+		{
+			// Always return true to force timestamps to be updated
+			return Task.FromResult(true);
+		}
+
+		/// <inheritdoc/>
+		public async Task SetRefAsync(BlobHash Key, BlobHash? Value)
+		{
+			string Path = GetItemPath(RefsPath, Key);
+			if (Value == null)
+			{
+				await Backend.DeleteAsync(Path);
+			}
+			else
+			{
+				await Backend.WriteAsync(Path, Value.Value.Memory);
+			}
+		}
+
+		/// <inheritdoc/>
+		public async Task TouchRefAsync(BlobHash Key)
+		{
+			string Path = GetItemPath(RefsPath, Key);
+			await Backend.TouchAsync(Path);
+		}
+
+		/// <inheritdoc/>
+		public async Task<BlobHash?> GetRefAsync(BlobHash Key, TimeSpan MaxDrift = default)
+		{
+			ReadOnlyMemory<byte>? Memory = await Backend.ReadAsync(GetItemPath(RefsPath, Key));
+			if (Memory == null)
+			{
+				return null;
+			}
+			else
+			{
+				return new BlobHash(Memory.Value);
+			}
+		}
+	}
+}
diff --git a/Engine/Source/Programs/Horde/HordeServer/Tasks/ITaskSource.cs b/Engine/Source/Programs/Horde/HordeServer/Tasks/ITaskSource.cs
new file mode 100644
--- /dev/null
+++ b/Engine/Source/Programs/Horde/HordeServer/Tasks/ITaskSource.cs
@@ -0,0 +1,143 @@
+using Google.Protobuf;
+using Google.Protobuf.Reflection;
+using Google.Protobuf.WellKnownTypes;
+using HordeServer.Models;
+using HordeServer.Services;
+using MongoDB.Bson;
+using System;
+using System.Collections.Generic;
+using System.Linq;
+using System.Threading.Tasks;
+
+namespace HordeServer.Tasks
+{
+	/// <summary>
+	/// Information about a new task
+	/// </summary>
+	public class NewLeaseInfo
+	{
+		/// <summary>
+		/// Lease to execute the task
+		/// </summary>
+		public AgentLease Lease { get; }
+
+		/// <summary>
+		/// Callback for connection to the agent being lost
+		/// </summary>
+		public Action? OnConnectionLost { get; }
+
+		/// <summary>
+		/// Constructor
+		/// </summary>
+		/// <param name="Lease"></param>
+		/// <param name="OnConnectionLost"></param>
+		public NewLeaseInfo(AgentLease Lease, Action? OnConnectionLost = null)
+		{
+			this.Lease = Lease;
+			this.OnConnectionLost = OnConnectionLost;
+		}
+	}
+
+	/// <summary>
+	/// Interface for a an object that waits for a task
+	/// </summary>
+	public interface ITaskListener : IAsyncDisposable
+	{
+		/// <summary>
+		/// Task returning a lease for the agent. If a null lease is returned, subsequent task sources will not be checked for leases.
+		/// </summary>
+		Task<NewLeaseInfo?> LeaseTask { get; }
+
+		/// <summary>
+		/// Whether this lease has been accepted. An lease which has not been accepted will be returned to the pool when the subscription is disposed.
+		/// </summary>
+		bool Accepted { get; set; }
+	}
+
+	class TaskSubscription : ITaskListener
+	{
+		public Task<NewLeaseInfo?> LeaseTask { get; }
+
+		public bool Accepted { get; set; }
+
+		public TaskSubscription(Task<NewLeaseInfo?> LeaseTask)
+		{
+			this.LeaseTask = LeaseTask;
+		}
+
+		public static TaskSubscription FromResult(AgentLease? Lease)
+		{
+			if (Lease == null)
+			{
+				return new TaskSubscription(Task.FromResult<NewLeaseInfo?>(null));
+			}
+			else
+			{
+				return new TaskSubscription(Task.FromResult<NewLeaseInfo?>(new NewLeaseInfo(Lease)));
+			}
+		}
+
+		public ValueTask DisposeAsync()
+		{
+			return default;
+		}
+	}
+
+	/// <summary>
+	/// Handler for a certain lease type
+	/// </summary>
+	public interface ITaskSource
+	{
+		/// <summary>
+		/// Descriptor for this lease type
+		/// </summary>
+		MessageDescriptor Descriptor { get; }
+
+		/// <summary>
+		/// Subscribes an agent to execute leases for this source. The lifetime of the returned subscription object indicates the time that the agent is volunteering for work.
+		/// </summary>
+		/// <param name="Agent">The agent state</param>
+		/// <returns>The subscription object</returns>
+		Task<ITaskListener?> SubscribeAsync(IAgent Agent);
+
+		/// <summary>
+		/// Force a lease as complete
+		/// </summary>
+		/// <param name="Agent">The agent that was allocated to the lease</param>
+		/// <param name="LeaseId">The lease id</param>
+		/// <param name="Payload">The lease payload</param>
+		/// <returns>Async task</returns>
+		Task AbortTaskAsync(IAgent Agent, ObjectId LeaseId, Any Payload);
+	}
+
+	/// <summary>
+	/// Strongly typed base class for implementation of ITaskHandler
+	/// </summary>
+	/// <typeparam name="T">Type of task message</typeparam>
+	public abstract class TaskSourceBase<T> : ITaskSource where T : IMessage<T>, new()
+	{
+		/// <inheritdoc/>
+		public MessageDescriptor Descriptor { get; } = new T().Descriptor;
+
+		/// <inheritdoc/>
+		public abstract Task<ITaskListener?> SubscribeAsync(IAgent Agent);
+
+		/// <inheritdoc/>
+		public Task AbortTaskAsync(IAgent Agent, ObjectId LeaseId, Any Payload)
+		{
+			return AbortTaskAsync(Agent, LeaseId, Payload.Unpack<T>());
+		}
+
+		/// <summary>
+		/// Cancels the given task
+		/// </summary>
+		/// <param name="Agent">Agent executing the task</param>
+		/// <param name="LeaseId">The lease id</param>
+		/// <param name="Payload">The task type</param>
+		/// <returns>Async task</returns>
+		protected virtual Task AbortTaskAsync(IAgent Agent, ObjectId LeaseId, T Payload)
+		{
+			return Task.CompletedTask;
+		}
+	}
+}
