Initial drop of Horde Server/Agent for UE5/Main. Still missing some Epic-specific functionality required to build/run correctly.

diff --git a/Engine/Source/Programs/Horde/HordeServer/Tasks/Impl/ActionTaskSource.cs b/Engine/Source/Programs/Horde/HordeServer/Tasks/Impl/ActionTaskSource.cs
new file mode 100644
--- /dev/null
+++ b/Engine/Source/Programs/Horde/HordeServer/Tasks/Impl/ActionTaskSource.cs
@@ -0,0 +1,555 @@
+using Build.Bazel.Remote.Execution.V2;
+using EpicGames.Core;
+using Google.LongRunning;
+using Google.Protobuf;
+using Google.Protobuf.Reflection;
+using Google.Protobuf.WellKnownTypes;
+using Grpc.Core;
+using HordeCommon;
+using HordeCommon.Rpc.Tasks;
+using HordeServer.Api;
+using HordeServer.Models;
+using HordeServer.Rpc;
+using HordeServer.Services;
+using Microsoft.Extensions.Logging;
+using MongoDB.Bson;
+using System;
+using System.Collections.Concurrent;
+using System.Collections.Generic;
+using System.Diagnostics;
+using System.Diagnostics.CodeAnalysis;
+using System.Linq;
+using System.Threading;
+using System.Threading.Tasks;
+
+namespace HordeServer.Tasks.Impl
+{
+	/// <summary>
+	/// Interface for an action operation
+	/// </summary>
+	public interface IActionExecuteOperation
+	{
+		/// <summary>
+		/// The unique id for this operation
+		/// </summary>
+		ObjectId Id { get; }
+
+		/// <summary>
+		/// Sets the result for this operation
+		/// </summary>
+		/// <param name="Result">The result for this operation. If null, the operation will be rescheduled to operate on a new worker</param>
+		/// <returns></returns>
+		bool TrySetResult(ActionResult? Result);
+
+		/// <summary>
+		/// Read status updates from the operation
+		/// </summary>
+		/// <returns>Sequence of operation updates</returns>
+		IAsyncEnumerable<Operation> ReadStatusUpdatesAsync();
+	}
+
+	/// <summary>
+	/// Dispatches remote actions. Does not implement any cross-pod communication to satisfy leases; only agents connected to this server instance will be stored.
+	/// </summary>
+	public class ActionTaskSource : ITaskSource
+	{
+		/// <summary>
+		/// Number of buckets to create for storing operations. Operations are garbage collected from random buckets.
+		/// </summary>
+		const int NumBuckets = 1024;
+
+		/// <summary>
+		/// Number of buckets to sample when finding an operation to evict from the cache. Only operations older than <see cref="MinAgeForGC"/> will be considered for deletion.
+		/// </summary>
+		const int NumSamplesForGC = 4;
+
+		/// <summary>
+		/// Minimum age for operations to be garbage collected
+		/// </summary>
+		static readonly long MinAgeForGC = TimeSpan.FromMinutes(15.0).Ticks;
+
+		/// <summary>
+		/// Tracks an execution operation
+		/// </summary>
+		class ExecuteOperation : IActionExecuteOperation
+		{
+			/// <summary>
+			/// The operation id
+			/// </summary>
+			public ObjectId Id { get; }
+
+			/// <summary>
+			/// The execution request
+			/// </summary>
+			public ExecuteRequest Request { get; }
+
+			/// <summary>
+			/// Timestamp indicating the last time the operation was accessed. Used to evict items from the list of tracked operations.
+			/// </summary>
+			public long Timestamp = Stopwatch.GetTimestamp();
+
+			/// <summary>
+			/// The executing task
+			/// </summary>
+			public Task? BackgroundTask;
+
+			/// <summary>
+			/// The current status of the operation
+			/// </summary>
+			public Operation? Status;
+
+			/// <summary>
+			/// Event which is signalled whenever the status changes
+			/// </summary>
+			TaskCompletionSource<bool> StateChangedEvent = new TaskCompletionSource<bool>();
+
+			/// <summary>
+			/// The result from executing this operation
+			/// </summary>
+			public TaskCompletionSource<ActionResult?> ResultTaskSource = new TaskCompletionSource<ActionResult?>();
+
+			/// <summary>
+			/// Constructor
+			/// </summary>
+			/// <param name="Request">The request to execute</param>
+			public ExecuteOperation(ExecuteRequest Request)
+			{
+				this.Id = ObjectId.GenerateNewId();
+				this.Request = Request;
+			}
+
+			/// <summary>
+			/// Sets the operation result
+			/// </summary>
+			/// <param name="Result">The result of the operation</param>
+			public bool TrySetResult(ActionResult? Result)
+			{
+				return ResultTaskSource.TrySetResult(Result);
+			}
+
+			/// <summary>
+			/// Updates the current status
+			/// </summary>
+			/// <param name="Stage">The current stage of execution</param>
+			/// <param name="Response">The final response from the operation</param>
+			/// <returns>False if the operation is already complete</returns>
+			public bool SetStatus(ExecutionStage.Types.Value Stage, Any? Response)
+			{
+				ExecuteOperationMetadata Metadata = new ExecuteOperationMetadata();
+				Metadata.Stage = Stage;
+				Metadata.ActionDigest = Request.ActionDigest;
+
+				Operation Operation = new Operation();
+				if (Response != null)
+				{
+					Operation.Response = Response;
+					Operation.Done = true;
+				}
+				Operation.Metadata = Any.Pack(Metadata);
+
+				return TrySetStatus(Operation);
+			}
+
+			/// <summary>
+			/// Updates the current status
+			/// </summary>
+			/// <param name="NewStatus">The operation status</param>
+			/// <returns>False if the operation is already complete</returns>
+			public bool TrySetStatus(Operation NewStatus)
+			{
+				for (; ; )
+				{
+					TaskCompletionSource<bool> PrevStateChangedEvent = StateChangedEvent;
+
+					Operation? CurrentStatus = Status;
+					if (CurrentStatus != null && CurrentStatus.Done)
+					{
+						return false;
+					}
+					if (Interlocked.CompareExchange(ref Status, NewStatus, CurrentStatus) == CurrentStatus)
+					{
+						TaskCompletionSource<bool> NextStateChangedEvent = new TaskCompletionSource<bool>();
+						while (Interlocked.CompareExchange(ref StateChangedEvent, NextStateChangedEvent, PrevStateChangedEvent) != PrevStateChangedEvent)
+						{
+							PrevStateChangedEvent = StateChangedEvent;
+						}
+						PrevStateChangedEvent.TrySetResult(true);
+						return true;
+					}
+				}
+			}
+
+			/// <summary>
+			/// Read status updates from the operation
+			/// </summary>
+			/// <returns>Sequence of operation updates</returns>
+			public async IAsyncEnumerable<Operation> ReadStatusUpdatesAsync()
+			{
+				Operation? LastStatus = null;
+				for (; ; )
+				{
+					TaskCompletionSource<bool> StateChangedEventCopy = StateChangedEvent;
+
+					Operation? NextStatus = Status;
+					if (NextStatus != LastStatus)
+					{
+						yield return NextStatus!;
+
+						if (NextStatus!.Done)
+						{
+							break;
+						}
+					}
+					LastStatus = NextStatus;
+
+					await StateChangedEventCopy.Task;
+
+					Interlocked.CompareExchange(ref StateChangedEvent, new TaskCompletionSource<bool>(), StateChangedEventCopy);
+				}
+			}
+		}
+
+		/// <summary>
+		/// Tracks an agent waiting for work
+		/// </summary>
+		class Subscription : ITaskListener
+		{
+			/// <summary>
+			/// The owning instance
+			/// </summary>
+			ActionTaskSource Outer;
+
+			/// <summary>
+			/// Agent waiting for work
+			/// </summary>
+			public IAgent Agent { get; }
+
+			/// <summary>
+			/// The operation assigned to this subscription. May be set to null if the subscription is cancelled.
+			/// </summary>
+			public TaskCompletionSource<NewLeaseInfo?> LeaseTaskSource = new TaskCompletionSource<NewLeaseInfo?>();
+
+			/// <inheritdoc/>
+			public Task<NewLeaseInfo?> LeaseTask => LeaseTaskSource.Task;
+
+			/// <inheritdoc/>
+			public bool Accepted { get; set; }
+
+			/// <summary>
+			/// Constructor
+			/// </summary>
+			/// <param name="Outer">The task source</param>
+			/// <param name="Agent">The agent instance</param>
+			public Subscription(ActionTaskSource Outer, IAgent Agent)
+			{
+				this.Outer = Outer;
+				this.Agent = Agent;
+
+				lock (Outer.Subscriptions)
+				{
+					Outer.Subscriptions.Add(this);
+				}
+			}
+
+			/// <summary>
+			/// Attempts to assign a lease to this agent
+			/// </summary>
+			/// <param name="Lease">The lease to assign</param>
+			/// <param name="OnConnectionLost"></param>
+			/// <returns>True if the lease was set, false otherwise</returns>
+			public bool TrySetLease(AgentLease Lease, System.Action OnConnectionLost)
+			{
+				return LeaseTaskSource.TrySetResult(new NewLeaseInfo(Lease, OnConnectionLost));
+			}
+
+			/// <summary>
+			/// Cancel of this subscription
+			/// </summary>
+			/// <returns></returns>
+			public async ValueTask DisposeAsync()
+			{
+				// Remove from the list of active subscriptions
+				lock (Outer.Subscriptions)
+				{
+					Outer.Subscriptions.Remove(this);
+				}
+
+				// If an operation has been assigned to this subscription and it wasn't accepted, attempt to reassign it
+				if (!LeaseTaskSource.TrySetResult(null) && !Accepted)
+				{
+					NewLeaseInfo? NewLeaseInfo = await LeaseTaskSource.Task;
+					if (NewLeaseInfo != null)
+					{
+						await Outer.CancelLeaseAsync(NewLeaseInfo.Lease);
+					}
+				}
+			}
+		}
+
+		ActionCacheService ActionCacheService;
+		ILogFileService LogFileService;
+		HashSet<Subscription> Subscriptions = new HashSet<Subscription>();
+		Dictionary<ObjectId, ExecuteOperation>[] ActiveOperationBuckets = new Dictionary<ObjectId, ExecuteOperation>[NumBuckets];
+		Random Random = new Random();
+		ILogger Logger;
+
+		/// <inheritdoc/>
+		public MessageDescriptor Descriptor => ActionTask.Descriptor;
+
+		/// <summary>
+		/// Constructor
+		/// </summary>
+		/// <param name="ActionCacheService"></param>
+		/// <param name="LogFileService">The log file service instance</param>
+		/// <param name="Logger">The logger instance</param>
+		public ActionTaskSource(ActionCacheService ActionCacheService, ILogFileService LogFileService, ILogger<ActionTaskSource> Logger)
+		{
+			this.ActionCacheService = ActionCacheService;
+			this.LogFileService = LogFileService;
+			this.Logger = Logger;
+
+			for (int Idx = 0; Idx < NumBuckets; Idx++)
+			{
+				ActiveOperationBuckets[Idx] = new Dictionary<ObjectId, ExecuteOperation>();
+			}
+		}
+
+		Task CancelLeaseAsync(AgentLease Lease)
+		{
+			ExecuteOperation? Operation;
+			if (TryGetOperation(Lease.Id, out Operation))
+			{
+				Operation.TrySetResult(null);
+			}
+			return Task.CompletedTask;
+		}
+
+		static int GetBucketIdx(ObjectId Id)
+		{
+			return Id.GetHashCode() & (NumBuckets - 1);
+		}
+
+		Dictionary<ObjectId, ExecuteOperation> GetBucket(ObjectId Id)
+		{
+			return ActiveOperationBuckets[GetBucketIdx(Id)];
+		}
+
+		/// <summary>
+		/// Attempts to assign an action for execution
+		/// </summary>
+		/// <param name="Request">Execution request</param>
+		/// <returns>The operation created to execute this action</returns>
+		public IActionExecuteOperation Execute(ExecuteRequest Request)
+		{
+			ExecuteOperation ExecuteOperation = new ExecuteOperation(Request);
+			ExecuteOperation.BackgroundTask = Task.Run(() => RunAsync(ExecuteOperation));
+
+			Dictionary<ObjectId, ExecuteOperation> Bucket = GetBucket(ExecuteOperation.Id);
+			lock (Bucket)
+			{
+				Bucket[ExecuteOperation.Id] = ExecuteOperation;
+			}
+
+			TrimOnce();
+
+			return ExecuteOperation;
+		}
+
+		/// <summary>
+		/// Run the task as a background operation
+		/// </summary>
+		/// <returns></returns>
+		async Task RunAsync(ExecuteOperation Operation)
+		{
+			ExecuteResponse Response = await RunInternalAsync(Operation);
+			Operation.SetStatus(ExecutionStage.Types.Value.Completed, Any.Pack(Response));
+		}
+
+		/// <summary>
+		/// Run the task as a background operation
+		/// </summary>
+		/// <returns></returns>
+		async Task<ExecuteResponse> RunInternalAsync(ExecuteOperation Operation)
+		{
+			ExecuteRequest Request = Operation.Request;
+			if (!Request.SkipCacheLookup)
+			{
+				Operation.SetStatus(ExecutionStage.Types.Value.CacheCheck, null);
+
+				ActionResult? Result = await ActionCacheService.TryGetResult(Request.ActionDigest);
+				if (Result != null)
+				{
+					ExecuteResponse CachedResponse = new ExecuteResponse();
+					CachedResponse.Result = Result;
+					CachedResponse.CachedResult = true;
+					CachedResponse.Status = new Status(StatusCode.OK, String.Empty);
+					return CachedResponse;
+				}
+			}
+
+			for(; ;)
+			{
+				Operation.SetStatus(ExecutionStage.Types.Value.Queued, null);
+
+				// Get a new subscription
+				Subscription? Subscription = await WaitForSubscriberAsync(TimeSpan.FromSeconds(20.0));
+				if (Subscription == null)
+				{
+					ExecuteResponse FailedResponse = new ExecuteResponse();
+					FailedResponse.Status = new Status(StatusCode.ResourceExhausted, "Unable to find agent to execute request");
+					return FailedResponse;
+				}
+
+				// Create a lease for it
+				string LeaseName = $"Remote action ({Request.ActionDigest.Hash.Substring(0, 16)}...)";
+
+				ILogFile LogFile = await LogFileService.CreateLogFileAsync(ObjectId.Empty, Subscription.Agent.SessionId, LogType.Json);
+
+				ActionTask ActionTask = new ActionTask();
+				ActionTask.Digest = Request.ActionDigest;
+				ActionTask.LogId = LogFile.Id.ToString();
+
+				byte[] Payload = Any.Pack(ActionTask).ToByteArray();
+				AgentLease Lease = new AgentLease(Operation.Id, LeaseName, null, null, LogFile.Id, LeaseState.Pending, Payload, new AgentRequirements(), null);
+
+				// Try to set the lease on the subscription. If it fails, the subscriber has already been allocated
+				Operation.ResultTaskSource = new TaskCompletionSource<ActionResult?>();
+				if (Subscription.TrySetLease(Lease, () => Operation.ResultTaskSource.TrySetResult(null)))
+				{
+					Operation.SetStatus(ExecutionStage.Types.Value.Executing, null);
+
+					ActionResult? Result = await Operation.ResultTaskSource.Task;
+					if (Result != null)
+					{
+						ExecuteResponse SuccessResponse = new ExecuteResponse();
+						SuccessResponse.Status = new Status(StatusCode.OK, String.Empty);
+						SuccessResponse.Result = Result;
+						return SuccessResponse;
+					}
+				}
+			}
+		}
+
+		/// <summary>
+		/// Waits until there is an available agent for executing work
+		/// </summary>
+		/// <param name="TimeLimit"></param>
+		/// <returns></returns>
+		async Task<Subscription?> WaitForSubscriberAsync(TimeSpan TimeLimit)
+		{
+			Stopwatch Timer = Stopwatch.StartNew();
+			for (; ; )
+			{
+				// Try to allocate a subscription
+				lock (Subscriptions)
+				{
+					Subscription? Subscription = Subscriptions.FirstOrDefault();
+					if (Subscription != null)
+					{
+						Subscriptions.Remove(Subscription);
+						return Subscription;
+					}
+				}
+
+				// Wait until the next poll time
+				TimeSpan Delay = TimeLimit - Timer.Elapsed;
+				if (Delay < TimeSpan.Zero)
+				{
+					return null;
+				}
+				TimeSpan MaxDelay = TimeSpan.FromMilliseconds(200);
+				if (Delay > MaxDelay)
+				{
+					Delay = MaxDelay;
+				}
+				await Task.Delay(Delay);
+			}
+		}
+
+		/// <summary>
+		/// Attempts to get an existing operation by id
+		/// </summary>
+		/// <param name="OperationId"></param>
+		/// <param name="Operation"></param>
+		/// <returns></returns>
+		bool TryGetOperation(ObjectId OperationId, [NotNullWhen(true)] out ExecuteOperation? Operation)
+		{
+			Dictionary<ObjectId, ExecuteOperation> Bucket = GetBucket(OperationId);
+			lock (Bucket)
+			{
+				return Bucket.TryGetValue(OperationId, out Operation);
+			}
+		}
+
+		/// <summary>
+		/// Attempts to get an existing operation by id
+		/// </summary>
+		/// <param name="OperationId"></param>
+		/// <param name="Operation"></param>
+		/// <returns></returns>
+		public bool TryGetOperation(ObjectId OperationId, [NotNullWhen(true)] out IActionExecuteOperation? Operation)
+		{
+			ExecuteOperation? ExecuteOperation;
+			if (TryGetOperation(OperationId, out ExecuteOperation))
+			{
+				Operation = ExecuteOperation;
+				return true;
+			}
+			else
+			{
+				Operation = null;
+				return false;
+			}
+		}
+
+		void TrimOnce()
+		{
+			ObjectId? LeaseId = null;
+			ExecuteOperation? Operation = null;
+			Dictionary<ObjectId, ExecuteOperation>? Bucket = null;
+
+			// Sample several buckets and find the oldest operation
+			long Timestamp = Stopwatch.GetTimestamp() - MinAgeForGC;
+			for (int AttemptIdx = 0; AttemptIdx < NumSamplesForGC; AttemptIdx++)
+			{
+				Dictionary<ObjectId, ExecuteOperation> NextBucket = ActiveOperationBuckets[Random.Next() & (NumBuckets - 1)];
+				lock(NextBucket)
+				{
+					foreach (KeyValuePair<ObjectId, ExecuteOperation> NextItem in NextBucket)
+					{
+						if (NextItem.Value.Timestamp < Timestamp && (NextItem.Value.BackgroundTask?.IsCompleted ?? true))
+						{
+							Bucket = NextBucket;
+							LeaseId = NextItem.Key;
+							Operation = NextItem.Value;
+							Timestamp = NextItem.Value.Timestamp;
+							break;
+						}
+					}
+				}
+			}
+
+			// If we found one, 
+			if (LeaseId != null)
+			{
+				lock (Bucket!)
+				{
+					Operation!.TrySetResult(null);
+					Bucket.Remove(LeaseId.Value);
+				}
+			}
+		}
+
+		/// <inheritdoc/>
+		public Task<ITaskListener?> SubscribeAsync(IAgent Agent)
+		{
+			return Task.FromResult<ITaskListener?>(new Subscription(this, Agent));
+		}
+
+		/// <inheritdoc/>
+		public Task AbortTaskAsync(IAgent Agent, ObjectId LeaseId, Any Payload)
+		{
+			return Task.CompletedTask;
+		}
+	}
+}
diff --git a/Engine/Source/Programs/Horde/HordeServer/Tasks/Impl/ConformTaskSource.cs b/Engine/Source/Programs/Horde/HordeServer/Tasks/Impl/ConformTaskSource.cs
new file mode 100644
--- /dev/null
+++ b/Engine/Source/Programs/Horde/HordeServer/Tasks/Impl/ConformTaskSource.cs
@@ -0,0 +1,232 @@
+using Google.Protobuf;
+using Google.Protobuf.Reflection;
+using Google.Protobuf.WellKnownTypes;
+using HordeCommon;
+using HordeCommon.Rpc.Tasks;
+using HordeServer.Api;
+using HordeServer.Models;
+using HordeServer.Services;
+using HordeServer.Utilities;
+using Microsoft.Extensions.Logging;
+using MongoDB.Bson;
+using System;
+using System.Collections.Generic;
+using System.Linq;
+using System.Threading.Tasks;
+
+namespace HordeServer.Tasks.Impl
+{
+	/// <summary>
+	/// Generates tasks telling agents to sync their workspaces
+	/// </summary>
+	public class ConformTaskSource : TaskSourceBase<ConformTask>
+	{
+		DatabaseService DatabaseService;
+		PoolService PoolService;
+		SingletonDocument<ConformList> ConformList;
+		PerforceLoadBalancer PerforceLoadBalancer;
+		ILogFileService LogService;
+		ILogger Logger;
+
+		/// <summary>
+		/// Constructor
+		/// </summary>
+		/// <param name="DatabaseService"></param>
+		/// <param name="PoolService"></param>
+		/// <param name="LogService"></param>
+		/// <param name="PerforceLoadBalancer"></param>
+		/// <param name="Logger"></param>
+		public ConformTaskSource(DatabaseService DatabaseService, PoolService PoolService, ILogFileService LogService, PerforceLoadBalancer PerforceLoadBalancer, ILogger<ConformTaskSource> Logger)
+		{
+			this.DatabaseService = DatabaseService;
+			this.PoolService = PoolService;
+			this.ConformList = new SingletonDocument<ConformList>(DatabaseService);
+			this.PerforceLoadBalancer = PerforceLoadBalancer;
+			this.LogService = LogService;
+			this.Logger = Logger;
+		}
+
+		/// <inheritdoc/>
+		public override async Task<ITaskListener?> SubscribeAsync(IAgent Agent)
+		{
+			DateTime UtcNow = DateTime.UtcNow;
+			if (!await IsConformPendingAsync(Agent, UtcNow))
+			{
+				return null;
+			}
+
+			if (Agent.Leases.Count == 0)
+			{
+				ObjectId LeaseId = ObjectId.GenerateNewId();
+				if (await AllocateConformLeaseAsync(LeaseId))
+				{
+					ConformTask Task = new ConformTask();
+					if (!await GetWorkspacesAsync(Agent, Task.Workspaces))
+					{
+						return null;
+					}
+
+					ILogFile Log = await LogService.CreateLogFileAsync(ObjectId.Empty, Agent.SessionId, LogType.Json);
+					Task.LogId = Log.Id.ToString();
+
+					byte[] Payload = Any.Pack(Task).ToByteArray();
+
+					AgentLease Lease = new AgentLease(LeaseId, "Updating workspaces", null, null, Log.Id, LeaseState.Pending, Payload, new AgentRequirements(), null);
+					return TaskSubscription.FromResult(Lease);
+				}
+			}
+
+			if (Agent.RequestConform)
+			{
+				return TaskSubscription.FromResult(null);
+			}
+
+			return null;
+		}
+
+		/// <inheritdoc/>
+		public async Task<bool> GetWorkspacesAsync(IAgent Agent, IList<HordeCommon.Rpc.Messages.AgentWorkspace> Workspaces)
+		{
+			Globals Globals = await DatabaseService.GetGlobalsAsync();
+
+			HashSet<AgentWorkspace> ConformWorkspaces = await PoolService.GetWorkspacesAsync(Agent, DateTime.UtcNow);
+			foreach (AgentWorkspace ConformWorkspace in ConformWorkspaces)
+			{
+				if (!await Agent.TryAddWorkspaceMessage(ConformWorkspace, Globals, PerforceLoadBalancer, Workspaces))
+				{
+					return false;
+				}
+			}
+
+			return true;
+		}
+
+		/// <inheritdoc/>
+		protected override Task AbortTaskAsync(IAgent Agent, ObjectId LeaseId, ConformTask Task)
+		{
+			return ReleaseConformLeaseAsync(LeaseId);
+		}
+
+		/// <summary>
+		/// Atempt to allocate a conform resource for the given lease
+		/// </summary>
+		/// <param name="LeaseId">The lease id</param>
+		/// <returns>True if the resource was allocated, false otherwise</returns>
+		private async Task<bool> AllocateConformLeaseAsync(ObjectId LeaseId)
+		{
+			for (; ; )
+			{
+				ConformList CurrentValue = await ConformList.GetAsync();
+				if (CurrentValue.MaxCount != 0 && CurrentValue.LeaseIds.Count >= CurrentValue.MaxCount)
+				{
+					return false;
+				}
+
+				CurrentValue.LeaseIds.Add(LeaseId);
+				if (await ConformList.TryUpdateAsync(CurrentValue))
+				{
+					Logger.LogInformation("Added conform lease {LeaseId}", LeaseId);
+					return true;
+				}
+			}
+		}
+
+		/// <summary>
+		/// Terminate a conform lease
+		/// </summary>
+		/// <param name="LeaseId">The lease id</param>
+		/// <returns>Async task</returns>
+		public async Task ReleaseConformLeaseAsync(ObjectId LeaseId)
+		{
+			for (; ; )
+			{
+				ConformList CurrentValue = await ConformList.GetAsync();
+				if (!CurrentValue.LeaseIds.Remove(LeaseId))
+				{
+					Logger.LogInformation("Conform lease {LeaseId} is not in singelton", LeaseId);
+					break;
+				}
+				if (await ConformList.TryUpdateAsync(CurrentValue))
+				{
+					Logger.LogInformation("Removed conform lease {LeaseId}", LeaseId);
+					break;
+				}
+			}
+		}
+
+		/// <summary>
+		/// Determine if an agent should be conformed
+		/// </summary>
+		/// <param name="Agent">The agent to test</param>
+		/// <param name="UtcNow">Current time</param>
+		/// <returns>True if the agent should be conformed</returns>
+		private async Task<bool> IsConformPendingAsync(IAgent Agent, DateTime UtcNow)
+		{
+			// If a conform was manually requested, allow it to run even if the agent is disabled
+			if (Agent.RequestConform)
+			{
+				return !IsConformCoolDownPeriod(Agent, UtcNow);
+			}
+
+			// Otherwise only run if the agent is enabled
+			if (Agent.Enabled)
+			{
+				// If we've attempted (and failed) a conform, run again after a certain time
+				if (Agent.ConformAttemptCount.HasValue)
+				{
+					return !IsConformCoolDownPeriod(Agent, UtcNow);
+				}
+
+				// Always run a conform every 24h
+				if (UtcNow > Agent.LastConformTime + TimeSpan.FromDays(1.0))
+				{
+					return true;
+				}
+
+				// Check if the workspaces have changed (first check against a cached list of workspaces, then an accurate one)
+				HashSet<AgentWorkspace> Workspaces = await PoolService.GetWorkspacesAsync(Agent, UtcNow - TimeSpan.FromSeconds(30.0));
+				if (!Workspaces.SetEquals(Agent.Workspaces))
+				{
+					Workspaces = await PoolService.GetWorkspacesAsync(Agent, UtcNow);
+					if (!Workspaces.SetEquals(Agent.Workspaces))
+					{
+						return true;
+					}
+				}
+			}
+			return false;
+		}
+
+		/// <summary>
+		/// Determines if the given agent is in a cooldown period, waiting to retry after a failed conform
+		/// </summary>
+		/// <param name="Agent">The agent to test</param>
+		/// <param name="UtcNow">Current time</param>
+		/// <returns></returns>
+		private static bool IsConformCoolDownPeriod(IAgent Agent, DateTime UtcNow)
+		{
+			if (!Agent.ConformAttemptCount.HasValue)
+			{
+				return false;
+			}
+
+			TimeSpan RetryTime;
+			switch (Agent.ConformAttemptCount.Value)
+			{
+				case 1:
+					RetryTime = TimeSpan.FromMinutes(5.0);
+					break;
+				case 2:
+					RetryTime = TimeSpan.FromMinutes(20.0);
+					break;
+				case 3:
+					RetryTime = TimeSpan.FromHours(1.0);
+					break;
+				default:
+					RetryTime = TimeSpan.FromHours(6.0);
+					break;
+			}
+			return UtcNow < Agent.LastConformTime + RetryTime;
+		}
+	}
+}
diff --git a/Engine/Source/Programs/Horde/HordeServer/Tasks/Impl/JobTaskSource.cs b/Engine/Source/Programs/Horde/HordeServer/Tasks/Impl/JobTaskSource.cs
new file mode 100644
--- /dev/null
+++ b/Engine/Source/Programs/Horde/HordeServer/Tasks/Impl/JobTaskSource.cs
@@ -0,0 +1,1098 @@
+// Copyright Epic Games, Inc. All Rights Reserved.
+
+using EpicGames.Core;
+using Google.Protobuf;
+using Google.Protobuf.Reflection;
+using Google.Protobuf.WellKnownTypes;
+using HordeServer.Api;
+using HordeServer.Collections;
+using HordeCommon;
+using HordeServer.Models;
+using HordeCommon.Rpc;
+using HordeCommon.Rpc.Tasks;
+using HordeServer.Tasks;
+using HordeServer.Utilities;
+using Microsoft.Extensions.Hosting;
+using Microsoft.Extensions.Logging;
+using Microsoft.Extensions.Options;
+using MongoDB.Bson;
+using MongoDB.Driver;
+using System;
+using System.Collections.Generic;
+using System.Diagnostics;
+using System.Diagnostics.CodeAnalysis;
+using System.Linq;
+using System.Runtime.CompilerServices;
+using System.Text;
+using System.Threading;
+using System.Threading.Tasks;
+
+using PoolId = HordeServer.Utilities.StringId<HordeServer.Models.IPool>;
+using StreamId = HordeServer.Utilities.StringId<HordeServer.Models.IStream>;
+using HordeServer.Services;
+
+namespace HordeServer.Tasks.Impl
+{
+	/// <summary>
+	/// Background service to dispatch pending work to agents in priority order.
+	/// </summary>
+	public class JobTaskSource : TickedBackgroundService, ITaskSource
+	{
+		/// <summary>
+		/// An item in the queue to be executed
+		/// </summary>
+		class QueueItem
+		{
+			/// <summary>
+			/// The stream for this job
+			/// </summary>
+			public IStream Stream;
+
+			/// <summary>
+			/// The job instance
+			/// </summary>
+			public IJob Job;
+
+			/// <summary>
+			/// Index of the batch within this job to be executed
+			/// </summary>
+			public int BatchIdx;
+
+			/// <summary>
+			/// Requirements for the agent to execute this item
+			/// </summary>
+			public AgentRequirements Requirements;
+
+			/// <summary>
+			/// The pool of machines to allocate from
+			/// </summary>
+			public IPool Pool;
+
+			/// <summary>
+			/// The type of workspace that this item should run in
+			/// </summary>
+			public AgentWorkspace Workspace;
+
+			/// <summary>
+			/// Task for creating a lease and assigning to a waiter
+			/// </summary>
+			public Task? AssignTask;
+
+			/// <summary>
+			/// Accessor for the batch referenced by this item
+			/// </summary>
+			public IJobStepBatch Batch
+			{
+				get { return Job.Batches[BatchIdx]; }
+			}
+
+			/// <summary>
+			/// Returns an identifier describing this unique batch
+			/// </summary>
+			public (ObjectId, SubResourceId) Id
+			{
+				get { return (Job.Id, Batch.Id); }
+			}
+
+			/// <summary>
+			/// Constructor
+			/// </summary>
+			/// <param name="Stream">The stream containing this job</param>
+			/// <param name="Job">The job instance</param>
+			/// <param name="BatchIdx">The batch index to execute</param>
+			/// <param name="Pool">Unique id of the pool of machines to allocate from</param>
+			/// <param name="Requirements">Requirements for the agent to execute this batch</param>
+			/// <param name="Workspace">The workspace that this job should run in</param>
+			public QueueItem(IStream Stream, IJob Job, int BatchIdx, IPool Pool, AgentRequirements Requirements, AgentWorkspace Workspace)
+			{
+				this.Stream = Stream;
+				this.Job = Job;
+				this.BatchIdx = BatchIdx;
+				this.Pool = Pool;
+				this.Requirements = Requirements;
+				this.Workspace = Workspace;
+			}
+		}
+
+		/// <summary>
+		/// Comparer for items in the queue
+		/// </summary>
+		class QueueItemComparer : IComparer<QueueItem>
+		{
+			/// <summary>
+			/// Compare two items
+			/// </summary>
+			/// <param name="X">First item to compare</param>
+			/// <param name="Y">Second item to compare</param>
+			/// <returns>Negative value if X is a higher priority than Y</returns>
+			public int Compare([AllowNull] QueueItem X, [AllowNull] QueueItem Y)
+			{
+				if (X == null)
+				{
+					return (Y == null) ? 0 : -1;
+				}
+				else if (Y == null)
+				{
+					return 1;
+				}
+
+				int Delta = Y.Batch.SchedulePriority - X.Batch.SchedulePriority;
+				if (Delta == 0)
+				{
+					Delta = X.Job.Id.CompareTo(Y.Job.Id);
+					if (Delta == 0)
+					{
+						Delta = (int)X.Batch.Id.Value - (int)Y.Batch.Id.Value;
+					}
+				}
+				return Delta;
+			}
+		}
+
+		/// <summary>
+		/// Information about an agent waiting for work
+		/// </summary>
+		class QueueWaiter : ITaskListener
+		{
+			/// <summary>
+			/// The dispatch service instance
+			/// </summary>
+			JobTaskSource DispatchService;
+
+			/// <summary>
+			/// The agent performing the wait
+			/// </summary>
+			public IAgent Agent { get; }
+
+			/// <summary>
+			/// Completion source for the waiting agent. If a new queue item becomes available, the result will be passed through this.
+			/// </summary>
+			public TaskCompletionSource<NewLeaseInfo?> LeaseSource { get; } = new TaskCompletionSource<NewLeaseInfo?>();
+
+			/// <inheritdoc/>
+			public bool Accepted { get; set; }
+
+			/// <summary>
+			/// Constructor
+			/// </summary>
+			/// <param name="DispatchService"></param>
+			/// <param name="Agent">The agent waiting for a task</param>
+			public QueueWaiter(JobTaskSource DispatchService, IAgent Agent)
+			{
+				this.DispatchService = DispatchService;
+				this.Agent = Agent;
+
+				lock (DispatchService)
+				{
+					DispatchService.Waiters.Add(this);
+				}
+			}
+
+			public Task<NewLeaseInfo?> LeaseTask => LeaseSource.Task;
+
+			public void Accept()
+			{
+				Accepted = true;
+			}
+
+			public async ValueTask DisposeAsync()
+			{
+				lock (DispatchService.LockObject)
+				{
+					DispatchService.Waiters.Remove(this);
+				}
+
+				if (!LeaseSource.TrySetResult(null) && !Accepted)
+				{
+					NewLeaseInfo? NewLeaseInfo = await LeaseSource.Task;
+					if(NewLeaseInfo != null)
+					{
+						Any Any = Any.Parser.ParseFrom(NewLeaseInfo.Lease.Payload);
+						ExecuteJobTask Task = Any.Unpack<ExecuteJobTask>();
+						await DispatchService.CancelLeaseAsync(Agent, Task.JobId.ToObjectId(), Task.BatchId.ToSubResourceId());
+					}
+				}
+			}
+		}
+
+		/// <summary>
+		/// The database service instance
+		/// </summary>
+		DatabaseService DatabaseService;
+
+		/// <summary>
+		/// The stream service instance
+		/// </summary>
+		StreamService StreamService;
+
+		/// <summary>
+		/// The log file service instance
+		/// </summary>
+		ILogFileService LogFileService;
+
+		/// <summary>
+		/// Collection of agent documents
+		/// </summary>
+		IAgentCollection AgentsCollection;
+
+		/// <summary>
+		/// Collection of job documents
+		/// </summary>
+		IJobCollection Jobs;
+
+		/// <summary>
+		/// Collection of jobstepref documents
+		/// </summary>
+		IJobStepRefCollection JobStepRefs;
+
+		/// <summary>
+		/// Collection of graph documents
+		/// </summary>
+		IGraphCollection Graphs;
+
+		/// <summary>
+		/// Collection of pool documents
+		/// </summary>
+		IPoolCollection PoolCollection;
+
+		/// <summary>
+		/// Collection of UGS metadata documents
+		/// </summary>
+		IUgsMetadataCollection UgsMetadataCollection;
+
+		/// <summary>
+		/// The Perforce load balancer
+		/// </summary>
+		PerforceLoadBalancer PerforceLoadBalancer;
+
+		/// <summary>
+		/// The application lifetime
+		/// </summary>
+		IHostApplicationLifetime ApplicationLifetime;
+
+		/// <summary>
+		/// Task which waits for the application stopping event to start
+		/// </summary>
+		CancellationTask StoppingTask;
+
+		/// <summary>
+		/// Settings instance
+		/// </summary>
+		IOptionsMonitor<ServerSettings> Settings;
+
+		/// <summary>
+		/// List of active conform tasks
+		/// </summary>
+		SingletonDocument<ConformList> ConformList;
+
+		/// <summary>
+		/// Log writer
+		/// </summary>
+		ILogger<JobTaskSource> Logger;
+
+		/// <summary>
+		/// Object used for ensuring mutual exclusion to the queues
+		/// </summary>
+		object LockObject = new object();
+
+		/// <summary>
+		/// List of items waiting to be executed
+		/// </summary>
+		SortedSet<QueueItem> Queue = new SortedSet<QueueItem>(new QueueItemComparer());
+
+		/// <summary>
+		/// Map from batch id to the corresponding queue item
+		/// </summary>
+		Dictionary<(ObjectId, SubResourceId), QueueItem> BatchIdToQueueItem = new Dictionary<(ObjectId, SubResourceId), QueueItem>();
+
+		/// <summary>
+		/// Set of long-poll tasks waiting to be satisfied 
+		/// </summary>
+		HashSet<QueueWaiter> Waiters = new HashSet<QueueWaiter>();
+
+		/// <summary>
+		/// During a background queue refresh operation, any updated batches are added to this dictionary for merging into the updated queue.
+		/// </summary>
+		List<QueueItem>? NewQueueItemsDuringUpdate;
+
+		/// <summary>
+		/// Cache of pools
+		/// </summary>
+		Dictionary<PoolId, IPool> CachedPoolIdToInstance = new Dictionary<PoolId, IPool>();
+
+		/// <summary>
+		/// Cache of stream objects. Used to resolve agent types.
+		/// </summary>
+		private Dictionary<StreamId, IStream> Streams = new Dictionary<StreamId, IStream>();
+
+		/// <summary>
+		/// Interval between querying the database for jobs to execute
+		/// </summary>
+		static readonly TimeSpan RefreshInterval = TimeSpan.FromSeconds(5.0);
+
+		/// <inheritdoc/>
+		public MessageDescriptor Descriptor => ExecuteJobTask.Descriptor;
+
+		/// <summary>
+		/// The next server index that should be used to determine what Perforce server an agent should use.
+		/// </summary>
+		public uint NextServerIndex { get; set; } = 0;
+
+		/// <summary>
+		/// Constructor
+		/// </summary>
+		/// <param name="DatabaseService">The database service instance</param>
+		/// <param name="Agents">The agents collection</param>
+		/// <param name="Jobs">The jobs collection</param>
+		/// <param name="JobStepRefs">Collection of jobstepref documents</param>
+		/// <param name="Graphs">The graphs collection</param>
+		/// <param name="Pools">The pools collection</param>
+		/// <param name="UgsMetadataCollection">Ugs metadata collection</param>
+		/// <param name="StreamService">The stream service instance</param>
+		/// <param name="LogFileService">The log file service instance</param>
+		/// <param name="PerforceLoadBalancer">Perforce load balancer</param>
+		/// <param name="ApplicationLifetime">Application lifetime interface</param>
+		/// <param name="Settings">Settings for the server</param>
+		/// <param name="Logger">Log writer</param>
+		public JobTaskSource(DatabaseService DatabaseService, IAgentCollection Agents, IJobCollection Jobs, IJobStepRefCollection JobStepRefs, IGraphCollection Graphs, IPoolCollection Pools, IUgsMetadataCollection UgsMetadataCollection, StreamService StreamService, ILogFileService LogFileService, PerforceLoadBalancer PerforceLoadBalancer, IHostApplicationLifetime ApplicationLifetime, IOptionsMonitor<ServerSettings> Settings, ILogger<JobTaskSource> Logger)
+			: base(RefreshInterval, Logger)
+		{
+			this.DatabaseService = DatabaseService;
+			this.AgentsCollection = Agents;
+			this.Jobs = Jobs;
+			this.JobStepRefs = JobStepRefs;
+			this.Graphs = Graphs;
+			this.PoolCollection = Pools;
+			this.UgsMetadataCollection = UgsMetadataCollection;
+			this.StreamService = StreamService;
+			this.LogFileService = LogFileService;
+			this.PerforceLoadBalancer = PerforceLoadBalancer;
+			this.ApplicationLifetime = ApplicationLifetime;
+			this.StoppingTask = new CancellationTask(ApplicationLifetime.ApplicationStopping);
+			this.ConformList = new SingletonDocument<ConformList>(DatabaseService);
+			this.Settings = Settings;
+			this.Logger = Logger;
+		}
+
+		/// <summary>
+		/// Gets an object containing the stats of the queue for diagnostic purposes.
+		/// </summary>
+		/// <returns>Status object</returns>
+		public async Task<object> GetStatus()
+		{
+			List<IPool> AllPools = await PoolCollection.GetAsync();
+			lock (LockObject)
+			{
+				List<object> OutputItems = new List<object>();
+				foreach (QueueItem QueueItem in Queue)
+				{
+					OutputItems.Add(new { JobId = QueueItem.Job.Id.ToString(), BatchId = QueueItem.Batch.Id.ToString(), PoolId = QueueItem.Pool.Id.ToString(), Workspace = QueueItem.Workspace, Requirements = QueueItem.Requirements });
+				}
+
+				List<object> OutputWaiters = new List<object>();
+				foreach (QueueWaiter Waiter in Waiters)
+				{
+					OutputWaiters.Add(new { Id = Waiter.Agent.Id.ToString(), Pools = Waiter.Agent.GetPools(AllPools).Select(x => x.Id.ToString()).ToList(), Workspaces = Waiter.Agent.Workspaces, Capabilities = Waiter.Agent.Capabilities });
+				}
+
+				return new { Items = OutputItems, Waiters = OutputWaiters };
+			}
+		}
+
+		/// <summary>
+		/// Cancel any pending wait for an agent, allowing it to cycle its session state immediately
+		/// </summary>
+		/// <param name="AgentId">The agent id</param>
+		public void CancelLongPollForAgent(AgentId AgentId)
+		{
+			QueueWaiter? Waiter;
+			lock (LockObject)
+			{
+				Waiter = Waiters.FirstOrDefault(x => x.Agent.Id == AgentId);
+			}
+			if(Waiter != null)
+			{
+				Waiter.LeaseSource.TrySetResult(null);
+			}
+		}
+
+		/// <summary>
+		/// Background task
+		/// </summary>
+		/// <param name="StoppingToken">Token that indicates that the service should shut down</param>
+		/// <returns>Async task</returns>
+		protected async override Task TickAsync(CancellationToken StoppingToken)
+		{
+			// Set the NewBatchIdToQueueItem member, so we capture any updated jobs during the DB query.
+			lock (LockObject)
+			{
+				NewQueueItemsDuringUpdate = new List<QueueItem>();
+			}
+
+			// Query all the current streams
+			List<IStream> StreamsList = await StreamService.GetStreamsAsync();
+			Streams = StreamsList.ToDictionary(x => x.Id, x => x);
+
+			// Find all the pools which are valid (ie. have at least one online agent)
+			DateTime UtcNow = DateTime.UtcNow;
+			List<IAgent> Agents = await AgentsCollection.FindAsync();
+			List<IPool> Pools = await PoolCollection.GetAsync();
+
+			// Find all the pools which are currently online
+			HashSet<PoolId> OnlinePools = new HashSet<PoolId>(Agents.Where(x => x.IsSessionValid(UtcNow)).SelectMany(x => x.ExplicitPools));
+			foreach (IPool Pool in Pools)
+			{
+				if (Pool.Requirements != null && !OnlinePools.Contains(Pool.Id) && Agents.Any(x => x.IsSessionValid(UtcNow) && x.InPool(Pool)))
+				{
+					OnlinePools.Add(Pool.Id);
+				}
+			}
+
+			// Find lists of valid pools and online pools
+			HashSet<PoolId> ValidPools = new HashSet<PoolId>(OnlinePools.Union(Agents.Where(x => !x.IsSessionValid(UtcNow)).SelectMany(x => x.ExplicitPools)));
+			foreach (IPool Pool in Pools)
+			{
+				if (Pool.Requirements != null && !ValidPools.Contains(Pool.Id) && Agents.Any(x => !x.IsSessionValid(UtcNow) && x.InPool(Pool)))
+				{
+					ValidPools.Add(Pool.Id);
+				}
+			}
+
+			// Query all the current pools
+			CachedPoolIdToInstance = Pools.ToDictionary(x => x.Id, x => x);
+
+			// New list of queue items
+			SortedSet<QueueItem> NewQueue = new SortedSet<QueueItem>(Queue.Comparer);
+			Dictionary<(ObjectId, SubResourceId), QueueItem> NewBatchIdToQueueItem = new Dictionary<(ObjectId, SubResourceId), QueueItem>();
+
+			// Query for a new list of jobs for the queue
+			List<IJob> NewJobs = await Jobs.GetDispatchQueueAsync();
+			for(int Idx = 0; Idx < NewJobs.Count; Idx++)
+			{
+				IJob? NewJob = NewJobs[Idx];
+
+				// Get the graph for this job
+				IGraph Graph = await Graphs.GetAsync(NewJob.GraphHash);
+
+				// Get the stream. If it fails, skip the whole job.
+				IStream? Stream;
+				if (!Streams.TryGetValue(NewJob.StreamId, out Stream))
+				{
+					NewJob = await Jobs.SkipAllBatchesAsync(NewJob, Graph, JobStepBatchError.UnknownStream);
+					continue;
+				}
+
+				// Update all the batches
+				for (int BatchIdx = 0; NewJob != null && BatchIdx < NewJob.Batches.Count; BatchIdx++)
+				{
+					// Skip any batches which aren't ready to execute
+					IJobStepBatch Batch = NewJob.Batches[BatchIdx];
+					if (Batch.State == JobStepBatchState.Ready)
+					{
+						// Validate the agent type and workspace settings
+						IPool? Pool;
+						if (!Stream.AgentTypes.TryGetValue(Graph.Groups[Batch.GroupIdx].AgentType, out AgentType? AgentType))
+						{
+							NewJob = await SkipBatchAsync(NewJob, BatchIdx, Graph, JobStepBatchError.UnknownAgentType);
+						}
+						else if (!CachedPoolIdToInstance.TryGetValue(AgentType.Pool, out Pool))
+						{
+							NewJob = await SkipBatchAsync(NewJob, BatchIdx, Graph, JobStepBatchError.UnknownPool);
+						}
+						else if (!ValidPools.Contains(AgentType.Pool))
+						{
+							NewJob = await SkipBatchAsync(NewJob, BatchIdx, Graph, JobStepBatchError.NoAgentsInPool);
+						}
+						else if (!OnlinePools.Contains(AgentType.Pool))
+						{
+							NewJob = await SkipBatchAsync(NewJob, BatchIdx, Graph, JobStepBatchError.NoAgentsOnline);
+						}
+						else if (!Stream.TryGetAgentWorkspace(AgentType, out AgentWorkspace? Workspace))
+						{
+							NewJob = await SkipBatchAsync(NewJob, BatchIdx, Graph, JobStepBatchError.UnknownWorkspace);
+						}
+						else
+						{
+							AgentRequirements Requirements = new AgentRequirements();
+
+							QueueItem NewQueueItem = new QueueItem(Stream, NewJob, BatchIdx, Pool, Requirements, Workspace);
+							NewQueue.Add(NewQueueItem);
+							NewBatchIdToQueueItem[(NewJob.Id, Batch.Id)] = NewQueueItem;
+						}
+					}
+				}
+			}
+
+			// Update the queue
+			lock (LockObject)
+			{
+				Queue = NewQueue;
+				BatchIdToQueueItem = NewBatchIdToQueueItem;
+
+				// Merge the new queue items with the queue
+				foreach (QueueItem NewQueueItem in NewQueueItemsDuringUpdate)
+				{
+					QueueItem? ExistingQueueItem;
+					if (!NewBatchIdToQueueItem.TryGetValue((NewQueueItem.Job.Id, NewQueueItem.Batch.Id), out ExistingQueueItem))
+					{
+						// Always just add this item
+						Queue.Add(NewQueueItem);
+						BatchIdToQueueItem[NewQueueItem.Id] = NewQueueItem;
+					}
+					else if (NewQueueItem.Job.UpdateIndex > ExistingQueueItem.Job.UpdateIndex)
+					{
+						// Replace the existing item
+						Queue.Remove(ExistingQueueItem);
+						Queue.Add(NewQueueItem);
+						BatchIdToQueueItem[NewQueueItem.Id] = NewQueueItem;
+					}
+				}
+
+				// Clear out the list to capture queue items during an update
+				NewQueueItemsDuringUpdate = null;
+			}
+		}
+
+		private async Task<IJob?> SkipBatchAsync(IJob Job, int BatchIdx, IGraph Graph, JobStepBatchError Reason)
+		{
+			IReadOnlyList<(LabelState, LabelOutcome)> OldLabelStates = Job.GetLabelStates(Graph);
+			IJob? NewJob = await Jobs.SkipBatchAsync(Job, BatchIdx, Graph, Reason);
+			if(NewJob != null)
+			{
+				IReadOnlyList<(LabelState, LabelOutcome)> NewLabelStates = NewJob.GetLabelStates(Graph);
+				await UpdateUgsBadges(NewJob, Graph, OldLabelStates, NewLabelStates);
+			}
+			return NewJob;
+		}
+
+		/// <summary>
+		/// Updates the current state of a job
+		/// </summary>
+		/// <param name="Job">The job that has been updated</param>
+		/// <param name="Graph">Graph for the job</param>
+		/// <returns>Async task</returns>
+		public void UpdateQueuedJob(IJob Job, IGraph Graph)
+		{
+			IStream? Stream;
+			Streams.TryGetValue(Job.StreamId, out Stream);
+			UpdateQueuedJob(Job, Graph, Stream);
+		}
+
+		void AssignAnyQueueItemToWaiter(QueueWaiter Waiter)
+		{
+			lock (Waiters)
+			{
+				foreach(QueueItem Item in BatchIdToQueueItem.Values)
+				{
+					if (TryAssignItemToWaiter(Item, Waiter))
+					{
+						break;
+					}
+				}
+			}
+		}
+
+		/// <summary>
+		/// Attempt to find a waiter that can handle the given queue item
+		/// </summary>
+		/// <param name="Item">The queue item</param>
+		/// <returns></returns>
+		void AssignQueueItemToAnyWaiter(QueueItem Item)
+		{
+			if (Item.AssignTask == null && Item.Batch.SessionId == null)
+			{
+				lock (Waiters)
+				{
+					if (Item.AssignTask == null && Item.Batch.SessionId == null)
+					{
+						foreach (QueueWaiter Waiter in Waiters)
+						{
+							if (TryAssignItemToWaiter(Item, Waiter))
+							{
+								break;
+							}
+						}
+					}
+				}
+			}
+		}
+
+		bool TryAssignItemToWaiter(QueueItem Item, QueueWaiter Waiter)
+		{
+			if (Item.AssignTask == null && Item.Batch.SessionId == null)
+			{
+				List<AgentLeaseDevice>? LeasedDevices;
+				if (Waiter.Agent.TryCreateLease(Item.Pool, Item.Requirements, out LeasedDevices))
+				{
+					Task StartTask = new Task<Task>(() => TryCreateLeaseAsync(Item, Waiter, LeasedDevices));
+					Task ExecuteTask = StartTask.ContinueWith(Task => Task, TaskScheduler.Default);
+
+					if (Interlocked.CompareExchange(ref Item.AssignTask, ExecuteTask, null) == null)
+					{
+						StartTask.Start(TaskScheduler.Default);
+						return true;
+					}
+				}
+			}
+			return false;
+		}
+
+		/// <summary>
+		/// Updates the current state of a job
+		/// </summary>
+		/// <param name="Job">The job that has been updated</param>
+		/// <param name="Graph">Graph for the job</param>
+		/// <param name="Stream">The stream containing the job</param>
+		public void UpdateQueuedJob(IJob Job, IGraph Graph, IStream? Stream)
+		{
+			List<TaskCompletionSource<bool>> CompleteWaiters = new List<TaskCompletionSource<bool>>();
+			lock (LockObject)
+			{
+				for (int BatchIdx = 0; BatchIdx < Job.Batches.Count; BatchIdx++)
+				{
+					IJobStepBatch Batch = Job.Batches[BatchIdx];
+					if (Batch.State == JobStepBatchState.Ready && Stream != null && Batch.AgentId == null)
+					{
+						// Check if this item is already in the list.
+						QueueItem? ExistingItem;
+						if (BatchIdToQueueItem.TryGetValue((Job.Id, Batch.Id), out ExistingItem))
+						{
+							// Make sure this is a newer version of the job. There's no guarantee that this is the latest revision.
+							if (Job.UpdateIndex > ExistingItem.Job.UpdateIndex)
+							{
+								if (Batch.SchedulePriority == ExistingItem.Batch.SchedulePriority)
+								{
+									ExistingItem.Job = Job;
+									ExistingItem.BatchIdx = BatchIdx;
+								}
+								else
+								{
+									RemoveQueueItem(ExistingItem);
+									InsertQueueItem(Stream, Job, BatchIdx, ExistingItem.Pool, ExistingItem.Requirements, ExistingItem.Workspace);
+								}
+							}
+							continue;
+						}
+
+						// Get the group being executed by this batch
+						INodeGroup Group = Graph.Groups[Batch.GroupIdx];
+
+						// Get the requirements for the new queue item
+						AgentType? AgentType;
+						if (Stream.AgentTypes.TryGetValue(Group.AgentType, out AgentType))
+						{
+							AgentWorkspace? Workspace;
+							if (Stream.TryGetAgentWorkspace(AgentType, out Workspace))
+							{
+								AgentRequirements AgentRequirements = new AgentRequirements();
+
+								// Get the pool for this agent type
+								IPool? Pool;
+								if (CachedPoolIdToInstance.TryGetValue(AgentType.Pool, out Pool))
+								{
+									InsertQueueItem(Stream, Job, BatchIdx, Pool, AgentRequirements, Workspace);
+								}
+							}
+						}
+					}
+					else
+					{
+						// Check if this item is already in the list. Remove it if it is.
+						QueueItem? ExistingItem;
+						if (BatchIdToQueueItem.TryGetValue((Job.Id, Batch.Id), out ExistingItem))
+						{
+							if (Job.UpdateIndex > ExistingItem.Job.UpdateIndex)
+							{
+								RemoveQueueItem(ExistingItem);
+							}
+						}
+					}
+				}
+			}
+
+			// Awake all the threads that have been assigned new work items. Has do be done outside the lock to prevent continuations running within it (see Waiter.CompletionSource for more info).
+			foreach (TaskCompletionSource<bool> CompleteWaiter in CompleteWaiters)
+			{
+				CompleteWaiter.TrySetResult(true);
+			}
+		}
+
+		/// <summary>
+		/// Gets the index that should be used to choose the Perforce server an agent should use for a lease.
+		/// </summary>
+		public uint GetPerforceServerIndex()
+		{
+			return NextServerIndex++;
+		}
+
+		/// <inheritdoc/>
+		public Task<ITaskListener?> SubscribeAsync(IAgent Agent)
+		{
+			QueueWaiter Waiter = new QueueWaiter(this, Agent);
+			lock (LockObject)
+			{
+				AssignAnyQueueItemToWaiter(Waiter);
+			}
+			return Task.FromResult<ITaskListener?>(Waiter);
+		}
+
+		/// <summary>
+		/// Assign a new batch to be executed by the given agent
+		/// </summary>
+		/// <param name="Item">The item to create a lease for</param>
+		/// <param name="Waiter">The agent waiting for work</param>
+		/// <param name="LeasedDevices">The devices to include in the lease</param>
+		/// <returns>New work to execute</returns>
+		private async Task<AgentLease?> TryCreateLeaseAsync(QueueItem Item, QueueWaiter Waiter, List<AgentLeaseDevice>? LeasedDevices)
+		{
+			IJob Job = Item.Job;
+			IJobStepBatch Batch = Job.Batches[Item.BatchIdx];
+			IAgent Agent = Waiter.Agent;
+			Logger.LogDebug("Assigning job {JobId}, batch {BatchId} to waiter (agent {AgentID})", Job.Id, Batch.Id, Agent.Id);
+
+			// Generate a new unique id for the lease
+			ObjectId LeaseId = ObjectId.GenerateNewId();
+
+			// The next time to try assigning to another agent
+			DateTime BackOffTime = DateTime.UtcNow + TimeSpan.FromMinutes(1.0);
+
+			// Try to update the job with this agent id
+			ObjectId LogId = (await LogFileService.CreateLogFileAsync(Job.Id, Agent.SessionId, LogType.Json)).Id;
+			if (await Jobs.TryAssignLeaseAsync(Item.Job, Item.BatchIdx, Item.Pool.Id, Agent.Id, Agent.SessionId!.Value, LeaseId, LogId))
+			{
+				// Get the lease name
+				StringBuilder LeaseName = new StringBuilder($"{Item.Stream.Name} - ");
+				if (Job.PreflightChange > 0)
+				{
+					LeaseName.Append((Job.Change > 0) ? $"Preflight CL {Job.PreflightChange} against CL {Job.Change}" : $"Preflight CL {Job.PreflightChange} against latest");
+				}
+				else
+				{
+					LeaseName.Append((Job.Change > 0) ? $"CL {Job.Change}" : "Latest CL");
+				}
+				LeaseName.Append($" - {Job.Name}");
+
+				// Get the global settings
+				Globals Globals = await DatabaseService.GetGlobalsAsync();
+
+				// Encode the payload
+				ExecuteJobTask? Task = await CreateExecuteJobTaskAsync(Item.Stream, Job, Batch, Agent, Item.Workspace, LogId);
+				if (Task != null)
+				{
+					byte[] Payload = Any.Pack(Task).ToByteArray();
+
+					// Create the lease and try to set it on the waiter. If this fails, the waiter has already moved on, and the lease can be cancelled.
+					AgentLease Lease = new AgentLease(LeaseId, LeaseName.ToString(), Job.StreamId, Item.Pool.Id, LogId, LeaseState.Pending, Payload, Item.Requirements, LeasedDevices);
+					if (Waiter.LeaseSource.TrySetResult(new NewLeaseInfo(Lease)))
+					{
+						Logger.LogDebug("Assigned lease {LeaseId} to agent {AgentId}", LeaseId, Agent.Id);
+						return Lease;
+					}
+				}
+
+				// Cancel the lease
+				Logger.LogDebug("Unable to assign lease {LeaseId} to agent {AgentId}, cancelling", LeaseId, Agent.Id);
+				await CancelLeaseAsync(Waiter.Agent, Job.Id, Batch.Id);
+			}
+			else
+			{
+				// Unable to assign job
+				Logger.LogDebug("Refreshing queue entries for job {JobId}", Job.Id);
+
+				// Get the new copy of the job
+				IJob? NewJob = await Jobs.GetAsync(Job.Id);
+				if (NewJob == null)
+				{
+					lock (LockObject)
+					{
+						List<QueueItem> RemoveItems = Queue.Where(x => x.Job == Job).ToList();
+						foreach (QueueItem RemoveItem in RemoveItems)
+						{
+							RemoveQueueItem(RemoveItem);
+						}
+					}
+				}
+				else
+				{
+					IGraph Graph = await Graphs.GetAsync(NewJob.GraphHash);
+					UpdateQueuedJob(NewJob, Graph);
+				}
+			}
+
+			// Clear out the assignment for this item, and try to reassign it
+			Item.AssignTask = null;
+			AssignQueueItemToAnyWaiter(Item);
+			return null;
+		}
+
+		async Task<ExecuteJobTask?> CreateExecuteJobTaskAsync(IStream Stream, IJob Job, IJobStepBatch Batch, IAgent Agent, AgentWorkspace Workspace, ObjectId LogId)
+		{
+			// Get the lease name
+			StringBuilder LeaseName = new StringBuilder($"{Stream.Name} - ");
+			if (Job.PreflightChange > 0)
+			{
+				LeaseName.Append((Job.Change > 0) ? $"Preflight CL {Job.PreflightChange} against CL {Job.Change}" : $"Preflight CL {Job.PreflightChange} against latest");
+			}
+			else
+			{
+				LeaseName.Append((Job.Change > 0) ? $"CL {Job.Change}" : "Latest CL");
+			}
+			LeaseName.Append($" - {Job.Name}");
+
+			// Get the global settings
+			Globals Globals = await DatabaseService.GetGlobalsAsync();
+
+			// Encode the payload
+			ExecuteJobTask Task = new ExecuteJobTask();
+			Task.JobId = Job.Id.ToString();
+			Task.BatchId = Batch.Id.ToString();
+			Task.LogId = LogId.ToString();
+			Task.JobName = LeaseName.ToString();
+
+			List<HordeCommon.Rpc.Messages.AgentWorkspace> Workspaces = new List<HordeCommon.Rpc.Messages.AgentWorkspace>();
+
+			AgentWorkspace? AutoSdkWorkspace = Agent.Workspaces.FirstOrDefault(x => x.Identifier == AgentWorkspace.AutoSdkIdentifier);
+			if (AutoSdkWorkspace != null)
+			{
+				if (!await Agent.TryAddWorkspaceMessage(AutoSdkWorkspace, Globals, PerforceLoadBalancer, Workspaces))
+				{
+					return null;
+				}
+				Task.AutoSdkWorkspace = Workspaces.Last();
+			}
+
+			if (!await Agent.TryAddWorkspaceMessage(Workspace, Globals, PerforceLoadBalancer, Workspaces))
+			{
+				return null;
+			}
+
+			Task.Workspace = Workspaces.Last();
+			return Task;
+		}
+
+		/// <summary>
+		/// Send any badge updates for this job
+		/// </summary>
+		/// <param name="Job">The job being updated</param>
+		/// <param name="Graph">Graph for the job</param>
+		/// <param name="OldLabelStates">Previous badge states for the job</param>
+		/// <returns>Async task</returns>
+		public async Task UpdateUgsBadges(IJob Job, IGraph Graph, IReadOnlyList<(LabelState, LabelOutcome)> OldLabelStates)
+		{
+			await UpdateUgsBadges(Job, Graph, OldLabelStates, Job.GetLabelStates(Graph));
+		}
+
+		/// <summary>
+		/// Send any badge updates for this job
+		/// </summary>
+		/// <param name="Job">The job being updated</param>
+		/// <param name="Graph">Graph for the job</param>
+		/// <param name="OldLabelStates">Previous badge states for the job</param>
+		/// <param name="NewLabelStates">The new badge states for the job</param>
+		/// <returns>Async task</returns>
+		public async Task UpdateUgsBadges(IJob Job, IGraph Graph, IReadOnlyList<(LabelState, LabelOutcome)> OldLabelStates, IReadOnlyList<(LabelState, LabelOutcome)> NewLabelStates)
+		{
+			if (!Job.ShowUgsBadges)
+			{
+				return;
+			}
+
+			IReadOnlyDictionary<int, UgsBadgeState> OldStates = Job.GetUgsBadgeStates(Graph, OldLabelStates);
+			IReadOnlyDictionary<int, UgsBadgeState> NewStates = Job.GetUgsBadgeStates(Graph, NewLabelStates);
+
+			// Figure out a list of all the badges that have been modified
+			List<int> UpdateLabels = new List<int>();
+			foreach (KeyValuePair<int, UgsBadgeState> Pair in OldStates)
+			{
+				if (!NewStates.ContainsKey(Pair.Key))
+				{
+					UpdateLabels.Add(Pair.Key);
+				}
+			}
+			foreach (KeyValuePair<int, UgsBadgeState> Pair in NewStates)
+			{
+				if (!OldStates.TryGetValue(Pair.Key, out UgsBadgeState PrevState) || PrevState != Pair.Value)
+				{
+					UpdateLabels.Add(Pair.Key);
+				}
+			}
+
+			// Cached stream for this job
+			IStream? Stream = null;
+
+			// Send all the updates
+			Dictionary<int, IUgsMetadata> MetadataCache = new Dictionary<int, IUgsMetadata>();
+			foreach (int LabelIdx in UpdateLabels)
+			{
+				ILabel Label = Graph.Labels[LabelIdx];
+
+				// Skip if this label has no UGS name.
+				if (Label.UgsName == null)
+				{
+					continue;
+				}
+
+				// Get the new state
+				if (!NewStates.TryGetValue(LabelIdx, out UgsBadgeState NewState))
+				{
+					NewState = UgsBadgeState.Skipped;
+				}
+
+				// Get the stream
+				if (Stream == null)
+				{
+					Stream = await StreamService.GetStreamAsync(Job.StreamId);
+					if (Stream == null)
+					{
+						Logger.LogError("Unable to fetch definition for stream {StreamId}", Job.StreamId);
+						break;
+					}
+				}
+
+				// The changelist number to display the badge for
+				int Change;
+				if (Label.Change == LabelChange.Code)
+				{
+					Change = Job.CodeChange;
+				}
+				else
+				{
+					Change = Job.Change;
+				}
+
+				// Get the current metadata state
+				IUgsMetadata? Metadata;
+				if (!MetadataCache.TryGetValue(Change, out Metadata))
+				{
+					Metadata = await UgsMetadataCollection.FindOrAddAsync(Stream.Name, Change, Label.UgsProject);
+					MetadataCache[Change] = Metadata;
+				}
+
+				// Apply the update
+				Uri LabelUrl = new Uri(Settings.CurrentValue.DashboardUrl, $"job/{Job.Id}?label={LabelIdx}");
+				Logger.LogInformation("Updating state of badge {BadgeName} at {Change} to {NewState} ({LabelUrl})", Label.UgsName, Change, NewState, LabelUrl);
+				Metadata = await UgsMetadataCollection.UpdateBadgeAsync(Metadata, Label.UgsName!, LabelUrl, NewState);
+				MetadataCache[Change] = Metadata;
+			}
+		}
+
+		/// <inheritdoc/>
+		public async Task AbortTaskAsync(IAgent Agent, ObjectId LeaseId, Any Any)
+		{
+			AgentId AgentId = Agent.Id;
+			ExecuteJobTask Task = Any.Unpack<ExecuteJobTask>();
+			ObjectId JobId = Task.JobId.ToObjectId();
+			SubResourceId BatchId = Task.BatchId.ToSubResourceId();
+
+			// Update the batch
+			for (; ; )
+			{
+				IJob? Job = await Jobs.GetAsync(JobId);
+				if (Job == null)
+				{
+					break;
+				}
+
+				int BatchIdx = Job.Batches.FindIndex(x => x.Id == BatchId);
+				if (BatchIdx == -1)
+				{
+					break;
+				}
+
+				IJobStepBatch Batch = Job.Batches[BatchIdx];
+				if (Batch.AgentId != AgentId)
+				{
+					break;
+				}
+
+				int RunningStepIdx = Batch.Steps.FindIndex(x => x.State == JobStepState.Running);
+
+				IGraph Graph = await Graphs.GetAsync(Job.GraphHash);
+				if (await Jobs.TryFailBatchAsync(Job, BatchIdx, Graph))
+				{
+					if (Batch.Error != JobStepBatchError.None)
+					{
+						Logger.LogError("Failed job {JobId}, batch {BatchId} with error {Error}", Job.Id, Batch.Id, Batch.Error);
+					}
+					if (RunningStepIdx != -1)
+					{
+						await JobStepRefs.UpdateAsync(Job, Batch, Batch.Steps[RunningStepIdx], Graph);
+					}
+					break;
+				}
+			}
+		}
+
+		/// <summary>
+		/// 
+		/// </summary>
+		/// <param name="Agent"></param>
+		/// <param name="JobId"></param>
+		/// <param name="BatchId"></param>
+		/// <returns></returns>
+		async Task CancelLeaseAsync(IAgent Agent, ObjectId JobId, SubResourceId BatchId)
+		{
+			Logger.LogDebug("Cancelling lease for job {JobId}, batch {BatchId}", JobId, BatchId);
+
+			// Update the batch
+			for (; ; )
+			{
+				IJob? Job = await Jobs.GetAsync(JobId);
+				if (Job == null)
+				{
+					break;
+				}
+
+				int BatchIdx = Job.Batches.FindIndex(x => x.Id == BatchId);
+				if (BatchIdx == -1)
+				{
+					break;
+				}
+
+				IJobStepBatch Batch = Job.Batches[BatchIdx];
+				if (Batch.AgentId != Agent.Id)
+				{
+					break;
+				}
+
+				if (await Jobs.TryCancelLeaseAsync(Job, BatchIdx))
+				{
+					break;
+				}
+			}
+		}
+
+		/// <summary>
+		/// Inserts an item into the queue
+		/// </summary>
+		/// <param name="Stream">The stream containing the job</param>
+		/// <param name="Job"></param>
+		/// <param name="BatchIdx"></param>
+		/// <param name="Pool">The pool to use</param>
+		/// <param name="Requirements">The agent requirements for this item</param>
+		/// <param name="Workspace">The workspace for this item to run in</param>
+		/// <returns></returns>
+		void InsertQueueItem(IStream Stream, IJob Job, int BatchIdx, IPool Pool, AgentRequirements Requirements, AgentWorkspace Workspace)
+		{
+			Logger.LogDebug("Adding queued job {JobId}, batch {BatchId} [Pool: {Pool}, Workspace: {Workspace}]", Job.Id, Job.Batches[BatchIdx].Id, Pool.Id, Workspace.Identifier);
+
+			QueueItem NewItem = new QueueItem(Stream, Job, BatchIdx, Pool, Requirements, Workspace);
+			BatchIdToQueueItem[NewItem.Id] = NewItem;
+			Queue.Add(NewItem);
+
+			AssignQueueItemToAnyWaiter(NewItem);
+		}
+
+		/// <summary>
+		/// Removes an item from the queue
+		/// </summary>
+		/// <param name="Item">Item to remove</param>
+		void RemoveQueueItem(QueueItem Item)
+		{
+			Logger.LogDebug("Removing queued job {JobId}, batch {BatchId}", Item.Job.Id, Item.Batch.Id);
+
+			Queue.Remove(Item);
+			BatchIdToQueueItem.Remove(Item.Id);
+		}
+	}
+}
diff --git a/Engine/Source/Programs/Horde/HordeServer/Tasks/Impl/RestartTaskSource.cs b/Engine/Source/Programs/Horde/HordeServer/Tasks/Impl/RestartTaskSource.cs
new file mode 100644
--- /dev/null
+++ b/Engine/Source/Programs/Horde/HordeServer/Tasks/Impl/RestartTaskSource.cs
@@ -0,0 +1,50 @@
+using Google.Protobuf;
+using Google.Protobuf.Reflection;
+using Google.Protobuf.WellKnownTypes;
+using HordeCommon;
+using HordeCommon.Rpc.Tasks;
+using HordeServer.Api;
+using HordeServer.Models;
+using HordeServer.Services;
+using MongoDB.Bson;
+using System;
+using System.Collections.Generic;
+using System.Linq;
+using System.Threading.Tasks;
+
+using AgentSoftwareVersion = HordeServer.Utilities.StringId<HordeServer.Collections.IAgentSoftwareCollection>;
+
+namespace HordeServer.Tasks.Impl
+{
+	class RestartTaskSource : TaskSourceBase<RestartTask>
+	{
+		ILogFileService LogService;
+
+		public RestartTaskSource(ILogFileService LogService)
+		{
+			this.LogService = LogService;
+		}
+
+		public override async Task<ITaskListener?> SubscribeAsync(IAgent Agent)
+		{
+			if (!Agent.RequestRestart)
+			{
+				return null;
+			}
+			if (Agent.Leases.Count > 0)
+			{
+				return TaskSubscription.FromResult(null);
+			}
+
+			ILogFile Log = await LogService.CreateLogFileAsync(ObjectId.Empty, Agent.SessionId, LogType.Json);
+
+			RestartTask Task = new RestartTask();
+			Task.LogId = Log.Id.ToString();
+
+			byte[] Payload = Any.Pack(Task).ToByteArray();
+
+			AgentLease? Lease = new AgentLease(ObjectId.GenerateNewId(), "Restart", null, null, Log.Id, LeaseState.Pending, Payload, new AgentRequirements(), null);
+			return TaskSubscription.FromResult(Lease);
+		}
+	}
+}
diff --git a/Engine/Source/Programs/Horde/HordeServer/Tasks/Impl/ShutdownTaskSource.cs b/Engine/Source/Programs/Horde/HordeServer/Tasks/Impl/ShutdownTaskSource.cs
new file mode 100644
--- /dev/null
+++ b/Engine/Source/Programs/Horde/HordeServer/Tasks/Impl/ShutdownTaskSource.cs
@@ -0,0 +1,50 @@
+using Google.Protobuf;
+using Google.Protobuf.Reflection;
+using Google.Protobuf.WellKnownTypes;
+using HordeCommon;
+using HordeCommon.Rpc.Tasks;
+using HordeServer.Api;
+using HordeServer.Models;
+using HordeServer.Services;
+using MongoDB.Bson;
+using System;
+using System.Collections.Generic;
+using System.Linq;
+using System.Threading.Tasks;
+
+using AgentSoftwareVersion = HordeServer.Utilities.StringId<HordeServer.Collections.IAgentSoftwareCollection>;
+
+namespace HordeServer.Tasks.Impl
+{
+	class ShutdownTaskSource : TaskSourceBase<ShutdownTask>
+	{
+		ILogFileService LogService;
+
+		public ShutdownTaskSource(ILogFileService LogService)
+		{
+			this.LogService = LogService;
+		}
+
+		public override async Task<ITaskListener?> SubscribeAsync(IAgent Agent)
+		{
+			if (!Agent.RequestShutdown)
+			{
+				return null;
+			}
+			if (Agent.Leases.Count > 0)
+			{
+				return TaskSubscription.FromResult(null);
+			}
+
+			ILogFile Log = await LogService.CreateLogFileAsync(ObjectId.Empty, Agent.SessionId, LogType.Json);
+
+			ShutdownTask Task = new ShutdownTask();
+			Task.LogId = Log.Id.ToString();
+
+			byte[] Payload = Any.Pack(Task).ToByteArray();
+
+			AgentLease? Lease = new AgentLease(ObjectId.GenerateNewId(), "Shutdown", null, null, Log.Id, LeaseState.Pending, Payload, new AgentRequirements(), null);
+			return TaskSubscription.FromResult(Lease);
+		}
+	}
+}
diff --git a/Engine/Source/Programs/Horde/HordeServer/Tasks/Impl/UpgradeTaskSource.cs b/Engine/Source/Programs/Horde/HordeServer/Tasks/Impl/UpgradeTaskSource.cs
new file mode 100644
--- /dev/null
+++ b/Engine/Source/Programs/Horde/HordeServer/Tasks/Impl/UpgradeTaskSource.cs
@@ -0,0 +1,75 @@
+using Google.Protobuf;
+using Google.Protobuf.Reflection;
+using Google.Protobuf.WellKnownTypes;
+using HordeCommon;
+using HordeCommon.Rpc.Tasks;
+using HordeServer.Api;
+using HordeServer.Models;
+using HordeServer.Services;
+using MongoDB.Bson;
+using System;
+using System.Collections.Generic;
+using System.Linq;
+using System.Threading.Tasks;
+
+using AgentSoftwareVersion = HordeServer.Utilities.StringId<HordeServer.Collections.IAgentSoftwareCollection>;
+using AgentSoftwareChannelName = HordeServer.Utilities.StringId<HordeServer.Services.AgentSoftwareChannels>;
+
+namespace HordeServer.Tasks.Impl
+{
+	class UpgradeTaskSource : ITaskSource
+	{
+		AgentSoftwareService AgentSoftwareService;
+		ILogFileService LogService;
+		IClock Clock;
+
+		public MessageDescriptor Descriptor => UpgradeTask.Descriptor;
+
+		public UpgradeTaskSource(AgentSoftwareService AgentSoftwareService, ILogFileService LogService, IClock Clock)
+		{
+			this.AgentSoftwareService = AgentSoftwareService;
+			this.LogService = LogService;
+			this.Clock = Clock;
+		}
+
+		public async Task<ITaskListener?> SubscribeAsync(IAgent Agent)
+		{
+			AgentSoftwareVersion? RequiredVersion = await GetRequiredSoftwareVersion(Agent);
+			if (RequiredVersion != null && Agent.Version != RequiredVersion)
+			{
+				AgentLease? Lease = null;
+				if (Agent.Leases.Count == 0 && (Agent.LastUpgradeTime == null || Agent.LastUpgradeTime.Value + TimeSpan.FromMinutes(5.0) < Clock.UtcNow || Agent.LastUpgradeVersion != RequiredVersion.ToString()))
+				{
+					ILogFile LogFile = await LogService.CreateLogFileAsync(ObjectId.Empty, Agent.SessionId, LogType.Json);
+
+					UpgradeTask Task = new UpgradeTask();
+					Task.SoftwareId = RequiredVersion.ToString();
+					Task.LogId = LogFile.Id.ToString();
+
+					byte[] Payload = Any.Pack(Task).ToByteArray();
+
+					Lease = new AgentLease(ObjectId.GenerateNewId(), $"Upgrade to {RequiredVersion}", null, null, LogFile.Id, LeaseState.Pending, Payload, new AgentRequirements(), null);
+				}
+				return TaskSubscription.FromResult(Lease);
+			}
+			return null;
+		}
+
+		public Task AbortTaskAsync(IAgent Agent, ObjectId LeaseId, Any Payload)
+		{
+			return Task.CompletedTask;
+		}
+
+		/// <summary>
+		/// Determines the client software version that should be installed on an agent
+		/// </summary>
+		/// <param name="Agent">The agent instance</param>
+		/// <returns>Unique id of the client version this agent should be running</returns>
+		public async Task<AgentSoftwareVersion?> GetRequiredSoftwareVersion(IAgent Agent)
+		{
+			AgentSoftwareChannelName ChannelName = Agent.Channel ?? AgentSoftwareService.DefaultChannelName;
+			IAgentSoftwareChannel? Channel = await AgentSoftwareService.GetCachedChannelAsync(ChannelName);
+			return Channel?.Version;
+		}
+	}
+}
diff --git a/Engine/Source/Programs/Horde/HordeServer/Utilities/ArtifactChunkStream.cs b/Engine/Source/Programs/Horde/HordeServer/Utilities/ArtifactChunkStream.cs
new file mode 100644
--- /dev/null
+++ b/Engine/Source/Programs/Horde/HordeServer/Utilities/ArtifactChunkStream.cs
@@ -0,0 +1,122 @@
+using Grpc.Core;
+using HordeCommon.Rpc;
+using System;
+using System.Collections.Generic;
+using System.IO;
+using System.Linq;
+using System.Threading;
+using System.Threading.Tasks;
+
+namespace HordeServer.Utilities
+{
+	/// <summary>
+	/// Stream that takes in rpc chunks from an upload and copies them?
+	/// </summary>
+	public class ArtifactChunkStream : Stream
+	{
+		/// <summary>
+		/// The grpc client reader. Should probably be templatized
+		/// </summary>
+		IAsyncStreamReader<UploadArtifactRequest> Reader;
+
+		/// <summary>
+		/// Position within the stream
+		/// </summary>
+		long StreamPosition;
+
+		/// <summary>
+		/// The length of the stream
+		/// </summary>
+		long StreamLength;
+
+		/// <summary>
+		/// The current request being read from
+		/// </summary>
+		UploadArtifactRequest? Request;
+
+		/// <summary>
+		/// Position within the current request
+		/// </summary>
+		int RequestPos;
+
+		/// <summary>
+		/// Constructor
+		/// </summary>
+		/// <param name="Reader">the grpc reader</param>
+		/// <param name="Length">filesize reported by the client</param>
+		public ArtifactChunkStream(IAsyncStreamReader<UploadArtifactRequest> Reader, long Length)
+		{
+			this.Reader = Reader;
+			this.StreamLength = Length;
+		}
+
+		/// <inheritdoc/>
+		public override long Position
+		{
+			get { return StreamPosition; }
+			set { throw new NotSupportedException(); }
+		}
+
+		/// <inheritdoc/>
+		public override long Length => StreamLength;
+
+		/// <inheritdoc/>
+		public override void SetLength(long value) => throw new NotSupportedException();
+
+		/// <inheritdoc/>
+		public override bool CanSeek => false;
+
+		/// <inheritdoc/>
+		public override long Seek(long Offset, SeekOrigin Origin) => throw new NotSupportedException();
+
+		/// <inheritdoc/>
+		public override bool CanWrite => false;
+
+		/// <inheritdoc/>
+		public override void Write(byte[] buffer, int offset, int count) => throw new NotSupportedException();
+
+		/// <inheritdoc/>
+		public override bool CanRead => true;
+
+		/// <inheritdoc/>
+		public override int Read(byte[] Buffer, int Offset, int Count)
+		{
+			return ReadAsync(Buffer, Offset, Count, CancellationToken.None).Result;
+		}
+
+		/// <inheritdoc/>
+		public override async Task<int> ReadAsync(byte[] Buffer, int Offset, int Count, CancellationToken CancellationToken)
+		{
+			int BytesRead = 0;
+			while (BytesRead < Count && StreamPosition < StreamLength)
+			{
+				if (Request == null || RequestPos == Request.Data.Length)
+				{
+					// Read the next request object
+					if (!await Reader.MoveNext())
+					{
+						throw new EndOfStreamException("Unexpected end of stream while reading artifact");
+					}
+
+					Request = Reader.Current;
+					RequestPos = 0;
+				}
+				else
+				{
+					// Copy data from the current request object
+					int NumBytesToCopy = Math.Min(Count - BytesRead, Request.Data.Length - RequestPos);
+					Request.Data.Span.Slice(RequestPos, NumBytesToCopy).CopyTo(Buffer.AsSpan(BytesRead, NumBytesToCopy));
+					RequestPos += NumBytesToCopy;
+					BytesRead += NumBytesToCopy;
+					StreamPosition += NumBytesToCopy;
+				}
+			}
+			return BytesRead;
+		}
+
+		/// <inheritdoc/>
+		public override void Flush()
+		{
+		}
+	}
+}
diff --git a/Engine/Source/Programs/Horde/HordeServer/Utilities/BlockKit/ActionsBlock.cs b/Engine/Source/Programs/Horde/HordeServer/Utilities/BlockKit/ActionsBlock.cs
new file mode 100644
--- /dev/null
+++ b/Engine/Source/Programs/Horde/HordeServer/Utilities/BlockKit/ActionsBlock.cs
@@ -0,0 +1,159 @@
+// Copyright Epic Games, Inc. All Rights Reserved.
+
+using System;
+using System.Collections.Generic;
+using System.Globalization;
+using System.Text.Json;
+
+namespace HordeServer.Utilities.Slack.BlockKit
+{
+	/// <summary>
+	/// Wrapper for a BlockKit Actions block that contains interactive elements.
+	/// </summary>
+	public class ActionsBlock : BlockBase
+	{
+		/// <summary>
+		/// A collection of interactive elements.
+		/// </summary>
+		public List<ActionElement> Elements { get; } = new List<ActionElement>();
+
+		/// <inheritdoc/>
+		public override void Write(Utf8JsonWriter Writer, JsonSerializerOptions Options)
+		{
+			Writer.WriteStartObject();
+			Writer.WriteString("type", "actions");
+			Writer.WriteStartArray("elements");
+			foreach (ActionElement Element in Elements)
+			{
+				Element.Write(Writer, Options);
+			}
+			Writer.WriteEndArray();
+			Writer.WriteEndObject();
+		}
+
+		/// <summary>
+		/// Helper method to add a new <see cref="ActionButton"/> to the block.
+		/// </summary>
+		/// <param name="Text">The button text.</param>
+		/// <param name="URL">The URL to navigate to when clicked, if any.</param>
+		/// <param name="Value">
+		/// The value, if any, that should be passed in the interaction payload if the button is clicked.
+		/// </param>
+		/// <param name="ActionId">An optional identifier for handling interaction logic.</param>
+		/// <param name="Style">An option button style to apply to the button.</param>
+		public void AddButton(string Text, Uri? URL = null, string? Value = null, string? ActionId = null, ActionButton.ButtonStyle Style = ActionButton.ButtonStyle.Default)
+		{
+			Elements.Add(new ActionButton(Text, URL, Value, ActionId, Style));
+		}
+	}
+
+	/// <summary>
+	/// A base class for interaction Action elements to derive from.
+	/// </summary>
+	public abstract class ActionElement
+	{
+		/// <summary>
+		/// Writes the Json for an Action block element to the given <see cref="Utf8JsonWriter"/>.
+		/// </summary>
+		/// <param name="Writer"></param>
+		/// <param name="Options"></param>
+		public abstract void Write(Utf8JsonWriter Writer, JsonSerializerOptions Options);
+	}
+
+	/// <summary>
+	/// An button element that can be added to an <see cref="ActionsBlock"/>.
+	/// </summary>
+	public class ActionButton : ActionElement
+	{
+		/// <summary>
+		/// Determines the visual color scheme to use for the button.
+		/// </summary>
+		public enum ButtonStyle
+		{
+			/// <summary>
+			/// The default color scheme.
+			/// </summary>
+			Default,
+
+			/// <summary>
+			/// Used for affirmation/confirmation actions.
+			/// </summary>
+			Primary,
+
+			/// <summary>
+			/// Used when an action is destructive and cannot be undone.
+			/// </summary>
+			Danger,
+		}
+
+		/// <summary>
+		/// Construct a new Button action element.
+		/// </summary>
+		/// <param name="Text">The button text.</param>
+		/// <param name="URL">The URL to navigate to when clicked, if any.</param>
+		/// <param name="Value">
+		/// The value, if any, that should be passed in the interaction payload if the button is clicked.
+		/// </param>
+		/// <param name="ActionId">An optional identifier for handling interaction logic.</param>
+		/// <param name="Style">An option button style to apply to the button.</param>
+		public ActionButton(string Text, Uri? URL = null, string? Value = null, string? ActionId = null, ButtonStyle Style = ButtonStyle.Default)
+		{
+			this.Text = new TextObject(Text, IsMarkdown: false);
+			this.URL = URL;
+			this.Value = Value;
+			this.ActionId = ActionId;
+			this.Style = Style;
+		}
+
+		/// <summary>
+		/// The button text.
+		/// </summary>
+		public TextObject Text { get; } = new TextObject();
+
+		/// <summary>
+		/// The URL to navigate to when clicked, if any.
+		/// </summary>
+		public Uri? URL { get; set; }
+
+		/// <summary>
+		/// The value, if any, that should be passed in the interaction payload if the button is clicked.
+		/// </summary>
+		public string? Value { get; set; }
+
+		/// <summary>
+		/// An optional identifier for handling interaction logic.
+		/// </summary>
+		public string? ActionId { get; set; }
+
+		/// <summary>
+		/// Determines how the button appears in the message.
+		/// </summary>
+		public ButtonStyle Style { get; set; }
+
+		/// <inheritdoc/>
+		public override void Write(Utf8JsonWriter Writer, JsonSerializerOptions Options)
+		{
+			Writer.WriteStartObject();
+			Writer.WriteString("type", "button");
+			Text.Write(Writer);
+			if (URL != null)
+			{
+				Writer.WriteString("url", URL.AbsoluteUri);
+			}
+			if (!string.IsNullOrWhiteSpace(Value))
+			{
+				Writer.WriteString("value", Value);
+			}
+			if (!string.IsNullOrWhiteSpace(ActionId))
+			{
+				Writer.WriteString("action_id", ActionId);
+			}
+			if (Style != ButtonStyle.Default)
+			{
+				Writer.WriteString("style", Style.ToString().ToLower(CultureInfo.CurrentCulture));
+			}
+
+			Writer.WriteEndObject();
+		}
+	}
+}
diff --git a/Engine/Source/Programs/Horde/HordeServer/Utilities/BlockKit/BlockBase.cs b/Engine/Source/Programs/Horde/HordeServer/Utilities/BlockKit/BlockBase.cs
new file mode 100644
--- /dev/null
+++ b/Engine/Source/Programs/Horde/HordeServer/Utilities/BlockKit/BlockBase.cs
@@ -0,0 +1,42 @@
+// Copyright Epic Games, Inc. All Rights Reserved.
+
+using System;
+using System.Text.Json;
+using System.Text.Json.Serialization;
+
+namespace HordeServer.Utilities.Slack.BlockKit
+{
+	/// <summary>
+	/// Converter for overriding default Json serialization/deserialization for BlockKit blocks.
+	/// </summary>
+	public class BlockBaseConverter : JsonConverter<BlockBase>
+	{
+		/// <summary>
+		/// Reads a BlockKit block Json blob and turns it into the appropriate BlockKit block object.
+		/// </summary>
+		/// <returns>A BlockKit block object of the appropriate type.</returns>
+		public override BlockBase Read(ref Utf8JsonReader reader, Type typeToConvert, JsonSerializerOptions options) => throw new NotImplementedException();
+
+		/// <summary>
+		/// Writes the Json for a BlockKit block to the given <see cref="Utf8JsonWriter"/>.
+		/// </summary>
+		public override void Write(Utf8JsonWriter Writer, BlockBase Value, JsonSerializerOptions Options)
+		{
+			Value.Write(Writer, Options);
+		}
+	}
+
+	/// <summary>
+	/// A base class for all BlockKit blocks to derive from.
+	/// </summary>
+	[JsonConverter(typeof(BlockBaseConverter))]
+	public abstract class BlockBase
+	{
+		/// <summary>
+		/// Write a BlockKit block to a <see cref="Utf8JsonWriter"/>
+		/// </summary>
+		/// <param name="Writer"></param>
+		/// <param name="Options"></param>
+		public abstract void Write(Utf8JsonWriter Writer, JsonSerializerOptions Options);
+	}
+}
diff --git a/Engine/Source/Programs/Horde/HordeServer/Utilities/BlockKit/BlockKitAttachment.cs b/Engine/Source/Programs/Horde/HordeServer/Utilities/BlockKit/BlockKitAttachment.cs
new file mode 100644
--- /dev/null
+++ b/Engine/Source/Programs/Horde/HordeServer/Utilities/BlockKit/BlockKitAttachment.cs
@@ -0,0 +1,97 @@
+// Copyright Epic Games, Inc. All Rights Reserved.
+
+using System;
+using System.Collections.Generic;
+using System.Drawing;
+using System.Text.Json;
+using System.Text.Json.Serialization;
+
+namespace HordeServer.Utilities.Slack.BlockKit
+{
+	/// <summary>
+	/// Converter for overriding default Json serialization/deserialization for BlockKitAttachment objects.
+	/// </summary>
+	public class BlockKitAttachmentConverter : JsonConverter<BlockKitAttachment>
+	{
+		/// <summary>
+		/// Reads a BlockKitAttachment Json blob and turns it into the appropriate BlockKitAttachment object.
+		/// </summary>
+		/// <returns>A BlockKitAttachment object of the appropriate type.</returns>
+		public override BlockKitAttachment Read(ref Utf8JsonReader reader, Type typeToConvert, JsonSerializerOptions options) => throw new NotImplementedException();
+
+		/// <summary>
+		/// Writes the Json for a BlockKitAttachment to the given <see cref="Utf8JsonWriter"/>.
+		/// </summary>
+		public override void Write(Utf8JsonWriter Writer, BlockKitAttachment Value, JsonSerializerOptions Options)
+		{
+			Value.Write(Writer, Options);
+		}
+	}
+
+	/// <summary>
+	/// Static helpers for consistent color usage for <see cref="BlockKitMessage"/> messages.
+	/// </summary>
+	public static class BlockKitAttachmentColors
+	{
+		/// <summary>
+		/// The color to use for error messages.
+		/// </summary>
+		public static Color Error => Color.FromArgb(236, 76, 71);
+
+		/// <summary>
+		/// The color to use for warning messages.
+		/// </summary>
+		public static Color Warning => Color.FromArgb(247, 209, 84);
+
+		/// <summary>
+		/// The color to use for success messages.
+		/// </summary>
+		public static Color Success => Color.FromArgb(77, 181, 7);
+	}
+
+	/// <summary>
+	/// An attachment to a block kit message that is indented with a vertical, colored line on the left.
+	/// Should only be used for secondary information if possible.
+	/// </summary>
+	[JsonConverter(typeof(BlockKitAttachmentConverter))]
+	public class BlockKitAttachment
+	{
+		/// <summary>
+		/// Used for the "toast" notifications sent by Slack. If not set, and the Text property of the
+		/// base message isn't set, the notification will have a message saying it has no preview.
+		/// </summary>
+		public string? FallbackText { get; set; }
+
+		/// <summary>
+		/// The color of the line down the left side of the attachment.
+		/// </summary>
+		public Color Color { get; set; }
+
+		/// <summary>
+		/// A collection of BlockKit blocks the attachment is composed from.
+		/// </summary>
+		public List<BlockBase> Blocks { get; } = new List<BlockBase>();
+
+		/// <summary>
+		/// Write a BlockKit block to a <see cref="Utf8JsonWriter"/>
+		/// </summary>
+		/// <param name="Writer"></param>
+		/// <param name="Options"></param>
+		public void Write(Utf8JsonWriter Writer, JsonSerializerOptions Options)
+		{
+			Writer.WriteStartObject();
+			if (FallbackText != null)
+			{
+				Writer.WriteString("fallback", FallbackText);
+			}
+			Writer.WriteString("color", $"#{Color.ToArgb() & 0xFFFFFF:x6}");
+			Writer.WriteStartArray("blocks");
+			foreach (BlockBase Block in Blocks)
+			{
+				Block.Write(Writer, Options);
+			}
+			Writer.WriteEndArray();
+			Writer.WriteEndObject();
+		}
+	}
+}
diff --git a/Engine/Source/Programs/Horde/HordeServer/Utilities/BlockKit/BlockKitMessage.cs b/Engine/Source/Programs/Horde/HordeServer/Utilities/BlockKit/BlockKitMessage.cs
new file mode 100644
--- /dev/null
+++ b/Engine/Source/Programs/Horde/HordeServer/Utilities/BlockKit/BlockKitMessage.cs
@@ -0,0 +1,58 @@
+// Copyright Epic Games, Inc. All Rights Reserved.
+
+using System.Collections.Generic;
+using System.Text.Json.Serialization;
+
+namespace HordeServer.Utilities.Slack.BlockKit
+{
+    /// <summary>
+    /// An entire BlockKit message to be presented in Slack.
+    /// </summary>
+    public class BlockKitMessage
+    {
+		/// <summary>
+		/// The channel ID or user ID to deliver the notification to.
+		/// </summary>
+		[JsonPropertyName("channel")]
+		public string? Recipient { get; set; }
+
+		/// <summary>
+		/// Timestamp of the message to update
+		/// </summary>
+		[JsonPropertyName("ts")]
+		public string? Ts { get; set; }
+
+		/// <summary>
+		/// If no blocks are provided, this text will be displayed. This is also what is used on the pop up notification.
+		/// </summary>
+		[JsonPropertyName("text")]
+		public string? Text { get; set; }
+
+        /// <summary>
+        /// A collection of BlockKit blocks the message is composed from.
+        /// </summary>
+        [JsonPropertyName("blocks")]
+        public List<BlockBase> Blocks { get; } = new List<BlockBase>();
+
+        /// <summary>
+        /// Optional value that determines if Slack should try to replace the original message for the command that
+        /// this message is associated with.
+        /// </summary>
+        [JsonPropertyName("replace_original")]
+        public bool? ReplaceOriginal { get; set; }
+
+        /// <summary>
+        /// Optional value that determines if Slack should try to replace to delete original message for the command
+        /// that this message is associated with.
+        /// </summary>
+        [JsonPropertyName("delete_original")]
+        public bool? DeleteOriginal { get; set; }
+
+		/// <summary>
+		/// A collection of attachments.
+		/// NOTE: This feature is not supported any more but is not deprecated.
+		/// </summary>
+		[JsonPropertyName("attachments")]
+		public List<BlockKitAttachment> Attachments { get; } = new List<BlockKitAttachment>();
+    }
+}
diff --git a/Engine/Source/Programs/Horde/HordeServer/Utilities/BlockKit/DividerBlock.cs b/Engine/Source/Programs/Horde/HordeServer/Utilities/BlockKit/DividerBlock.cs
new file mode 100644
--- /dev/null
+++ b/Engine/Source/Programs/Horde/HordeServer/Utilities/BlockKit/DividerBlock.cs
@@ -0,0 +1,20 @@
+// Copyright Epic Games, Inc. All Rights Reserved.
+
+using System.Text.Json;
+
+namespace HordeServer.Utilities.Slack.BlockKit
+{
+	/// <summary>
+	/// Wrapper for a BlockKit Divider block that is used to add a divider between two blocks.
+	/// </summary>
+	public class DividerBlock : BlockBase
+	{
+		/// <inheritdoc/>
+		public override void Write(Utf8JsonWriter Writer, JsonSerializerOptions Options)
+		{
+			Writer.WriteStartObject();
+			Writer.WriteString("type", "divider");
+			Writer.WriteEndObject();
+		}
+	}
+}
diff --git a/Engine/Source/Programs/Horde/HordeServer/Utilities/BlockKit/HeaderBlock.cs b/Engine/Source/Programs/Horde/HordeServer/Utilities/BlockKit/HeaderBlock.cs
new file mode 100644
--- /dev/null
+++ b/Engine/Source/Programs/Horde/HordeServer/Utilities/BlockKit/HeaderBlock.cs
@@ -0,0 +1,37 @@
+// Copyright Epic Games, Inc. All Rights Reserved.
+
+using System.Text.Json;
+
+namespace HordeServer.Utilities.Slack.BlockKit
+{
+	/// <summary>
+	/// Represents a BlockKit Section block
+	/// </summary>
+	public class HeaderBlock : BlockBase
+	{
+		/// <summary>
+		/// Constructor
+		/// </summary>
+		/// <param name="Text">Any text to initiale the Section with</param>
+		/// <param name="IsMarkdown">If true, text is markdown, else text is considered plain text</param>
+		/// <param name="Emoji">If true, allows escaping of emojis with colons.</param>
+		public HeaderBlock(string Text = "", bool IsMarkdown = true, bool Emoji = false)
+		{
+			this.Text = new TextObject(Text, IsMarkdown, Emoji);
+		}
+
+		/// <summary>
+		/// The text of the Section block
+		/// </summary>
+		public TextObject Text { get; }
+
+		/// <inheritdoc/>
+		public override void Write(Utf8JsonWriter Writer, JsonSerializerOptions Options)
+		{
+			Writer.WriteStartObject();
+			Writer.WriteString("type", "header");
+			Text.Write(Writer);
+			Writer.WriteEndObject();
+		}
+	}
+}
diff --git a/Engine/Source/Programs/Horde/HordeServer/Utilities/BlockKit/SectionBlock.cs b/Engine/Source/Programs/Horde/HordeServer/Utilities/BlockKit/SectionBlock.cs
new file mode 100644
--- /dev/null
+++ b/Engine/Source/Programs/Horde/HordeServer/Utilities/BlockKit/SectionBlock.cs
@@ -0,0 +1,37 @@
+// Copyright Epic Games, Inc. All Rights Reserved.
+
+using System.Text.Json;
+
+namespace HordeServer.Utilities.Slack.BlockKit
+{
+	/// <summary>
+	/// Represents a BlockKit Section block
+	/// </summary>
+	public class SectionBlock : BlockBase
+	{
+		/// <summary>
+		/// Constructor
+		/// </summary>
+		/// <param name="Text">Any text to initiale the Section with</param>
+		/// <param name="IsMarkdown">If true, text is markdown, else text is considered plain text</param>
+		/// <param name="Emoji">If true, allows escaping of emoji with colon.</param>
+		public SectionBlock(string Text = "", bool IsMarkdown = true, bool Emoji = false)
+		{
+			this.Text = new TextObject(Text, IsMarkdown, Emoji);
+		}
+
+		/// <summary>
+		/// The text of the Section block
+		/// </summary>
+		public TextObject Text { get; }
+
+		/// <inheritdoc/>
+		public override void Write(Utf8JsonWriter Writer, JsonSerializerOptions Options)
+		{
+			Writer.WriteStartObject();
+			Writer.WriteString("type", "section");
+			Text.Write(Writer);
+			Writer.WriteEndObject();
+		}
+	}
+}
diff --git a/Engine/Source/Programs/Horde/HordeServer/Utilities/BlockKit/SlackResponses.cs b/Engine/Source/Programs/Horde/HordeServer/Utilities/BlockKit/SlackResponses.cs
new file mode 100644
--- /dev/null
+++ b/Engine/Source/Programs/Horde/HordeServer/Utilities/BlockKit/SlackResponses.cs
@@ -0,0 +1,45 @@
+// Copyright Epic Games, Inc. All Rights Reserved.
+
+using System.Drawing;
+using HordeServer.Utilities.Slack.BlockKit;
+
+namespace HordeServer.Utilities.BlockKit
+{
+	/// <summary>
+	/// Helper methods for creating <see cref="BlockKitMessage"/> responses to send
+	/// to users on Slack via the response_url of an interaction payload.
+	/// </summary>
+	public static class SlackResponses
+	{
+		/// <summary>
+		/// Creates a response to send to Slack about receiving bad data for a request.
+		/// </summary>
+		/// <param name="Message">A message explaining what was wrong with the data.</param>
+		/// <returns>A <see cref="BlockKitMessage"/> to be sent to Slack.</returns>
+		public static BlockKitMessage CreateBadDataResponse(string Message)
+		{
+			BlockKitMessage BlockKitMessage = new BlockKitMessage();
+			BlockKitAttachment Attachment = new BlockKitAttachment();
+			Attachment.Color = BlockKitAttachmentColors.Error;
+			Attachment.Blocks.Add(new SectionBlock("Failed to process request due to bad data from Slack:"));
+			Attachment.Blocks.Add(new SectionBlock(Message));
+			BlockKitMessage.Attachments.Add(Attachment);
+			return BlockKitMessage;
+		}
+
+		/// <summary>
+		/// Creates a response to send to Slack when failing to assign an issue due to it not
+		/// being found in the DB.
+		/// </summary>
+		/// <returns>A <see cref="BlockKitMessage"/> to be sent to Slack.</returns>
+		public static BlockKitMessage CreateIssueNotFoundResponse(int IssueId)
+		{
+			BlockKitMessage BlockKitMessage = new BlockKitMessage();
+			BlockKitAttachment Attachment = new BlockKitAttachment();
+			Attachment.Color = BlockKitAttachmentColors.Error;
+			Attachment.Blocks.Add(new SectionBlock($"Failed to assign issue as issue {IssueId} could not be found."));
+			BlockKitMessage.Attachments.Add(Attachment);
+			return BlockKitMessage;
+		}
+	}
+}
diff --git a/Engine/Source/Programs/Horde/HordeServer/Utilities/BlockKit/TextObject.cs b/Engine/Source/Programs/Horde/HordeServer/Utilities/BlockKit/TextObject.cs
new file mode 100644
--- /dev/null
+++ b/Engine/Source/Programs/Horde/HordeServer/Utilities/BlockKit/TextObject.cs
@@ -0,0 +1,58 @@
+// Copyright Epic Games, Inc. All Rights Reserved.
+
+using System;
+using System.Text.Json;
+
+namespace HordeServer.Utilities.Slack.BlockKit
+{
+	/// <summary>
+	/// Represents a BlockKit TextObject which is used to present markdown or plain text
+	/// </summary>
+	public class TextObject
+	{
+		/// <summary>
+		/// Constructor
+		/// </summary>
+		/// <param name="Text">Any text to initialize with</param>
+		/// <param name="IsMarkdown">If true, text is markdown, else text is considered plain text</param>
+		/// <param name="Emoji">If true, allows escaping of emojis with a colon</param>
+		public TextObject(string Text = "", bool IsMarkdown = true, bool Emoji = false)
+		{
+			this.Text = Text;
+			this.IsMarkdown = IsMarkdown;
+			// force emoji to false if markdown is true as this invalidates the schema
+			this.Emoji = this.IsMarkdown ? false : Emoji;
+		}
+
+		/// <summary>
+		/// If true, text is markdown, else text is considered plain text
+		/// </summary>
+		public bool IsMarkdown { get; set; }
+
+		/// <summary>
+		/// If true, allows escaping of emojis with a colon
+		/// </summary>
+		public bool Emoji { get; set; }
+
+		/// <summary>
+		/// The plain text or markdown for this TextObject
+		/// </summary>
+		public string Text { get; set; }
+
+		/// <summary>
+		/// Writes the TextObject JSON to a <see cref="Utf8JsonWriter"/>
+		/// </summary>
+		/// <param name="Writer"></param>
+		public void Write(Utf8JsonWriter Writer)
+		{
+			Writer.WriteStartObject("text");
+			Writer.WriteString("type", IsMarkdown ? "mrkdwn" : "plain_text");
+			Writer.WriteString("text", Text);
+			if (Emoji)
+			{
+				Writer.WriteBoolean("emoji", Emoji);
+			}
+			Writer.WriteEndObject();
+		}
+	}
+}
diff --git a/Engine/Source/Programs/Horde/HordeServer/Utilities/BloomFilter.cs b/Engine/Source/Programs/Horde/HordeServer/Utilities/BloomFilter.cs
new file mode 100644
--- /dev/null
+++ b/Engine/Source/Programs/Horde/HordeServer/Utilities/BloomFilter.cs
@@ -0,0 +1,228 @@
+using System;
+using System.Collections.Generic;
+using System.Linq;
+using System.Threading.Tasks;
+
+namespace HordeServer.Utilities
+{
+	struct LogIndex
+	{
+		const int NumHashes = 3;
+		const int NgramLength = 3;
+		const int NgramMask = (1 << (NgramLength * 8)) - 1;
+
+		/// <summary>
+		/// Enumerates the hash codes for a given block of text
+		/// </summary>
+		/// <param name="Text">The text to parse</param>
+		/// <returns></returns>
+		public static IEnumerable<int> GetHashCodes(ReadOnlyMemory<byte> Text)
+		{
+			if(Text.Length >= NgramLength)
+			{
+				int Value = 0;
+				for(int Idx = 0; Idx < NgramLength - 1; Idx++)
+				{
+					Value = (Value << 8) | ToLowerUtf8(Text.Span[Idx]);
+				}
+				for(int Idx = NgramLength - 1; Idx < Text.Length; Idx++)
+				{
+					Value = ((Value << 8) | ToLowerUtf8(Text.Span[Idx])) & NgramMask;
+
+					int HashValue = Value;
+					yield return HashValue;
+
+					for(int HashIdx = 1; HashIdx < NumHashes; HashIdx++)
+					{
+						HashValue = Scramble(HashValue);
+						yield return HashValue;
+					}
+				}
+			}
+		}
+
+		/// <summary>
+		/// Converts a character into a format for hashing
+		/// </summary>
+		/// <param name="Value">The input byte</param>
+		/// <returns>The value to include in the trigram</returns>
+		static byte ToLowerUtf8(byte Value)
+		{
+			if (Value >= 'A' && Value <= 'Z')
+			{
+				return (byte)(Value + 'a' - 'A');
+			}
+			else
+			{
+				return Value;
+			}
+		}
+
+		/// <summary>
+		/// Scramble the given number using a pseudo-RNG
+		/// </summary>
+		/// <param name="Value">Initial value</param>
+		/// <returns>The scrambled value</returns>
+		private static int Scramble(int Value)
+		{
+			Value ^= Value << 13;
+			Value ^= Value >> 17;
+			Value ^= Value << 5;
+			return Value;
+		}
+	}
+
+	/// <summary>
+	/// Read only implementation of a bloom filter
+	/// </summary>
+	public class ReadOnlyBloomFilter
+	{
+		/// <summary>
+		/// Data for the bloom filter
+		/// </summary>
+		public ReadOnlyMemory<byte> Data { get; }
+
+		/// <summary>
+		/// Constructor
+		/// </summary>
+		/// <param name="Data">Data to construct from</param>
+		public ReadOnlyBloomFilter(ReadOnlyMemory<byte> Data)
+		{
+			this.Data = Data;
+		}
+
+		/// <summary>
+		/// Determines if the filter may contain a value
+		/// </summary>
+		/// <param name="HashCodes">Sequence of hash codes</param>
+		/// <returns>True if the value is in the set</returns>
+		public bool Contains(IEnumerable<int> HashCodes)
+		{
+			ReadOnlySpan<byte> Span = Data.Span;
+			foreach (int HashCode in HashCodes)
+			{
+				int Index = HashCode & ((Span.Length << 3) - 1);
+				if ((Span[Index >> 3] & (1 << (Index & 7))) == 0)
+				{
+					return false;
+				}
+			}
+			return true;
+		}
+
+		/// <summary>
+		/// Calculates the load of this filter (ie. the proportion of bits which are set)
+		/// </summary>
+		/// <returns>Value between 0 and 1</returns>
+		public double CalculateLoadFactor()
+		{
+			int SetCount = 0;
+
+			ReadOnlySpan<byte> Span = Data.Span;
+			for (int Idx = 0; Idx < Span.Length; Idx++)
+			{
+				for (int Bit = 0; Bit < 8; Bit++)
+				{
+					if ((Span[Idx] & (1 << Bit)) != 0)
+					{
+						SetCount++;
+					}
+				}
+			}
+
+			return (double)SetCount / (Span.Length * 8);
+		}
+	}
+
+	/// <summary>
+	/// Implementation of a bloom filter
+	/// </summary>
+	public class BloomFilter
+	{
+		/// <summary>
+		/// The array of bits in the filter
+		/// </summary>
+		public byte[] Data { get; }
+
+		/// <summary>
+		/// Constructs a new filter of the given size
+		/// </summary>
+		/// <param name="MinSize">The minimum size of the filter. Will be rounded up to the next power of two.</param>
+		public BloomFilter(int MinSize)
+		{
+			int Size = MinSize;
+			while ((Size & (Size - 1)) != 0)
+			{
+				Size |= (Size - 1);
+				Size++;
+			}
+			Data = new byte[Size];
+		}
+
+		/// <summary>
+		/// Constructs a filter from raw data
+		/// </summary>
+		/// <param name="Data">The data to construct from</param>
+		public BloomFilter(byte[] Data)
+		{
+			this.Data = Data;
+			if ((Data.Length & (Data.Length - 1)) != 0)
+			{
+				throw new ArgumentException("Array for bloom filter must be a power of 2 in size.");
+			}
+		}
+
+		/// <summary>
+		/// Add a value to the filter
+		/// </summary>
+		/// <param name="HashCodes">Sequence of hash codes</param>
+		public void Add(IEnumerable<int> HashCodes)
+		{
+			foreach(int HashCode in HashCodes)
+			{
+				int Index = HashCode & ((Data.Length << 3) - 1);
+				Data[Index >> 3] = (byte)(Data[Index >> 3] | (1 << (Index & 7)));
+			}
+		}
+
+		/// <summary>
+		/// Determines if the filter may contain a value
+		/// </summary>
+		/// <param name="HashCodes">Sequence of hash codes</param>
+		/// <returns>True if the value is in the set</returns>
+		public bool Contains(IEnumerable<int> HashCodes)
+		{
+			foreach (int HashCode in HashCodes)
+			{
+				int Index = HashCode & ((Data.Length << 3) - 1);
+				if ((Data[Index >> 3] & (1 << (Index & 7))) == 0)
+				{
+					return false;
+				}
+			}
+			return true;
+		}
+
+		/// <summary>
+		/// Calculates the load of this filter (ie. the proportion of bits which are set)
+		/// </summary>
+		/// <returns>Value between 0 and 1</returns>
+		public double CalculateLoadFactor()
+		{
+			int SetCount = 0;
+			int AllCount = 0;
+			for (int Idx = 0; Idx < Data.Length; Idx++)
+			{
+				for (int Bit = 0; Bit < 8; Bit++)
+				{
+					if ((Data[Idx] & (1 << Bit)) != 0)
+					{
+						SetCount++;
+					}
+					AllCount++;
+				}
+			}
+			return (double)SetCount / AllCount;
+		}
+	}
+}
