Initial drop of Horde Server/Agent for UE5/Main. Still missing some Epic-specific functionality required to build/run correctly.

diff --git a/Engine/Source/Programs/Horde/HordeServer/Utilities/BsonDocumentExtensions.cs b/Engine/Source/Programs/Horde/HordeServer/Utilities/BsonDocumentExtensions.cs
new file mode 100644
--- /dev/null
+++ b/Engine/Source/Programs/Horde/HordeServer/Utilities/BsonDocumentExtensions.cs
@@ -0,0 +1,108 @@
+using MongoDB.Bson;
+using System;
+using System.Collections.Generic;
+using System.Diagnostics.CodeAnalysis;
+using System.Linq;
+using System.Threading.Tasks;
+
+namespace HordeServer.Utilities
+{
+	/// <summary>
+	/// Extension methods for BsonDocument
+	/// </summary>
+	static class BsonDocumentExtensions
+	{
+		/// <summary>
+		/// Gets a property value from a document or subdocument, indicated with dotted notation
+		/// </summary>
+		/// <param name="Document">Document to get a property for</param>
+		/// <param name="Name">Name of the property</param>
+		/// <param name="Type">Expected type of the property</param>
+		/// <param name="OutValue">Receives the property value</param>
+		/// <returns>True if the property exists and was of the correct type</returns>
+		public static bool TryGetPropertyValue(this BsonDocument Document, string Name, BsonType Type, [NotNullWhen(true)] out BsonValue? OutValue)
+		{
+			int DotIdx = Name.IndexOf('.', StringComparison.Ordinal);
+			if (DotIdx == -1)
+			{
+				return TryGetDirectPropertyValue(Document, Name, Type, out OutValue);
+			}
+
+			BsonValue? DocValue;
+			if (TryGetDirectPropertyValue(Document, Name.Substring(0, DotIdx), BsonType.Document, out DocValue))
+			{
+				return TryGetPropertyValue(DocValue.AsBsonDocument, Name.Substring(DotIdx + 1), Type, out OutValue);
+			}
+
+			OutValue = null;
+			return false;
+		}
+
+		/// <summary>
+		/// Gets a property value that's an immediate child of the document
+		/// </summary>
+		/// <param name="Document">Document to get a property for</param>
+		/// <param name="Name">Name of the property</param>
+		/// <param name="Type">Expected type of the property</param>
+		/// <param name="OutValue">Receives the property value</param>
+		/// <returns>True if the property exists and was of the correct type</returns>
+		private static bool TryGetDirectPropertyValue(this BsonDocument Document, string Name, BsonType Type, [NotNullWhen(true)] out BsonValue? OutValue)
+		{
+			BsonValue Value;
+			if (Document.TryGetValue(Name, out Value) && Value.BsonType == Type)
+			{
+				OutValue = Value;
+				return true;
+			}
+			else
+			{
+				OutValue = null;
+				return false;
+			}
+		}
+
+		/// <summary>
+		/// Gets an int32 value from the document
+		/// </summary>
+		/// <param name="Document">Document to get a property for</param>
+		/// <param name="Name">Name of the property</param>
+		/// <param name="OutValue">Receives the property value</param>
+		/// <returns>True if the property was retrieved</returns>
+		public static bool TryGetInt32(this BsonDocument Document, string Name, out int OutValue)
+		{
+			BsonValue? Value;
+			if (Document.TryGetPropertyValue(Name, BsonType.Int32, out Value))
+			{
+				OutValue = Value.AsInt32;
+				return true;
+			}
+			else
+			{
+				OutValue = 0;
+				return false;
+			}
+		}
+
+		/// <summary>
+		/// Gets a string value from the document
+		/// </summary>
+		/// <param name="Document">Document to get a property for</param>
+		/// <param name="Name">Name of the property</param>
+		/// <param name="OutValue">Receives the property value</param>
+		/// <returns>True if the property was retrieved</returns>
+		public static bool TryGetString(this BsonDocument Document, string Name, [NotNullWhen(true)] out string? OutValue)
+		{
+			BsonValue? Value;
+			if (Document.TryGetPropertyValue(Name, BsonType.String, out Value))
+			{
+				OutValue = Value.AsString;
+				return true;
+			}
+			else
+			{
+				OutValue = null;
+				return false;
+			}
+		}
+	}
+}
diff --git a/Engine/Source/Programs/Horde/HordeServer/Utilities/CaseInsensitiveStringSet.cs b/Engine/Source/Programs/Horde/HordeServer/Utilities/CaseInsensitiveStringSet.cs
new file mode 100644
--- /dev/null
+++ b/Engine/Source/Programs/Horde/HordeServer/Utilities/CaseInsensitiveStringSet.cs
@@ -0,0 +1,31 @@
+using MongoDB.Bson.Serialization;
+using MongoDB.Bson.Serialization.Serializers;
+using System;
+using System.Collections.Generic;
+using System.Linq;
+using System.Threading.Tasks;
+
+namespace HordeServer.Utilities
+{
+	/// <summary>
+	/// Case insensitive set of strings
+	/// </summary>
+	public class CaseInsensitiveStringSet : HashSet<string>
+	{
+		/// <summary>
+		/// Constructor
+		/// </summary>
+		public CaseInsensitiveStringSet()
+			: base(StringComparer.OrdinalIgnoreCase)
+		{
+		}
+
+		/// <summary>
+		/// Constructor
+		/// </summary>
+		public CaseInsensitiveStringSet(IEnumerable<string> Items)
+			: base(Items, StringComparer.OrdinalIgnoreCase)
+		{
+		}
+	}
+}
diff --git a/Engine/Source/Programs/Horde/HordeServer/Utilities/CollectionExtensions.cs b/Engine/Source/Programs/Horde/HordeServer/Utilities/CollectionExtensions.cs
new file mode 100644
--- /dev/null
+++ b/Engine/Source/Programs/Horde/HordeServer/Utilities/CollectionExtensions.cs
@@ -0,0 +1,31 @@
+using Microsoft.AspNetCore.Http.Features;
+using Microsoft.AspNetCore.Mvc.Formatters;
+using Microsoft.IdentityModel.Tokens;
+using System;
+using System.Collections.Generic;
+using System.Linq;
+using System.Threading.Tasks;
+
+namespace HordeServer.Utilities
+{
+	/// <summary>
+	/// Extension methods for collections
+	/// </summary>
+	public static class CollectionExtensions
+	{
+		/// <summary>
+		/// Adds an arbitrary sequence of items to a protobuf map field
+		/// </summary>
+		/// <typeparam name="TKey">The key type</typeparam>
+		/// <typeparam name="TValue">The value type</typeparam>
+		/// <param name="Map">The map to update</param>
+		/// <param name="Sequence">Sequence of items to add</param>
+		public static void Add<TKey, TValue>(this Google.Protobuf.Collections.MapField<TKey, TValue> Map, IEnumerable<KeyValuePair<TKey, TValue>> Sequence)
+		{
+			foreach(KeyValuePair<TKey, TValue> Pair in Sequence)
+			{
+				Map.Add(Pair.Key, Pair.Value);
+			}
+		}
+	}
+}
diff --git a/Engine/Source/Programs/Horde/HordeServer/Utilities/CompressionExtensions.cs b/Engine/Source/Programs/Horde/HordeServer/Utilities/CompressionExtensions.cs
new file mode 100644
--- /dev/null
+++ b/Engine/Source/Programs/Horde/HordeServer/Utilities/CompressionExtensions.cs
@@ -0,0 +1,74 @@
+using Datadog.Trace;
+using EpicGames.Core;
+using ICSharpCode.SharpZipLib.BZip2;
+using System;
+using System.Buffers.Binary;
+using System.Collections.Generic;
+using System.Globalization;
+using System.IO;
+using System.Linq;
+using System.Threading.Tasks;
+
+namespace HordeServer.Utilities
+{
+	/// <summary>
+	/// Extension methods for compression
+	/// </summary>
+	static class CompressionExtensions
+	{
+		/// <summary>
+		/// Compress a block of data with bzip2
+		/// </summary>
+		/// <returns>The compressed data</returns>
+		public static byte[] CompressBzip2(this ReadOnlyMemory<byte> Memory)
+		{
+			using Scope Scope = Tracer.Instance.StartActive("CompressBzip2");
+			Scope.Span.SetTag("DecompressedSize", Memory.Length.ToString(CultureInfo.InvariantCulture));
+
+			byte[] CompressedData;
+			using (MemoryStream Stream = new MemoryStream())
+			{
+				byte[] DecompressedSize = new byte[4];
+				BinaryPrimitives.WriteInt32LittleEndian(DecompressedSize.AsSpan(), Memory.Length);
+				Stream.Write(DecompressedSize.AsSpan());
+
+				using (BZip2OutputStream CompressedStream = new BZip2OutputStream(Stream))
+				{
+					CompressedStream.Write(Memory.Span);
+				}
+
+				CompressedData = Stream.ToArray();
+			}
+
+			Scope.Span.SetTag("CompressedSize", CompressedData.Length.ToString(CultureInfo.InvariantCulture));
+			return CompressedData;
+		}
+
+		/// <summary>
+		/// Decompress the data
+		/// </summary>
+		/// <returns>The decompressed data</returns>
+		public static byte[] DecompressBzip2(this ReadOnlyMemory<byte> Memory)
+		{
+			int DecompressedSize = BinaryPrimitives.ReadInt32LittleEndian(Memory.Span);
+
+			using Scope Scope = Tracer.Instance.StartActive("DecompressBzip2");
+			Scope.Span.SetTag("CompressedSize", Memory.Length.ToString(CultureInfo.InvariantCulture));
+			Scope.Span.SetTag("DecompressedSize", DecompressedSize.ToString(CultureInfo.InvariantCulture));
+
+			byte[] Data = new byte[DecompressedSize];
+			using (ReadOnlyMemoryStream Stream = new ReadOnlyMemoryStream(Memory.Slice(4)))
+			{
+				using (BZip2InputStream DecompressedStream = new BZip2InputStream(Stream))
+				{
+					int ReadSize = DecompressedStream.Read(Data.AsSpan());
+					if (ReadSize != Data.Length)
+					{
+						throw new InvalidDataException($"Compressed data is too short (expected {Data.Length} bytes, got {ReadSize} bytes)");
+					}
+				}
+			}
+			return Data;
+		}
+	}
+}
diff --git a/Engine/Source/Programs/Horde/HordeServer/Utilities/ContentHashSerializer.cs b/Engine/Source/Programs/Horde/HordeServer/Utilities/ContentHashSerializer.cs
new file mode 100644
--- /dev/null
+++ b/Engine/Source/Programs/Horde/HordeServer/Utilities/ContentHashSerializer.cs
@@ -0,0 +1,52 @@
+using EpicGames.Core;
+using MongoDB.Bson.Serialization;
+using System;
+using System.Collections.Generic;
+using System.Linq;
+using System.Threading.Tasks;
+
+namespace HordeServer.Utilities
+{
+	/// <summary>
+	/// Serializer for JobStepRefId objects
+	/// </summary>
+	public sealed class ContentHashSerializer : IBsonSerializer<ContentHash>
+	{
+		/// <inheritdoc/>
+		public Type ValueType
+		{
+			get { return typeof(ContentHash); }
+		}
+
+		/// <inheritdoc/>
+		void IBsonSerializer.Serialize(BsonSerializationContext Context, BsonSerializationArgs Args, object Value)
+		{
+			Serialize(Context, Args, (ContentHash)Value);
+		}
+
+		/// <inheritdoc/>
+		object IBsonSerializer.Deserialize(BsonDeserializationContext Context, BsonDeserializationArgs Args)
+		{
+			return ((IBsonSerializer<ContentHash>)this).Deserialize(Context, Args);
+		}
+
+		/// <inheritdoc/>
+		public ContentHash Deserialize(BsonDeserializationContext Context, BsonDeserializationArgs Args)
+		{
+			if (Context.Reader.CurrentBsonType == MongoDB.Bson.BsonType.ObjectId)
+			{
+				return new ContentHash(Context.Reader.ReadObjectId().ToByteArray());
+			}
+			else
+			{
+				return ContentHash.Parse(Context.Reader.ReadString());
+			}
+		}
+
+		/// <inheritdoc/>
+		public void Serialize(BsonSerializationContext Context, BsonSerializationArgs Args, ContentHash Value)
+		{
+			Context.Writer.WriteString(((ContentHash)Value).ToString());
+		}
+	}
+}
diff --git a/Engine/Source/Programs/Horde/HordeServer/Utilities/CustomFileCallbackResult.cs b/Engine/Source/Programs/Horde/HordeServer/Utilities/CustomFileCallbackResult.cs
new file mode 100644
--- /dev/null
+++ b/Engine/Source/Programs/Horde/HordeServer/Utilities/CustomFileCallbackResult.cs
@@ -0,0 +1,94 @@
+using Microsoft.AspNetCore.Mvc;
+using Microsoft.AspNetCore.Mvc.Infrastructure;
+using Microsoft.Extensions.DependencyInjection;
+using Microsoft.Extensions.Logging;
+using System;
+using System.Collections.Generic;
+using System.IO;
+using System.Linq;
+using System.Net.Http.Headers;
+using System.Net.Mime;
+using System.Threading;
+using System.Threading.Tasks;
+
+namespace HordeServer.Utilities
+{
+    /// <summary>
+    /// Class deriving from a FileResult that allows custom file types (used for zip file creation)
+    /// </summary>
+    public class CustomFileCallbackResult : FileResult
+    {
+        private Func<Stream, ActionContext, Task> Callback;
+		string FileName;
+		bool Inline;
+
+		/// <summary>
+		/// Constructor
+		/// </summary>
+		/// <param name="FileName">Default filename for the downloaded file</param>
+		/// <param name="MimeType">Content type for the file</param>
+		/// <param name="Inline">Whether to display the file inline in the browser</param>
+		/// <param name="Callback">Callback used to write the data</param>
+		public CustomFileCallbackResult(string FileName, string MimeType, bool Inline, Func<Stream, ActionContext, Task> Callback)
+            : base(MimeType)
+        {
+			this.FileName = FileName;
+			this.Inline = Inline;
+            this.Callback = Callback;
+        }
+
+        /// <summary>
+        /// Executes the action result
+        /// </summary>
+        /// <param name="Context">The controller context</param>
+        /// <returns></returns>
+        public override Task ExecuteResultAsync(ActionContext Context)
+        {
+			ContentDisposition ContentDisposition = new ContentDisposition();
+			ContentDisposition.Inline = Inline;
+			ContentDisposition.FileName = FileName;
+			Context.HttpContext.Response.Headers.Add("Content-Disposition", ContentDisposition.ToString());
+
+			CustomFileCallbackResultExecutor Executor = new CustomFileCallbackResultExecutor(Context.HttpContext.RequestServices.GetRequiredService<ILoggerFactory>());
+            return Executor.ExecuteAsync(Context, this);
+        }
+
+        /// <summary>
+        /// Exectutor for the custom FileResult
+        /// </summary>
+        private sealed class CustomFileCallbackResultExecutor : FileResultExecutorBase
+        {
+            /// <summary>
+            /// Constructor
+            /// </summary>
+            /// <param name="LoggerFactory">The logger</param>
+            public CustomFileCallbackResultExecutor(ILoggerFactory LoggerFactory)
+                : base(CreateLogger<CustomFileCallbackResultExecutor>(LoggerFactory))
+            {
+            }
+
+            /// <summary>
+            /// Executes a CustomFileResult callback
+            /// </summary>
+            /// <param name="Context">The controller context</param>
+            /// <param name="Result">The custom file result</param>
+            /// <returns></returns>
+            public Task ExecuteAsync(ActionContext Context, CustomFileCallbackResult Result)
+            {
+                SetHeadersAndLog(Context, Result, null, false);
+                return Result.Callback(Context.HttpContext.Response.Body, Context);
+            }
+        }
+    }
+
+    /// <summary>
+    /// Stream overriding CanSeek to false so the zip file plays nice with it.
+    /// </summary>
+    public class CustomBufferStream : MemoryStream
+    {
+        /// <summary>
+        /// Always report unseekable.
+        /// </summary>
+        public override bool CanSeek => false;
+    }
+}
diff --git a/Engine/Source/Programs/Horde/HordeServer/Utilities/DateTimeOffsetStringSerializer.cs b/Engine/Source/Programs/Horde/HordeServer/Utilities/DateTimeOffsetStringSerializer.cs
new file mode 100644
--- /dev/null
+++ b/Engine/Source/Programs/Horde/HordeServer/Utilities/DateTimeOffsetStringSerializer.cs
@@ -0,0 +1,40 @@
+using MongoDB.Bson;
+using MongoDB.Bson.Serialization;
+using MongoDB.Bson.Serialization.Serializers;
+using Newtonsoft.Json.Bson;
+using System;
+using System.Collections.Generic;
+using System.Globalization;
+using System.Linq;
+using System.Threading.Tasks;
+
+namespace HordeServer.Utilities
+{
+	/// <summary>
+	/// Serializer for DateTimeOffset structs, which writes data as a string
+	/// </summary>
+	class DateTimeOffsetStringSerializer : StructSerializerBase<DateTimeOffset>
+	{
+		/// <inheritdoc/>
+		public override DateTimeOffset Deserialize(BsonDeserializationContext Context, BsonDeserializationArgs Args)
+		{
+			switch (Context.Reader.CurrentBsonType)
+			{
+				case BsonType.String:
+					string StringValue = Context.Reader.ReadString();
+					return DateTimeOffset.Parse(StringValue, DateTimeFormatInfo.InvariantInfo);
+				case BsonType.DateTime:
+					long Ticks = Context.Reader.ReadDateTime();
+					return DateTimeOffset.FromUnixTimeMilliseconds(Ticks);
+				default:
+					throw new FormatException($"Unable to deserialize a DateTimeOffset from a {Context.Reader.CurrentBsonType}");
+			}
+		}
+
+		/// <inheritdoc/>
+		public override void Serialize(BsonSerializationContext Context, BsonSerializationArgs Args, DateTimeOffset Value)
+		{
+			Context.Writer.WriteString(Value.ToString(DateTimeFormatInfo.InvariantInfo));
+		}
+	}
+}
diff --git a/Engine/Source/Programs/Horde/HordeServer/Utilities/DefaultDiscriminatorConvention.cs b/Engine/Source/Programs/Horde/HordeServer/Utilities/DefaultDiscriminatorConvention.cs
new file mode 100644
--- /dev/null
+++ b/Engine/Source/Programs/Horde/HordeServer/Utilities/DefaultDiscriminatorConvention.cs
@@ -0,0 +1,82 @@
+using MongoDB.Bson;
+using MongoDB.Bson.IO;
+using MongoDB.Bson.Serialization.Conventions;
+using System;
+using System.Collections.Generic;
+using System.Linq;
+using System.Threading.Tasks;
+
+namespace HordeServer.Utilities
+{
+	/// <summary>
+	/// Implements a discriminator convention, which defaults to a specific concrete type if no discriminator is specified.
+	/// 
+	/// This class should also be registered as the discriminator for the concrete type as well as the interface, because the default behavior
+	/// of the serializer is to recognize that the concrete type is discriminated and fall back to trying to read the discriminator element,
+	/// failing, and falling back to serializing it as the interface recursively until we get a stack overflow.
+	/// 
+	/// Registering this as the discriminator for the concrete type prevents this, because we can discriminate the concrete type correctly.
+	/// </summary>
+	class DefaultDiscriminatorConvention : IDiscriminatorConvention
+	{
+		/// <summary>
+		/// The normal discriminator to use
+		/// </summary>
+		static IDiscriminatorConvention Inner { get; } = StandardDiscriminatorConvention.Hierarchical;
+
+		/// <summary>
+		/// The nominal type
+		/// </summary>
+		Type BaseType;
+
+		/// <summary>
+		/// Default type to use if the inner discriminator returns the nominal type
+		/// </summary>
+		Type DefaultType;
+
+		/// <summary>
+		/// Constructor
+		/// </summary>
+		/// <param name="BaseType">The base class</param>
+		/// <param name="DefaultType">Default type to use if the inner discriminator returns an interface</param>
+		public DefaultDiscriminatorConvention(Type BaseType, Type DefaultType)
+		{
+			this.BaseType = BaseType;
+			this.DefaultType = DefaultType;
+		}
+
+		/// <inheritdoc/>
+		public string ElementName => Inner.ElementName;
+
+		/// <inheritdoc/>
+		public Type GetActualType(IBsonReader BsonReader, Type NominalType)
+		{
+			Type ActualType = Inner.GetActualType(BsonReader, NominalType);
+			if (ActualType == BaseType)
+			{
+				ActualType = DefaultType;
+			}
+			return ActualType;
+		}
+
+		/// <inheritdoc/>
+		public BsonValue GetDiscriminator(Type NominalType, Type ActualType)
+		{
+			return Inner.GetDiscriminator(NominalType, ActualType);
+		}
+	}
+
+	/// <summary>
+	/// Generic version of <see cref="DefaultDiscriminatorConvention"/>
+	/// </summary>
+	class DefaultDiscriminatorConvention<TBase, TDerived> : DefaultDiscriminatorConvention
+	{
+		/// <summary>
+		/// Constructor
+		/// </summary>
+		public DefaultDiscriminatorConvention()
+			: base(typeof(TBase), typeof(TDerived))
+		{
+		}
+	}
+}
diff --git a/Engine/Source/Programs/Horde/HordeServer/Utilities/ElectedBackgroundService.cs b/Engine/Source/Programs/Horde/HordeServer/Utilities/ElectedBackgroundService.cs
new file mode 100644
--- /dev/null
+++ b/Engine/Source/Programs/Horde/HordeServer/Utilities/ElectedBackgroundService.cs
@@ -0,0 +1,118 @@
+using HordeServer.Models;
+using HordeServer.Services;
+using Microsoft.Extensions.Logging;
+using MongoDB.Bson;
+using System;
+using System.Collections.Generic;
+using System.Linq;
+using System.Threading;
+using System.Threading.Tasks;
+
+namespace HordeServer.Utilities
+{
+	/// <summary>
+	/// Service which is ticked on one machine in the cluster at a specified interval
+	/// </summary>
+	public abstract class ElectedBackgroundService : TickedBackgroundService
+	{
+		class State : SingletonBase
+		{
+			public string? Host { get; set; }
+			public DateTime NextUpdateTime { get; set; }
+		}
+
+		SingletonDocument<State> StateAccessor;
+		TimeSpan PollInterval;
+		TimeSpan ElectionTimeout;
+
+		/// <summary>
+		/// Constructor
+		/// </summary>
+		/// <param name="DatabaseService">The database service instance</param>
+		/// <param name="ServiceId">Object id for the singleton tracking the update state</param>
+		/// <param name="Logger">Logging device</param>
+		public ElectedBackgroundService(DatabaseService DatabaseService, ObjectId ServiceId, ILogger Logger)
+			: this(DatabaseService, ServiceId, TimeSpan.FromMinutes(1.0), TimeSpan.FromMinutes(5.0), Logger)
+		{
+		}
+
+		/// <summary>
+		/// Constructor
+		/// </summary>
+		/// <param name="DatabaseService">The database service instance</param>
+		/// <param name="ServiceId">Object id for the singleton tracking the update state</param>
+		/// <param name="PollInterval">Frequency to poll state</param>
+		/// <param name="ElectionTimeout"></param>
+		/// <param name="Logger">Logging device</param>
+		public ElectedBackgroundService(DatabaseService DatabaseService, ObjectId ServiceId, TimeSpan PollInterval, TimeSpan ElectionTimeout, ILogger Logger)
+			: base(null, Logger)
+		{
+			this.StateAccessor = new SingletonDocument<State>(DatabaseService, ServiceId);
+			this.PollInterval = PollInterval;
+			this.ElectionTimeout = ElectionTimeout;
+		}
+
+		/// <summary>
+		/// Tick the local service, and poll for whether to tick the shared service.
+		/// </summary>
+		/// <param name="StoppingToken">Stopping token</param>
+		/// <returns>Async task</returns>
+		protected override sealed async Task TickAsync(CancellationToken StoppingToken)
+		{
+			DateTime UtcNow = DateTime.UtcNow;
+			for (; ; )
+			{
+				State State = await StateAccessor.GetAsync();
+
+				// Check if it's time to update yet
+				TimeSpan Delay = State.NextUpdateTime - UtcNow;
+				if(Delay > TimeSpan.Zero)
+				{
+					TimeSpan RetryDelay = Delay + TimeSpan.FromSeconds(0.5);
+					if (PollInterval < RetryDelay)
+					{
+						RetryDelay = PollInterval;
+					}
+
+					Interval = RetryDelay;
+					break;
+				}
+
+				// Try to become the leader and execute the update
+				State.Host = System.Net.Dns.GetHostName();
+				State.NextUpdateTime = UtcNow + ElectionTimeout;
+
+				if (await StateAccessor.TryUpdateAsync(State))
+				{
+					// Update the next update time while we're still running the tick
+					Task<DateTime> TickTask = Task.Run(() => TickLeaderAsync(StoppingToken));
+					while (await Task.WhenAny(TickTask, Task.Delay(ElectionTimeout / 2)) != TickTask)
+					{
+						State.NextUpdateTime = DateTime.UtcNow + ElectionTimeout;
+						await StateAccessor.TryUpdateAsync(State);
+					}
+					State.NextUpdateTime = await TickTask;
+					await StateAccessor.TryUpdateAsync(State);
+				}
+			}
+		}
+
+		/// <summary>
+		/// Tick the shared service
+		/// </summary>
+		/// <param name="StoppingToken">Stopping token</param>
+		/// <returns>Async task</returns>
+		protected abstract Task<DateTime> TickLeaderAsync(CancellationToken StoppingToken);
+		
+		/// <summary>
+		/// Tick the implemented method from above
+		/// **Only for use in testing!**
+		/// </summary>
+		/// <param name="StoppingToken">Optional token signalling that the task should terminate</param>
+		/// <returns>Async task</returns>
+		public Task TickSharedOnlyForTestingAsync(CancellationToken? StoppingToken = null)
+		{
+			return TickLeaderAsync(StoppingToken ?? new CancellationToken());
+		}
+	}
+}
diff --git a/Engine/Source/Programs/Horde/HordeServer/Utilities/HordeClaimTypes.cs b/Engine/Source/Programs/Horde/HordeServer/Utilities/HordeClaimTypes.cs
new file mode 100644
--- /dev/null
+++ b/Engine/Source/Programs/Horde/HordeServer/Utilities/HordeClaimTypes.cs
@@ -0,0 +1,110 @@
+using HordeServer.Collections;
+using MongoDB.Bson;
+using System;
+using System.Collections.Generic;
+using System.Linq;
+using System.Security.Claims;
+using System.Threading.Tasks;
+
+namespace HordeServer.Utilities
+{
+	/// <summary>
+	/// Claim types that are specific to horde
+	/// </summary>
+	public static class HordeClaimTypes
+	{
+		/// <summary>
+		/// Base URI for all Horde claims.
+		/// </summary>
+		const string Prefix = "http://epicgames.com/ue/horde/";
+
+		/// <summary>
+		/// Claim for a particular role.
+		/// </summary>
+		public const string Role = Prefix + "role";
+
+		/// <summary>
+		/// Claim for a particular session
+		/// </summary>
+		public const string AgentSessionId = Prefix + "session";
+
+		/// <summary>
+		/// User name for horde
+		/// </summary>
+		public const string User = Prefix + "user";
+
+		/// <summary>
+		/// Unique id of the horde user (see <see cref="IUserCollection"/>)
+		/// </summary>
+		public const string UserId = Prefix + "user-id-v2";
+
+		/// <summary>
+		/// Claim for downloading artifacts from a particular job
+		/// </summary>
+		public const string JobArtifacts = Prefix + "job-artifacts";
+
+		/// <summary>
+		/// Claim for the Perforce username
+		/// </summary>
+		public const string PerforceUser = Prefix + "perforce-user";
+	}
+
+	/// <summary>
+	/// Extension methods for getting Horde claims from a principal
+	/// </summary>
+	public static class HordeClaimExtensions
+	{
+		/// <summary>
+		/// Gets the Horde user id from a principal
+		/// </summary>
+		/// <param name="Principal"></param>
+		/// <returns></returns>
+		public static ObjectId? GetUserId(this ClaimsPrincipal Principal)
+		{
+			string? IdValue = Principal.FindFirstValue(HordeClaimTypes.UserId);
+			if(IdValue == null)
+			{
+				return null;
+			}
+			else
+			{
+				return new ObjectId(IdValue);
+			}
+		}
+
+		/// <summary>
+		/// Gets the Horde user name from a principal
+		/// </summary>
+		/// <param name="Principal"></param>
+		/// <returns></returns>
+		public static string? GetUserName(this ClaimsPrincipal Principal)
+		{
+			return Principal.FindFirstValue(HordeClaimTypes.User) ?? Principal.FindFirstValue(ClaimTypes.Name);
+		}
+
+		/// <summary>
+		/// Get the email for the given user
+		/// </summary>
+		/// <param name="User">The user to query the email for</param>
+		/// <returns>The user's email address or null if not found</returns>
+		public static string? GetEmail(this ClaimsPrincipal User)
+		{
+			return User.Claims.FirstOrDefault(x => x.Type == ClaimTypes.Email)?.Value;
+		}
+
+		/// <summary>
+		/// Gets the perforce username for the given principal
+		/// </summary>
+		/// <param name="User">The principal to get the Perforce user for</param>
+		/// <returns>Perforce user name</returns>
+		public static string? GetPerforceUser(this ClaimsPrincipal User)
+		{
+			Claim? Claim = User.FindFirst(HordeClaimTypes.PerforceUser);
+			if (Claim == null)
+			{
+				return null;
+			}
+			return Claim.Value;
+		}
+	}
+}
diff --git a/Engine/Source/Programs/Horde/HordeServer/Utilities/JsonElementExtensions.cs b/Engine/Source/Programs/Horde/HordeServer/Utilities/JsonElementExtensions.cs
new file mode 100644
--- /dev/null
+++ b/Engine/Source/Programs/Horde/HordeServer/Utilities/JsonElementExtensions.cs
@@ -0,0 +1,204 @@
+using MongoDB.Bson.IO;
+using System;
+using System.Collections.Generic;
+using System.Diagnostics.CodeAnalysis;
+using System.Linq;
+using System.Text;
+using System.Text.Json;
+using System.Threading.Tasks;
+
+namespace HordeServer.Utilities
+{
+	/// <summary>
+	/// Extensions for parsing values out of generic dictionary objects
+	/// </summary>
+	static class JsonElementExtensions
+	{
+		/// <summary>
+		/// Checks if the element has a property with the given value
+		/// </summary>
+		/// <param name="Element">Element to search</param>
+		/// <param name="Name">Name of the property</param>
+		/// <param name="Value">Expected value of the property</param>
+		/// <returns>True if the property exists and matches</returns>
+		public static bool HasStringProperty(this JsonElement Element, ReadOnlySpan<byte> Name, string Value)
+		{
+			JsonElement Property;
+			return Element.ValueKind == JsonValueKind.Object && Element.TryGetProperty(Name, out Property) && Property.ValueEquals(Value);
+		}
+
+		/// <summary>
+		/// Checks if the element has a property with the given value
+		/// </summary>
+		/// <param name="Element">Element to search</param>
+		/// <param name="Name">Name of the property</param>
+		/// <param name="Value">Expected value of the property</param>
+		/// <returns>True if the property exists and matches</returns>
+		public static bool HasStringProperty(this JsonElement Element, string Name, string Value)
+		{
+			JsonElement Property;
+			return Element.ValueKind == JsonValueKind.Object && Element.TryGetProperty(Name, out Property) && Property.ValueEquals(Value);
+		}
+
+		/// <summary>
+		/// Gets a property value of a certain type
+		/// </summary>
+		/// <param name="Element">The element to get a property from</param>
+		/// <param name="Name">Name of the element</param>
+		/// <param name="ValueKind">The required type of property</param>
+		/// <param name="Value">Value of the property</param>
+		/// <returns>True if the property exists and was a string</returns>
+		public static bool TryGetProperty(this JsonElement Element, ReadOnlySpan<byte> Name, JsonValueKind ValueKind, [NotNullWhen(true)] out JsonElement Value)
+		{
+			JsonElement Property;
+			if (Element.TryGetProperty(Name, out Property) && Property.ValueKind == ValueKind)
+			{
+				Value = Property;
+				return true;
+			}
+			else
+			{
+				Value = new JsonElement();
+				return false;
+			}
+		}
+
+		/// <summary>
+		/// Gets a property value of a certain type
+		/// </summary>
+		/// <param name="Element">The element to get a property from</param>
+		/// <param name="Name">Name of the element</param>
+		/// <param name="ValueKind">The required type of property</param>
+		/// <param name="Value">Value of the property</param>
+		/// <returns>True if the property exists and was a string</returns>
+		public static bool TryGetProperty(this JsonElement Element, string Name, JsonValueKind ValueKind, [NotNullWhen(true)] out JsonElement Value)
+		{
+			JsonElement Property;
+			if (Element.TryGetProperty(Name, out Property) && Property.ValueKind == ValueKind)
+			{
+				Value = Property;
+				return true;
+			}
+			else
+			{
+				Value = new JsonElement();
+				return false;
+			}
+		}
+
+		/// <summary>
+		/// Gets a string property value
+		/// </summary>
+		/// <param name="Element">The element to get a property from</param>
+		/// <param name="Name">Name of the element</param>
+		/// <param name="Value">Value of the property</param>
+		/// <returns>True if the property exists and was a string</returns>
+		public static bool TryGetStringProperty(this JsonElement Element, ReadOnlySpan<byte> Name, [NotNullWhen(true)] out string? Value)
+		{
+			JsonElement Property;
+			if (Element.TryGetProperty(Name, out Property) && Property.ValueKind == JsonValueKind.String)
+			{
+				Value = Property.GetString();
+				return true;
+			}
+			else
+			{
+				Value = null;
+				return false;
+			}
+		}
+
+		/// <summary>
+		/// Gets a string property value
+		/// </summary>
+		/// <param name="Element">The element to get a property from</param>
+		/// <param name="Name">Name of the element</param>
+		/// <param name="Value">Value of the property</param>
+		/// <returns>True if the property exists and was a string</returns>
+		public static bool TryGetStringProperty(this JsonElement Element, string Name, [NotNullWhen(true)] out string? Value)
+		{
+			return TryGetStringProperty(Element, Encoding.UTF8.GetBytes(Name).AsSpan(), out Value);
+		}
+
+		/// <summary>
+		/// Gets a property value from a document or subdocument, indicated with dotted notation
+		/// </summary>
+		/// <param name="Element">Document to get a property for</param>
+		/// <param name="Name">Name of the property</param>
+		/// <param name="OutElement">Receives the nexted element</param>
+		/// <returns>True if the property exists and was of the correct type</returns>
+		public static bool TryGetNestedProperty(this JsonElement Element, ReadOnlySpan<byte> Name, [NotNullWhen(true)] out JsonElement OutElement)
+		{
+			int DotIdx = Name.IndexOf((byte)'.');
+			if (DotIdx == -1)
+			{
+				return Element.TryGetProperty(Name, out OutElement);
+			}
+
+			JsonElement DocValue;
+			if (Element.TryGetProperty(Name.Slice(0, DotIdx), out DocValue) && DocValue.ValueKind == JsonValueKind.Object)
+			{
+				return TryGetNestedProperty(DocValue, Name.Slice(DotIdx + 1), out OutElement);
+			}
+
+			OutElement = new JsonElement();
+			return false;
+		}
+
+		/// <summary>
+		/// Gets a property value from a document or subdocument, indicated with dotted notation
+		/// </summary>
+		/// <param name="Element">Document to get a property for</param>
+		/// <param name="Name">Name of the property</param>
+		/// <param name="OutElement">Receives the nexted element</param>
+		/// <returns>True if the property exists and was of the correct type</returns>
+		public static bool TryGetNestedProperty(this JsonElement Element, string Name, [NotNullWhen(true)] out JsonElement OutElement)
+		{
+			return TryGetNestedProperty(Element, Encoding.UTF8.GetBytes(Name).AsSpan(), out OutElement);
+		}
+
+		/// <summary>
+		/// Gets an int32 value from the document
+		/// </summary>
+		/// <param name="Element">Document to get a property for</param>
+		/// <param name="Name">Name of the property</param>
+		/// <param name="OutValue">Receives the property value</param>
+		/// <returns>True if the property was retrieved</returns>
+		public static bool TryGetNestedProperty(this JsonElement Element, string Name, out int OutValue)
+		{
+			JsonElement Value;
+			if (Element.TryGetNestedProperty(Name, out Value) && Value.ValueKind == JsonValueKind.Number)
+			{
+				OutValue = Value.GetInt32();
+				return true;
+			}
+			else
+			{
+				OutValue = 0;
+				return false;
+			}
+		}
+
+		/// <summary>
+		/// Gets a string value from the document
+		/// </summary>
+		/// <param name="Element">Document to get a property for</param>
+		/// <param name="Name">Name of the property</param>
+		/// <param name="OutValue">Receives the property value</param>
+		/// <returns>True if the property was retrieved</returns>
+		public static bool TryGetNestedProperty(this JsonElement Element, string Name, [NotNullWhen(true)] out string? OutValue)
+		{
+			JsonElement Value;
+			if (Element.TryGetNestedProperty(Name, out Value) && Value.ValueKind == JsonValueKind.String)
+			{
+				OutValue = Value.GetString();
+				return OutValue != null;
+			}
+			else
+			{
+				OutValue = null;
+				return false;
+			}
+		}
+	}
+}
diff --git a/Engine/Source/Programs/Horde/HordeServer/Utilities/JwtUtils.cs b/Engine/Source/Programs/Horde/HordeServer/Utilities/JwtUtils.cs
new file mode 100644
--- /dev/null
+++ b/Engine/Source/Programs/Horde/HordeServer/Utilities/JwtUtils.cs
@@ -0,0 +1,75 @@
+using Microsoft.AspNetCore.Http;
+using System;
+using System.Collections.Generic;
+using System.Diagnostics.CodeAnalysis;
+using System.IdentityModel.Tokens.Jwt;
+using System.Linq;
+using System.Text.Json;
+using System.Threading.Tasks;
+
+namespace HordeServer.Utilities
+{
+	/// <summary>
+	/// Helper functions for dealing with JWTs
+	/// </summary>
+	public static class JwtUtils
+	{
+		/// <summary>
+		/// Gets the bearer token from an HTTP request
+		/// </summary>
+		/// <param name="Request">The request to read from</param>
+		/// <param name="BearerPrefix">The bearer token prefix, ex. "Bearer "</param>
+		/// <param name="Token">On success, receives the bearer token</param>
+		/// <returns>True if the bearer token was read</returns>
+		public static bool TryGetBearerToken(HttpRequest Request, string BearerPrefix, [NotNullWhen(true)] out string? Token)
+		{
+			// Get the authorization header
+			string? Authorization = Request.Headers["Authorization"];
+			if (String.IsNullOrEmpty(Authorization))
+			{
+				Token = null;
+				return false;
+			}
+
+			// Check if it's a bearer token				
+			if (!Authorization.StartsWith(BearerPrefix, StringComparison.OrdinalIgnoreCase))
+			{
+				Token = null;
+				return false;
+			}
+
+			// Get the token
+			Token = Authorization.Substring(BearerPrefix.Length).Trim();
+			return true;
+		}
+
+		/// <summary>
+		/// Tries to parse a JWT and check the issuer matches
+		/// </summary>
+		/// <param name="Token">The token to parse</param>
+		/// <param name="JwtToken">On success, receives the parsed JWT</param>
+		/// <returns>True if the jwt was parsed</returns>
+		public static bool TryParseJwt(string Token, [NotNullWhen(true)] out JwtSecurityToken? JwtToken)
+		{
+			// Check if it's a JWT
+			JwtSecurityTokenHandler Handler = new JwtSecurityTokenHandler();
+			if (!Handler.CanReadToken(Token))
+			{
+				JwtToken = null;
+				return false;
+			}
+
+			// Try to parse the JWT
+			try
+			{
+				JwtToken = Handler.ReadJwtToken(Token);
+				return true;
+			}
+			catch (JsonException)
+			{
+				JwtToken = null;
+				return false;
+			}
+		}
+	}
+}
diff --git a/Engine/Source/Programs/Horde/HordeServer/Utilities/LazyCachedValue.cs b/Engine/Source/Programs/Horde/HordeServer/Utilities/LazyCachedValue.cs
new file mode 100644
--- /dev/null
+++ b/Engine/Source/Programs/Horde/HordeServer/Utilities/LazyCachedValue.cs
@@ -0,0 +1,90 @@
+using System;
+using System.Collections.Generic;
+using System.Diagnostics;
+using System.Linq;
+using System.Reflection.Metadata.Ecma335;
+using System.Threading.Tasks;
+
+namespace HordeServer.Utilities
+{
+	/// <summary>
+	/// Stores a value that expires after a given time
+	/// </summary>
+	/// <typeparam name="T"></typeparam>
+	public class LazyCachedValue<T> where T : class
+	{
+		/// <summary>
+		/// The current value
+		/// </summary>
+		T? Value;
+
+		/// <summary>
+		/// Generator for the new value
+		/// </summary>
+		Func<T> Generator;
+
+		/// <summary>
+		/// Time since the value was updated
+		/// </summary>
+		Stopwatch Timer = Stopwatch.StartNew();
+
+		/// <summary>
+		/// Default expiry time
+		/// </summary>
+		TimeSpan DefaultMaxAge;
+
+		/// <summary>
+		/// Default constructor
+		/// </summary>
+		public LazyCachedValue(Func<T> Generator, TimeSpan MaxAge)
+		{
+			this.Generator = Generator;
+			this.DefaultMaxAge = MaxAge;
+		}
+
+		/// <summary>
+		/// Sets the new value
+		/// </summary>
+		/// <param name="Value">The value to store</param>
+		public void Set(T Value)
+		{
+			this.Value = Value;
+			Timer.Restart();
+		}
+
+		/// <summary>
+		/// Tries to get the current value
+		/// </summary>
+		/// <returns>The cached value, if valid</returns>
+		public T GetCached()
+		{
+			return GetCached(DefaultMaxAge);
+		}
+
+		/// <summary>
+		/// Tries to get the current value
+		/// </summary>
+		/// <returns>The cached value, if valid</returns>
+		public T GetCached(TimeSpan MaxAge)
+		{
+			T? Current = Value;
+			if (Current == null || Timer.Elapsed > MaxAge)
+			{
+				Current = Generator();
+				Set(Current);
+			}
+			return Current;
+		}
+
+		/// <summary>
+		/// Gets the latest value, updating the cache
+		/// </summary>
+		/// <returns>The latest value</returns>
+		public T GetLatest()
+		{
+			T NewValue = Generator();
+			Set(NewValue);
+			return NewValue;
+		}
+	}
+}
diff --git a/Engine/Source/Programs/Horde/HordeServer/Utilities/MongoQueryCache.cs b/Engine/Source/Programs/Horde/HordeServer/Utilities/MongoQueryCache.cs
new file mode 100644
--- /dev/null
+++ b/Engine/Source/Programs/Horde/HordeServer/Utilities/MongoQueryCache.cs
@@ -0,0 +1,87 @@
+using Microsoft.Extensions.Caching.Memory;
+using MongoDB.Bson;
+using MongoDB.Bson.Serialization;
+using MongoDB.Driver;
+using System;
+using System.Collections.Generic;
+using System.Diagnostics;
+using System.Linq;
+using System.Threading.Tasks;
+
+namespace HordeServer.Utilities
+{
+	/// <summary>
+	/// Caches queries against a given collection
+	/// </summary>
+	/// <typeparam name="TDocument">The document type</typeparam>
+	class MongoQueryCache<TDocument> : IDisposable
+	{
+		class QueryState
+		{
+			public Stopwatch Timer = Stopwatch.StartNew();
+			public List<TDocument>? Results;
+			public Task? QueryTask;
+		}
+
+		IMongoCollection<TDocument> Collection;
+		MemoryCache Cache;
+		TimeSpan MaxLatency;
+
+		public MongoQueryCache(IMongoCollection<TDocument> Collection, TimeSpan MaxLatency)
+		{
+			this.Collection = Collection;
+
+			MemoryCacheOptions Options = new MemoryCacheOptions();
+			this.Cache = new MemoryCache(Options);
+
+			this.MaxLatency = MaxLatency;
+		}
+
+		public void Dispose()
+		{
+			Cache.Dispose();
+		}
+
+		async Task RefreshAsync(QueryState State, FilterDefinition<TDocument> Filter)
+		{
+			State.Results = await Collection.Find(Filter).ToListAsync();
+			State.Timer.Restart();
+		}
+
+		public async Task<List<TDocument>> FindAsync(FilterDefinition<TDocument> Filter, int Index, int Count)
+		{
+			BsonDocument Rendered = Filter.Render(BsonSerializer.LookupSerializer<TDocument>(), BsonSerializer.SerializerRegistry);
+			BsonDocument Document = new BsonDocument { new BsonElement("filter", Rendered), new BsonElement("index", Index), new BsonElement("count", Count) };
+
+			string FilterKey = Document.ToString();
+
+			QueryState? State;
+			if (!Cache.TryGetValue(FilterKey, out State) || State == null)
+			{
+				State = new QueryState();
+				using (ICacheEntry CacheEntry = Cache.CreateEntry(FilterKey))
+				{
+					CacheEntry.SetSlidingExpiration(TimeSpan.FromMinutes(5.0));
+					CacheEntry.SetValue(State);
+				}
+			}
+
+			if(State.QueryTask != null && State.QueryTask.IsCompleted)
+			{
+				await State.QueryTask;
+				State.QueryTask = null;
+			}
+
+			if (State.QueryTask == null && (State.Results == null || State.Timer.Elapsed > MaxLatency))
+			{
+				State.QueryTask = Task.Run(() => RefreshAsync(State, Filter));
+			}
+			if (State.QueryTask != null && (State.Results == null || State.Timer.Elapsed > MaxLatency))
+			{
+				await State.QueryTask;
+			}
+
+			return State.Results!;
+		}
+	}
+}
diff --git a/Engine/Source/Programs/Horde/HordeServer/Utilities/NoOpDogStatsd.cs b/Engine/Source/Programs/Horde/HordeServer/Utilities/NoOpDogStatsd.cs
new file mode 100644
--- /dev/null
+++ b/Engine/Source/Programs/Horde/HordeServer/Utilities/NoOpDogStatsd.cs
@@ -0,0 +1,103 @@
+﻿using System;
+using System.Diagnostics.CodeAnalysis;
+using System.IO;
+using StatsdClient;
+
+namespace HordeServer.Utilities
+{
+	/// <summary>
+	/// Empty implementation that does nothing but still satisfies the interface
+	/// Used when DogStatsD is not available.
+	/// </summary>
+	public sealed class NoOpDogStatsd : IDogStatsd
+	{
+		/// <inheritdoc />
+		public ITelemetryCounters TelemetryCounters { get; } = null!;
+		
+		/// <inheritdoc />
+		public void Configure(StatsdConfig Config)
+		{
+		}
+
+		/// <inheritdoc />
+		public void Counter(string StatName, double Value, double SampleRate = 1, string[] Tags = null!)
+		{
+		}
+
+		/// <inheritdoc />
+		public void Decrement(string StatName, int Value = 1, double SampleRate = 1, params string[] Tags)
+		{
+		}
+
+		/// <inheritdoc />
+		public void Event(string Title, string Text, string AlertType = null!, string AggregationKey = null!, string SourceType = null!,
+			int? DateHappened = null, string Priority = null!, string Hostname = null!, string[] Tags = null!)
+		{
+		}
+
+		/// <inheritdoc />
+		public void Gauge(string StatName, double Value, double SampleRate = 1, string[] Tags = null!)
+		{
+		}
+
+		/// <inheritdoc />
+		public void Histogram(string StatName, double Value, double SampleRate = 1, string[] Tags = null!)
+		{
+		}
+
+		/// <inheritdoc />
+		public void Distribution(string StatName, double Value, double SampleRate = 1, string[] Tags = null!)
+		{
+		}
+
+		/// <inheritdoc />
+		public void Increment(string StatName, int Value = 1, double SampleRate = 1, string[] Tags = null!)
+		{
+		}
+
+		/// <inheritdoc />
+		public void Set<T>(string StatName, T Value, double SampleRate = 1, string[] Tags = null!)
+		{
+		}
+
+		/// <inheritdoc />
+		public void Set(string StatName, string Value, double SampleRate = 1, string[] Tags = null!)
+		{
+		}
+
+		/// <inheritdoc />
+		public IDisposable StartTimer(string Name, double SampleRate = 1, string[] Tags = null!)
+		{
+			// Some random object to satisfy IDisposable
+			return new MemoryStream();
+		}
+
+		/// <inheritdoc />
+		public void Time(Action Action, string StatName, double SampleRate = 1, string[] Tags = null!)
+		{
+		}
+
+		/// <inheritdoc />
+		public T Time<T>(Func<T> Func, string StatName, double SampleRate = 1, string[] Tags = null!)
+		{
+			return Func();
+		}
+
+		/// <inheritdoc />
+		public void Timer(string StatName, double Value, double SampleRate = 1, string[] Tags = null!)
+		{
+		}
+
+		/// <inheritdoc />
+		public void ServiceCheck(string Name, Status Status, int? Timestamp = null, string Hostname = null!, string[] Tags = null!, string Message = null!)
+		{
+		}
+		
+		/// <inheritdoc />
+		[SuppressMessage("Design", "CA1063:Implement IDisposable correctly")]
+		[SuppressMessage("Usage", "CA1816:Call GC.SuppressFinalize correctly")]
+		public void Dispose()
+		{
+		}
+	}
+}
\ No newline at end of file
diff --git a/Engine/Source/Programs/Horde/HordeServer/Utilities/PropertyFilter.cs b/Engine/Source/Programs/Horde/HordeServer/Utilities/PropertyFilter.cs
new file mode 100644
--- /dev/null
+++ b/Engine/Source/Programs/Horde/HordeServer/Utilities/PropertyFilter.cs
@@ -0,0 +1,245 @@
+using Microsoft.AspNetCore.Mvc;
+using Microsoft.AspNetCore.Mvc.ModelBinding;
+using System;
+using System.Collections;
+using System.Collections.Generic;
+using System.ComponentModel;
+using System.Diagnostics.CodeAnalysis;
+using System.Globalization;
+using System.Reflection;
+using System.Text;
+using System.Threading.Tasks;
+
+namespace HordeServer.Utilities
+{
+	/// <summary>
+	/// Utility class to filter a response to include a subset of fields
+	/// </summary>
+	[TypeConverter(typeof(PropertyFilterTypeConverter))]
+	public class PropertyFilter
+	{
+		/// <summary>
+		/// Specific filters for properties of this object. A key of '*' will match any property not matched by anything else, and a value of null will include the whole object.
+		/// </summary>
+		Dictionary<string, PropertyFilter?> NameToFilter = new Dictionary<string, PropertyFilter?>(StringComparer.OrdinalIgnoreCase);
+
+		/// <summary>
+		/// Parse a list of fields as a property filter
+		/// </summary>
+		/// <param name="Fields">List of fields to parse, separated by commas</param>
+		/// <returns>New property filter</returns>
+		public static PropertyFilter Parse(string? Fields)
+		{
+			PropertyFilter RootFilter = new PropertyFilter();
+			if(Fields != null)
+			{
+				foreach (string Field in Fields.Split(','))
+				{
+					// Split the field name into segments
+					string[] Segments = Field.Split('.');
+
+					// Create the tree of segments
+					Dictionary<string, PropertyFilter?> NameToFilter = RootFilter.NameToFilter;
+					for (int Idx = 0; Idx < Segments.Length - 1; Idx++)
+					{
+						PropertyFilter? NextFilter;
+						if (!NameToFilter.TryGetValue(Segments[Idx], out NextFilter))
+						{
+							NextFilter = new PropertyFilter();
+							NameToFilter[Segments[Idx]] = NextFilter;
+						}
+						else if (NextFilter == null)
+						{
+							break;
+						}
+						NameToFilter = NextFilter.NameToFilter;
+					}
+					NameToFilter[Segments[Segments.Length - 1]] = null;
+				}
+			}
+			return RootFilter;
+		}
+
+		/// <summary>
+		/// Checks whether the filter contains the given field
+		/// </summary>
+		/// <param name="Filter">The filter to check</param>
+		/// <param name="Field">Name of the field to check for</param>
+		/// <returns>True if the filter includes the given field</returns>
+		public static bool Includes(PropertyFilter? Filter, string Field)
+		{
+			return Filter == null || Filter.Includes(Field);
+		}
+
+		/// <summary>
+		/// Attempts to apply a property filter to a response object, where the property filter may be null.
+		/// </summary>
+		/// <param name="Response">Response object</param>
+		/// <param name="Filter">The filter to apply</param>
+		/// <returns>Filtered object</returns>
+		public static object Apply(object Response, PropertyFilter? Filter)
+		{
+			return (Filter == null) ? Response : Filter.ApplyTo(Response);
+		}
+
+		/// <summary>
+		/// Checks whether the filter contains the given field
+		/// </summary>
+		/// <param name="Field">Name of the field to check for</param>
+		/// <returns>True if the filter includes the given field</returns>
+		public bool Includes(string Field)
+		{
+			return NameToFilter.Count == 0 || NameToFilter.ContainsKey(Field) || NameToFilter.ContainsKey("*");
+		}
+
+		/// <summary>
+		/// Extract the given fields from the response
+		/// </summary>
+		/// <param name="Response">The response to filter</param>
+		/// <returns>Filtered response</returns>
+		public object ApplyTo(object Response)
+		{
+			// If there are no filters at this depth, return the whole object
+			if (NameToFilter.Count == 0)
+			{
+				return Response;
+			}
+
+			// Check if this object is a list. If it is, we'll apply the filter to each element in the list.
+			IList? List = Response as IList;
+			if (List != null)
+			{
+				List<object?> NewResponse = new List<object?>();
+				for (int Index = 0; Index < List.Count; Index++)
+				{
+					object? Value = List[Index];
+					if (Value == null)
+					{
+						NewResponse.Add(Value);
+					}
+					else
+					{
+						NewResponse.Add(ApplyTo(Value));
+					}
+				}
+				return NewResponse;
+			}
+			else
+			{
+				Dictionary<string, object?> NewResponse = new Dictionary<string, object?>(StringComparer.Ordinal);
+
+				// Check if this object is a dictionary.
+				IDictionary? Dictionary = Response as IDictionary;
+				if (Dictionary != null)
+				{
+					foreach (DictionaryEntry? Entry in Dictionary)
+					{
+						if (Entry.HasValue)
+						{
+							string? Key = Entry.Value.Key.ToString();
+							if (Key != null)
+							{
+								Key = ConvertPascalToCamelCase(Key);
+
+								PropertyFilter? Filter;
+								if (NameToFilter.TryGetValue(Key, out Filter))
+								{
+									object? Value = Entry.Value.Value;
+									if (Filter == null)
+									{
+										NewResponse[Key] = Value;
+									}
+									else if (Value != null)
+									{
+										NewResponse[Key] = Filter.ApplyTo(Value);
+									}
+								}
+							}
+						}
+					}
+				}
+				else
+				{
+					// Otherwise try to get the properties of this object
+					foreach (PropertyInfo Property in Response.GetType().GetProperties(BindingFlags.Public | BindingFlags.GetProperty | BindingFlags.Instance))
+					{
+						PropertyFilter? Filter;
+						if (NameToFilter.TryGetValue(Property.Name, out Filter))
+						{
+							string Key = ConvertPascalToCamelCase(Property.Name);
+							object? Value = Property.GetValue(Response);
+							if (Filter == null)
+							{
+								NewResponse[Key] = Value;
+							}
+							else if (Value != null)
+							{
+								NewResponse[Key] = Filter.ApplyTo(Value);
+							}
+						}
+					}
+				}
+				return NewResponse;
+			}
+		}
+
+		/// <summary>
+		/// Convert a property name to camel case
+		/// </summary>
+		/// <param name="Name">The name to convert to camelCase</param>
+		/// <returns>Camelcase version of the name</returns>
+		static string ConvertPascalToCamelCase(string Name)
+		{
+			if (Name[0] >= 'A' && Name[0] <= 'Z')
+			{
+				return (char)(Name[0] - 'A' + 'a') + Name.Substring(1);
+			}
+			else
+			{
+				return Name;
+			}
+		}
+	}
+
+	/// <summary>
+	/// Type converter from strings to PropertyFilter objects
+	/// </summary>
+	public class PropertyFilterTypeConverter : TypeConverter
+	{
+		/// <inheritdoc/>
+		public override bool CanConvertFrom(ITypeDescriptorContext context, Type sourceType)
+		{
+			return sourceType == typeof(string);
+		}
+
+		/// <inheritdoc/>
+		public override object ConvertFrom(ITypeDescriptorContext context, CultureInfo culture, object value)
+		{
+			return PropertyFilter.Parse((string)value);
+		}
+	}
+
+	/// <summary>
+	/// Extension methods for property filters
+	/// </summary>
+	public static class PropertyFilterExtensions
+	{
+		/// <summary>
+		/// Apply a filter to a response object
+		/// </summary>
+		/// <param name="Obj">The object to filter</param>
+		/// <param name="Filter">The filter to apply</param>
+		/// <returns>The filtered object</returns>
+		public static object ApplyFilter(this object Obj, PropertyFilter? Filter)
+		{
+			if (Filter == null)
+			{
+				return Obj;
+			}
+			else
+			{
+				return Filter.ApplyTo(Obj);
+			}
+		}
+	}
+}
diff --git a/Engine/Source/Programs/Horde/HordeServer/Utilities/ReadOnlyTrie.cs b/Engine/Source/Programs/Horde/HordeServer/Utilities/ReadOnlyTrie.cs
new file mode 100644
--- /dev/null
+++ b/Engine/Source/Programs/Horde/HordeServer/Utilities/ReadOnlyTrie.cs
@@ -0,0 +1,297 @@
+using EpicGames.Core;
+using HordeServer.Controllers;
+using System;
+using System.Collections;
+using System.Collections.Generic;
+using System.Linq;
+using System.Runtime.InteropServices;
+using System.Threading.Tasks;
+
+namespace HordeServer.Logs
+{
+	/// <summary>
+	/// A read-only trie of long values
+	/// </summary>
+	public class ReadOnlyTrie : IEnumerable<ulong>
+	{
+		/// <summary>
+		/// Stack item for traversing the tree
+		/// </summary>
+		struct StackItem
+		{
+			/// <summary>
+			/// The current node index
+			/// </summary>
+			public int Index;
+
+			/// <summary>
+			/// Value in the current node (0-15)
+			/// </summary>
+			public int Value;
+		}
+
+		/// <summary>
+		/// Delegate for filtering values during a tree traversal
+		/// </summary>
+		/// <param name="Value">The current value</param>
+		/// <param name="Mask">Mask for which bits in the value are valid</param>
+		/// <returns>True if values matching the given mask should be enumerated</returns>
+		public delegate bool VisitorDelegate(ulong Value, ulong Mask);
+
+		/// <summary>
+		/// Height of the tree
+		/// </summary>
+		const int Height = sizeof(ulong) * 2;
+
+		/// <summary>
+		/// Array of bitmasks for each node in the tree
+		/// </summary>
+		public ushort[] NodeData { get; }
+
+		/// <summary>
+		/// Array of child offsets for each node. Excludes the last layer of the tree.
+		/// </summary>
+		int[] FirstChildIndex;
+
+		/// <summary>
+		/// Empty index definition
+		/// </summary>
+		public static ReadOnlyTrie Empty { get; } = new ReadOnlyTrie(new ushort[1]);
+
+		/// <summary>
+		/// Constructor
+		/// </summary>
+		/// <param name="NodeData">Node data</param>
+		public ReadOnlyTrie(ushort[] NodeData)
+		{
+			this.NodeData = NodeData;
+			this.FirstChildIndex = CreateChildLookup(NodeData);
+		}
+
+		/// <summary>
+		/// Tests whether the given value is in the trie
+		/// </summary>
+		/// <param name="Value">The value to check for</param>
+		/// <returns>True if the value is in the trie</returns>
+		public bool Contains(ulong Value)
+		{
+			int Index = 0;
+			for (int Shift = (sizeof(ulong) * 8) - 4; Shift >= 0; Shift -= 4)
+			{
+				int Mask = NodeData[Index];
+				int Flag = 1 << (int)((Value >> Shift) & 15);
+				if ((Mask & Flag) == 0)
+				{
+					return false;
+				}
+
+				Index = FirstChildIndex[Index];
+				for (; ; )
+				{
+					Mask &= (Mask - 1);
+					if ((Mask & Flag) == 0)
+					{
+						break;
+					}
+					Index++;
+				}
+			}
+			return true;
+		}
+
+		/// <inheritdoc/>
+		public IEnumerator<ulong> GetEnumerator()
+		{
+			return EnumerateRange(ulong.MinValue, ulong.MaxValue).GetEnumerator();
+		}
+
+		/// <inheritdoc/>
+		IEnumerator IEnumerable.GetEnumerator() => GetEnumerator();
+
+		/// <summary>
+		/// Enumerate all values matching a given filter
+		/// </summary>
+		/// <param name="Predicate">Predicate for which values to include</param>
+		/// <returns>Values satisfying the given predicate</returns>
+		public IEnumerable<ulong> EnumerateValues(VisitorDelegate Predicate)
+		{
+			int Depth = 0;
+			ulong Value = 0;
+
+			StackItem[] Stack = new StackItem[Height];
+			Stack[1].Index = FirstChildIndex[0];
+
+			for (; ; )
+			{
+				StackItem Current = Stack[Depth];
+				if (Current.Value >= 16)
+				{
+					// Move up the tree if we've enumerated all the branches at the current level
+					Depth--;
+					if (Depth < 0)
+					{
+						yield break;
+					}
+					Stack[Depth].Value++;
+
+					// Increment the child index too. These are stored sequentially for a given parent. The value will be cleared when we recurse into it.
+					Stack[Depth + 1].Index++;
+				}
+				else if ((NodeData[Current.Index] & (1 << Current.Value)) == 0)
+				{
+					// This branch does not exist. Skip it.
+					Stack[Depth].Value++;
+				}
+				else
+				{
+					// Get the value and mask for the current node
+					int Shift = (Stack.Length - Depth - 1) * 4;
+					ulong Mask = ~((1UL << Shift) - 1);
+					Value = (Value & (Mask << 4)) | ((ulong)(uint)Current.Value << Shift);
+
+					if (!Predicate(Value, Mask))
+					{
+						// This node is excluded, skip it
+						Stack[Depth].Value++;
+						if (Depth + 1 < Stack.Length)
+						{
+							Stack[Depth + 1].Index++;
+						}
+					}
+					else if(Depth + 1 < Stack.Length)
+					{
+						// Move down the tree
+						Depth++;
+						Stack[Depth].Value = 0;
+						if (Depth + 1 < Stack.Length)
+						{
+							Stack[Depth + 1].Index = FirstChildIndex[Stack[Depth].Index];
+						}
+					}
+					else
+					{
+						// Yield the current value
+						yield return Value;
+						Stack[Depth].Value++;
+					}
+				}
+			}
+		}
+
+		/// <summary>
+		/// Enumerates all values in the trie between the given ranges
+		/// </summary>
+		/// <param name="MinValue">Minimum value to enumerate</param>
+		/// <param name="MaxValue">Maximum value to enumerate</param>
+		/// <returns>Sequence of values</returns>
+		public IEnumerable<ulong> EnumerateRange(ulong MinValue, ulong MaxValue)
+		{
+			return EnumerateValues((Value, Mask) => (Value >= (MinValue & Mask) && Value <= (MaxValue & Mask)));
+		}
+
+		/// <summary>
+		/// Creates a lookup for child node offsets from raw node data
+		/// </summary>
+		/// <param name="NodeData">Array of masks for each node</param>
+		/// <returns>Array of offsets</returns>
+		static int[] CreateChildLookup(ushort[] NodeData)
+		{
+			List<int> ChildOffsets = new List<int>();
+			if (NodeData.Length > 0)
+			{
+				int NodeCount = 1;
+
+				int Index = 0;
+				int ChildIndex = NodeCount;
+
+				for (int Level = 0; Level < Height; Level++)
+				{
+					int NextNodeCount = 0;
+					for (int Idx = 0; Idx < NodeCount; Idx++)
+					{
+						ushort Node = NodeData[Index++];
+
+						int NumChildren = CountBits(Node);
+						ChildOffsets.Add(ChildIndex);
+						ChildIndex += NumChildren;
+
+						NextNodeCount += NumChildren;
+					}
+					NodeCount = NextNodeCount;
+				}
+			}
+			return ChildOffsets.ToArray();
+		}
+
+		/// <summary>
+		/// Count the number of set bits in the given value
+		/// </summary>
+		/// <param name="Value">Value to test</param>
+		/// <returns>Number of set bits</returns>
+		static int CountBits(ushort Value)
+		{
+			int Count = Value;
+			Count = (Count & 0b0101010101010101) + ((Count >> 1) & 0b0101010101010101);
+			Count = (Count & 0b0011001100110011) + ((Count >> 2) & 0b0011001100110011);
+			Count = (Count & 0b0000111100001111) + ((Count >> 4) & 0b0000111100001111);
+			Count = (Count & 0b0000000011111111) + ((Count >> 8) & 0b0000000011111111);
+			return Count;
+		}
+
+		/// <summary>
+		/// Read a trie from the given buffer
+		/// </summary>
+		/// <param name="Reader">Reader to read from</param>
+		/// <returns>New trie</returns>
+		public static ReadOnlyTrie Read(MemoryReader Reader)
+		{
+			ReadOnlyMemory<byte> Nodes = Reader.ReadVariableLengthBytes();
+			ushort[] NodeData = MemoryMarshal.Cast<byte, ushort>(Nodes.Span).ToArray();
+			return new ReadOnlyTrie(NodeData);
+		}
+
+		/// <summary>
+		/// Write this trie to the given buffer
+		/// </summary>
+		/// <param name="Writer">Writer to output to</param>
+		public void Write(MemoryWriter Writer)
+		{
+			Writer.WriteVariableLengthBytes(MemoryMarshal.AsBytes<ushort>(NodeData));
+		}
+
+		/// <summary>
+		/// Gets the serialized size of this trie
+		/// </summary>
+		/// <returns></returns>
+		public int GetSerializedSize()
+		{
+			return (sizeof(int) + NodeData.Length * sizeof(ushort));
+		}
+	}
+
+	/// <summary>
+	/// Extension methods for serializing tries
+	/// </summary>
+	public static class LogTokenSetExtensions
+	{
+		/// <summary>
+		/// Read a trie from the given buffer
+		/// </summary>
+		/// <param name="Reader">Reader to read from</param>
+		/// <returns>New trie</returns>
+		public static ReadOnlyTrie ReadTrie(this MemoryReader Reader)
+		{
+			return ReadOnlyTrie.Read(Reader);
+		}
+
+		/// <summary>
+		/// Write this trie to the given buffer
+		/// </summary>
+		/// <param name="Writer">Writer to output to</param>
+		/// <param name="Trie">Trie to write</param>
+		public static void WriteTrie(this MemoryWriter Writer, ReadOnlyTrie Trie)
+		{
+			Trie.Write(Writer);
+		}
+	}
+}
diff --git a/Engine/Source/Programs/Horde/HordeServer/Utilities/ReadOnlyTrieBuilder.cs b/Engine/Source/Programs/Horde/HordeServer/Utilities/ReadOnlyTrieBuilder.cs
new file mode 100644
--- /dev/null
+++ b/Engine/Source/Programs/Horde/HordeServer/Utilities/ReadOnlyTrieBuilder.cs
@@ -0,0 +1,124 @@
+using Amazon.Runtime;
+using EpicGames.Core;
+using HordeServer.Logs;
+using Microsoft.AspNetCore.Mvc;
+using Microsoft.AspNetCore.Routing.Constraints;
+using System;
+using System.Collections.Generic;
+using System.Linq;
+using System.Runtime.CompilerServices;
+using System.Runtime.InteropServices;
+using System.Threading.Tasks;
+
+namespace HordeServer.Logs
+{
+	/// <summary>
+	/// Structure used for generating prefix trees
+	/// </summary>
+	public class ReadOnlyTrieBuilder
+	{
+		/// <summary>
+		/// Node within the trie
+		/// </summary>
+		class Node
+		{
+			public Node?[]? Children;
+		}
+
+		/// <summary>
+		/// The root node
+		/// </summary>
+		Node Root;
+
+		/// <summary>
+		/// Default constructor
+		/// </summary>
+		public ReadOnlyTrieBuilder()
+		{
+			Root = new Node();
+		}
+
+		/// <summary>
+		/// Adds a value to the trie
+		/// </summary>
+		/// <param name="Value">Value to add</param>
+		public void Add(ulong Value)
+		{
+			// Loop through the tree until we've added the item
+			Node Leaf = Root;
+			for (int Shift = (sizeof(ulong) * 8) - 4; Shift >= 0; Shift -= 4)
+			{
+				int Index = (int)(Value >> Shift) & 15;
+				if (Leaf.Children == null)
+				{
+					Leaf.Children = new Node[16];
+				}
+				if (Leaf.Children[Index] == null)
+				{
+					Leaf.Children[Index] = new Node();
+				}
+				Leaf = Leaf.Children[Index]!;
+			}
+		}
+
+		/// <summary>
+		/// Searches for the given item in the trie
+		/// </summary>
+		/// <param name="Value">Value to add</param>
+		public bool Contains(ulong Value)
+		{
+			// Loop through the tree until we've added the item
+			Node Leaf = Root;
+			for (int Shift = (sizeof(ulong) * 8) - 4; Shift >= 0; Shift -= 4)
+			{
+				int Index = (int)(Value >> Shift) & 15;
+				if (Leaf.Children == null)
+				{
+					return false;
+				}
+				if (Leaf.Children[Index] == null)
+				{
+					return false;
+				}
+				Leaf = Leaf.Children[Index]!;
+			}
+			return true;
+		}
+
+		/// <summary>
+		/// Creates a <see cref="ReadOnlyTrie"/> from this data
+		/// </summary>
+		/// <returns></returns>
+		public ReadOnlyTrie Build()
+		{
+			List<ushort> Values = new List<ushort>();
+
+			List<Node> Nodes = new List<Node>();
+			Nodes.Add(Root);
+
+			for(int Bits = 0; Bits < (sizeof(ulong) * 8); Bits += 4)
+			{
+				List<Node> NextNodes = new List<Node>();
+				foreach(Node Node in Nodes)
+				{
+					ushort Value = 0;
+					if (Node.Children != null)
+					{
+						for (int Idx = 0; Idx < Node.Children.Length; Idx++)
+						{
+							if(Node.Children[Idx] != null)
+							{
+								Value |= (ushort)(1 << Idx);
+								NextNodes.Add(Node.Children[Idx]!);
+							}
+						}
+					}
+					Values.Add(Value);
+				}
+				Nodes = NextNodes;
+			}
+
+			return new ReadOnlyTrie(Values.ToArray());
+		}
+	}
+}
diff --git a/Engine/Source/Programs/Horde/HordeServer/Utilities/RedisConnectionPool.cs b/Engine/Source/Programs/Horde/HordeServer/Utilities/RedisConnectionPool.cs
new file mode 100644
--- /dev/null
+++ b/Engine/Source/Programs/Horde/HordeServer/Utilities/RedisConnectionPool.cs
@@ -0,0 +1,73 @@
+// Copyright Epic Games, Inc. All Rights Reserved.
+
+using System;
+using System.Collections.Concurrent;
+using System.Collections.Generic;
+using System.Linq;
+using StackExchange.Redis;
+
+namespace HordeServer.Utilities
+{
+	/// <summary>
+	/// A connection pool for Redis client
+	///
+	/// Wraps multiple ConnectionMultiplexers as lazy values and initializes them as needed.
+	/// If full, the least loaded connection will be picked.
+	/// </summary>
+	public class RedisConnectionPool
+	{
+		private readonly ConcurrentBag<Lazy<ConnectionMultiplexer>> Connections;
+		private readonly int DefaultDatabaseIndex;
+
+		/// <summary>
+		/// Constructor
+		/// </summary>
+		/// <param name="PoolSize">Size of the pool (i.e max number of connections)</param>
+		/// <param name="RedisConfString">Configuration string for Redis</param>
+		/// <param name="DefaultDatabaseIndex">The Redis database index to use. Use -1 for the default one</param>
+		public RedisConnectionPool(int PoolSize, string RedisConfString, int DefaultDatabaseIndex = -1)
+		{
+			this.DefaultDatabaseIndex = DefaultDatabaseIndex;
+			Connections = new ConcurrentBag<Lazy<ConnectionMultiplexer>>();
+			for (int i = 0; i < PoolSize; i++)
+			{
+				Connections.Add(new Lazy<ConnectionMultiplexer>(() => ConnectionMultiplexer.Connect(RedisConfString)));
+			}
+		}
+
+		/// <summary>
+		/// Get a connection from the pool
+		///
+		/// It will pick the least loaded connection or create a new one (if pool size allows)
+		/// </summary>
+		/// <returns>A Redis database connection</returns>
+		public IConnectionMultiplexer GetConnection()
+		{
+			Lazy<ConnectionMultiplexer> LazyConnection;
+			var LazyConnections = Connections.Where(x => x.IsValueCreated);
+
+			if (LazyConnections.Count() == Connections.Count)
+			{
+				// No more new connections can be created, pick the least loaded one
+				LazyConnection = Connections.OrderBy(x => x.Value.GetCounters().TotalOutstanding).First();
+			}
+			else
+			{
+				// Create a new connection by picking a not yet initialized lazy value 
+				LazyConnection = Connections.First(x => !x.IsValueCreated);
+			}
+
+			return LazyConnection.Value;
+		}
+
+		/// <summary>
+		/// Shortcut to getting a IDatabase
+		/// </summary>
+		/// <returns>A Redis database</returns>
+		public IDatabase GetDatabase()
+		{
+			return GetConnection().GetDatabase(DefaultDatabaseIndex);
+		}
+
+	}
+}
\ No newline at end of file
diff --git a/Engine/Source/Programs/Horde/HordeServer/Utilities/RemoteExtensions.cs b/Engine/Source/Programs/Horde/HordeServer/Utilities/RemoteExtensions.cs
new file mode 100644
--- /dev/null
+++ b/Engine/Source/Programs/Horde/HordeServer/Utilities/RemoteExtensions.cs
@@ -0,0 +1,106 @@
+// Copyright Epic Games, Inc. All Rights Reserved.
+
+using Build.Bazel.Remote.Execution.V2;
+using HordeServer.Storage;
+using HordeServer.Storage.Primitives;
+using Google.Protobuf;
+using Google.Protobuf.Reflection;
+using System;
+using System.Collections.Generic;
+using System.Linq;
+using System.Threading.Tasks;
+
+namespace HordeServer.Utility
+{
+	static class RemoteExtensions
+	{
+		public static BlobHash ToHashValue(this Digest Digest)
+		{
+			return BlobHash.Parse(Digest.Hash);
+		}
+
+		public static void MergeFrom<T>(this IMessage<T> Message, ReadOnlyMemory<byte> Memory) where T : IMessage<T>
+		{
+			// TODO: optimize this to avoid array allocation
+			Message.MergeFrom(Memory.ToArray());
+		}
+
+		public static T ParseFrom<T>(this MessageParser<T> Parser, ReadOnlyMemory<byte> Memory) where T : IMessage<T>
+		{
+			// TODO: optimize this to avoid array allocation
+			return Parser.ParseFrom(Memory.ToArray());
+		}
+
+		public static async Task<T> GetProtoMessageAsync<T>(this IStorageService StorageProvider, BlobHash HashValue, DateTime Deadline = default) where T : class, IMessage<T>, new()
+		{
+			T? Result = await TryGetProtoMessageAsync<T>(StorageProvider, HashValue, Deadline);
+			if (Result == null)
+			{
+				throw new KeyNotFoundException($"Unable to get blob {HashValue}");
+			}
+			return Result;
+		}
+
+		public static async Task<T?> TryGetProtoMessageAsync<T>(this IStorageService StorageProvider, BlobHash HashValue, DateTime Deadline = default) where T : class, IMessage<T>, new()
+		{
+			ReadOnlyMemory<byte>? Data = await StorageProvider.TryGetBlobAsync(HashValue, Deadline);
+			if (Data == null)
+			{
+				return null;
+			}
+
+			T NewItem = new T();
+			NewItem.MergeFrom(Data.Value.ToArray());
+			return NewItem;
+		}
+
+		public static async Task<BlobHash> PutProtoMessageAsync<T>(this IStorageService StorageProvider, IMessage<T> Message) where T : class, IMessage<T>
+		{
+			HashSet<BlobHash> References = new HashSet<BlobHash>();
+			FindReferences(Message, References);
+
+			byte[] MessageData = Message.ToByteArray(); // TODO: Could generate this directly into the buffer
+
+			BlobHash HashValue = BlobHash.Compute(MessageData);
+			await StorageProvider.PutBlobAsync(HashValue, MessageData);
+
+			return HashValue;
+		}
+
+		public static Task<ReadOnlyMemory<byte>> GetBulkDataTreeAsync(this IStorageService StorageProvider, BlobHash Hash)
+		{
+			return StorageProvider.GetBlobAsync(Hash);
+		}
+
+		public static async Task<BlobHash> PutBulkDataTreeAsync(this IStorageService StorageProvider, ReadOnlyMemory<byte> BulkData)
+		{
+			BlobHash HashValue = BlobHash.Compute(BulkData);
+			await StorageProvider.PutBlobAsync(HashValue, BulkData);
+			return HashValue;
+		}
+
+		static void FindReferences(this IMessage Message, HashSet<BlobHash> References)
+		{
+			foreach (FieldDescriptor FieldDescriptor in Message.Descriptor.Fields.InFieldNumberOrder())
+			{
+				if (FieldDescriptor.FieldType == FieldType.Message)
+				{
+					object? Value = FieldDescriptor.Accessor.GetValue(Message);
+					if (Value != null)
+					{
+						if (FieldDescriptor.MessageType == Digest.Descriptor)
+						{
+							Digest Digest = (Digest)Value;
+							References.Add(Digest.ToHashValue());
+						}
+						else
+						{
+							IMessage SubMessage = (IMessage)Value;
+							FindReferences(SubMessage, References);
+						}
+					}
+				}
+			}
+		}
+	}
+}
diff --git a/Engine/Source/Programs/Horde/HordeServer/Utilities/SearchText.cs b/Engine/Source/Programs/Horde/HordeServer/Utilities/SearchText.cs
new file mode 100644
--- /dev/null
+++ b/Engine/Source/Programs/Horde/HordeServer/Utilities/SearchText.cs
@@ -0,0 +1,185 @@
+using System;
+using System.Collections.Generic;
+using System.Linq;
+using System.Text;
+using System.Threading.Tasks;
+
+namespace HordeServer.Logs
+{
+	/// <summary>
+	/// Stores cached information about a utf8 search term 
+	/// </summary>
+	public class SearchText
+	{
+		/// <summary>
+		/// The search text
+		/// </summary>
+		public string Text { get; }
+
+		/// <summary>
+		/// The utf-8 bytes to search for
+		/// </summary>
+		public ReadOnlyMemory<byte> Bytes { get; }
+
+		/// <summary>
+		/// Normalized (lowercase) utf-8 bytes to search for
+		/// </summary>
+		byte[] SearchBytes;
+
+		/// <summary>
+		/// Skip table for comparisons
+		/// </summary>
+		byte[] SkipTable;
+
+		/// <summary>
+		/// Constructor
+		/// </summary>
+		/// <param name="Text">The text to search for</param>
+		public SearchText(string Text)
+		{
+			byte[] Bytes = Encoding.UTF8.GetBytes(Text);
+			this.Bytes = Bytes;
+			this.Text = Text;
+
+			// Find the byte sequence to search for, in lowercase
+			SearchBytes = new byte[Bytes.Length];
+			for (int Idx = 0; Idx < SearchBytes.Length; Idx++)
+			{
+				if (Bytes[Idx] >= 'A' && Bytes[Idx] <= 'Z')
+				{
+					SearchBytes[Idx] = (byte)('a' + (Bytes[Idx] - 'A'));
+				}
+				else
+				{
+					SearchBytes[Idx] = Bytes[Idx];
+				}
+			}
+
+			// Build a table indicating how many characters to skip before attempting the next comparison
+			SkipTable = new byte[256];
+			for (int Idx = 0; Idx < 256; Idx++)
+			{
+				SkipTable[Idx] = (byte)SearchBytes.Length;
+			}
+			for (int Idx = 0; Idx < SearchBytes.Length - 1; Idx++)
+			{
+				byte Character = SearchBytes[Idx];
+
+				byte SkipBytes = (byte)(SearchBytes.Length - 1 - Idx);
+				SkipTable[Character] = SkipBytes;
+
+				if (Character >= 'a' && Character <= 'z')
+				{
+					SkipTable['A' + (Character - 'a')] = SkipBytes;
+				}
+			}
+		}
+
+		/// <summary>
+		/// Find all ocurrences of the text in the given buffer
+		/// </summary>
+		/// <param name="Buffer">The buffer to search</param>
+		/// <param name="Text">The text to search for</param>
+		/// <returns>Sequence of offsets within the buffer</returns>
+		public static IEnumerable<int> FindOcurrences(ReadOnlyMemory<byte> Buffer, SearchText Text)
+		{
+			for(int Offset = 0; ;Offset++)
+			{
+				Offset = FindNextOcurrence(Buffer.Span, Offset, Text);
+				if(Offset == -1)
+				{
+					break;
+				}
+				yield return Offset;
+			}
+		}
+
+		/// <summary>
+		/// Perform a case sensitive search for the next occurerence of the search term in a given buffer
+		/// </summary>
+		/// <param name="Buffer">The buffer to search</param>
+		/// <param name="Offset">Starting offset for the search</param>
+		/// <param name="Text">The text to search for</param>
+		/// <returns>Offset of the next occurence, or -1</returns>
+		public static int FindNextOcurrence(ReadOnlySpan<byte> Buffer, int Offset, SearchText Text)
+		{
+			while (Offset + Text.SearchBytes.Length <= Buffer.Length)
+			{
+				if (Matches(Buffer, Offset, Text))
+				{
+					return Offset;
+				}
+				else
+				{
+					Offset += Text.SkipTable[Buffer[Offset + Text.SearchBytes.Length - 1]];
+				}
+			}
+			return -1;
+		}
+
+		/// <summary>
+		/// Compare the search term against the given buffer
+		/// </summary>
+		/// <param name="Buffer">The buffer to search</param>
+		/// <param name="Offset">Starting offset for the search</param>
+		/// <param name="Text">The text to search for</param>
+		/// <returns>True if the text matches, false otherwise</returns>
+		public static bool Matches(ReadOnlySpan<byte> Buffer, int Offset, SearchText Text)
+		{
+			for (int Idx = Text.SearchBytes.Length - 1; Idx >= 0; Idx--)
+			{
+				byte Char = Buffer[Offset + Idx];
+				if (Char >= 'A' && Char <= 'Z')
+				{
+					Char = (byte)('a' + (Char - 'A'));
+				}
+				if (Char != Text.SearchBytes[Idx])
+				{
+					return false;
+				}
+			}
+			return true;
+		}
+	}
+
+	/// <summary>
+	/// Stores cached information about a utf8 search term 
+	/// </summary>
+	public static class SearchTextExtensions
+	{
+		/// <summary>
+		/// Find all ocurrences of the text in the given buffer
+		/// </summary>
+		/// <param name="Buffer">The buffer to search</param>
+		/// <param name="Text">The text to search for</param>
+		/// <returns>Sequence of offsets within the buffer</returns>
+		public static IEnumerable<int> FindOcurrences(this ReadOnlyMemory<byte> Buffer, SearchText Text)
+		{
+			return SearchText.FindOcurrences(Buffer, Text);
+		}
+
+		/// <summary>
+		/// Perform a case sensitive search for the next occurerence of the search term in a given buffer
+		/// </summary>
+		/// <param name="Buffer">The buffer to search</param>
+		/// <param name="Offset">Starting offset for the search</param>
+		/// <param name="Text">The text to search for</param>
+		/// <returns>Offset of the next occurence, or -1</returns>
+		public static int FindNextOcurrence(this ReadOnlySpan<byte> Buffer, int Offset, SearchText Text)
+		{
+			return SearchText.FindNextOcurrence(Buffer, Offset, Text);
+		}
+
+		/// <summary>
+		/// Compare the search term against the given buffer
+		/// </summary>
+		/// <param name="Buffer">The buffer to search</param>
+		/// <param name="Offset">Starting offset for the search</param>
+		/// <param name="Text">The text to search for</param>
+		/// <returns>True if the text matches, false otherwise</returns>
+		public static bool Matches(this ReadOnlySpan<byte> Buffer, int Offset, SearchText Text)
+		{
+			return SearchText.Matches(Buffer, Offset, Text);
+		}
+	}
+}
diff --git a/Engine/Source/Programs/Horde/HordeServer/Utilities/SemaphoreSlimExtensions.cs b/Engine/Source/Programs/Horde/HordeServer/Utilities/SemaphoreSlimExtensions.cs
new file mode 100644
--- /dev/null
+++ b/Engine/Source/Programs/Horde/HordeServer/Utilities/SemaphoreSlimExtensions.cs
@@ -0,0 +1,67 @@
+// Copyright Epic Games, Inc. All Rights Reserved.
+
+using System;
+using System.Collections.Generic;
+using System.Linq;
+using System.Threading;
+using System.Threading.Tasks;
+
+namespace HordeServer.Utilities
+{
+	/// <summary>
+	/// Extension method to wrap a semaphore in a using statement
+	/// </summary>
+	public static class SemaphoreSlimExtensions
+	{
+		/// <summary>
+		/// Returns a disposable semaphore
+		/// </summary>
+		/// <param name="Semaphore">the semaphore to wrap</param>
+		/// <param name="CancelToken">optional cancellation token</param>
+		/// <returns></returns>
+		public static async Task<IDisposable> UseWaitAsync(this SemaphoreSlim Semaphore, CancellationToken CancelToken = default(CancellationToken))
+		{
+			await Semaphore.WaitAsync(CancelToken).ConfigureAwait(false);
+			return new ReleaseWrapper(Semaphore);
+		}
+
+		/// <summary>
+		/// Semaphore wrapper implementing IDisposable
+		/// </summary>
+		private class ReleaseWrapper : IDisposable
+		{
+			/// <summary>
+			/// The semaphore to wrap
+			/// </summary>
+			private readonly SemaphoreSlim Semaphore;
+
+			/// <summary>
+			/// Whether this has been disposed
+			/// </summary>
+			private bool bIsDisposed;
+
+			/// <summary>
+			/// Constructor
+			/// </summary>
+			/// <param name="Semaphore">The semaphore to wrap</param>
+			public ReleaseWrapper(SemaphoreSlim Semaphore)
+			{
+				this.Semaphore = Semaphore;
+			}
+
+			/// <summary>
+			/// Releases the lock on dispose
+			/// </summary>
+			public void Dispose()
+			{
+				if (bIsDisposed)
+				{
+					return;
+				}
+
+				Semaphore.Release();
+				bIsDisposed = true;
+			}
+		}
+	}
+}
diff --git a/Engine/Source/Programs/Horde/HordeServer/Utilities/SingletonDocument.cs b/Engine/Source/Programs/Horde/HordeServer/Utilities/SingletonDocument.cs
new file mode 100644
--- /dev/null
+++ b/Engine/Source/Programs/Horde/HordeServer/Utilities/SingletonDocument.cs
@@ -0,0 +1,138 @@
+using HordeServer.Models;
+using HordeServer.Services;
+using MongoDB.Bson;
+using System;
+using System.Collections.Generic;
+using System.Linq;
+using System.Reflection;
+using System.Threading.Tasks;
+
+namespace HordeServer.Utilities
+{
+	/// <summary>
+	/// Attribute specifying the unique id for a singleton document
+	/// </summary>
+	[AttributeUsage(AttributeTargets.Class)]
+	public class SingletonDocumentAttribute : Attribute
+	{
+		/// <summary>
+		/// Unique id for the singleton document
+		/// </summary>
+		public string Id { get; }
+
+		/// <summary>
+		/// Constructor
+		/// </summary>
+		/// <param name="Id">Unique id for the singleton document</param>
+		public SingletonDocumentAttribute(string Id)
+		{
+			this.Id = Id;
+		}
+	}
+
+	/// <summary>
+	/// Interface for the getting and setting the singleton
+	/// </summary>
+	/// <typeparam name="T">Type of document</typeparam>
+	public interface ISingletonDocument<T>
+	{
+		/// <summary>
+		/// Gets the current document
+		/// </summary>
+		/// <returns>The current document</returns>
+		Task<T> GetAsync();
+
+		/// <summary>
+		/// Attempts to update the document
+		/// </summary>
+		/// <param name="Value">New state of the document</param>
+		/// <returns>True if the document was updated, false otherwise</returns>
+		Task<bool> TryUpdateAsync(T Value);
+	}
+
+	/// <summary>
+	/// Concrete implementation of <see cref="ISingletonDocument{T}"/>
+	/// </summary>
+	/// <typeparam name="T">The document type</typeparam>
+	public class SingletonDocument<T> : ISingletonDocument<T> where T : SingletonBase, new()
+	{
+		/// <summary>
+		/// The database service instance
+		/// </summary>
+		DatabaseService DatabaseService;
+
+		/// <summary>
+		/// Unique id for the singleton document
+		/// </summary>
+		ObjectId ObjectId;
+
+		/// <summary>
+		/// Constructor
+		/// </summary>
+		/// <param name="DatabaseService">The database service instance</param>
+		public SingletonDocument(DatabaseService DatabaseService)
+		{
+			this.DatabaseService = DatabaseService;
+
+			SingletonDocumentAttribute? Attribute = typeof(T).GetCustomAttribute<SingletonDocumentAttribute>();
+			if (Attribute == null)
+			{
+				throw new Exception($"Type {typeof(T).Name} is missing a {nameof(SingletonDocumentAttribute)} annotation");
+			}
+
+			ObjectId = new ObjectId(Attribute.Id);
+		}
+
+		/// <summary>
+		/// Constructor
+		/// </summary>
+		/// <param name="DatabaseService">The database service instance</param>
+		/// <param name="ObjectId">The singleton document object id</param>
+		public SingletonDocument(DatabaseService DatabaseService, ObjectId ObjectId)
+		{
+			this.DatabaseService = DatabaseService;
+			this.ObjectId = ObjectId;
+		}
+
+		/// <inheritdoc/>
+		public async Task<T> GetAsync()
+		{
+			T Value = await DatabaseService.GetSingletonAsync<T>(ObjectId);
+			Value.Id = ObjectId;
+			return Value;
+		}
+
+		/// <inheritdoc/>
+		public Task<bool> TryUpdateAsync(T Value)
+		{
+			return DatabaseService.TryUpdateSingletonAsync<T>(Value);
+		}
+	}
+
+	/// <summary>
+	/// Extension methods for singletons
+	/// </summary>
+	public static class SingletonDocumentExtensions
+	{
+		/// <summary>
+		/// Update a singleton
+		/// </summary>
+		/// <typeparam name="T"></typeparam>
+		/// <param name="Singleton"></param>
+		/// <param name="UpdateAction"></param>
+		/// <returns></returns>
+		public static async Task UpdateAsync<T>(this ISingletonDocument<T> Singleton, Action<T> UpdateAction) where T : SingletonBase, new()
+		{
+			for (; ; )
+			{
+				T Value = await Singleton.GetAsync();
+				UpdateAction(Value);
+
+				if (await Singleton.TryUpdateAsync(Value))
+				{
+					break;
+				}
+			}
+		}
+	}
+}
diff --git a/Engine/Source/Programs/Horde/HordeServer/Utilities/StreamExtensions.cs b/Engine/Source/Programs/Horde/HordeServer/Utilities/StreamExtensions.cs
new file mode 100644
--- /dev/null
+++ b/Engine/Source/Programs/Horde/HordeServer/Utilities/StreamExtensions.cs
@@ -0,0 +1,46 @@
+using System;
+using System.Collections.Generic;
+using System.IO;
+using System.Linq;
+using System.Threading.Tasks;
+
+namespace HordeServer.Utilities
+{
+	/// <summary>
+	/// Extension methods for working with streams
+	/// </summary>
+	static class StreamExtensions
+	{
+		/// <summary>
+		/// Reads data of a fixed size into the buffer. Throws an exception if the whole block cannot be read.
+		/// </summary>
+		/// <param name="Stream">The stream to read from</param>
+		/// <param name="Data">Buffer to receive the read data</param>
+		/// <param name="Offset">Offset within the buffer to read the new data</param>
+		/// <param name="Length">Length of the data to read</param>
+		/// <returns>Async task</returns>
+		public static Task ReadFixedSizeDataAsync(this Stream Stream, byte[] Data, int Offset, int Length)
+		{
+			return ReadFixedSizeDataAsync(Stream, Data.AsMemory(Offset, Length));
+		}
+
+		/// <summary>
+		/// Reads data of a fixed size into the buffer. Throws an exception if the whole block cannot be read.
+		/// </summary>
+		/// <param name="Stream">The stream to read from</param>
+		/// <param name="Data">Buffer to receive the read data</param>
+		/// <returns>Async task</returns>
+		public static async Task ReadFixedSizeDataAsync(this Stream Stream, Memory<byte> Data)
+		{
+			while (Data.Length > 0)
+			{
+				int Count = await Stream.ReadAsync(Data);
+				if (Count == 0)
+				{
+					throw new EndOfStreamException();
+				}
+				Data = Data.Slice(Count);
+			}
+		}
+	}
+}
\ No newline at end of file
diff --git a/Engine/Source/Programs/Horde/HordeServer/Utilities/StringId.cs b/Engine/Source/Programs/Horde/HordeServer/Utilities/StringId.cs
new file mode 100644
--- /dev/null
+++ b/Engine/Source/Programs/Horde/HordeServer/Utilities/StringId.cs
@@ -0,0 +1,301 @@
+using HordeServer.Models;
+using MongoDB.Bson.Serialization;
+using MongoDB.Bson.Serialization.Attributes;
+using MongoDB.Bson.Serialization.Serializers;
+using System;
+using System.Collections.Generic;
+using System.ComponentModel;
+using System.Diagnostics.CodeAnalysis;
+using System.Globalization;
+using System.Linq;
+using System.Text;
+using System.Text.Json;
+using System.Text.Json.Serialization;
+using System.Threading.Tasks;
+
+namespace HordeServer.Utilities
+{
+	/// <summary>
+	/// Normalized string identifier for a resource
+	/// </summary>
+	[TypeConverter(typeof(StringIdTypeConverter))]
+	public struct StringId<T> : IEquatable<StringId<T>>
+	{
+		/// <summary>
+		/// The text representing this id
+		/// </summary>
+		readonly string Text;
+
+		/// <summary>
+		/// Constructor
+		/// </summary>
+		/// <param name="Text">Unique id for the string</param>
+		public StringId(string Text)
+		{
+			this.Text = Text;
+
+			if (Text.Length == 0)
+			{
+//				throw new ArgumentException("String id may not be empty");
+			}
+
+			const int MaxLength = 64;
+			if (Text.Length > MaxLength)
+			{
+				throw new ArgumentException($"String id may not be longer than {MaxLength} characters");
+			}
+
+			for (int Idx = 0; Idx < Text.Length; Idx++)
+			{
+				if (!IsValidCharacter(Text[Idx]))
+				{
+					throw new ArgumentException($"{Text} is not a valid string id");
+				}
+			}
+		}
+
+		/// <summary>
+		/// Constructs from a nullable string
+		/// </summary>
+		/// <param name="Text">The text to construct from</param>
+		/// <returns></returns>
+		[SuppressMessage("Design", "CA1000:Do not declare static members on generic types")]
+		public static StringId<T>? FromNullable(string? Text)
+		{
+			if (String.IsNullOrEmpty(Text))
+			{
+				return null;
+			}
+			else
+			{
+				return new StringId<T>(Text);
+			}
+		}
+
+		/// <summary>
+		/// Generates a new string id from the given text
+		/// </summary>
+		/// <param name="Text">Text to generate from</param>
+		/// <returns>New string id</returns>
+		[SuppressMessage("Design", "CA1000:Do not declare static members on generic types", Justification = "<Pending>")]
+		public static StringId<T> Sanitize(string Text)
+		{
+			StringBuilder Result = new StringBuilder();
+			for (int Idx = 0; Idx < Text.Length; Idx++)
+			{
+				char Character = Text[Idx];
+				if (Character >= 'A' && Character <= 'Z')
+				{
+					Result.Append((char)('a' + (Character - 'A')));
+				}
+				else if (IsValidCharacter(Character))
+				{
+					Result.Append(Character);
+				}
+				else if (Result.Length > 0 && Result[Result.Length - 1] != '-')
+				{
+					Result.Append('-');
+				}
+			}
+			while(Result.Length > 0 && Result[Result.Length - 1] == '-')
+			{
+				Result.Remove(Result.Length - 1, 1);
+			}
+			return new StringId<T>(Result.ToString());
+		}
+
+		/// <summary>
+		/// Checks whether this StringId is set
+		/// </summary>
+		public bool IsEmpty
+		{
+			get { return String.IsNullOrEmpty(Text); }
+		}
+
+		/// <summary>
+		/// Checks whether the given character is valid within a string id
+		/// </summary>
+		/// <param name="Character">The character to check</param>
+		/// <returns>True if the character is valid</returns>
+		static bool IsValidCharacter(char Character)
+		{
+			if (Character >= 'a' && Character <= 'z')
+			{
+				return true;
+			}
+			if (Character >= '0' && Character <= '9')
+			{
+				return true;
+			}
+			if (Character == '-' || Character == '_' || Character == '.')
+			{
+				return true;
+			}
+			return false;
+		}
+
+		/// <inheritdoc/>
+		public override bool Equals(object? Obj)
+		{
+			return Obj is StringId<T> && Equals((StringId<T>)Obj);
+		}
+
+		/// <inheritdoc/>
+		public override int GetHashCode()
+		{
+			return Text.GetHashCode(StringComparison.Ordinal);
+		}
+
+		/// <inheritdoc/>
+		public bool Equals(StringId<T> Other)
+		{
+			return Text.Equals(Other.Text, StringComparison.Ordinal);
+		}
+
+		/// <inheritdoc/>
+		public override string ToString()
+		{
+			return Text;
+		}
+
+		/// <summary>
+		/// Compares two string ids for equality
+		/// </summary>
+		/// <param name="Left">The first string id</param>
+		/// <param name="Right">Second string id</param>
+		/// <returns>True if the two string ids are equal</returns>
+		public static bool operator ==(StringId<T> Left, StringId<T> Right)
+		{
+			return Left.Equals(Right);
+		}
+
+		/// <summary>
+		/// Compares two string ids for inequality
+		/// </summary>
+		/// <param name="Left">The first string id</param>
+		/// <param name="Right">Second string id</param>
+		/// <returns>True if the two string ids are not equal</returns>
+		public static bool operator !=(StringId<T> Left, StringId<T> Right)
+		{
+			return !Left.Equals(Right);
+		}
+	}
+
+	/// <summary>
+	/// Converts <see cref="StringId{T}"/> values to and from JSON
+	/// </summary>
+	public class StringIdJsonConverter<T> : JsonConverter<StringId<T>>
+	{
+		/// <inheritdoc/>
+		public override StringId<T> Read(ref Utf8JsonReader Reader, Type TypeToConvert, JsonSerializerOptions Options)
+		{
+			return new StringId<T>(Reader.GetString());
+		}
+
+		/// <inheritdoc/>
+		public override void Write(Utf8JsonWriter Writer, StringId<T> Value, JsonSerializerOptions Options)
+		{
+			Writer.WriteStringValue(Value.ToString());
+		}
+	}
+
+	/// <summary>
+	/// Converts <see cref="StringId{T}"/> values to and from JSON
+	/// </summary>
+	public class JsonStringIdConverterFactory : JsonConverterFactory
+	{
+		/// <inheritdoc/>
+		public override bool CanConvert(Type TypeToConvert)
+		{
+			return TypeToConvert.IsGenericType && TypeToConvert.GetGenericTypeDefinition() == typeof(StringId<>);
+		}
+
+		/// <inheritdoc/>
+		public override JsonConverter? CreateConverter(Type Type, JsonSerializerOptions Options)
+		{
+			return (JsonConverter?)Activator.CreateInstance(typeof(StringIdJsonConverter<>).MakeGenericType(Type.GetGenericArguments()));
+		}
+	}
+
+	/// <summary>
+	/// Serializer for StringId objects
+	/// </summary>
+	public sealed class StringIdBsonSerializer<T> : SerializerBase<StringId<T>>
+	{
+		/// <inheritdoc/>
+		public override StringId<T> Deserialize(BsonDeserializationContext Context, BsonDeserializationArgs Args)
+		{
+			string Argument;
+			if (Context.Reader.CurrentBsonType == MongoDB.Bson.BsonType.ObjectId)
+			{
+				Argument = Context.Reader.ReadObjectId().ToString();
+			}
+			else
+			{
+				Argument = Context.Reader.ReadString();
+			}
+			return new StringId<T>(Argument);
+		}
+
+		/// <inheritdoc/>
+		public override void Serialize(BsonSerializationContext Context, BsonSerializationArgs Args, StringId<T> Value)
+		{
+			Context.Writer.WriteString(Value.ToString());
+		}
+	}
+
+	/// <summary>
+	/// Serializer for StringId objects
+	/// </summary>
+	public sealed class StringIdSerializationProvider : BsonSerializationProviderBase
+	{
+		/// <inheritdoc/>
+		public override IBsonSerializer? GetSerializer(Type type, IBsonSerializerRegistry serializerRegistry)
+		{
+			if (type.IsGenericType && type.GetGenericTypeDefinition() == typeof(StringId<>))
+			{
+				return (IBsonSerializer?)Activator.CreateInstance(typeof(StringIdBsonSerializer<>).MakeGenericType(type.GetGenericArguments()));
+			}
+			else
+			{
+				return null;
+			}
+		}
+	}
+
+	/// <summary>
+	/// Type converter from strings to PropertyFilter objects
+	/// </summary>
+	public sealed class StringIdTypeConverter : TypeConverter
+	{
+		/// <inheritdoc/>
+		public override bool CanConvertFrom(ITypeDescriptorContext Context, Type SourceType)
+		{
+			return SourceType == typeof(string);
+		}
+
+		/// <inheritdoc/>
+		public override bool CanConvertTo(ITypeDescriptorContext Context, Type DestinationType)
+		{
+			if (DestinationType.IsGenericType)
+			{
+				Type GenericTypeDefinition = DestinationType.GetGenericTypeDefinition();
+				if (GenericTypeDefinition == typeof(StringId<>))
+				{
+					return true;
+				}
+				if (GenericTypeDefinition == typeof(Nullable<>))
+				{
+					return CanConvertTo(Context, GenericTypeDefinition.GetGenericArguments()[0]);
+				}
+			}
+			return false;
+		}
+
+		/// <inheritdoc/>
+		public override object? ConvertTo(ITypeDescriptorContext Context, CultureInfo Culture, object Value, Type DestinationType)
+		{
+			return Activator.CreateInstance(DestinationType, Value);
+		}
+	}
+}
diff --git a/Engine/Source/Programs/Horde/HordeServer/Utilities/StructuredRpcException.cs b/Engine/Source/Programs/Horde/HordeServer/Utilities/StructuredRpcException.cs
new file mode 100644
--- /dev/null
+++ b/Engine/Source/Programs/Horde/HordeServer/Utilities/StructuredRpcException.cs
@@ -0,0 +1,88 @@
+using Grpc.Core;
+using System;
+using System.Collections.Generic;
+using System.Globalization;
+using System.Linq;
+using System.Text;
+using System.Threading.Tasks;
+
+namespace HordeServer.Utilities
+{
+	/// <summary>
+	/// Exception class designed to allow logging structured log messages
+	/// </summary>
+	[System.Diagnostics.CodeAnalysis.SuppressMessage("Design", "CA1032:Implement standard exception constructors", Justification = "<Pending>")]
+	public class StructuredRpcException : RpcException
+	{
+		/// <summary>
+		/// The format string with named parameters
+		/// </summary>
+		public string Format { get; }
+
+		/// <summary>
+		/// The argument list
+		/// </summary>
+		public object[] Args { get; }
+
+		/// <summary>
+		/// Constructor
+		/// </summary>
+		/// <param name="StatusCode">Status code to return</param>
+		/// <param name="Format">The format string</param>
+		/// <param name="Args">Arguments for the format string</param>
+		public StructuredRpcException(StatusCode StatusCode, string Format, params object[] Args)
+			: base(new Status(StatusCode, FormatMessage(Format, Args)))
+		{
+			this.Format = Format;
+			this.Args = Args;
+		}
+
+		/// <summary>
+		/// Replace named arguments in the format message with their values
+		/// </summary>
+		/// <param name="Format"></param>
+		/// <param name="Args"></param>
+		/// <returns></returns>
+		static string FormatMessage(string Format, params object[] Args)
+		{
+			string NewFormat = ConvertToFormatString(Format);
+			return String.Format(CultureInfo.CurrentCulture, NewFormat, Args);
+		}
+
+		/// <summary>
+		/// Converts a named parameter format string to a String.Format style string
+		/// </summary>
+		/// <param name="Format"></param>
+		/// <returns></returns>
+		static string ConvertToFormatString(string Format)
+		{
+			int ArgIdx = 0;
+
+			StringBuilder NewFormat = new StringBuilder();
+			for (int Idx = 0; Idx < Format.Length; Idx++)
+			{
+				char Character = Format[Idx];
+				NewFormat.Append(Character);
+
+				if (Character == '{' && Idx + 1 < Format.Length)
+				{
+					char NextCharacter = Format[Idx + 1];
+					if ((NextCharacter >= 'a' && NextCharacter <= 'z') || (NextCharacter >= 'A' && NextCharacter <= 'Z') || NextCharacter == '_')
+					{
+						for (int EndIdx = Idx + 2; EndIdx < Format.Length; EndIdx++)
+						{
+							if (Format[EndIdx] == ':' || Format[EndIdx] == '}')
+							{
+								NewFormat.Append(ArgIdx++);
+								Idx = EndIdx - 1;
+								break;
+							}
+						}
+					}
+				}
+			}
+
+			return NewFormat.ToString();
+		}
+	}
+}
diff --git a/Engine/Source/Programs/Horde/HordeServer/Utilities/TickedBackgroundService.cs b/Engine/Source/Programs/Horde/HordeServer/Utilities/TickedBackgroundService.cs
new file mode 100644
--- /dev/null
+++ b/Engine/Source/Programs/Horde/HordeServer/Utilities/TickedBackgroundService.cs
@@ -0,0 +1,130 @@
+using Datadog.Trace;
+using EpicGames.Core;
+using Microsoft.Extensions.Hosting;
+using Microsoft.Extensions.Logging;
+using System;
+using System.Collections.Generic;
+using System.Diagnostics;
+using System.Linq;
+using System.Threading;
+using System.Threading.Tasks;
+
+namespace HordeServer.Utilities
+{
+	/// <summary>
+	/// Base class for a background service which is ticked on a regular schedule
+	/// </summary>
+	public abstract class TickedBackgroundService : BackgroundService
+	{
+		/// <summary>
+		/// Frequency that the service will be ticked
+		/// </summary>
+		public TimeSpan? Interval { get; set; }
+
+		/// <summary>
+		/// Logger interface
+		/// </summary>
+		private ILogger Logger;
+
+		/// <summary>
+		/// Event which can be signalled to make the background service tick immediately
+		/// </summary>
+		private AsyncEvent TickImmediatelyEvent = new AsyncEvent();
+
+		/// <summary>
+		/// Constructor
+		/// </summary>
+		/// <param name="Interval">Interval for ticking the service</param>
+		/// <param name="Logger">Logger for debug messages and errors</param>
+		public TickedBackgroundService(TimeSpan? Interval, ILogger Logger)
+		{
+			this.Interval = Interval;
+			this.Logger = Logger;
+		}
+
+		/// <summary>
+		/// Ticks the service immediately
+		/// </summary>
+		protected void TickNow()
+		{
+			TickImmediatelyEvent.Set();
+		}
+
+		/// <inheritdoc/>
+		[System.Diagnostics.CodeAnalysis.SuppressMessage("Design", "CA1031:Do not catch general exception types", Justification = "<Pending>")]
+		protected sealed override async Task ExecuteAsync(CancellationToken StoppingToken)
+		{
+			using (CancellationTask StoppingTask = new CancellationTask(StoppingToken))
+			{
+				// Wait a random amount of time before the first iteration
+				TimeSpan? IntervalCopy = Interval;
+				if(IntervalCopy.HasValue)
+				{
+					TimeSpan Delay = new Random().NextDouble() * IntervalCopy.Value;
+					await Task.WhenAny(StoppingTask.Task, Task.Delay(Delay));
+				}
+
+				// Enter the main tick loop
+				Stopwatch Timer = new Stopwatch();
+				while (!StoppingToken.IsCancellationRequested)
+				{
+					// If the tick immediately event is set, reset it 
+					if(TickImmediatelyEvent.IsSet())
+					{
+						TickImmediatelyEvent = new AsyncEvent();
+					}
+
+					// Run the tick method
+					Timer.Restart();
+					using (Scope Scope = Tracer.Instance.StartActive($"{GetType().Name}.{nameof(TickAsync)}"))
+					{
+						try
+						{
+							await TickAsync(StoppingToken);
+						}
+						catch (Exception Ex)
+						{
+							if (StoppingToken.IsCancellationRequested)
+							{
+								break;
+							}
+							else
+							{
+								Logger.LogError(Ex, "Unhandled exception in {ServiceName}", GetType().Name);
+							}
+						}
+					}
+
+					// Wait until the next interval has elapsed
+					IntervalCopy = Interval;
+					if(IntervalCopy.HasValue)
+					{
+						TimeSpan Delay = IntervalCopy.Value - Timer.Elapsed;
+						if(Delay > TimeSpan.Zero)
+						{
+							await Task.WhenAny(Task.Delay(Delay), StoppingTask.Task);
+						}
+					}
+				}
+			}
+		}
+
+		/// <summary>
+		/// Abstract tick method to be implemented by derived classes
+		/// </summary>
+		/// <param name="StoppingToken">Token signalling that the task should terminate</param>
+		/// <returns>Async task</returns>
+		protected abstract Task TickAsync(CancellationToken StoppingToken);
+		
+		/// <summary>
+		/// Tick the implemented method from above
+		/// **Only for use in testing!**
+		/// </summary>
+		/// <param name="StoppingToken">Optional token signalling that the task should terminate</param>
+		/// <returns>Async task</returns>
+		public Task TickOnlyForTestingAsync(CancellationToken? StoppingToken = null)
+		{
+			return TickAsync(StoppingToken ?? new CancellationToken());
+		}
+	}
+}
