Initial drop of Horde Server/Agent for UE5/Main. Still missing some Epic-specific functionality required to build/run correctly.

diff --git a/Engine/Source/Programs/Horde/HordeServer/appsettings.Development.json b/Engine/Source/Programs/Horde/HordeServer/appsettings.Development.json
new file mode 100644
--- /dev/null
+++ b/Engine/Source/Programs/Horde/HordeServer/appsettings.Development.json
@@ -0,0 +1,10 @@
+{
+	"Horde":
+	{
+		"CorsEnabled": true,
+		"CorsOrigin": "https://localhost:5003;http://localhost:3000",
+		"LogFileProviderType": "Local",
+		"DisableAuth": false,
+		"DisableSchedules": true
+	}
+}
diff --git a/Engine/Source/Programs/Horde/HordeServer/appsettings.Production.json b/Engine/Source/Programs/Horde/HordeServer/appsettings.Production.json
new file mode 100644
--- /dev/null
+++ b/Engine/Source/Programs/Horde/HordeServer/appsettings.Production.json
@@ -0,0 +1,6 @@
+{
+    "Horde": 
+	{
+		"LogJsonToStdOut": true
+    }
+}
diff --git a/Engine/Source/Programs/Horde/HordeServer/appsettings.json b/Engine/Source/Programs/Horde/HordeServer/appsettings.json
new file mode 100644
--- /dev/null
+++ b/Engine/Source/Programs/Horde/HordeServer/appsettings.json
@@ -0,0 +1,10 @@
+{
+	"Horde": 
+	{
+		"DatabaseConnectionString": "mongodb://localhost:27017",
+		"RedisConnectionConfig": "localhost",
+		"DatabaseName": "Horde",
+		"DashboardUrl": "https://localhost:3000"
+	},
+	"AllowedHosts": "*"
+}
diff --git a/Engine/Source/Programs/Horde/HordeServerTests/Fixture.cs b/Engine/Source/Programs/Horde/HordeServerTests/Fixture.cs
new file mode 100644
--- /dev/null
+++ b/Engine/Source/Programs/Horde/HordeServerTests/Fixture.cs
@@ -0,0 +1,160 @@
+using System.Collections.Generic;
+using System.IO;
+using System.Text;
+using System.Threading.Tasks;
+using EpicGames.Core;
+using HordeCommon;
+using HordeServer.Collections;
+using HordeServer.Models;
+using HordeServer.Services;
+using HordeServerTests.Stubs.Services;
+using MongoDB.Bson;
+using StreamId = HordeServer.Utilities.StringId<HordeServer.Models.IStream>;
+using TemplateRefId = HordeServer.Utilities.StringId<HordeServer.Models.TemplateRef>;
+using ProjectId = HordeServer.Utilities.StringId<HordeServer.Models.IProject>;
+using HordeServer.Api;
+
+namespace HordeServerTests
+{
+	public class Fixture
+	{
+		public IJob Job1 { get; private set; } = null!;
+		public IJob Job2 { get; private set; } = null!;
+		public ITemplate Template { get; private set; } = null!;
+		public IGraph Graph { get; private set; } = null!;
+		public IStream? Stream { get; private set; }
+		public TemplateRefId TemplateRefId1 { get; private set; }
+		public TemplateRefId TemplateRefId2 { get; private set; }
+		public TemplateRef TemplateRef1 { get; private set; } = null!;
+		public TemplateRef TemplateRef2 { get; private set; } = null!;
+		public Artifact Job1Artifact { get; private set; } = null!;
+		public string Job1ArtifactData { get; private set; } = null!;
+		public IAgent Agent1 { get; private set; } = null!;
+		public string Agent1Name { get; private set; } = null!;
+
+		private static Fixture? _fixture;
+
+		public static async Task<Fixture> Create(bool ForceNewFixture, IGraphCollection GraphCollection, TemplateService TemplateService, JobService JobService, ArtifactService ArtifactService, StreamService StreamService, AgentService AgentService, IPerforceService PerforceService)
+		{
+			if (_fixture == null || ForceNewFixture)
+			{
+				_fixture = new Fixture();
+				await _fixture.Populate(GraphCollection, TemplateService, JobService, ArtifactService, StreamService, AgentService, PerforceService);
+			}
+			
+			(PerforceService as PerforceServiceStub)?.AddChange("//UE5/Main", 112233, "leet.coder", "Did stuff", new []{"file.cpp"});
+			(PerforceService as PerforceServiceStub)?.AddChange("//UE5/Main", 1111, "swarm", "A shelved CL here", new []{"renderer.cpp"});
+		
+			return _fixture;
+		}
+
+		private async Task Populate(IGraphCollection GraphCollection, TemplateService TemplateService, JobService JobService, ArtifactService ArtifactService, StreamService StreamService, AgentService AgentService, IPerforceService PerforceService)
+		{
+			var Fg = new FixtureGraph();
+			Fg.Id = ContentHash.Empty;
+			Fg.Schema = 1122;
+			Fg.Groups = new List<INodeGroup>();
+			Fg.Aggregates = new List<IAggregate>();
+			Fg.Labels = new List<ILabel>();
+
+			Template = await TemplateService.CreateTemplateAsync("Test template", null, false, null, null,
+				new List<TemplateCounter>(), new List<string>(), new List<Parameter>());
+			Graph = await GraphCollection.AddAsync(Template);
+
+			TemplateRefId1 = new TemplateRefId("template1");
+			TemplateRefId2 = new TemplateRefId("template2");
+			TemplateRef1 = new TemplateRef(Template, false, false, null, null, null);
+			TemplateRef2 = new TemplateRef(Template, false, false, null, null, null);
+
+			Dictionary<TemplateRefId, TemplateRef> TemplateRefs = new Dictionary<TemplateRefId, TemplateRef>
+			{
+				{TemplateRefId1, TemplateRef1},
+				{TemplateRefId2, TemplateRef2}
+			};
+			
+			List<StreamTab> Tabs = new List<StreamTab>();
+			Tabs.Add(new JobsTab("foo", true, new List<TemplateRefId> { TemplateRefId1, TemplateRefId2 }, new List<string>(), new List<JobsTabColumn>()));
+
+			Stream = await StreamService.TryCreateStreamAsync(
+				Id: new StreamId("ue5-main"),
+				Name: "//UE5/Main",
+				ProjectId: new ProjectId("does-not-exist"),
+				Order: null,
+				Tabs: Tabs,
+				AgentTypes: null,
+				WorkspaceTypes: null,
+				TemplateRefs: TemplateRefs,
+				Properties: null,
+				Acl: null
+			);
+			
+			Job1 = await JobService.CreateJobAsync(
+				JobId: new ObjectId("5f283932841e7fdbcafb6ab5"),
+				StreamId: Stream!.Id,
+				TemplateRefId: TemplateRefId1,
+				TemplateHash: Template.Id,
+				Graph: Graph,
+				Name: "hello1",
+				Change: 1000001,
+				CodeChange: 1000002,
+				PreflightChange: 1001,
+				ClonedPreflightChange: null,
+				StartedByUserId: null,
+				StartedByUserName: "SomeUser",
+				Priority: Priority.Normal,
+				null,
+				null,
+				false,
+				false,
+				null,
+				null,
+				null,
+				Template.Counters,
+				Arguments: new List<string>()
+			);
+			Job1 = (await JobService.GetJobAsync(Job1.Id))!;
+
+			Job2 = await JobService.CreateJobAsync(
+				JobId: new ObjectId("5f69ea1b68423e921b035106"),
+				StreamId: Stream!.Id,
+				TemplateRefId: new TemplateRefId("template-id-1"),
+				TemplateHash: ContentHash.MD5("made-up-template-hash"),
+				Graph: Fg,
+				Name: "hello2",
+				Change: 2000001,
+				CodeChange: 2000002,
+				PreflightChange: null,
+				ClonedPreflightChange: null,
+				StartedByUserId: null,
+				StartedByUserName: "SomeUser",
+				Priority: Priority.Normal,
+				null,
+				null,
+				false,
+				false,
+				null,
+				null,
+				null,
+				Template.Counters,
+				Arguments: new List<string>()
+			);
+			Job2 = (await JobService.GetJobAsync(Job2.Id))!;
+
+			Job1ArtifactData = "For The Horde!";
+			Job1Artifact = await ArtifactService.CreateArtifactAsync(Job1.Id, SubResourceId.Parse("22"), "myFile.txt",
+				"text/plain", new MemoryStream(Encoding.UTF8.GetBytes(Job1ArtifactData)));
+
+			Agent1Name = "testAgent1";
+			Agent1 = await AgentService.CreateAgentAsync(Agent1Name, true, false, null, null);
+		}
+
+		private class FixtureGraph : IGraph
+		{
+			public ContentHash Id { get; set; } = ContentHash.Empty;
+			public int Schema { get; set; }
+			public IReadOnlyList<INodeGroup> Groups { get; set; } = null!;
+			public IReadOnlyList<IAggregate> Aggregates { get; set; } = null!;
+			public IReadOnlyList<ILabel> Labels { get; set; } = null!;
+		}
+	}
+}
\ No newline at end of file
diff --git a/Engine/Source/Programs/Horde/HordeServerTests/HordeServerTests.csproj b/Engine/Source/Programs/Horde/HordeServerTests/HordeServerTests.csproj
new file mode 100644
--- /dev/null
+++ b/Engine/Source/Programs/Horde/HordeServerTests/HordeServerTests.csproj
@@ -0,0 +1,52 @@
+<Project Sdk="Microsoft.NET.Sdk">
+
+  <PropertyGroup>
+    <TargetFramework>netcoreapp3.1</TargetFramework>
+    <Nullable>enable</Nullable>
+    <IsPackable>false</IsPackable>
+  </PropertyGroup>
+
+  <ItemGroup>
+    <Compile Remove="Stubs\Collections\JobCollectionStub.cs" />
+    <Compile Remove="Stubs\Collections\StreamCollectionStub.cs" />
+  </ItemGroup>
+
+  <ItemGroup>
+    <None Include="Stubs\Collections\JobCollectionStub.cs" />
+    <None Include="Stubs\Collections\StreamCollectionStub.cs" />
+  </ItemGroup>
+
+  <ItemGroup>
+    <PackageReference Include="Microsoft.AspNetCore.Mvc.Testing" Version="3.1.6" />
+    <PackageReference Include="Microsoft.NET.Test.Sdk" Version="16.5.0" />
+    <PackageReference Include="Moq" Version="4.14.5" />
+    <PackageReference Include="MSTest.TestAdapter" Version="2.1.0" />
+    <PackageReference Include="MSTest.TestFramework" Version="2.1.0" />
+    <PackageReference Include="coverlet.collector" Version="1.2.0" />
+  </ItemGroup>
+
+  <ItemGroup>
+    <ProjectReference Include="..\HordeAgent\HordeAgent.csproj" />
+    <ProjectReference Include="..\HordeServer\HordeServer.csproj" />
+  </ItemGroup>
+
+  <ItemGroup>
+    <Compile Update="Properties\Resources.Designer.cs">
+      <DesignTime>True</DesignTime>
+      <AutoGen>True</AutoGen>
+      <DependentUpon>Resources.resx</DependentUpon>
+    </Compile>
+  </ItemGroup>
+
+  <ItemGroup>
+    <EmbeddedResource Update="Properties\Resources.resx">
+      <Generator>ResXFileCodeGenerator</Generator>
+      <LastGenOutput>Resources.Designer.cs</LastGenOutput>
+    </EmbeddedResource>
+  </ItemGroup>
+
+  <ItemGroup>
+    <Folder Include="Resources\" />
+  </ItemGroup>
+
+</Project>
diff --git a/Engine/Source/Programs/Horde/HordeServerTests/JobServiceTests.cs b/Engine/Source/Programs/Horde/HordeServerTests/JobServiceTests.cs
new file mode 100644
--- /dev/null
+++ b/Engine/Source/Programs/Horde/HordeServerTests/JobServiceTests.cs
@@ -0,0 +1,193 @@
+using System;
+using HordeCommon;
+using HordeServer.Api;
+using HordeServer.Models;
+using Microsoft.VisualStudio.TestTools.UnitTesting;
+using MongoDB.Bson;
+using System.Collections.Generic;
+using System.Text;
+using System.Threading.Tasks;
+using EpicGames.Core;
+using ProjectId = HordeServer.Utilities.StringId<HordeServer.Models.IProject>;
+using StreamId = HordeServer.Utilities.StringId<HordeServer.Models.IStream>;
+using TemplateRefId = HordeServer.Utilities.StringId<HordeServer.Models.TemplateRef>;
+
+namespace HordeServerTests
+{
+	[TestClass]
+	public class JobServiceTests : DatabaseIntegrationTest
+	{
+		[TestMethod]
+		public async Task TestChainedJobs()
+		{
+			TestSetup TestSetup = await GetTestSetup();
+			
+			IProject? Project = await TestSetup.ProjectService.TryCreateProjectAsync(new ProjectId("ue5"), "UE5", null, null, null);
+			Assert.IsNotNull(Project);
+
+			ITemplate Template = await TestSetup.TemplateService.CreateTemplateAsync("Test template", null, false, null, null, new List<TemplateCounter>(), new List<string>(), new List<Parameter>());
+			IGraph Graph = await TestSetup.GraphCollection.AddAsync(Template);
+
+			TemplateRefId TemplateRefId1 = new TemplateRefId("template1");
+			TemplateRefId TemplateRefId2 = new TemplateRefId("template2");
+			TemplateRef TemplateRef1 = new TemplateRef(Template, Triggers: new List<ChainedJobTemplate> { new ChainedJobTemplate("Setup Build", TemplateRefId2) });
+			TemplateRef TemplateRef2 = new TemplateRef(Template);
+
+			Dictionary<TemplateRefId, TemplateRef> TemplateRefs = new Dictionary<TemplateRefId, TemplateRef>
+			{
+				{ TemplateRefId1, TemplateRef1 },
+				{ TemplateRefId2, TemplateRef2 }
+			};
+
+			List<StreamTab> Tabs = new List<StreamTab>();
+			Tabs.Add(new JobsTab("foo", true, new List<TemplateRefId> { TemplateRefId1, TemplateRefId2 }, new List<string>(), new List<JobsTabColumn>()));
+
+			StreamId StreamId = new StreamId("ue5-main");
+			await TestSetup.StreamService.TryCreateStreamAsync(new StreamId("ue5-main"), "//UE5/Main", Project!.Id);
+
+			IStream? Stream = await TestSetup.StreamService.GetStreamAsync(StreamId);
+			Assert.IsNotNull(Stream);
+			await TestSetup.StreamService.UpdateStreamAsync(Stream, NewTabs: Tabs, NewTemplateRefs: TemplateRefs);
+
+			IJob Job = await TestSetup.JobService.CreateJobAsync(null, Stream!.Id, TemplateRefId1, Template.Id, Graph, "Hello", 1234, 1233, 999, null, null, "joe", null, null, TemplateRef1.ChainedJobs, true, true, null, null, null, Template.Counters, new List<string>());
+			Assert.AreEqual(1, Job.ChainedJobs.Count);
+
+			Assert.IsTrue(await TestSetup.JobService.UpdateBatchAsync(Job, Job.Batches[0].Id, ObjectId.GenerateNewId(), JobStepBatchState.Running));
+			Assert.IsNotNull(await TestSetup.JobService.UpdateStepAsync(Job, Job.Batches[0].Id, Job.Batches[0].Steps[0].Id, JobStepState.Running));
+			Assert.IsNotNull(await TestSetup.JobService.UpdateStepAsync(Job, Job.Batches[0].Id, Job.Batches[0].Steps[0].Id, JobStepState.Completed, JobStepOutcome.Success));
+
+			Assert.IsNotNull(Job.ChainedJobs[0].JobId);
+
+			IJob? ChainedJob = await TestSetup.JobCollection.GetAsync(Job.ChainedJobs[0].JobId!.Value);
+			Assert.IsNotNull(ChainedJob);
+			Assert.AreEqual(ChainedJob!.Id, Job!.ChainedJobs[0].JobId);
+
+			Assert.AreEqual(ChainedJob!.Change, Job!.Change);
+			Assert.AreEqual(ChainedJob!.CodeChange, Job!.CodeChange);
+			Assert.AreEqual(ChainedJob!.PreflightChange, Job!.PreflightChange);
+			Assert.AreEqual(ChainedJob!.StartedByUser, Job!.StartedByUser);
+		}
+
+		[TestMethod]
+		public async Task StopAnyDuplicateJobsByPreflight()
+		{
+			TestSetup TestSetup = await GetTestSetup();
+
+			string[] Args = {"-Target=bogus"};
+			IJob OrgJob = await CreatePreflightJob(TestSetup, "tpl-ref-1", "tpl-hash-1", "elvis", 1000, Args);
+			IJob NewJob = await CreatePreflightJob(TestSetup, "tpl-ref-1", "tpl-hash-1", "elvis", 1000, Args);
+			IJob DifferentTplRef = await CreatePreflightJob(TestSetup, "tpl-ref-other", "tpl-hash-1", "elvis", 1000, Args);
+			IJob DifferentTplHash = await CreatePreflightJob(TestSetup, "tpl-ref-1", "tpl-hash-other", "elvis", 1000, Args);
+			IJob DifferentUserName = await CreatePreflightJob(TestSetup, "tpl-ref-1", "tpl-hash-1", "julia", 1000, Args);
+			IJob DifferentArgs = await CreatePreflightJob(TestSetup, "tpl-ref-1", "tpl-hash-1", "elvis", 1000, new []{"-Target=other"});
+			
+			OrgJob = (await TestSetup.JobService.GetJobAsync(OrgJob.Id))!;
+			NewJob = (await TestSetup.JobService.GetJobAsync(NewJob.Id))!;
+			DifferentTplRef = (await TestSetup.JobService.GetJobAsync(DifferentTplRef.Id))!;
+			DifferentTplHash = (await TestSetup.JobService.GetJobAsync(DifferentTplHash.Id))!;
+			DifferentUserName = (await TestSetup.JobService.GetJobAsync(DifferentUserName.Id))!;
+			DifferentArgs = (await TestSetup.JobService.GetJobAsync(DifferentArgs.Id))!;
+			
+			Assert.AreEqual("horde.duplicated.by.newer.CL", OrgJob.AbortedByUser);
+			Assert.IsNull(NewJob.AbortedByUser);
+			Assert.IsNull(DifferentTplRef.AbortedByUser);
+			Assert.IsNull(DifferentTplHash.AbortedByUser);
+			Assert.IsNull(DifferentUserName.AbortedByUser);
+			Assert.IsNull(DifferentArgs.AbortedByUser);
+		}
+
+		private Task<IJob> CreatePreflightJob(TestSetup TestSetup, string TemplateRefId, string TemplateHash, string StartedByUserName, int PreflightChange, string[] Arguments)
+		{
+			return TestSetup.JobService.CreateJobAsync(
+				JobId: ObjectId.GenerateNewId(),
+				StreamId: TestSetup.Fixture!.Stream!.Id,
+				TemplateRefId: new TemplateRefId(TemplateRefId),
+				TemplateHash: new ContentHash(Encoding.ASCII.GetBytes(TemplateHash)),
+				Graph: TestSetup.Fixture!.Graph,
+				Name: "hello1",
+				Change: 1000001,
+				CodeChange: 1000002,
+				PreflightChange: PreflightChange,
+				ClonedPreflightChange: null,
+				StartedByUserId: null,
+				StartedByUserName: StartedByUserName,
+				Priority: Priority.Normal,
+				null,
+				null,
+				false,
+				false,
+				null,
+				null,
+				null,
+				TestSetup.Fixture!.Template.Counters,
+				Arguments: new List<string>(Arguments)
+			);
+		}
+		
+		// Only test for cancelled preflights
+		// [TestMethod]
+		// public async Task StopAnyDuplicateJobsByChange()
+		// {
+		// 	TestSetup TestSetup = await GetTestSetup();
+		// 	Fixture Fix = TestSetup.Fixture!;
+		// 	IJob Job1 = Fix.Job1;
+		// 	IJob Job2 = Fix.Job2;
+		// 	
+		// 	Assert.AreEqual(JobState.Waiting, Job1.GetState());
+		// 	Assert.IsTrue(await TestSetup.JobService.UpdateBatchAsync(Job1, Job1.Batches[0].Id, ObjectId.GenerateNewId(), JobStepBatchState.Running));
+		//
+		// 	Job1 = (await TestSetup.JobService.GetJobAsync(Job1.Id))!;
+		// 	Job2 = (await TestSetup.JobService.GetJobAsync(Job2.Id))!;
+		// 	Assert.AreEqual(JobState.Running, Job1.GetState());
+		// 	Assert.AreEqual(JobState.Complete, Job2.GetState());
+		// 	
+		// 	IJob Job3 = await TestSetup.JobService.CreateJobAsync(null, Fix.Stream!.Id, Fix.TemplateRefId1, Fix.Template.Id, Fix.Graph, "Hello", Job1.Change, Job1.CodeChange, null, "joe", null, null, true, true, null, null, Fix.Template.Counters, new List<string>());
+		// 	await TestSetup.DispatchService.TickOnlyForTestingAsync();
+		// 	Job1 = (await TestSetup.JobService.GetJobAsync(Job1.Id))!;
+		// 	Assert.AreEqual(JobState.Complete, Job1.GetState());
+		// 	Assert.AreEqual(JobState.Waiting, Job3.GetState());
+		// }
+
+		[TestMethod]
+		public async Task TestRunEarly()
+		{
+			TestSetup TestSetup = await GetTestSetup();
+			
+			IProject? Project = await TestSetup.ProjectService.TryCreateProjectAsync(new ProjectId("ue5"), "UE5", null, null, null);
+			Assert.IsNotNull(Project);
+
+			StreamId StreamId = new StreamId("ue5-main");
+			await TestSetup.StreamService.TryCreateStreamAsync(new StreamId("ue5-main"), "//UE5/Main", Project!.Id);
+
+			ITemplate Template = await TestSetup.TemplateService.CreateTemplateAsync("Test template", null, false, null, null, new List<TemplateCounter>(), new List<string>(), new List<Parameter>());
+			IGraph Graph = await TestSetup.GraphCollection.AddAsync(Template);
+
+			CreateGroupRequest GroupA = new CreateGroupRequest("win", new List<CreateNodeRequest>());
+			GroupA.Nodes.Add(new CreateNodeRequest("Compile"));
+
+			CreateGroupRequest GroupB = new CreateGroupRequest("win", new List<CreateNodeRequest>());
+			GroupB.Nodes.Add(new CreateNodeRequest("Cook", RunEarly: true));
+			GroupB.Nodes.Add(new CreateNodeRequest("Middle"));
+			GroupB.Nodes.Add(new CreateNodeRequest("Pak", InputDependencies: new List<string> { "Compile", "Cook" }));
+
+			Graph = await TestSetup.GraphCollection.AppendAsync(Graph, new List<CreateGroupRequest> { GroupA, GroupB });
+
+			IJob Job = await TestSetup.JobService.CreateJobAsync(null, StreamId, new TemplateRefId("temp"), Template.Id, Graph, "Hello", 1234, 1233, 999, null, null, "joe", null, null, null, true, true, null, null, null, Template.Counters, new List<string> { "-Target=Pak" });
+
+			Assert.IsTrue(await TestSetup.JobService.UpdateBatchAsync(Job, Job.Batches[0].Id, ObjectId.GenerateNewId(), JobStepBatchState.Running));
+			Assert.IsNotNull(await TestSetup.JobService.UpdateStepAsync(Job, Job.Batches[0].Id, Job.Batches[0].Steps[0].Id, JobStepState.Running));
+			Assert.IsNotNull(await TestSetup.JobService.UpdateStepAsync(Job, Job.Batches[0].Id, Job.Batches[0].Steps[0].Id, JobStepState.Completed, JobStepOutcome.Success));
+
+			Assert.IsTrue(await TestSetup.JobService.UpdateBatchAsync(Job, Job.Batches[1].Id, ObjectId.GenerateNewId(), JobStepBatchState.Running));
+			Assert.IsNotNull(await TestSetup.JobService.UpdateStepAsync(Job, Job.Batches[1].Id, Job.Batches[1].Steps[0].Id, JobStepState.Running));
+
+			Assert.IsTrue(await TestSetup.JobService.UpdateBatchAsync(Job, Job.Batches[2].Id, ObjectId.GenerateNewId(), JobStepBatchState.Running));
+			Assert.IsNotNull(await TestSetup.JobService.UpdateStepAsync(Job, Job.Batches[2].Id, Job.Batches[2].Steps[0].Id, JobStepState.Running));
+			Assert.IsNotNull(await TestSetup.JobService.UpdateStepAsync(Job, Job.Batches[2].Id, Job.Batches[2].Steps[0].Id, JobStepState.Completed, JobStepOutcome.Success));
+
+			Job = (await TestSetup.JobService.GetJobAsync(Job.Id))!;
+
+			Assert.AreEqual(JobStepState.Waiting, Job.Batches[2].Steps[1].State);
+		}
+	}
+}
diff --git a/Engine/Source/Programs/Horde/HordeServerTests/KeyValueTreeTests.cs b/Engine/Source/Programs/Horde/HordeServerTests/KeyValueTreeTests.cs
new file mode 100644
--- /dev/null
+++ b/Engine/Source/Programs/Horde/HordeServerTests/KeyValueTreeTests.cs
@@ -0,0 +1,183 @@
+using HordeServer.Storage;
+using HordeServer.Storage.Impl;
+using HordeServer.Storage.Primitives;
+using HordeServer.Storage.Services;
+using Microsoft.VisualStudio.TestTools.UnitTesting;
+using System;
+using System.Buffers.Binary;
+using System.Collections.Generic;
+using System.Linq;
+using System.Text;
+using System.Threading.Tasks;
+
+namespace HordeServerTests
+{
+	[TestClass]
+	public class KeyValueTreeTests
+	{
+		class TestTreeMax2 : KeyValueTree
+		{
+			public TestTreeMax2()
+				: base(BlobRef<LeafBlob>.Type, BlobHash.NumBytes, new[] { (0, (IBlobType)BlobRef<LeafBlob>.Type) }, 2, 2)
+			{
+			}
+		}
+
+		class TestTreeMax256 : KeyValueTree
+		{
+			public TestTreeMax256()
+				: base(BlobRef<LeafBlob>.Type, BlobHash.NumBytes, new[] { (0, (IBlobType)BlobRef<LeafBlob>.Type) }, 256, 256)
+			{
+			}
+		}
+
+		[TestMethod]
+		public void TestChildKeys()
+		{
+			Assert.AreEqual(256, KeyValueTreeExtensions.GetChildKeyLength(128));
+			Assert.AreEqual(128, KeyValueTreeExtensions.GetChildKeyLength(64));
+			Assert.AreEqual(64, KeyValueTreeExtensions.GetChildKeyLength(32));
+
+			Assert.AreEqual(256, KeyValueTreeExtensions.GetChildKeyLength(128));
+
+			Assert.AreEqual(128, KeyValueTreeExtensions.GetChildKeyLength(64));
+			Assert.AreEqual(64, KeyValueTreeExtensions.GetChildKeyLength(56));
+			Assert.AreEqual(64, KeyValueTreeExtensions.GetChildKeyLength(48));
+			Assert.AreEqual(48, KeyValueTreeExtensions.GetChildKeyLength(40));
+
+			Assert.AreEqual(128 + 128, KeyValueTreeExtensions.GetChildKeyLength(128 + 64));
+			Assert.AreEqual(128 + 64, KeyValueTreeExtensions.GetChildKeyLength(128 + 56));
+			Assert.AreEqual(128 + 64, KeyValueTreeExtensions.GetChildKeyLength(128 + 48));
+			Assert.AreEqual(128 + 48, KeyValueTreeExtensions.GetChildKeyLength(128 + 40));
+
+			Assert.AreEqual(64, KeyValueTreeExtensions.GetChildKeyLength(32));
+		}
+
+		static KeyValueItem MakeUpdate(BlobHash HashValue)
+		{
+			byte[] Buffer = new byte[BlobHash.NumBytes];
+			new Random().NextBytes(Buffer);
+			return new KeyValueItem(HashValue, Buffer);
+		}
+
+		static KeyValueItem MakeItem(string HashString)
+		{
+			return MakeUpdate(BlobHash.Parse(HashString));
+		}
+
+		[TestMethod]
+		public async Task TestBasicInsertion()
+		{
+			TransientStorageBackend Backend = new TransientStorageBackend();
+			SimpleStorageService Storage = new SimpleStorageService(Backend);
+
+			BlobRef<TestTreeMax2> Tree = await Storage.CreateKeyValueTreeAsync<TestTreeMax2>();
+			{
+				ReadOnlyKeyValueNode RootNode = await Storage.GetKeyValueNodeAsync(Tree.Hash);
+				Assert.AreEqual(BlobHash.NumBits, RootNode.NumKeyBits);
+				Assert.AreEqual(0, RootNode.NumItems);
+				Assert.AreEqual(0, RootNode.NumItemsIfMerged);
+			}
+
+			KeyValueItem Item1 = MakeItem("aaaaaaaaaaaaaaaaaaaaaaaaaaaaaaaaaaaaaaaaaaaaaaaaaaaaaaaaaaaaaaaa");
+			Tree = await Storage.UpdateKeyValueTreeAsync(Tree, new[] { Item1 });
+			{
+				ReadOnlyKeyValueNode RootNode = await Storage.GetKeyValueNodeAsync(Tree.Hash);
+				Assert.AreEqual(BlobHash.NumBits, RootNode.NumKeyBits);
+				Assert.AreEqual(1, RootNode.NumItems);
+				Assert.AreEqual(1, RootNode.NumItemsIfMerged);
+			}
+
+			KeyValueItem Item2 = MakeItem("aaaaaaaaaaaaaaaaaaaaaaaaaaaaaaaaaaaaaaaaaaaaaaaaaaaaaaaaaaaaaaac");
+			Tree = await Storage.UpdateKeyValueTreeAsync(Tree, new[] { Item2 });
+			{
+				ReadOnlyKeyValueNode RootNode = await Storage.GetKeyValueNodeAsync(Tree.Hash);
+				Assert.AreEqual(BlobHash.NumBits, RootNode.NumKeyBits);
+				Assert.AreEqual(2, RootNode.NumItems);
+				Assert.AreEqual(2, RootNode.NumItemsIfMerged);
+			}
+
+			KeyValueItem Item3 = MakeItem("aaaaaaaaaaaaaaaaaaaaaaaaaaaaaaaaaaaaaaaaaaaaaaaaaaaaaaaaaaaaaaab");
+			Tree = await Storage.UpdateKeyValueTreeAsync(Tree, new[] { Item3 });
+			{
+				ReadOnlyKeyValueNode Node = await Storage.GetKeyValueNodeAsync(Tree.Hash);
+				Assert.IsTrue(Node.NumItems <= 2);
+				Assert.AreEqual(3, Node.NumItemsIfMerged);
+			}
+
+			ReadOnlyMemory<byte>? FindItem1 = await Storage.SearchKeyValueTreeAsync(Tree, Item1.Key);
+			Assert.IsNotNull(FindItem1);
+			Assert.IsTrue(FindItem1.Value.Span.SequenceEqual(Item1.Value.Span));
+
+			ReadOnlyMemory<byte>? FindItem2 = await Storage.SearchKeyValueTreeAsync(Tree, Item2.Key);
+			Assert.IsNotNull(FindItem2);
+			Assert.IsTrue(FindItem2.Value.Span.SequenceEqual(Item2.Value.Span));
+
+			ReadOnlyMemory<byte>? FindItem3 = await Storage.SearchKeyValueTreeAsync(Tree, Item3.Key);
+			Assert.IsNotNull(FindItem3);
+			Assert.IsTrue(FindItem3.Value.Span.SequenceEqual(Item3.Value.Span));
+		}
+
+		[TestMethod]
+		public async Task TestLoad()
+		{
+			Random Random = new Random(0x12345678);
+
+			TransientStorageBackend Backend = new TransientStorageBackend();
+			SimpleStorageService Storage = new SimpleStorageService(Backend);
+
+			BlobRef<TestTreeMax256> Tree = await Storage.CreateKeyValueTreeAsync<TestTreeMax256>();
+
+			Dictionary<BlobHash, BlobHash> Items = new Dictionary<BlobHash, BlobHash>();
+			for (int OuterIdx = 0; OuterIdx < 128; OuterIdx++)
+			{
+				List<KeyValueItem> NewItems = new List<KeyValueItem>();
+				for (int Idx = 0; Idx < 128; Idx++)
+				{
+					byte[] KeyBytes = new byte[BlobHash.NumBytes];
+					Random.NextBytes(KeyBytes);
+					BlobHash Key = new BlobHash(KeyBytes);
+
+					byte[] ValueBytes = new byte[BlobHash.NumBytes];
+					Random.NextBytes(ValueBytes);
+					BlobHash Value = new BlobHash(ValueBytes);
+
+					NewItems.Add(new KeyValueItem(Key, ValueBytes));
+					Items[Key] = Value;
+				}
+				Tree = await Storage.UpdateKeyValueTreeAsync(Tree, NewItems);
+			}
+
+			int NumItems = await CountItemsAsync(Storage, Tree.Hash);
+			Assert.AreEqual(Items.Count, NumItems);
+
+			foreach (KeyValuePair<BlobHash, BlobHash> Pair in Items)
+			{
+				ReadOnlyMemory<byte>? Value = await Storage.SearchKeyValueTreeAsync(Tree, Pair.Key);
+				Assert.IsNotNull(Value);
+				Assert.IsTrue(Value.Value.Span.SequenceEqual(Pair.Value.Span));
+			}
+		}
+
+		static async Task<int> CountItemsAsync(IStorageService Storage, BlobHash Hash)
+		{
+			ReadOnlyKeyValueNode Node = await Storage.GetKeyValueNodeAsync(Hash);
+			if (Node.NumKeyBits == BlobHash.NumBits)
+			{
+				return Node.NumItems;
+			}
+			else
+			{
+				int NumItems = 0;
+
+				ReadOnlyBlobHashArray Children = new ReadOnlyBlobHashArray(Node.Values);
+				foreach (BlobHash Child in Children)
+				{
+					NumItems += await CountItemsAsync(Storage, Child);
+				}
+
+				return NumItems;
+			}
+		}
+	}
+}
diff --git a/Engine/Source/Programs/Horde/HordeServerTests/ServiceAccountCollectionTests.cs b/Engine/Source/Programs/Horde/HordeServerTests/ServiceAccountCollectionTests.cs
new file mode 100644
--- /dev/null
+++ b/Engine/Source/Programs/Horde/HordeServerTests/ServiceAccountCollectionTests.cs
@@ -0,0 +1,75 @@
+﻿// Copyright Epic Games, Inc. All Rights Reserved.
+
+using System.Collections.Generic;
+using System.Threading.Tasks;
+using HordeServer.Collections;
+using HordeServer.Collections.Impl;
+using HordeServer.Models;
+using HordeServer.Services;
+using Microsoft.VisualStudio.TestTools.UnitTesting;
+using MongoDB.Bson;
+
+namespace HordeServerTests
+{
+	[TestClass]
+	public class ServiceAccountCollectionTests : DatabaseIntegrationTest
+	{
+		private readonly IServiceAccountCollection ServiceAccounts;
+		private readonly IServiceAccount ServiceAccount;
+		
+		public ServiceAccountCollectionTests()
+		{
+			DatabaseService DatabaseService = GetDatabaseService();
+			ServiceAccounts = new ServiceAccountCollection(DatabaseService);
+			ServiceAccount = ServiceAccounts.AddAsync(ObjectId.GenerateNewId().ToString(), new Dictionary<string, string> {{"myclaim", "myvalue"}}, "mydesc").Result;
+		}
+
+		[TestMethod]
+		public async Task Add()
+		{
+			IServiceAccount Sa = await ServiceAccounts.AddAsync("addToken", new Dictionary<string, string> {{"myclaim", "myvalue"}}, "mydesc");
+			Assert.AreEqual("addToken", Sa.SecretToken);
+			Assert.AreEqual(1, Sa.Claims.Count);
+			Assert.AreEqual("myvalue", Sa.Claims["myclaim"]);
+			Assert.IsTrue(Sa.Enabled);
+			Assert.AreEqual("mydesc", Sa.Description);
+		}
+		
+		[TestMethod]
+		public async Task Get()
+		{
+			IServiceAccount Sa = (await ServiceAccounts.GetAsync(ServiceAccount.Id))!;
+			Assert.AreEqual(ServiceAccount, Sa);
+		}
+		
+		[TestMethod]
+		public async Task GetBySecretToken()
+		{
+			IServiceAccount Sa = (await ServiceAccounts.GetBySecretTokenAsync(ServiceAccount.SecretToken))!;
+			Assert.AreEqual(ServiceAccount, Sa);
+		}
+		
+		[TestMethod]
+		public async Task Update()
+		{
+			Dictionary<string, string> NewClaims = new Dictionary<string, string> {{"newclaim1", "newvalue1"}, {"newclaim2", "newvalue2"}};
+			await ServiceAccounts.UpdateAsync(ServiceAccount.Id, "newtoken", NewClaims, false, "newdesc");
+			IServiceAccount GetSa = (await ServiceAccounts.GetAsync(ServiceAccount.Id))!;
+			
+			Assert.AreEqual("newtoken", GetSa.SecretToken);
+			Assert.AreEqual(2, GetSa.Claims.Count);
+			Assert.AreEqual("newvalue1", GetSa.Claims["newclaim1"]);
+			Assert.AreEqual("newvalue2", GetSa.Claims["newclaim2"]);
+			Assert.AreEqual(false, GetSa.Enabled);
+			Assert.AreEqual("newdesc", GetSa.Description);
+		}
+		
+		[TestMethod]
+		public async Task Delete()
+		{
+			await ServiceAccounts.DeleteAsync(ServiceAccount.Id);
+			IServiceAccount? Result = await ServiceAccounts.GetAsync(ServiceAccount.Id);
+			Assert.IsNull(Result);
+		}
+	}
+}
\ No newline at end of file
diff --git a/Engine/Source/Programs/Horde/HordeServerTests/Stubs/Services/PerforceServiceStub.cs b/Engine/Source/Programs/Horde/HordeServerTests/Stubs/Services/PerforceServiceStub.cs
new file mode 100644
--- /dev/null
+++ b/Engine/Source/Programs/Horde/HordeServerTests/Stubs/Services/PerforceServiceStub.cs
@@ -0,0 +1,139 @@
+using HordeServer.Models;
+using HordeServer.Services;
+using System;
+using System.Collections.Generic;
+using System.Diagnostics.CodeAnalysis;
+using System.Linq;
+using System.Security.Cryptography.X509Certificates;
+using System.Text;
+using System.Threading.Tasks;
+
+using StreamId = HordeServer.Utilities.StringId<HordeServer.Models.IStream>;
+
+namespace HordeServerTests.Stubs.Services
+{
+	class PerforceServiceStub : IPerforceService
+	{
+		class ChangeComparer : IComparer<int>
+		{
+			public int Compare(int X, int Y) => Y.CompareTo(X);
+		}
+
+		public Dictionary<string, SortedDictionary<int, ChangeDetails>> Changes { get; } = new Dictionary<string, SortedDictionary<int, ChangeDetails>>(StringComparer.OrdinalIgnoreCase);
+
+		public void AddChange(string StreamName, int Number, string Author, string Description, IEnumerable<string> Files)
+		{
+			SortedDictionary<int, ChangeDetails>? StreamChanges;
+			if (!Changes.TryGetValue(StreamName, out StreamChanges))
+			{
+				StreamChanges = new SortedDictionary<int, ChangeDetails>(new ChangeComparer());
+				Changes[StreamName] = StreamChanges;
+			}
+			StreamChanges.Add(Number, new ChangeDetails(Number, Author, Description, Files.ToList()));
+		}
+
+		public Task<List<ChangeSummary>> GetChangesAsync(string StreamName, int? MinChange, int? MaxChange, int NumResults, string? ImpersonateUser)
+		{
+			List<ChangeSummary> Results = new List<ChangeSummary>();
+
+			SortedDictionary<int, ChangeDetails>? StreamChanges;
+			if (Changes.TryGetValue(StreamName, out StreamChanges) && StreamChanges.Count > 0)
+			{
+				foreach (ChangeDetails Details in StreamChanges.Values)
+				{
+					if (MinChange.HasValue && Details.Number < MinChange)
+					{
+						break;
+					}
+					if (!MaxChange.HasValue || Details.Number <= MaxChange.Value)
+					{
+						Results.Add(new ChangeSummary(Details.Number, Details.Author, Details.Description));
+					}
+					if (NumResults > 0 && Results.Count >= NumResults)
+					{
+						break;
+					}
+				}
+			}
+
+			return Task.FromResult(Results);
+		}
+
+		public Task<PerforceUserInfo?> GetUserInfoAsync(string UserName)
+		{
+			return Task.FromResult<PerforceUserInfo?>(new PerforceUserInfo { Name = UserName, Email = $"[REDACTED_EMAIL]" });
+		}
+
+		public Task<List<ChangeDetails>> GetChangeDetailsAsync(string StreamName, IReadOnlyList<int> ChangeNumbers, string? ImpersonateUser)
+		{
+			List<ChangeDetails> Results = new List<ChangeDetails>();
+			foreach (int ChangeNumber in ChangeNumbers)
+			{
+				Results.Add(Changes[StreamName][ChangeNumber]);
+			}
+			return Task.FromResult(Results);
+		}
+
+		public Task<string> CreateTicket(string ImpersonateUser)
+		{
+			return Task.FromResult("bogus-ticket");
+		}
+
+		public Task<int> GetCodeChangeAsync(string StreamName, int Change)
+		{
+			int CodeChange = 0;
+
+			SortedDictionary<int, ChangeDetails>? StreamChanges;
+			if (Changes.TryGetValue(StreamName, out StreamChanges))
+			{
+				foreach (ChangeDetails Details in StreamChanges.Values)
+				{
+					if (Details.Number <= Change && Details.Files.Any(x => x.EndsWith(".h") || x.EndsWith(".cpp")))
+					{
+						CodeChange = Details.Number;
+						break;
+					}
+				}
+			}
+
+			return Task.FromResult(CodeChange);
+		}
+
+		public Task<int> CreateNewChangeAsync(string StreamName, string Path)
+		{
+			ChangeDetails NewChange = new ChangeDetails(Changes[StreamName].First().Key + 1, "", "", new List<string> { Path });
+			Changes[StreamName].Add(NewChange.Number, NewChange);
+			return Task.FromResult(NewChange.Number);
+		}
+
+		public Task<List<FileSummary>> FindFilesAsync(IEnumerable<string> Paths)
+		{
+			throw new NotImplementedException();
+		}
+
+		public Task<byte[]> PrintAsync(string Path)
+		{
+			throw new NotImplementedException();
+		}
+
+		public Task<(int? Change, string Message)> SubmitShelvedChangeAsync(int ShelvedChange, int OriginalChange)
+		{
+			throw new NotImplementedException();
+		}
+
+		public Task<int> DuplicateShelvedChangeAsync(int ShelvedChange)
+		{
+			return Task.FromResult(ShelvedChange);
+		}
+
+		public Task DeleteShelvedChangeAsync(int ShelvedChange)
+		{
+			return Task.CompletedTask;
+		}
+
+		public Task UpdateChangelistDescription(int Change, string Description)
+		{
+			return Task.CompletedTask;
+		}
+	}
+}
