Add missing copyright notices.

diff --git a/Engine/Source/Programs/Horde/HordeServer/Api/HelixSwarm.cs b/Engine/Source/Programs/Horde/HordeServer/Api/HelixSwarm.cs
--- a/Engine/Source/Programs/Horde/HordeServer/Api/HelixSwarm.cs
+++ b/Engine/Source/Programs/Horde/HordeServer/Api/HelixSwarm.cs
@@ -1,4 +1,6 @@
-﻿using System.ComponentModel.DataAnnotations;
+﻿// Copyright Epic Games, Inc. All Rights Reserved.
+
+using System.ComponentModel.DataAnnotations;
 using System.Diagnostics.CodeAnalysis;
 
 namespace HordeServer.Api
diff --git a/Engine/Source/Programs/Horde/HordeServer/Authentication/ServiceAccountAuthHandler.cs b/Engine/Source/Programs/Horde/HordeServer/Authentication/ServiceAccountAuthHandler.cs
--- a/Engine/Source/Programs/Horde/HordeServer/Authentication/ServiceAccountAuthHandler.cs
+++ b/Engine/Source/Programs/Horde/HordeServer/Authentication/ServiceAccountAuthHandler.cs
@@ -1,4 +1,6 @@
-﻿using Microsoft.AspNetCore.Authentication;
+﻿// Copyright Epic Games, Inc. All Rights Reserved.
+
+using Microsoft.AspNetCore.Authentication;
 using Microsoft.Extensions.Logging;
 using Microsoft.Extensions.Options;
 using System;
diff --git a/Engine/Source/Programs/Horde/HordeServer/Collections/DatadogTracingCollection.cs b/Engine/Source/Programs/Horde/HordeServer/Collections/DatadogTracingCollection.cs
--- a/Engine/Source/Programs/Horde/HordeServer/Collections/DatadogTracingCollection.cs
+++ b/Engine/Source/Programs/Horde/HordeServer/Collections/DatadogTracingCollection.cs
@@ -1,4 +1,6 @@
-﻿using System.Collections.Generic;
+﻿// Copyright Epic Games, Inc. All Rights Reserved.
+
+using System.Collections.Generic;
 using System.Threading;
 using System.Threading.Tasks;
 using Datadog.Trace;
diff --git a/Engine/Source/Programs/Horde/HordeServer/Services/RequestTrackerService.cs b/Engine/Source/Programs/Horde/HordeServer/Services/RequestTrackerService.cs
--- a/Engine/Source/Programs/Horde/HordeServer/Services/RequestTrackerService.cs
+++ b/Engine/Source/Programs/Horde/HordeServer/Services/RequestTrackerService.cs
@@ -1,4 +1,6 @@
-﻿using System;
+﻿// Copyright Epic Games, Inc. All Rights Reserved.
+
+using System;
 using System.Collections.Concurrent;
 using System.Collections.Generic;
 using System.Linq;
diff --git a/Engine/Source/Programs/Horde/HordeServer/Utilities/NoOpDogStatsd.cs b/Engine/Source/Programs/Horde/HordeServer/Utilities/NoOpDogStatsd.cs
--- a/Engine/Source/Programs/Horde/HordeServer/Utilities/NoOpDogStatsd.cs
+++ b/Engine/Source/Programs/Horde/HordeServer/Utilities/NoOpDogStatsd.cs
@@ -1,4 +1,6 @@
-﻿using System;
+﻿// Copyright Epic Games, Inc. All Rights Reserved.
+
+using System;
 using System.Diagnostics.CodeAnalysis;
 using System.IO;
 using StatsdClient;
diff --git a/Engine/Source/Programs/Horde/HordeServerTests/DatabaseRunnerTest.cs b/Engine/Source/Programs/Horde/HordeServerTests/DatabaseRunnerTest.cs
--- a/Engine/Source/Programs/Horde/HordeServerTests/DatabaseRunnerTest.cs
+++ b/Engine/Source/Programs/Horde/HordeServerTests/DatabaseRunnerTest.cs
@@ -1,4 +1,6 @@
-﻿using System.Threading;
+﻿// Copyright Epic Games, Inc. All Rights Reserved.
+
+using System.Threading;
 using Microsoft.VisualStudio.TestTools.UnitTesting;
 
 namespace HordeServerTests
diff --git a/Engine/Source/Programs/Horde/HordeServerTests/GrpcTestHelpers.cs b/Engine/Source/Programs/Horde/HordeServerTests/GrpcTestHelpers.cs
--- a/Engine/Source/Programs/Horde/HordeServerTests/GrpcTestHelpers.cs
+++ b/Engine/Source/Programs/Horde/HordeServerTests/GrpcTestHelpers.cs
@@ -1,4 +1,6 @@
-﻿#region Copyright notice and license
+﻿// Copyright Epic Games, Inc. All Rights Reserved.
+
+#region Copyright notice and license
 
 // Copyright 2019 The gRPC Authors
 //
diff --git a/Engine/Source/Programs/Horde/HordeServerTests/JobsControllerTest.cs b/Engine/Source/Programs/Horde/HordeServerTests/JobsControllerTest.cs
--- a/Engine/Source/Programs/Horde/HordeServerTests/JobsControllerTest.cs
+++ b/Engine/Source/Programs/Horde/HordeServerTests/JobsControllerTest.cs
@@ -1,4 +1,6 @@
-﻿using System.Net;
+﻿// Copyright Epic Games, Inc. All Rights Reserved.
+
+using System.Net;
 using System.Threading.Tasks;
 using Microsoft.VisualStudio.TestTools.UnitTesting;
 
diff --git a/Engine/Source/Programs/Horde/HordeServerTests/PoolsControllerTest.cs b/Engine/Source/Programs/Horde/HordeServerTests/PoolsControllerTest.cs
--- a/Engine/Source/Programs/Horde/HordeServerTests/PoolsControllerTest.cs
+++ b/Engine/Source/Programs/Horde/HordeServerTests/PoolsControllerTest.cs
@@ -1,4 +1,6 @@
-﻿using System.Threading.Tasks;
+﻿// Copyright Epic Games, Inc. All Rights Reserved.
+
+using System.Threading.Tasks;
 using Microsoft.VisualStudio.TestTools.UnitTesting;
 
 namespace HordeServerTests
