Inital drop of Horde Dashboard for UE5/Main.
Still missing some Epic-specific functionality required to build/run correctly.

diff --git a/Engine/Source/Programs/Horde/HordeDashboard/Dashboard/.eslintrc.js b/Engine/Source/Programs/Horde/HordeDashboard/Dashboard/.eslintrc.js
new file mode 100644
--- /dev/null
+++ b/Engine/Source/Programs/Horde/HordeDashboard/Dashboard/.eslintrc.js
@@ -0,0 +1,3 @@
+module.exports = {
+  "extends": ["react-app"]
+}
\ No newline at end of file
diff --git a/Engine/Source/Programs/Horde/HordeDashboard/Dashboard/README.md b/Engine/Source/Programs/Horde/HordeDashboard/Dashboard/README.md
new file mode 100644
--- /dev/null
+++ b/Engine/Source/Programs/Horde/HordeDashboard/Dashboard/README.md
@@ -0,0 +1,2 @@
+
+## Horde Dashboard
\ No newline at end of file
diff --git a/Engine/Source/Programs/Horde/HordeDashboard/Dashboard/package.json b/Engine/Source/Programs/Horde/HordeDashboard/Dashboard/package.json
new file mode 100644
--- /dev/null
+++ b/Engine/Source/Programs/Horde/HordeDashboard/Dashboard/package.json
@@ -0,0 +1,59 @@
+{
+  "name": "horde-dashboard",
+  "version": "0.1.0",
+  "private": true,
+  "main": "index.js",
+  "proxy": "https://horde.devtools.epicgames.com",
+  "dependencies": {
+    "@blueprintjs/icons": "^3.13.0",
+    "@datadog/browser-logs": "^1.12.4",
+    "@fluentui/react": "^7.155.0",
+    "@fluentui/react-charting": "^5.0.22",
+    "crypto-js": "^4.0.0",
+    "formik": "^2.1.1",
+    "idb": "^6.0.0",
+    "markdown-to-jsx": "^6.10.3",
+    "mobx": "^5.15.1",
+    "mobx-react-lite": "^1.5.2",
+    "moment-timezone": "^0.5.28",
+    "react": "^16.12.0",
+    "react-custom-scrollbars": "^4.2.1",
+    "react-dom": "^17.0.1",
+    "react-flame-graph": "^1.4.0",
+    "react-highlighter": "^0.4.3",
+    "react-router-dom": "5.2.0",
+    "react-router-hash-link": "^2.3.1",
+    "react-scripts": "^4.0.1",
+    "react-text-marquee": "^1.0.1"
+  },
+  "scripts": {
+    "start": "rimraf ./build&&set NODE_OPTIONS=\"--max-http-header-size=80000\"&&react-scripts start",
+    "build": "react-scripts build",
+    "test": "react-scripts test",
+    "eject": "react-scripts eject",
+    "lint-full": "eslint ./src/**/*.{ts,tsx}",
+    "line-full-fix": "eslint --fix ./src/**/*.{ts,tsx}"
+  },
+  "browserslist": {
+    "production": [
+      ">0.2%",
+      "not dead",
+      "not op_mini all"
+    ],
+    "development": [
+      "last 1 chrome version",
+      "last 1 firefox version",
+      "last 1 safari version"
+    ]
+  },
+  "devDependencies": {
+    "@types/crypto-js": "^4.0.1",
+    "@types/markdown-to-jsx": "6.9.1",
+    "@types/react-custom-scrollbars": "^4.0.6",
+    "@types/react-dom": "^17.0.0",
+    "@types/react-highlighter": "^0.3.4",
+    "@types/react-router-dom": "^5.1.6",
+    "rimraf": "^3.0.0",
+    "typescript": "^4.1.3"
+  }
+}
diff --git a/Engine/Source/Programs/Horde/HordeDashboard/Dashboard/public/fonts/LICENSE.txt b/Engine/Source/Programs/Horde/HordeDashboard/Dashboard/public/fonts/LICENSE.txt
new file mode 100644
--- /dev/null
+++ b/Engine/Source/Programs/Horde/HordeDashboard/Dashboard/public/fonts/LICENSE.txt
@@ -0,0 +1,202 @@
+
+                                 Apache License
+                           Version 2.0, January 2004
+                        http://www.apache.org/licenses/
+
+   TERMS AND CONDITIONS FOR USE, REPRODUCTION, AND DISTRIBUTION
+
+   1. Definitions.
+
+      "License" shall mean the terms and conditions for use, reproduction,
+      and distribution as defined by Sections 1 through 9 of this document.
+
+      "Licensor" shall mean the copyright owner or entity authorized by
+      the copyright owner that is granting the License.
+
+      "Legal Entity" shall mean the union of the acting entity and all
+      other entities that control, are controlled by, or are under common
+      control with that entity. For the purposes of this definition,
+      "control" means (i) the power, direct or indirect, to cause the
+      direction or management of such entity, whether by contract or
+      otherwise, or (ii) ownership of fifty percent (50%) or more of the
+      outstanding shares, or (iii) beneficial ownership of such entity.
+
+      "You" (or "Your") shall mean an individual or Legal Entity
+      exercising permissions granted by this License.
+
+      "Source" form shall mean the preferred form for making modifications,
+      including but not limited to software source code, documentation
+      source, and configuration files.
+
+      "Object" form shall mean any form resulting from mechanical
+      transformation or translation of a Source form, including but
+      not limited to compiled object code, generated documentation,
+      and conversions to other media types.
+
+      "Work" shall mean the work of authorship, whether in Source or
+      Object form, made available under the License, as indicated by a
+      copyright notice that is included in or attached to the work
+      (an example is provided in the Appendix below).
+
+      "Derivative Works" shall mean any work, whether in Source or Object
+      form, that is based on (or derived from) the Work and for which the
+      editorial revisions, annotations, elaborations, or other modifications
+      represent, as a whole, an original work of authorship. For the purposes
+      of this License, Derivative Works shall not include works that remain
+      separable from, or merely link (or bind by name) to the interfaces of,
+      the Work and Derivative Works thereof.
+
+      "Contribution" shall mean any work of authorship, including
+      the original version of the Work and any modifications or additions
+      to that Work or Derivative Works thereof, that is intentionally
+      submitted to Licensor for inclusion in the Work by the copyright owner
+      or by an individual or Legal Entity authorized to submit on behalf of
+      the copyright owner. For the purposes of this definition, "submitted"
+      means any form of electronic, verbal, or written communication sent
+      to the Licensor or its representatives, including but not limited to
+      communication on electronic mailing lists, source code control systems,
+      and issue tracking systems that are managed by, or on behalf of, the
+      Licensor for the purpose of discussing and improving the Work, but
+      excluding communication that is conspicuously marked or otherwise
+      designated in writing by the copyright owner as "Not a Contribution."
+
+      "Contributor" shall mean Licensor and any individual or Legal Entity
+      on behalf of whom a Contribution has been received by Licensor and
+      subsequently incorporated within the Work.
+
+   2. Grant of Copyright License. Subject to the terms and conditions of
+      this License, each Contributor hereby grants to You a perpetual,
+      worldwide, non-exclusive, no-charge, royalty-free, irrevocable
+      copyright license to reproduce, prepare Derivative Works of,
+      publicly display, publicly perform, sublicense, and distribute the
+      Work and such Derivative Works in Source or Object form.
+
+   3. Grant of Patent License. Subject to the terms and conditions of
+      this License, each Contributor hereby grants to You a perpetual,
+      worldwide, non-exclusive, no-charge, royalty-free, irrevocable
+      (except as stated in this section) patent license to make, have made,
+      use, offer to sell, sell, import, and otherwise transfer the Work,
+      where such license applies only to those patent claims licensable
+      by such Contributor that are necessarily infringed by their
+      Contribution(s) alone or by combination of their Contribution(s)
+      with the Work to which such Contribution(s) was submitted. If You
+      institute patent litigation against any entity (including a
+      cross-claim or counterclaim in a lawsuit) alleging that the Work
+      or a Contribution incorporated within the Work constitutes direct
+      or contributory patent infringement, then any patent licenses
+      granted to You under this License for that Work shall terminate
+      as of the date such litigation is filed.
+
+   4. Redistribution. You may reproduce and distribute copies of the
+      Work or Derivative Works thereof in any medium, with or without
+      modifications, and in Source or Object form, provided that You
+      meet the following conditions:
+
+      (a) You must give any other recipients of the Work or
+          Derivative Works a copy of this License; and
+
+      (b) You must cause any modified files to carry prominent notices
+          stating that You changed the files; and
+
+      (c) You must retain, in the Source form of any Derivative Works
+          that You distribute, all copyright, patent, trademark, and
+          attribution notices from the Source form of the Work,
+          excluding those notices that do not pertain to any part of
+          the Derivative Works; and
+
+      (d) If the Work includes a "NOTICE" text file as part of its
+          distribution, then any Derivative Works that You distribute must
+          include a readable copy of the attribution notices contained
+          within such NOTICE file, excluding those notices that do not
+          pertain to any part of the Derivative Works, in at least one
+          of the following places: within a NOTICE text file distributed
+          as part of the Derivative Works; within the Source form or
+          documentation, if provided along with the Derivative Works; or,
+          within a display generated by the Derivative Works, if and
+          wherever such third-party notices normally appear. The contents
+          of the NOTICE file are for informational purposes only and
+          do not modify the License. You may add Your own attribution
+          notices within Derivative Works that You distribute, alongside
+          or as an addendum to the NOTICE text from the Work, provided
+          that such additional attribution notices cannot be construed
+          as modifying the License.
+
+      You may add Your own copyright statement to Your modifications and
+      may provide additional or different license terms and conditions
+      for use, reproduction, or distribution of Your modifications, or
+      for any such Derivative Works as a whole, provided Your use,
+      reproduction, and distribution of the Work otherwise complies with
+      the conditions stated in this License.
+
+   5. Submission of Contributions. Unless You explicitly state otherwise,
+      any Contribution intentionally submitted for inclusion in the Work
+      by You to the Licensor shall be under the terms and conditions of
+      this License, without any additional terms or conditions.
+      Notwithstanding the above, nothing herein shall supersede or modify
+      the terms of any separate license agreement you may have executed
+      with Licensor regarding such Contributions.
+
+   6. Trademarks. This License does not grant permission to use the trade
+      names, trademarks, service marks, or product names of the Licensor,
+      except as required for reasonable and customary use in describing the
+      origin of the Work and reproducing the content of the NOTICE file.
+
+   7. Disclaimer of Warranty. Unless required by applicable law or
+      agreed to in writing, Licensor provides the Work (and each
+      Contributor provides its Contributions) on an "AS IS" BASIS,
+      WITHOUT WARRANTIES OR CONDITIONS OF ANY KIND, either express or
+      implied, including, without limitation, any warranties or conditions
+      of TITLE, NON-INFRINGEMENT, MERCHANTABILITY, or FITNESS FOR A
+      PARTICULAR PURPOSE. You are solely responsible for determining the
+      appropriateness of using or redistributing the Work and assume any
+      risks associated with Your exercise of permissions under this License.
+
+   8. Limitation of Liability. In no event and under no legal theory,
+      whether in tort (including negligence), contract, or otherwise,
+      unless required by applicable law (such as deliberate and grossly
+      negligent acts) or agreed to in writing, shall any Contributor be
+      liable to You for damages, including any direct, indirect, special,
+      incidental, or consequential damages of any character arising as a
+      result of this License or out of the use or inability to use the
+      Work (including but not limited to damages for loss of goodwill,
+      work stoppage, computer failure or malfunction, or any and all
+      other commercial damages or losses), even if such Contributor
+      has been advised of the possibility of such damages.
+
+   9. Accepting Warranty or Additional Liability. While redistributing
+      the Work or Derivative Works thereof, You may choose to offer,
+      and charge a fee for, acceptance of support, warranty, indemnity,
+      or other liability obligations and/or rights consistent with this
+      License. However, in accepting such obligations, You may act only
+      on Your own behalf and on Your sole responsibility, not on behalf
+      of any other Contributor, and only if You agree to indemnify,
+      defend, and hold each Contributor harmless for any liability
+      incurred by, or claims asserted against, such Contributor by reason
+      of your accepting any such warranty or additional liability.
+
+   END OF TERMS AND CONDITIONS
+
+   APPENDIX: How to apply the Apache License to your work.
+
+      To apply the Apache License to your work, attach the following
+      boilerplate notice, with the fields enclosed by brackets "[]"
+      replaced with your own identifying information. (Don't include
+      the brackets!)  The text should be enclosed in the appropriate
+      comment syntax for the file format. We also recommend that a
+      file or class name and description of purpose be included on the
+      same "printed page" as the copyright notice for easier
+      identification within third-party archives.
+
+   Copyright [yyyy] [name of copyright owner]
+
+   Licensed under the Apache License, Version 2.0 (the "License");
+   you may not use this file except in compliance with the License.
+   You may obtain a copy of the License at
+
+       http://www.apache.org/licenses/LICENSE-2.0
+
+   Unless required by applicable law or agreed to in writing, software
+   distributed under the License is distributed on an "AS IS" BASIS,
+   WITHOUT WARRANTIES OR CONDITIONS OF ANY KIND, either express or implied.
+   See the License for the specific language governing permissions and
+   limitations under the License.
diff --git a/Engine/Source/Programs/Horde/HordeDashboard/Dashboard/public/fonts/OFL.txt b/Engine/Source/Programs/Horde/HordeDashboard/Dashboard/public/fonts/OFL.txt
new file mode 100644
--- /dev/null
+++ b/Engine/Source/Programs/Horde/HordeDashboard/Dashboard/public/fonts/OFL.txt
@@ -0,0 +1,95 @@
+Copyright (c) 2010, Matt McInerney ([REDACTED_EMAIL]),
+Copyright (c) 2011, Pablo Impallari ([REDACTED_EMAIL]),
+Copyright (c) 2011, Rodrigo Fuenzalida ([REDACTED_EMAIL]), with Reserved Font Name Raleway
+
+This Font Software is licensed under the SIL Open Font License, Version 1.1.
+This license is copied below, and is also available with a FAQ at:
+http://scripts.sil.org/OFL
+
+
+-----------------------------------------------------------
+SIL OPEN FONT LICENSE Version 1.1 - 26 February 2007
+-----------------------------------------------------------
+
+PREAMBLE
+The goals of the Open Font License (OFL) are to stimulate worldwide
+development of collaborative font projects, to support the font creation
+efforts of academic and linguistic communities, and to provide a free and
+open framework in which fonts may be shared and improved in partnership
+with others.
+
+The OFL allows the licensed fonts to be used, studied, modified and
+redistributed freely as long as they are not sold by themselves. The
+fonts, including any derivative works, can be bundled, embedded, 
+redistributed and/or sold with any software provided that any reserved
+names are not used by derivative works. The fonts and derivatives,
+however, cannot be released under any other type of license. The
+requirement for fonts to remain under this license does not apply
+to any document created using the fonts or their derivatives.
+
+DEFINITIONS
+"Font Software" refers to the set of files released by the Copyright
+Holder(s) under this license and clearly marked as such. This may
+include source files, build scripts and documentation.
+
+"Reserved Font Name" refers to any names specified as such after the
+copyright statement(s).
+
+"Original Version" refers to the collection of Font Software components as
+distributed by the Copyright Holder(s).
+
+"Modified Version" refers to any derivative made by adding to, deleting,
+or substituting -- in part or in whole -- any of the components of the
+Original Version, by changing formats or by porting the Font Software to a
+new environment.
+
+"Author" refers to any designer, engineer, programmer, technical
+writer or other person who contributed to the Font Software.
+
+PERMISSION & CONDITIONS
+Permission is hereby granted, free of charge, to any person obtaining
+a copy of the Font Software, to use, study, copy, merge, embed, modify,
+redistribute, and sell modified and unmodified copies of the Font
+Software, subject to the following conditions:
+
+1) Neither the Font Software nor any of its individual components,
+in Original or Modified Versions, may be sold by itself.
+
+2) Original or Modified Versions of the Font Software may be bundled,
+redistributed and/or sold with any software, provided that each copy
+contains the above copyright notice and this license. These can be
+included either as stand-alone text files, human-readable headers or
+in the appropriate machine-readable metadata fields within text or
+binary files as long as those fields can be easily viewed by the user.
+
+3) No Modified Version of the Font Software may use the Reserved Font
+Name(s) unless explicit written permission is granted by the corresponding
+Copyright Holder. This restriction only applies to the primary font name as
+presented to the users.
+
+4) The name(s) of the Copyright Holder(s) or the Author(s) of the Font
+Software shall not be used to promote, endorse or advertise any
+Modified Version, except to acknowledge the contribution(s) of the
+Copyright Holder(s) and the Author(s) or with their explicit written
+permission.
+
+5) The Font Software, modified or unmodified, in part or in whole,
+must be distributed entirely under this license, and must not be
+distributed under any other license. The requirement for fonts to
+remain under this license does not apply to any document created
+using the Font Software.
+
+TERMINATION
+This license becomes null and void if any of the above conditions are
+not met.
+
+DISCLAIMER
+THE FONT SOFTWARE IS PROVIDED "AS IS", WITHOUT WARRANTY OF ANY KIND,
+EXPRESS OR IMPLIED, INCLUDING BUT NOT LIMITED TO ANY WARRANTIES OF
+MERCHANTABILITY, FITNESS FOR A PARTICULAR PURPOSE AND NONINFRINGEMENT
+OF COPYRIGHT, PATENT, TRADEMARK, OR OTHER RIGHT. IN NO EVENT SHALL THE
+COPYRIGHT HOLDER BE LIABLE FOR ANY CLAIM, DAMAGES OR OTHER LIABILITY,
+INCLUDING ANY GENERAL, SPECIAL, INDIRECT, INCIDENTAL, OR CONSEQUENTIAL
+DAMAGES, WHETHER IN AN ACTION OF CONTRACT, TORT OR OTHERWISE, ARISING
+FROM, OUT OF THE USE OR INABILITY TO USE THE FONT SOFTWARE OR FROM
+OTHER DEALINGS IN THE FONT SOFTWARE.
diff --git a/Engine/Source/Programs/Horde/HordeDashboard/Dashboard/public/fonts/blueprintjs_icons_license.txt b/Engine/Source/Programs/Horde/HordeDashboard/Dashboard/public/fonts/blueprintjs_icons_license.txt
new file mode 100644
--- /dev/null
+++ b/Engine/Source/Programs/Horde/HordeDashboard/Dashboard/public/fonts/blueprintjs_icons_license.txt
@@ -0,0 +1,176 @@
+                                 Apache License
+                           Version 2.0, January 2004
+                        http://www.apache.org/licenses/
+
+   TERMS AND CONDITIONS FOR USE, REPRODUCTION, AND DISTRIBUTION
+
+   1. Definitions.
+
+      "License" shall mean the terms and conditions for use, reproduction,
+      and distribution as defined by Sections 1 through 9 of this document.
+
+      "Licensor" shall mean the copyright owner or entity authorized by
+      the copyright owner that is granting the License.
+
+      "Legal Entity" shall mean the union of the acting entity and all
+      other entities that control, are controlled by, or are under common
+      control with that entity. For the purposes of this definition,
+      "control" means (i) the power, direct or indirect, to cause the
+      direction or management of such entity, whether by contract or
+      otherwise, or (ii) ownership of fifty percent (50%) or more of the
+      outstanding shares, or (iii) beneficial ownership of such entity.
+
+      "You" (or "Your") shall mean an individual or Legal Entity
+      exercising permissions granted by this License.
+
+      "Source" form shall mean the preferred form for making modifications,
+      including but not limited to software source code, documentation
+      source, and configuration files.
+
+      "Object" form shall mean any form resulting from mechanical
+      transformation or translation of a Source form, including but
+      not limited to compiled object code, generated documentation,
+      and conversions to other media types.
+
+      "Work" shall mean the work of authorship, whether in Source or
+      Object form, made available under the License, as indicated by a
+      copyright notice that is included in or attached to the work
+      (an example is provided in the Appendix below).
+
+      "Derivative Works" shall mean any work, whether in Source or Object
+      form, that is based on (or derived from) the Work and for which the
+      editorial revisions, annotations, elaborations, or other modifications
+      represent, as a whole, an original work of authorship. For the purposes
+      of this License, Derivative Works shall not include works that remain
+      separable from, or merely link (or bind by name) to the interfaces of,
+      the Work and Derivative Works thereof.
+
+      "Contribution" shall mean any work of authorship, including
+      the original version of the Work and any modifications or additions
+      to that Work or Derivative Works thereof, that is intentionally
+      submitted to Licensor for inclusion in the Work by the copyright owner
+      or by an individual or Legal Entity authorized to submit on behalf of
+      the copyright owner. For the purposes of this definition, "submitted"
+      means any form of electronic, verbal, or written communication sent
+      to the Licensor or its representatives, including but not limited to
+      communication on electronic mailing lists, source code control systems,
+      and issue tracking systems that are managed by, or on behalf of, the
+      Licensor for the purpose of discussing and improving the Work, but
+      excluding communication that is conspicuously marked or otherwise
+      designated in writing by the copyright owner as "Not a Contribution."
+
+      "Contributor" shall mean Licensor and any individual or Legal Entity
+      on behalf of whom a Contribution has been received by Licensor and
+      subsequently incorporated within the Work.
+
+   2. Grant of Copyright License. Subject to the terms and conditions of
+      this License, each Contributor hereby grants to You a perpetual,
+      worldwide, non-exclusive, no-charge, royalty-free, irrevocable
+      copyright license to reproduce, prepare Derivative Works of,
+      publicly display, publicly perform, sublicense, and distribute the
+      Work and such Derivative Works in Source or Object form.
+
+   3. Grant of Patent License. Subject to the terms and conditions of
+      this License, each Contributor hereby grants to You a perpetual,
+      worldwide, non-exclusive, no-charge, royalty-free, irrevocable
+      (except as stated in this section) patent license to make, have made,
+      use, offer to sell, sell, import, and otherwise transfer the Work,
+      where such license applies only to those patent claims licensable
+      by such Contributor that are necessarily infringed by their
+      Contribution(s) alone or by combination of their Contribution(s)
+      with the Work to which such Contribution(s) was submitted. If You
+      institute patent litigation against any entity (including a
+      cross-claim or counterclaim in a lawsuit) alleging that the Work
+      or a Contribution incorporated within the Work constitutes direct
+      or contributory patent infringement, then any patent licenses
+      granted to You under this License for that Work shall terminate
+      as of the date such litigation is filed.
+
+   4. Redistribution. You may reproduce and distribute copies of the
+      Work or Derivative Works thereof in any medium, with or without
+      modifications, and in Source or Object form, provided that You
+      meet the following conditions:
+
+      (a) You must give any other recipients of the Work or
+          Derivative Works a copy of this License; and
+
+      (b) You must cause any modified files to carry prominent notices
+          stating that You changed the files; and
+
+      (c) You must retain, in the Source form of any Derivative Works
+          that You distribute, all copyright, patent, trademark, and
+          attribution notices from the Source form of the Work,
+          excluding those notices that do not pertain to any part of
+          the Derivative Works; and
+
+      (d) If the Work includes a "NOTICE" text file as part of its
+          distribution, then any Derivative Works that You distribute must
+          include a readable copy of the attribution notices contained
+          within such NOTICE file, excluding those notices that do not
+          pertain to any part of the Derivative Works, in at least one
+          of the following places: within a NOTICE text file distributed
+          as part of the Derivative Works; within the Source form or
+          documentation, if provided along with the Derivative Works; or,
+          within a display generated by the Derivative Works, if and
+          wherever such third-party notices normally appear. The contents
+          of the NOTICE file are for informational purposes only and
+          do not modify the License. You may add Your own attribution
+          notices within Derivative Works that You distribute, alongside
+          or as an addendum to the NOTICE text from the Work, provided
+          that such additional attribution notices cannot be construed
+          as modifying the License.
+
+      You may add Your own copyright statement to Your modifications and
+      may provide additional or different license terms and conditions
+      for use, reproduction, or distribution of Your modifications, or
+      for any such Derivative Works as a whole, provided Your use,
+      reproduction, and distribution of the Work otherwise complies with
+      the conditions stated in this License.
+
+   5. Submission of Contributions. Unless You explicitly state otherwise,
+      any Contribution intentionally submitted for inclusion in the Work
+      by You to the Licensor shall be under the terms and conditions of
+      this License, without any additional terms or conditions.
+      Notwithstanding the above, nothing herein shall supersede or modify
+      the terms of any separate license agreement you may have executed
+      with Licensor regarding such Contributions.
+
+   6. Trademarks. This License does not grant permission to use the trade
+      names, trademarks, service marks, or product names of the Licensor,
+      except as required for reasonable and customary use in describing the
+      origin of the Work and reproducing the content of the NOTICE file.
+
+   7. Disclaimer of Warranty. Unless required by applicable law or
+      agreed to in writing, Licensor provides the Work (and each
+      Contributor provides its Contributions) on an "AS IS" BASIS,
+      WITHOUT WARRANTIES OR CONDITIONS OF ANY KIND, either express or
+      implied, including, without limitation, any warranties or conditions
+      of TITLE, NON-INFRINGEMENT, MERCHANTABILITY, or FITNESS FOR A
+      PARTICULAR PURPOSE. You are solely responsible for determining the
+      appropriateness of using or redistributing the Work and assume any
+      risks associated with Your exercise of permissions under this License.
+
+   8. Limitation of Liability. In no event and under no legal theory,
+      whether in tort (including negligence), contract, or otherwise,
+      unless required by applicable law (such as deliberate and grossly
+      negligent acts) or agreed to in writing, shall any Contributor be
+      liable to You for damages, including any direct, indirect, special,
+      incidental, or consequential damages of any character arising as a
+      result of this License or out of the use or inability to use the
+      Work (including but not limited to damages for loss of goodwill,
+      work stoppage, computer failure or malfunction, or any and all
+      other commercial damages or losses), even if such Contributor
+      has been advised of the possibility of such damages.
+
+   9. Accepting Warranty or Additional Liability. While redistributing
+      the Work or Derivative Works thereof, You may choose to offer,
+      and charge a fee for, acceptance of support, warranty, indemnity,
+      or other liability obligations and/or rights consistent with this
+      License. However, in accepting such obligations, You may act only
+      on Your own behalf and on Your sole responsibility, not on behalf
+      of any other Contributor, and only if You agree to indemnify,
+      defend, and hold each Contributor harmless for any liability
+      incurred by, or claims asserted against, such Contributor by reason
+      of your accepting any such warranty or additional liability.
+
+   END OF TERMS AND CONDITIONS
diff --git a/Engine/Source/Programs/Horde/HordeDashboard/Dashboard/public/fonts/open_sans_license.txt b/Engine/Source/Programs/Horde/HordeDashboard/Dashboard/public/fonts/open_sans_license.txt
new file mode 100644
--- /dev/null
+++ b/Engine/Source/Programs/Horde/HordeDashboard/Dashboard/public/fonts/open_sans_license.txt
@@ -0,0 +1,202 @@
+
+                                 Apache License
+                           Version 2.0, January 2004
+                        http://www.apache.org/licenses/
+
+   TERMS AND CONDITIONS FOR USE, REPRODUCTION, AND DISTRIBUTION
+
+   1. Definitions.
+
+      "License" shall mean the terms and conditions for use, reproduction,
+      and distribution as defined by Sections 1 through 9 of this document.
+
+      "Licensor" shall mean the copyright owner or entity authorized by
+      the copyright owner that is granting the License.
+
+      "Legal Entity" shall mean the union of the acting entity and all
+      other entities that control, are controlled by, or are under common
+      control with that entity. For the purposes of this definition,
+      "control" means (i) the power, direct or indirect, to cause the
+      direction or management of such entity, whether by contract or
+      otherwise, or (ii) ownership of fifty percent (50%) or more of the
+      outstanding shares, or (iii) beneficial ownership of such entity.
+
+      "You" (or "Your") shall mean an individual or Legal Entity
+      exercising permissions granted by this License.
+
+      "Source" form shall mean the preferred form for making modifications,
+      including but not limited to software source code, documentation
+      source, and configuration files.
+
+      "Object" form shall mean any form resulting from mechanical
+      transformation or translation of a Source form, including but
+      not limited to compiled object code, generated documentation,
+      and conversions to other media types.
+
+      "Work" shall mean the work of authorship, whether in Source or
+      Object form, made available under the License, as indicated by a
+      copyright notice that is included in or attached to the work
+      (an example is provided in the Appendix below).
+
+      "Derivative Works" shall mean any work, whether in Source or Object
+      form, that is based on (or derived from) the Work and for which the
+      editorial revisions, annotations, elaborations, or other modifications
+      represent, as a whole, an original work of authorship. For the purposes
+      of this License, Derivative Works shall not include works that remain
+      separable from, or merely link (or bind by name) to the interfaces of,
+      the Work and Derivative Works thereof.
+
+      "Contribution" shall mean any work of authorship, including
+      the original version of the Work and any modifications or additions
+      to that Work or Derivative Works thereof, that is intentionally
+      submitted to Licensor for inclusion in the Work by the copyright owner
+      or by an individual or Legal Entity authorized to submit on behalf of
+      the copyright owner. For the purposes of this definition, "submitted"
+      means any form of electronic, verbal, or written communication sent
+      to the Licensor or its representatives, including but not limited to
+      communication on electronic mailing lists, source code control systems,
+      and issue tracking systems that are managed by, or on behalf of, the
+      Licensor for the purpose of discussing and improving the Work, but
+      excluding communication that is conspicuously marked or otherwise
+      designated in writing by the copyright owner as "Not a Contribution."
+
+      "Contributor" shall mean Licensor and any individual or Legal Entity
+      on behalf of whom a Contribution has been received by Licensor and
+      subsequently incorporated within the Work.
+
+   2. Grant of Copyright License. Subject to the terms and conditions of
+      this License, each Contributor hereby grants to You a perpetual,
+      worldwide, non-exclusive, no-charge, royalty-free, irrevocable
+      copyright license to reproduce, prepare Derivative Works of,
+      publicly display, publicly perform, sublicense, and distribute the
+      Work and such Derivative Works in Source or Object form.
+
+   3. Grant of Patent License. Subject to the terms and conditions of
+      this License, each Contributor hereby grants to You a perpetual,
+      worldwide, non-exclusive, no-charge, royalty-free, irrevocable
+      (except as stated in this section) patent license to make, have made,
+      use, offer to sell, sell, import, and otherwise transfer the Work,
+      where such license applies only to those patent claims licensable
+      by such Contributor that are necessarily infringed by their
+      Contribution(s) alone or by combination of their Contribution(s)
+      with the Work to which such Contribution(s) was submitted. If You
+      institute patent litigation against any entity (including a
+      cross-claim or counterclaim in a lawsuit) alleging that the Work
+      or a Contribution incorporated within the Work constitutes direct
+      or contributory patent infringement, then any patent licenses
+      granted to You under this License for that Work shall terminate
+      as of the date such litigation is filed.
+
+   4. Redistribution. You may reproduce and distribute copies of the
+      Work or Derivative Works thereof in any medium, with or without
+      modifications, and in Source or Object form, provided that You
+      meet the following conditions:
+
+      (a) You must give any other recipients of the Work or
+          Derivative Works a copy of this License; and
+
+      (b) You must cause any modified files to carry prominent notices
+          stating that You changed the files; and
+
+      (c) You must retain, in the Source form of any Derivative Works
+          that You distribute, all copyright, patent, trademark, and
+          attribution notices from the Source form of the Work,
+          excluding those notices that do not pertain to any part of
+          the Derivative Works; and
+
+      (d) If the Work includes a "NOTICE" text file as part of its
+          distribution, then any Derivative Works that You distribute must
+          include a readable copy of the attribution notices contained
+          within such NOTICE file, excluding those notices that do not
+          pertain to any part of the Derivative Works, in at least one
+          of the following places: within a NOTICE text file distributed
+          as part of the Derivative Works; within the Source form or
+          documentation, if provided along with the Derivative Works; or,
+          within a display generated by the Derivative Works, if and
+          wherever such third-party notices normally appear. The contents
+          of the NOTICE file are for informational purposes only and
+          do not modify the License. You may add Your own attribution
+          notices within Derivative Works that You distribute, alongside
+          or as an addendum to the NOTICE text from the Work, provided
+          that such additional attribution notices cannot be construed
+          as modifying the License.
+
+      You may add Your own copyright statement to Your modifications and
+      may provide additional or different license terms and conditions
+      for use, reproduction, or distribution of Your modifications, or
+      for any such Derivative Works as a whole, provided Your use,
+      reproduction, and distribution of the Work otherwise complies with
+      the conditions stated in this License.
+
+   5. Submission of Contributions. Unless You explicitly state otherwise,
+      any Contribution intentionally submitted for inclusion in the Work
+      by You to the Licensor shall be under the terms and conditions of
+      this License, without any additional terms or conditions.
+      Notwithstanding the above, nothing herein shall supersede or modify
+      the terms of any separate license agreement you may have executed
+      with Licensor regarding such Contributions.
+
+   6. Trademarks. This License does not grant permission to use the trade
+      names, trademarks, service marks, or product names of the Licensor,
+      except as required for reasonable and customary use in describing the
+      origin of the Work and reproducing the content of the NOTICE file.
+
+   7. Disclaimer of Warranty. Unless required by applicable law or
+      agreed to in writing, Licensor provides the Work (and each
+      Contributor provides its Contributions) on an "AS IS" BASIS,
+      WITHOUT WARRANTIES OR CONDITIONS OF ANY KIND, either express or
+      implied, including, without limitation, any warranties or conditions
+      of TITLE, NON-INFRINGEMENT, MERCHANTABILITY, or FITNESS FOR A
+      PARTICULAR PURPOSE. You are solely responsible for determining the
+      appropriateness of using or redistributing the Work and assume any
+      risks associated with Your exercise of permissions under this License.
+
+   8. Limitation of Liability. In no event and under no legal theory,
+      whether in tort (including negligence), contract, or otherwise,
+      unless required by applicable law (such as deliberate and grossly
+      negligent acts) or agreed to in writing, shall any Contributor be
+      liable to You for damages, including any direct, indirect, special,
+      incidental, or consequential damages of any character arising as a
+      result of this License or out of the use or inability to use the
+      Work (including but not limited to damages for loss of goodwill,
+      work stoppage, computer failure or malfunction, or any and all
+      other commercial damages or losses), even if such Contributor
+      has been advised of the possibility of such damages.
+
+   9. Accepting Warranty or Additional Liability. While redistributing
+      the Work or Derivative Works thereof, You may choose to offer,
+      and charge a fee for, acceptance of support, warranty, indemnity,
+      or other liability obligations and/or rights consistent with this
+      License. However, in accepting such obligations, You may act only
+      on Your own behalf and on Your sole responsibility, not on behalf
+      of any other Contributor, and only if You agree to indemnify,
+      defend, and hold each Contributor harmless for any liability
+      incurred by, or claims asserted against, such Contributor by reason
+      of your accepting any such warranty or additional liability.
+
+   END OF TERMS AND CONDITIONS
+
+   APPENDIX: How to apply the Apache License to your work.
+
+      To apply the Apache License to your work, attach the following
+      boilerplate notice, with the fields enclosed by brackets "[]"
+      replaced with your own identifying information. (Don't include
+      the brackets!)  The text should be enclosed in the appropriate
+      comment syntax for the file format. We also recommend that a
+      file or class name and description of purpose be included on the
+      same "printed page" as the copyright notice for easier
+      identification within third-party archives.
+
+   Copyright [yyyy] [name of copyright owner]
+
+   Licensed under the Apache License, Version 2.0 (the "License");
+   you may not use this file except in compliance with the License.
+   You may obtain a copy of the License at
+
+       http://www.apache.org/licenses/LICENSE-2.0
+
+   Unless required by applicable law or agreed to in writing, software
+   distributed under the License is distributed on an "AS IS" BASIS,
+   WITHOUT WARRANTIES OR CONDITIONS OF ANY KIND, either express or implied.
+   See the License for the specific language governing permissions and
+   limitations under the License.
diff --git a/Engine/Source/Programs/Horde/HordeDashboard/Dashboard/public/index.html b/Engine/Source/Programs/Horde/HordeDashboard/Dashboard/public/index.html
new file mode 100644
--- /dev/null
+++ b/Engine/Source/Programs/Horde/HordeDashboard/Dashboard/public/index.html
@@ -0,0 +1,43 @@
+<!DOCTYPE html>
+<html lang="en">
+  <head>
+    <meta charset="utf-8" />
+    <link rel="icon" href="%PUBLIC_URL%/favicon.ico" />
+    <meta name="viewport" content="width=device-width, initial-scale=1" />
+    <meta name="theme-color" content="#000000" />
+    <meta
+      name="description"
+      content="Web site created using create-react-app"
+    />
+    <link rel="apple-touch-icon" href="%PUBLIC_URL%/logo192.png" />
+    <!--
+      manifest.json provides metadata used when your web app is installed on a
+      user's mobile device or desktop. See https://developers.google.com/web/fundamentals/web-app-manifest/
+    -->
+    <link rel="manifest" href="%PUBLIC_URL%/manifest.json" />
+    <!--
+      Notice the use of %PUBLIC_URL% in the tags above.
+      It will be replaced with the URL of the `public` folder during the build.
+      Only files inside the `public` folder can be referenced from the HTML.
+
+      Unlike "/favicon.ico" or "favicon.ico", "%PUBLIC_URL%/favicon.ico" will
+      work correctly both with client-side routing and a non-root public URL.
+      Learn how to configure a non-root public URL by running `npm run build`.
+    -->
+    <title>Horde</title>
+  </head>
+  <body>
+    <noscript>You need to enable JavaScript to run this app.</noscript>
+    <div id="root"></div>
+    <!--
+      This HTML file is a template.
+      If you open it directly in the browser, you will see an empty page.
+
+      You can add webfonts, meta tags, or analytics to this file.
+      The build step will place the bundled scripts into the <body> tag.
+
+      To begin the development, run `npm start` or `yarn start`.
+      To create a production bundle, use `npm run build` or `yarn build`.
+    -->
+  </body>
+</html>
diff --git a/Engine/Source/Programs/Horde/HordeDashboard/Dashboard/public/manifest.json b/Engine/Source/Programs/Horde/HordeDashboard/Dashboard/public/manifest.json
new file mode 100644
--- /dev/null
+++ b/Engine/Source/Programs/Horde/HordeDashboard/Dashboard/public/manifest.json
@@ -0,0 +1,25 @@
+{
+  "short_name": "HordeDashboard",
+  "name": "Horde Dashboard",
+  "icons": [
+    {
+      "src": "favicon.ico",
+      "sizes": "64x64 32x32 24x24 16x16",
+      "type": "image/x-icon"
+    },
+    {
+      "src": "logo192.png",
+      "type": "image/png",
+      "sizes": "192x192"
+    },
+    {
+      "src": "logo512.png",
+      "type": "image/png",
+      "sizes": "512x512"
+    }
+  ],
+  "start_url": ".",
+  "display": "standalone",
+  "theme_color": "#000000",
+  "background_color": "#ffffff"
+}
diff --git a/Engine/Source/Programs/Horde/HordeDashboard/Dashboard/public/robots.txt b/Engine/Source/Programs/Horde/HordeDashboard/Dashboard/public/robots.txt
new file mode 100644
--- /dev/null
+++ b/Engine/Source/Programs/Horde/HordeDashboard/Dashboard/public/robots.txt
@@ -0,0 +1,2 @@
+# https://www.robotstxt.org/robotstxt.html
+User-agent: *
diff --git a/Engine/Source/Programs/Horde/HordeDashboard/Dashboard/src/App.css b/Engine/Source/Programs/Horde/HordeDashboard/Dashboard/src/App.css
new file mode 100644
--- /dev/null
+++ b/Engine/Source/Programs/Horde/HordeDashboard/Dashboard/src/App.css
@@ -0,0 +1,38 @@
+.App {
+  text-align: center;
+}
+
+.App-logo {
+  height: 40vmin;
+  pointer-events: none;
+}
+
+@media (prefers-reduced-motion: no-preference) {
+  .App-logo {
+    animation: App-logo-spin infinite 20s linear;
+  }
+}
+
+.App-header {
+  background-color: #282c34;
+  min-height: 100vh;
+  display: flex;
+  flex-direction: column;
+  align-items: center;
+  justify-content: center;
+  font-size: calc(10px + 2vmin);
+  color: white;
+}
+
+.App-link {
+  color: #61dafb;
+}
+
+@keyframes App-logo-spin {
+  from {
+    transform: rotate(0deg);
+  }
+  to {
+    transform: rotate(360deg);
+  }
+}
diff --git a/Engine/Source/Programs/Horde/HordeDashboard/Dashboard/src/Plugins/SimpleTestReport/models/SimpleReport.ts b/Engine/Source/Programs/Horde/HordeDashboard/Dashboard/src/Plugins/SimpleTestReport/models/SimpleReport.ts
new file mode 100644
--- /dev/null
+++ b/Engine/Source/Programs/Horde/HordeDashboard/Dashboard/src/Plugins/SimpleTestReport/models/SimpleReport.ts
@@ -0,0 +1,29 @@
+// Copyright Epic Games, Inc. All Rights Reserved.
+
+export enum LogLevel {
+    Info = "Info",
+    Error = "Error",
+    Warning = "Warning",
+}
+
+export type SimpleReport = {
+    Description: string;
+    HasSucceeded: boolean;
+    Status: string;
+    ReportCreatedOn: string;
+    TotalDurationSeconds: number;
+    URLLink: string;
+    Logs: string[];
+    Errors: string[];
+    Warnings: string[];
+}
+
+export type HistoryItem = {
+    TestdataId: string;
+    Url: string;
+    Change: number;
+    HasSucceeded: boolean;
+    Date: string;
+    TotalDurationSeconds: number;
+    ErrorCount: number;
+}
\ No newline at end of file
diff --git a/Engine/Source/Programs/Horde/HordeDashboard/Dashboard/src/Plugins/UnrealAutomatedTests/models/UnrealAutomatedTests.ts b/Engine/Source/Programs/Horde/HordeDashboard/Dashboard/src/Plugins/UnrealAutomatedTests/models/UnrealAutomatedTests.ts
new file mode 100644
--- /dev/null
+++ b/Engine/Source/Programs/Horde/HordeDashboard/Dashboard/src/Plugins/UnrealAutomatedTests/models/UnrealAutomatedTests.ts
@@ -0,0 +1,61 @@
+// Copyright Epic Games, Inc. All Rights Reserved.
+
+export enum TestState {
+    Failed = "Fail",
+    Success = "Success",
+}
+
+export enum EventType {
+    Info = "Info",
+    Error = "Error",
+    Warning = "Warning",
+}
+
+export type TestEntryArtifact = {
+    Id: string;
+    Name: string;
+    Type: string;
+    Files: { Difference: string, Approved: string, Unapproved: string };
+}
+
+export type TestEntry = {
+    Filename: string;
+    LineNumber: number;
+    Timestamp: string;
+    Event: { Type: string, Message: string, Context: string, Artifact: string };
+}
+
+export type TestDetails = {
+    TestDisplayName: string;
+    FullTestPath: string;
+    State: string;
+    Errors: number;
+    Warnings: number;
+    Entries: TestEntry[];
+    Artifacts: TestEntryArtifact[];
+}
+
+export type TestResult = {
+    TestDisplayName: string;
+    FullTestPath: string;
+    State: string;
+    ArtifactName: string;
+}
+
+export type TestPassSummary = {
+    ClientDescriptor: string;
+    ReportURL: string;
+    FailedCount: number;
+    NotRunCount: number;
+    ReportCreatedOn: string;
+    SucceededCount: number;
+    SucceededWithWarningsCount: number;
+    TotalDurationSeconds: number;
+    Tests: TestResult[];
+}
+
+export type TestStateHistoryItem = {
+    TestdataId: string;
+    Change: number;
+    State: string;
+}
diff --git a/Engine/Source/Programs/Horde/HordeDashboard/Dashboard/src/backend/AgentStore.ts b/Engine/Source/Programs/Horde/HordeDashboard/Dashboard/src/backend/AgentStore.ts
new file mode 100644
--- /dev/null
+++ b/Engine/Source/Programs/Horde/HordeDashboard/Dashboard/src/backend/AgentStore.ts
@@ -0,0 +1,61 @@
+// Copyright Epic Games, Inc. All Rights Reserved.
+
+import { action, observable } from 'mobx';
+import { AgentData, PoolData } from './Api';
+import backend from '.';
+
+export class AgentStore {
+
+    @observable.ref
+	agents: AgentData[] = [];
+	
+	@observable.ref
+	pools: PoolData[] = [];
+
+	// don't need this for now
+    modifiedAfterDate: Date = new Date(0);
+
+    @action 
+    private _setPools(data: PoolData[]) {
+        this.pools = data;
+    }
+
+    @action
+    private _setAgents(data: AgentData[]) {
+        if(data.length !== 0) {
+            const toSet: AgentData[] = [...this.agents];
+            data.forEach(updatedAgent => {
+                const newUpdateIdx = toSet.findIndex(existingAgent => existingAgent.id === updatedAgent.id);
+                if(newUpdateIdx !== -1) {
+                    toSet[newUpdateIdx] = updatedAgent;
+                }
+                else {
+                    toSet.push(updatedAgent);
+                }
+            });
+            this.agents = toSet;
+        }
+	}
+
+    async update(slim = false): Promise<void> {
+        return new Promise<void>((resolve, reject) => {
+            const promises: any[] = [];
+            promises.push(backend.getAgents(this.modifiedAfterDate));
+            if(!slim) {
+                promises.push(backend.getPools());
+            }
+			Promise.all(promises).then(values => {
+                this.modifiedAfterDate = new Date();
+                this._setAgents(values[0]);
+                if(!slim){
+                    this._setPools(values[1]);
+                }
+                resolve();
+            }).catch(reason => {
+                console.error(`Error getting agents and pools: ${reason}`);
+            });
+        });
+    }
+}
+
+export const agentStore = new AgentStore();
\ No newline at end of file
diff --git a/Engine/Source/Programs/Horde/HordeDashboard/Dashboard/src/backend/Api.ts b/Engine/Source/Programs/Horde/HordeDashboard/Dashboard/src/backend/Api.ts
new file mode 100644
--- /dev/null
+++ b/Engine/Source/Programs/Horde/HordeDashboard/Dashboard/src/backend/Api.ts
@@ -0,0 +1,2689 @@
+// Copyright Epic Games, Inc. All Rights Reserved.
+
+//  Horde backend API types
+
+// The session type
+export enum SessionType {
+
+	// Execute a command
+	Job = 0,
+
+	// Upgrade the agent software
+	Upgrade = 1
+}
+
+export enum Priority {
+
+	// Lowest priority
+	Lowest = "Lowest",
+
+	// Below normal priority
+	BelowNormal = "BelowNormal",
+
+	// Normal priority
+	Normal = "Normal",
+
+	// Above normal priority
+	AboveNormal = "AboveNormal",
+
+	// High priority
+	High = "High",
+
+	// Highest priority
+	Highest = "Highest"
+}
+
+// The state of a jobstep
+export enum JobStepState {
+
+	// Waiting for dependencies of this step to complete (or paused)
+	Waiting = "Waiting",
+
+	// Ready to run, but has not been scheduled yet
+	Ready = "Ready",
+
+	// Dependencies of this step failed, so it cannot be executed
+	Skipped = "Skipped",
+
+	// There is an active instance of this step running
+	Running = "Running",
+
+	// This step has been run
+	Completed = "Completed",
+
+	// This step has been aborted
+	Aborted = "Aborted"
+
+}
+
+// Outcome of a jobstep run
+export enum JobStepOutcome {
+
+	// Finished with errors
+	Failure = "Failure",
+
+	// Finished with warnings
+	Warnings = "Warnings",
+
+	// Finished Succesfully
+	Success = "Success"
+}
+
+// The state of a particular run
+export enum JobStepBatchState {
+
+	// Waiting for dependencies of at least one jobstep to complete
+	Waiting = "Waiting",
+
+	// Ready to execute
+	Ready = "Ready",
+
+	// Currently running
+	Running = "Running",
+
+	// All steps have finished executing
+	Complete = "Complete"
+}
+
+/**Error code for a batch not being executed */
+export enum JobStepBatchError {
+
+	/** No error */
+	None = "None",
+
+	/** The stream for this job is unknown */
+	UnknownStream = "UnknownStream",
+
+	/** The given agent type for this batch was not valid for this stream */
+	UnknownAgentType = "UnknownAgentType",
+
+	/** The pool id referenced by the agent type was not found */
+	UnknownPool = "UnknownPool",
+
+	/** There are no agents in the given pool currently online */
+	NoAgentsInPool = "NoAgentsInPool",
+
+	/** There are no agents in this pool that are onlinbe */
+	NoAgentsOnline = "NoAgentsOnline",
+
+	/** Unknown workspace referenced by the agent type */
+	UnknownWorkspace = "UnknownWorkspace",
+
+	/** Cancelled */
+	Cancelled = "Cancelled",
+
+	/** Lost connection with the agent machine */
+	LostConnection = "LostConnection",
+
+	/** Lease terminated prematurely */
+	Incomplete = "Incomplete"
+
+}
+
+// The type of a template parameter
+export enum TemplateParameterType {
+
+	// An arbitrary string
+	String = "String"
+}
+
+export enum EventSeverity {
+
+	Unspecified = "Unspecified",
+	Information = "Information",
+	Warning = "Warning",
+	Error = "Error"
+}
+
+export enum AclActions {
+
+	// No actions allowed	
+	None = "None",
+
+	// Reading to the object	
+	Read = "Read",
+
+	// Modifying the object	
+	Write = "Write",
+
+	// Creating new objects	
+	Create = "Create",
+
+	// Deleting the object	
+	Delete = "Delete",
+
+	// Executing the object	
+	Execute = "Execute",
+
+	// Change permissions on the object	
+	ChangePermissions = "ChangePermissions"
+
+}
+
+export enum SubscriptonNotificationType {
+	// Slack notification
+	Slack = "Slack",
+}
+
+// aliases and extensions
+export type AgentData = GetAgentResponse
+export type LeaseData = GetAgentLeaseResponse
+export type SessionData = GetAgentSessionResponse
+export type PoolData = GetPoolResponse
+export type GroupData = GetGroupResponse
+export type NodeData = GetNodeResponse
+export type BatchData = GetBatchResponse
+export type EventData = GetLogEventResponse
+export type LogData = GetLogFileResponse
+export type ArtifactData = GetArtifactResponse
+export type AclData = GetAclResponse
+export type AclEntryData = GetAclEntryResponse
+export type SoftwareData = GetSoftwareResponse
+export type ScheduleData = GetScheduleResponse
+export type ChangeSummaryData = GetChangeSummaryResponse
+export type JobsTabData = GetJobsTabResponse
+export type JobsTabColumnData = GetJobsTabColumnResponse
+export type TimingInfo = GetTimingInfoResponse
+export type StepTimingInfo = GetStepTimingInfoResponse
+export type SubscriptionData = GetSubscriptionResponse
+
+export type IssueData = GetIssueResponse & {
+	events?: GetLogEventResponse[];
+}
+
+export type LabelData = GetLabelResponse & {
+	defaultLabel?: GetDefaultLabelStateResponse;
+}
+
+export type ProjectData = GetProjectResponse & {
+	streams?: StreamData[];
+}
+
+export type StreamData = GetStreamResponse & {
+	project?: ProjectData;
+	// full path as returned by server
+	fullname?: string;
+
+}
+
+export type TemplateData = GetTemplateResponse & {
+	ref?: GetTemplateRefResponse;
+}
+
+export type JobData = GetJobResponse & {
+	graphRef?: GetGraphResponse;
+}
+
+export type StepData = GetStepResponse & {
+
+	timing?: GetStepTimingInfoResponse;
+}
+
+export type TestData = GetTestDataResponse & {
+	data: object;
+}
+
+export enum LogType {
+	JSON = "JSON",
+	TEXT = "TEXT"
+}
+
+export enum LogLevel {
+	Informational = "Informational",
+	Warning = "Warning",
+	Error = "Error"
+}
+
+export type LogLine = {
+	time: string;
+	level: LogLevel;
+	message: string;
+	id?: number;
+	format?: string;
+	properties?: Record<string, string | Record<string, string | number>>;
+}
+
+export type LogLineData = {
+	format: LogType;
+	index: number;
+	count: number;
+	maxLineIndex: number;
+	lines?: LogLine[];
+};
+
+export type JobQuery = {
+	filter?: string;
+	streamId?: string;
+	name?: string;
+	template?: string[];
+	state?: string[];
+	outcome?: string[];
+	minChange?: number;
+	maxChange?: number;
+	preflightChange?: number;
+	includePreflight?: boolean,
+	preflightStartedByUser?: string,
+	minCreateTime?: string;
+	maxCreateTime?: string;
+	modifiedAfter?: string;
+	index?: number;
+	count?: number;
+};
+
+export type IssueQuery = {
+	jobId?: string;
+	batchId?: string;
+	stepId?: string;
+	label?: number;
+	streamId?: string;
+	change?: number;
+	minChange?: number;
+	maxChange?: number;
+	index?: number;
+	count?: number;
+	resolved?: boolean;
+}
+
+export type ScheduleQuery = {
+	streamId?: string;
+	filter?: string;
+}
+
+
+/* The severity of an issue */
+export enum IssueSeverity {
+	/** Unspecified severity */
+	Unspecified = "Unspecified",
+
+	/** This error represents a warning */
+	Warning = "Warning",
+
+	/** This issue represents an error */
+	Error = "Error"
+}
+
+/**Parameters to register a new agent */
+export type CreateAgentRequest = {
+
+	/**Friendly name for the agent */
+	name: string;
+
+	/**Whether the agent is currently enabled */
+	enabled: boolean;
+
+	/**Whether the agent is ephemeral (ie. should not be shown when inactive) */
+	ephemeral: boolean;
+
+	/**Per-agent override for the desired client version */
+	forceVersion?: string;
+
+	/**Pools for this agent */
+	pools?: string[];
+}
+
+/**Response from creating an agent */
+export type CreateAgentResponse = {
+
+	/**Unique id for the new agent */
+	id: string;
+
+	/**Bearer token for this agent to authorize itself with the server */
+	token: string;
+}
+
+/**Parameters to update an agent */
+export type UpdateAgentRequest = {
+
+	/**New name of the agent */
+	name?: string;
+
+	/**Whether the agent is currently enabled */
+	enabled?: boolean;
+
+	/**comment for the agent */
+	comment?: string;
+
+	/**Whether the agent is ephemeral */
+	ephemeral?: boolean;
+
+	/**Boolean to request a conform */
+	requestConform?: boolean;
+
+	requestRestart?: boolean;
+
+	/**Per-agent override for the desired client version */
+	forceVersion?: string;
+
+	/**Pools for this agent */
+	pools?: string[];
+
+	/**New ACL for this agent */
+	acl?: UpdateAclRequest;
+
+}
+
+/**Parameters to create a pool */
+export type CreatePoolRequest = {
+
+	/**New name of the pool */
+	name: string;
+
+	/**Arbitrary properties associated with this event */
+	properties?: { [key: string]: string };
+}
+
+
+/**Parameters to update a pool */
+export type UpdatePoolRequest = {
+
+	/**New name of the pool */
+	name?: string;
+
+	/**Arbitrary properties associated with this event */
+	properties?: { [key: string]: string };
+}
+
+/**Parameters to update a pool */
+export type BatchUpdatePoolRequest = {
+
+	/**ID of the pool to update  */
+	id: string;
+
+	/**New name of the pool */
+	name?: string;
+
+	/**Arbitrary properties associated with this event */
+	properties?: { [key: string]: string };
+
+	/**Soft-delete for this pool */
+	deleted?: boolean;
+}
+
+/**Response for queries to find a particular lease within an agent */
+export type GetAgentLeaseResponse = {
+
+	/**Identifier for the lease */
+	id: string;
+
+	/**Name of the lease */
+	name?: string;
+
+	/**Time at which the lease started (UTC) */
+	startTime: Date | string;
+
+	/**Time at which the lease ended (UTC) */
+	finishTime?: Date | string;
+
+	/**Whether this lease has started executing on the agent yet */
+	executing: boolean;
+
+	/**Details for this lease */
+	details?: { [key: string]: string };
+
+	/**agentId of lease */
+	agentId?: string;
+
+	/**What type of lease this is (parsed from Details array) */
+	type: string;
+
+	/**BatchId of this lease (parsed from Details array) */
+	batchId: string;
+
+	/**JobId of this lease (parsed from Details array) */
+	jobId: string;
+
+	/**logId of this lease (parsed from Details array) */
+	logId: string;
+}
+
+/**Information about an agent session */
+export type GetAgentSessionResponse = {
+
+	/**Unique id for this session */
+	id: string;
+
+	/**Start time for this session */
+	startTime: Date | string;
+
+	/**Finishing time for this session */
+	finishTime?: Date | string;
+
+	/**Capabilities of this agent */
+	capabilities?: GetAgentCapabilitiesResponse;
+
+	/**Version of the software running during this session */
+	version?: string;
+
+}
+
+export type GetAgentWorkspaceResponse = {
+
+	/**Server and port */
+	serverAndPort?: string;
+
+	/**username */
+	userName?: string;
+
+	/**password */
+	password?: string;
+
+	/**identifier */
+	identifier: string;
+
+	/**stream */
+	stream: string;
+
+	/**view */
+	view?: string[];
+
+	/**incremental */
+	bIncremental: boolean;
+}
+
+/**Information about an agent */
+export type GetAgentResponse = {
+
+	/**The agent's unique ID */
+	id: string;
+
+	/**Friendly name of the agent */
+	name: string;
+
+	/**Session ID */
+	sessionId: string;
+
+	/**Session expiry time */
+	sessionExpiresAt: Date | string;
+
+	/**Status of the agent */
+	online: boolean;
+
+	/**Whether the agent is currently enabled */
+	enabled: boolean;
+
+	/**Whether the agent is ephemeral */
+	ephemeral: boolean;
+
+	/**comment for the agent */
+	comment?: string;
+
+	/**The current client version */
+	version?: string;
+
+	/**Per-agent override for the desired client version */
+	forceVersion?: string;
+
+	/**Pools for this agent */
+	pools?: string[];
+
+	/**Capabilities of this agent */
+	capabilities?: GetAgentCapabilitiesResponse;
+
+	/**Array of active leases. */
+	leases?: GetAgentLeaseResponse[];
+
+	/**Per-object permissions */
+	acl?: GetAclResponse;
+
+	/**Last update time of the agent */
+	updateTime: Date | string;
+
+	/**Whether this agent is deleted */
+	deleted: boolean;
+
+	/** Whether agent is pending conform */
+	pendingConform: boolean;
+
+	/** Conform attempt count */
+	conformAttemptCount?: number;
+
+	/** Last conform time */
+	lastConformTime?: Date | string;
+
+	/** Next conform attempt time */
+	nextConformTime?: Date | string;
+
+	/** agent workspaces */
+	workspaces: GetAgentWorkspaceResponse[];
+
+}
+
+/**Information about an agent */
+export type GetPoolResponse = {
+
+	/**The agent's unique ID */
+	id: string;
+
+	/**Friendly name of the agent */
+	name: string;
+
+	/**Arbitrary properties associated with this event */
+	properties?: { [key: string]: string };
+
+}
+
+
+/**Parameters for creating a new event */
+export type CreateEventRequest = {
+
+	/**Time at which the event ocurred */
+	time: Date | string;
+
+	/**Severity of this event */
+	severity: EventSeverity;
+
+	/**Diagnostic text */
+	message: string;
+
+	/**Unique id of the log containing this event */
+	logId: string;
+
+	/**Minimum offset within the log containing this event */
+	minOffset: number;
+
+	/**Maximum offset within the log containing this event */
+	maxOffset: number;
+
+	/**Unique id of the agent associated with this event */
+	agentId?: string;
+
+	/**Unique id of the job associated with this event */
+	jobId?: string;
+
+	/**Unique id of the jobstep associated with this event */
+	jobStepId?: string;
+
+	/**Unique id of the jobstep associated with this event */
+	jobBatchId?: string;
+
+	/*The structured message data */
+	data?: { [key: string]: string };
+
+}
+
+/**Information about an uploaded event */
+export type GetLogEventResponse = {
+
+	/**Unique id of the log containing this event */
+	logId: string;
+
+	/**Severity of this event */
+	severity: EventSeverity;
+
+	/**Index of the first line for this event */
+	lineIndex: number;
+
+	/**Number of lines in the event */
+	lineCount: number;
+
+	/**The issue id associated with this event */
+	issueId: number | null;
+
+	/**The structured message data for this event */
+	lines: LogLine[];
+
+}
+
+
+
+/** Stats for a search */
+export type LogSearchStats = {
+
+	/// Number of blocks that were scanned
+	numScannedBlocks: number;
+
+	/// Number of bytes that had to be scanned for results
+	numScannedBytes: number;
+
+	/// Number of blocks that were skipped
+	numSkippedBlocks: number;
+
+	/// Number of blocks that had to be decompressed
+	numDecompressedBlocks: number;
+
+	/// Number of blocks that were searched but did not contain the search term
+	numFalsePositiveBlocks: number;
+}
+
+/** Response describing a log file */
+export type SearchLogFileResponse = {
+	/// List of line numbers containing the search text
+	lines: number[];
+
+	/// Stats for the search
+	stats?: LogSearchStats;
+}
+
+
+/**  Information about a group of nodes */
+export type CreateAggregateRequest = {
+
+	/** Name of the aggregate */
+	name: string;
+
+	/** Nodes which must be part of the job for the aggregate to be valid */
+	nodes: string[];
+
+}
+
+
+/** Response from creating a new aggregate */
+export type CreateAggregateResponse = {
+
+	/** Index of the first aggregate that was added	*/
+	firstIndex: number;
+}
+
+/**Information about a group of nodes */
+export type CreateLabelRequest = {
+
+	/**Category for this label */
+	category: string;
+
+	/**Name of the label */
+	name: string;
+
+	/**Nodes which must be part of the job for the label to be valid */
+	requiredNodes: string[];
+
+	/**Nodes which must be part of the job for the label to be valid */
+	includedNodes: string[];
+}
+
+
+
+/** Query selecting the base changelist to use */
+export type ChangeQueryRequest = {
+
+	/** The template id to query */
+	templateId?: string;
+
+	/** The target to query	*/
+	target?: string;
+}
+
+
+/**Parameters required to create a job */
+export type CreateJobRequest = {
+
+	/** The stream that this job belongs to */
+	streamId: string;
+
+	/** The template for this job */
+	templateId: string;
+
+	/** Name of the job */
+	name?: string;
+
+	/** The changelist number to build. Can be null for latest. */
+	change?: number;
+
+	/** Parameters to use when selecting the change to execute at. */
+	changeQuery?: ChangeQueryRequest;
+
+	/** The preflight changelist number */
+	preflightChange?: number;
+
+	/** Priority for the job */
+	priority?: Priority;
+
+	/** Whether to automatically submit the preflighted change on completion */
+	autoSubmit?: boolean;
+
+	/** Nodes for the new job */
+	groups?: CreateGroupRequest[];
+
+	/** Aggregates for the new job */
+	aggregates?: CreateAggregateRequest[];
+
+	/** Labels for the new job*/
+	labels?: CreateLabelRequest[];
+
+	/** Arguments for the job */
+	arguments?: string[];
+}
+
+/**Response from creating a new job */
+export type CreateJobResponse = {
+
+	/**Unique id for the new job */
+	id: string;
+}
+
+/**Updates an existing job */
+export type UpdateJobRequest = {
+
+	/** New name for the job */
+	name?: string;
+
+	/** New priority for the job */
+	priority?: Priority;
+
+	/**  Set whether the job should be automatically submitted or not */
+	autoSubmit?: boolean;
+
+	/** Mark this job as aborted */
+	aborted?: boolean;
+
+	/** New list of arguments for the job. Only -Target= arguments can be modified after the job has started.  */
+	arguments?: string[];
+
+	/**  Custom permissions for this object */
+	acl?: UpdateAclRequest;
+}
+
+export enum ReportPlacement {
+	/** On a panel of its own */
+	Panel = "Panel",
+	/** In the summary panel */
+	Summary = "Summary"
+}
+
+export enum ReportScope {
+	Job = "Job",
+	Step = "Step"
+}
+
+/**Information about a report associated with a job*/
+export type GetReportResponse = {
+
+	// Report scope
+	scope: ReportScope;
+
+	// Report placement
+	placement: ReportPlacement;
+
+	/** Name of the report */
+	name: string;
+
+	/** The artifact id */
+	artifactId: string;
+}
+
+
+/**Information about a job */
+export type GetJobResponse = {
+
+	/**Unique Id for the job */
+	id: string;
+
+	/**Unique id of the stream containing this job */
+	streamId: string;
+
+	/**Name of the job */
+	name: string;
+
+	/** The changelist number to build */
+	change?: number;
+
+	/** The preflight changelist number */
+	preflightChange?: number;
+
+	/** The template type */
+	templateId?: string;
+
+	/** Hash of the actual template data */
+	templateHash?: string;
+
+	/** Hash of the graph for this job */
+	graphHash?: string;
+
+	/** The user that started this job */
+	startedByUser?: string;
+
+	/** The user that aborted this job */
+	abortedByUser?: string;
+
+	/** The roles to impersonate when executing this job */
+	roles: string[];
+
+	/** Priority of the job */
+	priority: Priority;
+
+	/** Whether the change will automatically be submitted or not */
+	autoSubmit?: boolean;
+
+	/** The submitted changelist number */
+	autoSubmitChange?: number;
+
+	/**  Message produced by trying to auto-submit the change */
+	autoSubmitMessage?: string;
+
+	/**Time that the job was created */
+	createTime: Date | string;
+
+	/** The global job state */
+	state: JobState;
+
+	/** Array of jobstep batches */
+	batches?: GetBatchResponse[];
+
+	/**  List of labels */
+	labels?: GetLabelStateResponse[];
+
+	/** The default label, containing the state of all steps that are otherwise not matched. */
+	defaultLabel?: GetDefaultLabelStateResponse;
+
+	/** List of reports */
+	reports?: GetReportResponse[];
+
+	/**  Parameters for the job */
+	arguments: string[];
+
+	/**The last update time for this job*/
+	updateTime: Date | string;
+
+	/**  Custom permissions for this object */
+	acl?: GetAclResponse;
+
+}
+
+/**Request used to update a jobstep */
+export type UpdateStepRequest = {
+
+	/**The new jobstep state */
+	state?: JobStepState;
+
+	/**Outcome from the jobstep */
+	outcome?: JobStepOutcome;
+
+	/**If the step has been requested to abort */
+	abortRequested?: boolean;
+
+	/**Specifies the log file id for this step */
+	logId?: string;
+
+	/**Whether the step should be re-run */
+	retry?: boolean;
+
+	/**New priority for this step */
+	priority?: Priority;
+
+	/**Properties to set. Any entries with a null value will be removed. */
+	properties?: { [key: string]: string | null };
+}
+
+/**Returns information about a jobstep */
+export type GetStepResponse = {
+
+	/**The unique id of the step */
+	id: string;
+
+	/**Index of the node which this jobstep is to execute */
+	nodeIdx: number;
+
+	/**Current state of the job step. This is updated automatically when runs complete. */
+	state: JobStepState;
+
+	/**Current outcome of the jobstep */
+	outcome: JobStepOutcome;
+
+	/** If the step has been requested to abort	*/
+	abortRequested?: boolean;
+
+	/** Name of the user that requested the abort of this step */
+	abortByUser?: string;
+
+	/** Name of the user that requested this step be run again */
+	retryByUser?: string;
+
+	/**The log id for this step */
+	logId?: string;
+
+	/** Time at which the batch was ready (UTC). */
+	readyTime?: Date | string;
+
+	/**Time at which the run started (UTC). */
+	startTime?: Date | string;
+
+	/**Time at which the run finished (UTC) */
+	finishTime: Date | string;
+
+	/**User-defined properties for this jobstep. */
+	properties: { [key: string]: string };
+
+}
+
+//**Returns information about test data */
+export type GetTestDataResponse = {
+
+	/**The unique id of the test data */
+	id: string;
+
+	/**The key associated with this test data */
+	key: string;
+
+	/**The change id related to the job */
+	change: number;
+
+	/**The job id of this test data */
+	jobId: string;
+
+	/**The step id of this test data */
+	stepId: string;
+
+	/**The stream id of this test data */
+	streamId: string;
+
+	/**The template ref id of the related job */
+	templateRefId: string;
+}
+
+export type UpdateStepResponse = {
+	/** If a new step is created (due to specifying the retry flag), specifies the batch id */
+	batchId?: string;
+
+	/** If a step is retried, includes the new step id */
+	stepId?: string
+}
+
+/**Request to update a jobstep batch */
+export type UpdateBatchRequest = {
+
+	/**The new log file id */
+	logId?: string;
+
+	/**The state of this batch */
+	state?: JobStepBatchState;
+}
+
+/**Information about a jobstep batch */
+export type GetBatchResponse = {
+
+	/**Unique id for this batch */
+	id: string;
+
+	/**The unique log file id */
+	logId?: string;
+
+	/**Index of the group being executed */
+	groupIdx: number;
+
+	/**The state of this batch */
+	state: JobStepBatchState;
+
+	/**Error code for this batch */
+	error: JobStepBatchError;
+
+	/**Steps within this run */
+	steps: GetStepResponse[];
+
+	/**The agent assigned to execute this group */
+	agentId?: string;
+
+	/**The agent session holding this lease */
+	sessionId?: string;
+
+	/**The lease that's executing this group */
+	leaseId?: string;
+
+	/**The priority of this batch */
+	weightedPriority: number;
+
+	/**Time at which the group started (UTC). */
+	startTime?: Date | string;
+
+	/**Time at which the group finished (UTC) */
+	finishTime?: Date | string;
+
+}
+
+/**Describes the history of a step */
+export interface GetJobStepRefResponse {
+	/**The job id */
+	jobId: string;
+
+	/**The batch containing the step */
+	batchId: string;
+
+	/**The step identifier */
+	stepId: string;
+
+	/**The change number being built */
+	change: number;
+
+	/**The step log id */
+	logId?: string;
+
+	/**The pool id */
+	poolId?: string;
+
+	/** The agent id */
+	agentId?: string;
+
+	/**Outcome of the step, once complete. */
+	outcome?: JobStepOutcome;
+
+	/**Time at which the step started. */
+	startTime: Date | string;
+
+	/**Time at which the step finished. */
+	finishTime?: Date | string;
+
+}
+
+export interface GetJobStepTraceResponse {
+
+	Name: string;
+
+	Start: number
+
+	Finish: number
+
+	Service?: string;
+
+	Resource?: string;
+
+	Children: GetJobStepTraceResponse[];
+}
+
+/**Parameters required to create log file */
+export type CreateLogFileRequest = {
+	/**Job Id this log file belongs to */
+	jobId: string;
+
+}
+
+/**Response from creating a log file */
+export type CreateLogFileResponse = {
+	/**Unique id for this log file */
+	id: string;
+
+}
+
+/**Response describing a log file */
+export type GetLogFileResponse = {
+
+	/**Unique id of the log file */
+	id: string;
+
+	/** Unique id of the job for this log file */
+	jobId: string;
+
+	/** Type of events stored in this log */
+	type: LogType;
+
+	/** Length of the log, in bytes */
+	length: number;
+
+	/** Number of lines in the file	*/
+	lineCount: number;
+
+	/**Per-object permissions */
+	acl?: GetAclResponse;
+
+}
+
+/**Response describing an artifact */
+export type GetArtifactResponse = {
+
+	/**Unique id of the artifact */
+	id: string;
+
+	/** Unique id of the job for this artifact */
+	jobId: string;
+
+	/** Unique id of the job for this artifact */
+	stepId?: string;
+
+	/** Download code for this artifact */
+	code?: string;
+
+	/** Name of the artifact */
+	name: string;
+
+	/** MimeType of the artifact	*/
+	mimeType: string;
+
+	/** Length of the artifact, in bytes */
+	length: number;
+
+	/**Per-object permissions */
+	acl?: GetAclResponse;
+
+}
+
+/**Parameters request a zip file for artifacts */
+export type GetArtifactZipRequest = {
+
+	/**JobId to get all artifacts for */
+	jobId?: string;
+
+	/**StepId to get all artifacts for */
+	stepId?: string;
+
+	/**Individual file name if we've just got one file */
+	fileName?: string;
+
+	/** Whether or not we're a forced single file download (versus a link click for one file) */
+	isForcedDownload?: boolean;
+
+	/** Whether or not to open the resulting link in place or in a new tab */
+	isOpenNewTab?: boolean;
+
+	/**List of arbitrary artifact Ids to generate a zip for */
+	artifactIds?: string[];
+}
+
+
+/**Parameters required to update a log file */
+export type UpdateLogFileRequest = {
+
+	/** New permissions */
+
+	acl?: UpdateAclRequest;
+}
+
+
+export type CreateNodeRequest = {
+
+	/**The name of this node  */
+	name: string;
+
+	/**Indices of nodes which must have succeeded for this node to run */
+	inputDependencies?: string[];
+
+	/**Indices of nodes which must have completed for this node to run */
+	orderDependencies?: string[];
+
+	/**The priority of this node */
+	priority?: Priority;
+
+	/**This node can be run multiple times */
+	allowRetry?: boolean;
+
+	/**This node can start running early, before dependencies of other nodes in the same group are complete */
+	runEarly?: boolean;
+
+	/**Credentials required for this node to run. This dictionary maps from environment variable names to a credential property in the format 'CredentialName.PropertyName' */
+	credentials?: { [key: string]: string };
+
+	/**Properties for this node */
+	properties?: { [key: string]: string };
+}
+
+/**Information about a group of nodes */
+export type CreateGroupRequest = {
+
+	/**The executor to use for this group */
+	executor: string;
+
+	/**The type of agent to execute this group */
+	agentType: string;
+
+	/**Nodes in the group */
+	nodes: CreateNodeRequest[];
+
+}
+
+/**Information required to create a node */
+export type GetNodeResponse = {
+
+	/**The name of this node  */
+	name: string;
+
+	/**Indices of nodes which must have succeeded for this node to run */
+	inputDependencies: string[];
+
+	/**Indices of nodes which must have completed for this node to run */
+	orderDependencies: string[];
+
+	/**The priority of this node */
+	priority: Priority;
+
+	/**Whether this node can be retried */
+	allowRetry: boolean;
+
+	/**This node can start running early, before dependencies of other nodes in the same group are complete */
+	runEarly: boolean;
+
+	/**Sets this node as a target to be built */
+	target: boolean;
+
+	/**Expected time to execute this node based on historical trends */
+	averageDuration: number;
+
+	/**Aggregates that this node belongs do */
+	aggregates?: string[];
+
+	/**Properties for this node */
+	properties: { [key: string]: string };
+
+}
+
+/**Information about a group of nodes */
+export type GetGroupResponse = {
+
+	/**The executor to use for this group */
+	executor: string;
+
+	/**The type of agent to execute this group */
+	agentType: string;
+
+	/**Nodes in the group */
+	nodes: GetNodeResponse[];
+}
+
+/**Request to update a node */
+export type UpdateNodeRequest = {
+
+	/**The priority of this node */
+	priority?: Priority;
+
+}
+
+/**Parameters to create a new project */
+export type CreateProjectRequest = {
+
+	/**Name for the new project */
+	name: string;
+
+	/**Properties for the new project */
+	properties?: { [key: string]: string };
+}
+
+/**Response from creating a new project */
+export type CreateProjectResponse = {
+
+	/**Unique id for the new project */
+	id: string;
+}
+
+/**Parameters to update a project */
+export type UpdateProjectRequest = {
+
+	/**Optional new name for the project */
+	name?: string;
+
+	/**Properties to update for the project. Properties set to null will be removed. */
+	properties?: { [key: string]: string };
+
+
+	/**Custom permissions for this object */
+	acl?: UpdateAclRequest;
+}
+
+/**Information about a stream within a project */
+export interface GetProjectStreamResponse {
+	/**The stream id */
+	id: string;
+
+	/**The stream name */
+	name: string;
+}
+
+
+/**Information about a category to display for a stream */
+export type GetProjectCategoryResponse = {
+	/**Heading for this column */
+	name: string;
+
+	/**Index of the row to display this category on */
+	row: number;
+
+	/**Whether to show this category on the nav menu */
+	showOnNavMenu: boolean;
+
+	/**Patterns for stream names to include */
+	includePatterns: string[];
+
+	/**Patterns for stream names to exclude */
+	excludePatterns: string[];
+
+	/**Streams to include in this category */
+	streams: string[];
+}
+
+/**Response describing a project */
+export type GetProjectResponse = {
+	/**Unique id of the project */
+	id: string;
+
+	/**Name of the project */
+	name: string;
+
+	/**Order to display this project on the dashboard */
+	order: number;
+
+	/**List of streams that are in this project */
+	// streams?: ProjectStreamData[];
+
+	/**List of stream categories to display */
+	categories?: GetProjectCategoryResponse[];
+
+	/**Properties for this project. */
+	properties: { [key: string]: string };
+
+	/**Custom permissions for this object */
+	acl?: GetAclResponse;
+}
+
+/**Parameters to create a new schedule */
+export type CreateSchedulePatternRequest = {
+
+	/**Days of the week to run this schedule on. If null, the schedule will run every day. */
+	daysOfWeek?: string[];
+
+	/**Time during the day for the first schedule to trigger. Measured in minutes from midnight. */
+	minTime: number;
+
+	/**Time during the day for the last schedule to trigger. Measured in minutes from midnight. */
+	maxTime?: number;
+
+	/**Interval between each schedule triggering */
+	interval?: number;
+}
+
+/**Parameters to create a new schedule */
+export type CreateScheduleRequest = {
+
+	/**Name for the new schedule */
+	name: string;
+
+	/**Whether the schedule should be enabled */
+	enabled: boolean;
+
+	/**Maximum number of builds that can be active at once */
+	maxActive: number;
+
+	/**The stream to run this schedule in */
+	streamId: string;
+
+	/**The template job to execute */
+	templateId: string;
+
+	/**Parameters for the template */
+	templateParameters: { [key: string]: string };
+
+	/**New patterns for the schedule */
+	patterns: CreateSchedulePatternRequest[];
+
+}
+
+/**Response from creating a new schedule */
+export type CreateScheduleResponse = {
+
+	/**Unique id for the new schedule */
+	id: string;
+}
+
+/**Parameters to update a schedule */
+export type UpdateScheduleRequest = {
+
+	/**Optional new name for the schedule */
+	name?: string;
+
+	/**Whether to enable the schedule */
+	enabled?: boolean;
+
+	/**Maximum number of builds that can be active at once */
+	maxActive: number;
+
+	/**The template job to execute */
+	templateHash?: string;
+
+	/**Parameters for the template */
+	templateParameters?: { [key: string]: string };
+
+	/**New patterns for the schedule */
+	patterns?: CreateSchedulePatternRequest[];
+
+	/** Custom permissions for this object */
+	acl?: UpdateAclRequest;
+}
+
+/**Information about a schedule pattern */
+export type GetSchedulePatternResponse = {
+
+	/**Days of the week to run this schedule on. If null, the schedule will run every day. */
+	daysOfWeek?: string[];
+
+	/**Time during the day for the first schedule to trigger. Measured in minutes from midnight. */
+	minTime: number;
+
+	/**Time during the day for the last schedule to trigger. Measured in minutes from midnight. */
+	maxTime?: number;
+
+	/**Interval between each schedule triggering */
+	interval?: number;
+
+}
+
+/**Response describing a schedule */
+export type GetScheduleResponse = {
+
+	/**Whether the schedule is currently enabled */
+	enabled: boolean;
+
+	/**Maximum number of scheduled jobs at once */
+	maxActive: number;
+
+	/**Maximum number of changes the schedule can fall behind head revision. If greater than zero, builds will be triggered for every submitted changelist until the backlog is this size. */
+	maxChanges: number;
+
+	/**The template job to execute */
+	templateId: string;
+
+	/**Parameters for the template */
+	templateParameters: { [key: string]: string };
+
+	/**New patterns for the schedule */
+	patterns: GetSchedulePatternResponse[];
+
+	/* Last changelist number that this was triggered for */
+	lastTriggerChange: number;
+
+	/** Last time that the schedule was triggered */
+	lastTriggerTime: Date | string;
+
+	/** List of active jobs */
+	activeJobs: string[];
+
+}
+
+/**Response describing when a schedule is expected to trigger */
+export type GetScheduleForecastResponse = {
+
+	/**Next trigger times */
+	times: Date | string[];
+
+}
+
+
+/**Parameters for creating a new software archive */
+export type CreateSoftwareRequest = {
+
+	/**Whether this software should be the default */
+	default: boolean;
+}
+
+/**Information about a client version */
+export type CreateSoftwareResponse = {
+
+	/**The software id */
+	id: string;
+
+}
+
+/**Parameters for updating a software archive */
+export type UpdateSoftwareRequest = {
+
+	/**Whether this software should be the default */
+	default: boolean;
+}
+
+/**Information about an uploaded software archive */
+export type GetSoftwareResponse = {
+
+	/**Unique id for this enty */
+	id: string;
+
+	/**Name of the user that created this software */
+	uploadedByUser?: string;
+
+	/**Time at which the client was created */
+	uploadedAtTime: Date | string;
+
+	/**Name of the user that created this software */
+	madeDefaultByUser?: string;
+
+	/**Time at which the client was made default. */
+	madeDefaultAtTime?: Date | string;
+
+}
+
+
+/**Parameters to create a new stream */
+export type CreateStreamRequest = {
+
+	/**Unique id for the project */
+	projectId: string;
+
+	/**Name for the new stream */
+	name: string;
+
+	/**Properties for the new stream */
+	properties?: { [key: string]: string };
+
+
+
+}
+
+/**Response from creating a new stream */
+export type CreateStreamResponse = {
+
+	/**Unique id for the new stream */
+	id: string;
+
+}
+
+/**Parameters to update a stream */
+export type UpdateStreamRequest = {
+
+	/**Optional new name for the stream */
+	name?: string;
+
+	/**Properties to update on the stream. Properties with a value of null will be removed. */
+	properties?: { [key: string]: string | null };
+
+	/** Custom permissions for this object */
+	acl?: UpdateAclRequest;
+}
+
+/**Mapping from a BuildGraph agent type to a se t of machines on the farm */
+export type GetAgentTypeResponse = {
+
+	/**Pool of agents to use for this agent type */
+	pool: string;
+
+	/**Name of the workspace to sync */
+	workspace?: string;
+
+	/**Pool of agents to use to execute this work */
+	requirements?: GetAgentRequirementsResponse;
+
+	/**Path to the temporary storage dir */
+	tempStorageDir?: string;
+
+	/**Environment variables to be set when executing the job */
+	environment?: { [key: string]: string };
+
+}
+
+/**Information about a workspace type */
+export type GetWorkspaceTypeResponse = {
+
+	/**The Perforce server and port (eg. perforce:1666) */
+	serverAndPort?: string;
+
+	/**User to log into Perforce with (defaults to buildmachine) */
+	userName?: string;
+
+	/**Identifier to distinguish this workspace from other workspaces. Defaults to the workspace type name. */
+	identifier?: string;
+
+	/**Override for the stream to sync */
+	stream?: string;
+
+	/**Custom view for the workspace */
+	view?: string[];
+
+	/**Whether to use an incrementally synced workspace */
+	incremental: boolean;
+}
+
+
+/**Information about a template in this stream */
+export type GetTemplateRefResponse = {
+	/**Unique id of this template ref */
+	id: string;
+
+	/**Hash of the template definition */
+	hash: string;
+
+	/// The schedule for this ref
+	schedule?: GetScheduleResponse;
+
+	name: string;
+
+}
+
+/** Specifies defaults for running a preflight */
+export type DefaultPreflightRequest = {
+
+	/** The template id to query */
+	templateId?: string;
+
+	/**  The last successful job type to use for the base changelist */
+	changeTemplateId?: string;
+}
+
+/**Response describing a stream */
+export type GetStreamResponse = {
+
+	/**Unique id of the stream */
+	id: string;
+
+	/**Unique id of the project containing this stream */
+	projectId: string;
+
+	/**Name of the stream */
+	name: string;
+
+	/**List of tabs to display for this stream*/
+	tabs: GetStreamTabResponse[];
+
+	/**  Default template to use for preflights (deprecated)*/
+	defaultPreflightTemplate?: string;
+
+	/** Default template for running preflights */
+	defaultPreflight?: DefaultPreflightRequest;
+
+	/**Map of agent name to type */
+	agentTypes: { [key: string]: GetAgentTypeResponse };
+
+	/**Map of workspace name to type */
+	workspaceTypes?: { [key: string]: GetWorkspaceTypeResponse };
+
+	/**Templates for jobs in this stream */
+	templates: GetTemplateRefResponse[];
+
+	/**Properties for this stream */
+	properties: { [key: string]: string };
+
+	/** Custom permissions for this object */
+	acl?: GetAclResponse;
+
+}
+
+/**Information about a template parameter */
+export type CreateTemplateParameterRequest = {
+
+	/**Name of the parameter */
+	name: string;
+
+	/**Type of the template parameter */
+	type: TemplateParameterType;
+
+	/**Default value for this parameter */
+	default: string;
+
+	/**Whether this parameter is required */
+	required: boolean;
+}
+
+/**Parameters to create a new template */
+export type CreateTemplateRequest = {
+
+	/**Name for the new template */
+	name: string;
+
+	/**Default priority for this job */
+	priority: Priority;
+
+	/**Whether to allow preflights of this template */
+	allowPreflights: boolean;
+
+	/**Array of nodes for this job */
+	groups: CreateGroupRequest[];
+
+	/**List of aggregates for this template */
+	aggregates: CreateAggregateRequest[];
+
+	/**List of labels  for this template */
+	labels: CreateLabelRequest[];
+
+	/**Parameter names for this template. These will be assigned to properties of the job at startup. */
+	parameters: CreateTemplateParameterRequest[];
+
+	/**Properties for the new template */
+	properties: { [key: string]: string };
+}
+
+/**Describes how to render a group parameter */
+export enum GroupParameterStyle {
+	/**Separate tab on the form */
+	Tab = "Tab",
+
+	/**Section with heading */
+	Section = "Section"
+}
+
+/**Style of list parameter */
+export enum ListParameterStyle {
+	/**Regular drop-down list. One item is always selected. */
+	List = "List",
+
+	/**Drop-down list with checkboxes */
+	MultiList = "MultiList",
+
+	/**Tag picker from list of options */
+	TagPicker = "TagPicker"
+}
+
+export enum ParameterType {
+	Bool = "Bool",
+	List = "List",
+	Text = "Text",
+	Group = "Group"
+}
+
+
+/**Base class for template parameters */
+export type ParameterData = {
+	type: ParameterType;
+}
+
+/**Used to group a number of other parameters */
+export type GroupParameterData = ParameterData & {
+
+	/**Label to display next to this parameter */
+	label: string;
+
+	/**How to display this group */
+	style: GroupParameterStyle;
+
+	/**List of child parameters */
+	children: ParameterData[];
+}
+
+/**Free-form text entry parameter */
+export type TextParameterData = ParameterData & {
+
+	/**Name of the parameter associated with this parameter. */
+	label: string;
+
+	/**Argument to pass to the executor */
+	argument: string;
+
+	/**Default value for this argument */
+	default: string;
+
+	/**Hint text for this parameter */
+	hint?: string;
+
+	/**Regex used to validate this parameter */
+	validation?: string;
+
+	/**Message displayed if validation fails, informing user of valid values. */
+	validationError?: string;
+
+	/**Tool-tip text to display */
+	toolTip?: string;
+
+}
+
+/**Possible option for a list parameter */
+export type ListParameterItemData = ParameterData & {
+
+	/**Optional group heading to display this entry under, if the picker style supports it. */
+	group?: string;
+
+	/**Name of the parameter associated with this list. */
+	text: string;
+
+	/// <summary>
+	/// Argument to pass with this parameter, if enabled
+	/// </summary>
+	argumentIfEnabled?: string;
+
+	/// <summary>
+	/// Argument to pass with this parameter, if disabled
+	/// </summary>
+	argumentIfDisabled?: string;
+
+	/**Whether this item is selected by default */
+	default: boolean;
+
+}
+
+/**Allows the user to select a value from a constrained list of choices */
+export type ListParameterData = ParameterData & {
+	/**Label to display next to this parameter. Defaults to the parameter name. */
+	label: string;
+
+	/**The type of list parameter */
+	style: ListParameterStyle;
+
+	/**List of values to display in the list */
+	items: ListParameterItemData[];
+
+	/**Tool-tip text to display */
+	toolTip?: string;
+}
+
+/**Allows the user to toggle an option on or off */
+export type BoolParameterData = ParameterData & {
+	/**Name of the parameter associated with this parameter. */
+	label: string;
+
+	/**Value if enabled */
+	argumentIfEnabled?: string;
+
+	/**Value if disabled */
+	argumentIfDisabled?: string;
+
+	/**Whether this argument is enabled by default */
+	default: boolean;
+
+	/**Tool-tip text to display */
+	toolTip?: string;
+
+}
+
+
+/**Response describing a template */
+export type GetTemplateResponse = {
+
+	/**Unique id of the template */
+	id: string;
+
+	/**Name of the template */
+	name: string;
+
+	/**Default priority for this job */
+	priority?: Priority;
+
+	/**Whether to allow preflights of this template */
+	allowPreflights: boolean;
+
+	/**List of node groups for this job */
+	groups: CreateGroupRequest[];
+
+	/**List of template aggregates */
+	aggregates: CreateAggregateRequest[];
+
+	/**Parameters for the job. */
+	arguments: string[];
+
+	/**List of parameters for this template */
+	parameters: ParameterData[];
+
+	/**Custom permissions for this object */
+	acl?: GetAclResponse;
+}
+
+/** Information about a commit */
+export type GetChangeSummaryResponse = {
+
+	/**  The source changelist number */
+	number: number;
+
+	/**  Name of the user that authored this change */
+	author: string;
+
+	/**  The description text */
+	description: string;
+
+}
+
+
+/**Information about a device attached to this agent */
+export type GetDeviceCapabilitiesResponse = {
+
+	/**Logical name of this device */
+	name: string;
+
+	/**Required properties for this device, in the form "KEY=VALUE" */
+	properties?: string[];
+
+	/**Required resources for this node. If null, the node will assume exclusive access to the device. */
+	resources?: { [key: string]: number };
+
+}
+
+/**Information about the capabilities of this agent */
+export type GetAgentCapabilitiesResponse = {
+
+	/**Information about the devices required for this node to run */
+	devices?: GetDeviceCapabilitiesResponse[];
+
+	/**Global agent properties for this node */
+	properties?: string[];
+
+}
+
+/**Information about the device requirements of a node */
+export type CreateDeviceRequirementsRequest = {
+
+	/**Logical name of this device */
+	name: string;
+
+	/**Required properties for this device, in the form "KEY=VALUE" */
+	properties?: string[];
+
+	/**Required resources for this node. If null, the node will assume exclusive access to the device. */
+	resources?: { [key: string]: number };
+
+}
+
+/**Information about the agent requirements of node */
+export type CreateAgentRequirementsRequest = {
+
+	/**Information about the devices required for this node to run */
+	devices?: CreateDeviceRequirementsRequest[];
+
+	/**Global agent properties for this node */
+	properties?: string[];
+
+	/**Whether the agent can be shared with another job */
+	shared: boolean;
+}
+
+/**Information about the device requirements of a node */
+export type GetDeviceRequirementsResponse = {
+
+	/**Logical name of this device */
+	name: string;
+
+	/**Required properties for this device, in the form "KEY=VALUE" */
+	properties?: string[];
+
+	/**Required resources for this node. If null, the node will assume exclusive access to the device. */
+	resources?: { [key: string]: number };
+
+}
+
+/**Information about the agent requirements of node */
+export type GetAgentRequirementsResponse = {
+
+	/**Information about the devices required for this node to run */
+	devices?: GetDeviceRequirementsResponse[];
+
+	/**Global agent properties for this node */
+	properties?: string[];
+
+	/**Whether the agent can be shared with another job */
+	shared: boolean;
+
+}
+
+/**Individual entry in an ACL */
+export type UpdateAclEntryRequest = {
+	/**Name of the user or group */
+	roles: string[];
+
+	/**Array of actions to allow */
+	allow: string[];
+}
+
+/**Parameters to update an ACL */
+export type UpdateAclRequest = {
+
+	/**Entries to replace the existing ACL */
+	entries: UpdateAclEntryRequest[] | null;
+
+	/**Whether to inherit permissions from the parent ACL */
+	inheritPermissions: boolean | null;
+
+
+	/** List of exceptions to the inherited setting */
+	exceptions?: string[];
+}
+
+/**Individual entry in an ACL */
+export type GetAclEntryResponse = {
+	/**Names of the user or group */
+	roles: string[];
+
+	/**Array of actions to allow */
+	allow: string[];
+}
+
+/**Information about an ACL */
+export type GetAclResponse = {
+	/**Entries to replace the existing ACL */
+	entries: GetAclEntryResponse[];
+
+	/**Whether to inherit permissions from the parent entity */
+	inheritPermissions: boolean;
+}
+
+// Labels
+
+/**Type of a column in a jobs tab */
+export enum JobsTabColumnType {
+	/**Contains labels */
+	Labels = "Labels",
+
+	/**Contains parameters */
+	Parameter = "Parameter"
+}
+
+/**Describes a column to display on the jobs page */
+export type GetJobsTabColumnResponse = {
+
+	type: JobsTabColumnType;
+
+	/**Heading for this column */
+	heading: string;
+
+	/**Relative width of this column. */
+	relativeWidth?: number;
+}
+
+export type GetJobsTabLabelColumnResponse = GetJobsTabColumnResponse & {
+
+	/**Category of labels to display in this column. If null, includes any label not matched by another column. */
+	category?: string;
+
+}
+
+export type GetJobsTabParameterColumnResponse = GetJobsTabColumnResponse & {
+	/** Parameter to show in this column */
+	parameter?: string;
+}
+
+export enum TabType {
+
+	Jobs = "Jobs"
+
+}
+
+/**Information about a page to display in the dashboard for a stream */
+export type GetStreamTabResponse = {
+
+	/**Title of this page */
+	title: string;
+
+	type: TabType;
+
+};
+
+/**Describes a job page */
+export type GetJobsTabResponse = GetStreamTabResponse & {
+
+	/**  Whether to show names on the page */
+	showNames: boolean;
+
+
+	/** List of templates to show on the page */
+	templates?: string[];
+
+	/**Columns to display for different types of labels */
+	columns?: GetJobsTabColumnResponse[];
+
+};
+
+/** State of an label */
+export enum LabelState {
+
+	/** label is not currently being built (no required nodes are present)*/
+	Unspecified = "Unspecified",
+
+	/** Steps are still running */
+	Running = "Running",
+
+	/** All steps are complete */
+	Complete = "Complete"
+}
+
+/// Outcome of an aggregate
+export enum LabelOutcome {
+
+	/** Aggregate is not currently being built */
+	Unspecified = "Unspecified",
+
+	/** A step dependency failed */
+	Failure = "Failure",
+
+	/** A dependency finished with warnings */
+	Warnings = "Warnings",
+
+	/** Successful */
+	Success = "Success"
+}
+
+
+/** State of the job */
+export enum JobState {
+
+	/** Waiting for resources*/
+	Waiting = "Waiting",
+
+	/** Currently running one or more steps*/
+	Running = "Running",
+
+	/** All steps have completed */
+	Complete = "Complete"
+}
+
+
+/**Information about a label */
+export type GetLabelResponse = {
+
+	/**Category of the aggregate */
+	category: string;
+
+	/**Label for this aggregate */
+	name: string;
+
+	/**Label for this aggregate, currently mapped to name property on server */
+	dashboardName?: string;
+
+	/**Name to show for this label in UGS */
+	ugsName?: string;
+
+	/** Project to display this label for in UGS */
+	ugsProject?: string;
+
+	/**Nodes which must be part of the job for the aggregate to be shown */
+	requiredNodes: string[];
+
+	/**Nodes to include in the status of this aggregate, if present in the job */
+	includedNodes: string[];
+}
+
+
+/**Information about an aggregate */
+export type GetAggregateResponse = {
+
+	/**Name of the aggregate */
+	name: string;
+
+	/**Nodes which must be part of the job for the aggregate to be shown */
+	nodes: string[];
+
+}
+
+/**Information about a graph */
+export type GetGraphResponse = {
+
+	/**The hash of the graph */
+	hash: string;
+
+	/**Array of nodes for this job */
+	groups?: GetGroupResponse[];
+
+	/**List of aggregates */
+	namedAggregates?: GetAggregateResponse[];
+
+	/**List of labels for the graph */
+	labels?: GetLabelResponse[];
+
+}
+
+/**The timing info for  */
+export type GetJobTimingResponse = {
+
+	/**Timing info for each step */
+	steps: { [key: string]: GetStepTimingInfoResponse };
+
+	/**Timing information for each label */
+	labels: GetLabelTimingInfoResponse[];
+}
+
+/**Information about the timing info for a label */
+export type GetLabelTimingInfoResponse = GetTimingInfoResponse &
+{
+	/**Category for the label */
+	category: string;
+
+	/**Name of the label */
+	name: string;
+}
+
+/**State of an label within a job */
+export type GetLabelStateResponse = {
+	/**State of the label */
+	state?: LabelState;
+
+	/**Outcome of the label */
+	outcome?: LabelOutcome;
+}
+
+/**Information about the default label (ie. with inlined list of nodes) */
+export type GetDefaultLabelStateResponse = GetLabelStateResponse &
+{
+	/**List of nodes covered by default label */
+	nodes: string[];
+}
+
+
+/**Information about the timing info for a particular target */
+export type GetTimingInfoResponse = {
+
+	/**Wait time on the critical path */
+	totalWaitTime?: number;
+
+	/**Sync time on the critical path */
+	totalInitTime?: number;
+
+	/**Duration to this point */
+	totalTimeToComplete?: number;
+
+	/**Average wait time by the time the job reaches this point */
+	averageTotalWaitTime?: number;
+
+	/**Average sync time to this point */
+	averageTotalInitTime?: number;
+
+	/**Average duration to this point */
+	averageTotalTimeToComplete?: number;
+
+}
+
+/**Information about the timing info for a particular target */
+
+export type GetStepTimingInfoResponse = GetTimingInfoResponse & {
+
+	/** Name of this node */
+	name?: string;
+
+	/**Average wait time for this step */
+	averageStepWaitTime?: number;
+
+	/**Average init time for this step */
+	averageStepInitTime?: number;
+
+	/**Average duration for this step */
+	averageStepDuration?: number;
+
+}
+
+
+/**Request used to update notifications */
+export type UpdateNotificationsRequest = {
+
+	/** Notify via email */
+	email?: boolean;
+
+	/** Notify via Slack */
+	slack?: boolean;
+}
+
+/**Response describing notifications */
+export type GetNotificationResponse = {
+
+	/** Notify via email */
+	email?: boolean;
+
+	/** Notify via Slack */
+	slack?: boolean;
+}
+
+export type JobCompleteEventRecord = {
+	outcome: string;
+
+	type: string;
+
+	streamId: string;
+
+	templateId: string;
+}
+
+export type LabelCompleteEventRecord = {
+
+	type: string;
+
+	streamId: string;
+
+	templateId: string;
+
+	labelName: string;
+
+	categoryName: string;
+
+	outcome: string;
+}
+
+export type StepCompleteEventRecord = {
+
+	type: string;
+
+	streamId: string;
+
+	templateId: string;
+
+	stepName: string;
+
+	outcome: string;
+}
+
+/**Request used to create subscriptions */
+export type CreateSubscriptionRequest = {
+
+	event: JobCompleteEventRecord | LabelCompleteEventRecord | StepCompleteEventRecord;
+
+	userId?: string;
+
+	notificationType: SubscriptonNotificationType
+}
+
+/**Request used to create subscriptions */
+export type CreateSubscriptionResponse = {
+
+	id: string;
+}
+
+/**Request used to create subscriptions */
+export type GetSubscriptionResponse = {
+
+	id: string;
+
+	event: JobCompleteEventRecord | LabelCompleteEventRecord | StepCompleteEventRecord;
+
+	userId?: string;
+
+	notificationType: SubscriptonNotificationType
+}
+
+/**Identifies a particular changelist and job */
+export type GetIssueStepResponse = {
+
+	/**The changelist number */
+	change: number;
+
+	/** Name of the job containing this step */
+	jobName: string;
+
+	/**The unique job id */
+	jobId: string;
+
+	/**The unique batch id */
+	batchId: string;
+
+	/**The unique step id */
+	stepId: string;
+
+	/**  Time at which the step ran */
+	stepTime: Date | string;
+
+	/// The unique log id
+	logId?: string;
+
+}
+
+/** Information about a template affected by an issue */
+export type GetIssueAffectedTemplateResponse = {
+
+	/** The template id */
+	templateId: string;
+
+	/**  The template name */
+	templateName: string;
+
+	/**  Whether it has been resolved or not */
+	resolved: boolean;
+
+}
+
+/**Trace of a set of node failures across multiple steps */
+export type GetIssueNodeResponse = {
+
+	/** The template containing this step */
+	templateId: string;
+
+	/**Name of the step */
+	name: string;
+
+	/**The previous build  */
+	lastSuccess?: GetIssueStepResponse;
+
+	/**The failing builds for a particular event */
+	steps: GetIssueStepResponse[];
+
+	/**The following successful build */
+	nextSuccess?: GetIssueStepResponse;
+
+}
+
+/**Information about a particular step */
+export type GetIssueStreamResponse = {
+	/**Unique id of the stream */
+	streamId: string;
+
+	/**Minimum changelist affected by this issue (ie. last successful build) */
+	minChange?: number;
+
+	/**Maximum changelist affected by this issue (ie. next successful build) */
+	maxChange?: number;
+
+	/**Map of steps to (event signature id -> trace id) */
+	nodes: GetIssueNodeResponse[];
+}
+
+/**Outcome of a particular build */
+export enum IssueBuildOutcome {
+	/**Unknown outcome */
+	Unknown = "Unknown",
+
+	/**Build succeeded */
+	Success = "Success",
+
+	/**Build failed */
+	Error = "Error",
+
+	/**Build finished with warnings */
+	Warning = "Warning"
+}
+
+/**Information about a suspect changelist that may have caused an issue */
+export type GetIssueSuspectResponse = {
+	/**Number of the changelist that was submitted */
+	change: number;
+
+	/**Author of the changelist */
+	author: string;
+
+	/**The originating change */
+	originatingChange?: number;
+
+	/** Time at which the user declined this issue */
+	declinedAt?: Date | string;
+}
+
+/**Information about a diagnostic */
+export type GetIssueDiagnosticResponse = {
+	/**The corresponding build id */
+	buildId?: number;
+
+	/**Message for the diagnostic */
+	message: string;
+
+	/**Link to the error */
+	url: string;
+}
+
+/**Summary for the state of a stream in an issue */
+export interface GetIssueAffectedStreamResponse {
+	/**Id of the stream */
+	streamId: string;
+
+	/**Name of the stream */
+	streamName: string;
+
+	/**Whether the issue has been resolved in this stream */
+	resolved: boolean;
+
+	/** The affected templates */
+	affectedTemplates: GetIssueAffectedTemplateResponse[];
+
+	/**List of affected template ids */
+	templateIds: string[];
+
+	/** List of resolved template ids */
+	resolvedTemplateIds: string[];
+
+	/** List of resolved template ids */
+	unresolvedTemplateIds: string[];
+}
+
+
+/**Stores information about a build health issue */
+export type GetIssueResponse = {
+	/**The unique object id */
+	id: number;
+
+	/**Time at which the issue was created */
+	createdAt: Date | string;
+
+	/**Time at which the issue was retrieved */
+	retrievedAt: Date | string;
+
+	/**The associated project for the issue */
+	project?: string;
+
+	/**The summary text for this issue */
+	summary: string;
+
+	/**Details text describing the issue */
+	details?: string;
+
+	/** Severity of this issue	*/
+	severity: IssueSeverity;
+
+	/**Owner of the issue */
+	owner?: string;
+
+	/**User that nominated the current owner */
+	nominatedBy?: string;
+
+	/**Time that the issue was acknowledged */
+	acknowledgedAt?: Date | string;
+
+	/**Changelist that fixed this issue */
+	fixChange?: number;
+
+	/**Time at which the issue was resolved */
+	resolvedAt?: Date | string;
+
+	/**  List of stream paths affected by this issue */
+	streams: string[];
+
+	/** List of affected stream ids */
+	resolvedStreams: string[];
+
+	/** List of unresolved streams */
+	unresolvedStreams: string[];
+
+	affectedStreams: GetIssueAffectedStreamResponse[];
+
+	/** Most likely suspects for causing this issue */
+	primarySuspects: string[];
+
+	/** Whether to show alerts for this issue */
+	showDesktopAlerts: boolean;
+
+}
+
+/**Request an issue to be updated */
+export type UpdateIssueRequest = {
+
+	/** Summary of the issue */
+	summary?: string;
+
+	/**New owner of the issue */
+	owner?: string | null;
+
+	/**User than nominates the new owner */
+	nominatedBy?: string | null;
+
+	/**Whether the issue has been acknowledged */
+	acknowledged?: boolean;
+
+	/** Whether the user has declined this issue */
+	declined?: boolean;
+
+	/**The change at which the issue is claimed fixed */
+	fixChange?: number | null;
+
+	/**Whether the issue should be marked as resolved */
+	resolved?: boolean;
+}
+
+export type GetUtilizationTelemetryStream = {
+
+	streamId: string;
+
+	time: number;
+}
+
+export type GetUtilizationTelemetryPool = {
+
+	poolId: string;
+
+	numAgents: number;
+
+	adminTime: number;
+
+	otherTime: number;
+
+	streams: GetUtilizationTelemetryStream[];
+}
+
+export type GetUtilizationTelemetryResponse = {
+
+	startTime: Date;
+
+	finishTime: Date;
+
+	pools: GetUtilizationTelemetryPool[];
+
+	adminTime: number;
+
+	numAgents: number;
+}
+
+export type UserClaim = {
+
+	type: string;
+	value: string;
+}
+
+export enum DashboardPreference {
+	DisplayUTC = "DisplayUTC",
+	DisplayClock = "DisplayClock",
+	ColorSuccess = "ColorSuccess",
+	ColorWarning = "ColorWarning",
+	ColorError = "ColorError",
+	ColorRunning = "ColorRunning",
+	LocalCache = "LocalCache",
+}
+
+export type DashboardSettings = {
+
+	preferences: Map<DashboardPreference, string>;
+
+}
+
+export type UpdateDashboardSettings = {
+
+	preferences?: Record<DashboardPreference, string>;
+
+}
+
+
+/**  Response describing the current user */
+export type GetUserResponse = {
+
+	/** Id of the user */
+	id: string;
+
+	/** Claims for the user */
+	claims: UserClaim[];
+
+	/** List of pinned job ids */
+	pinnedJobIds: string[];
+
+	/** Whether to enable slack notifications for this user */
+	enableIssueNotifications: boolean;
+
+	/** Whether to enable experimental features for this user */
+	enableExperimentalFeatures: boolean;
+
+	/**  Settings for the dashboard */
+	dashboardSettings: DashboardSettings;
+
+}
+
+/// <summary>
+/// Request to update settings for a user
+/// </summary>
+export type UpdateUserRequest = {
+
+	/** New dashboard settings */
+	dashboardSettings?: UpdateDashboardSettings;
+
+	/** Job ids to add to the pinned list */
+	addPinnedJobIds?: string[];
+
+	/** Job ids to remove from the pinned list */
+	removePinnedJobIds?: string[];
+
+	/** Whether to enable slack notifications for this user */
+	enableIssueNotifications?: boolean;
+
+	/** Whether to enable experimental features for this user */
+	enableExperimentalFeatures?: boolean;
+
+}
+
+export type GetPerforceServerStatusResponse = {
+	serverAndPort: string;
+	baseServerAndPort: string;
+	cluster: string;
+	numLeases: number;
+	status: string;
+	detail: string;
+}
diff --git a/Engine/Source/Programs/Horde/HordeDashboard/Dashboard/src/backend/Backend.ts b/Engine/Source/Programs/Horde/HordeDashboard/Dashboard/src/backend/Backend.ts
new file mode 100644
--- /dev/null
+++ b/Engine/Source/Programs/Horde/HordeDashboard/Dashboard/src/backend/Backend.ts
@@ -0,0 +1,1161 @@
+// Copyright Epic Games, Inc. All Rights Reserved.
+
+import { configure } from "mobx";
+import { isNumber } from 'util';
+import templateCache from '../backend/TemplateCache';
+import { AgentData, ArtifactData, BatchUpdatePoolRequest, ChangeSummaryData, CreateJobRequest, CreateJobResponse, CreatePoolRequest, CreateSoftwareResponse, CreateSubscriptionRequest, CreateSubscriptionResponse, DashboardPreference, EventData, GetArtifactZipRequest, GetUserResponse, GetGraphResponse, GetIssueStreamResponse, GetJobStepRefResponse, GetJobStepTraceResponse, GetJobTimingResponse, GetLogEventResponse, GetNotificationResponse, GetSoftwareResponse, GetSubscriptionResponse, IssueData, IssueQuery, JobData, JobQuery, LeaseData, LogData, LogLineData, PoolData, ProjectData, ScheduleData, ScheduleQuery, SearchLogFileResponse, SessionData, StreamData, TemplateData, TestData, UpdateAgentRequest, UpdateIssueRequest, UpdateJobRequest, UpdateNotificationsRequest, UpdatePoolRequest, UpdateStepRequest, UpdateStepResponse, UpdateUserRequest, GetUtilizationTelemetryResponse, TabType, GetJobsTabResponse, JobsTabColumnType, GetPerforceServerStatusResponse } from './Api';
+import dashboard from './Dashboard';
+import { ChallengeStatus, Fetch } from './Fetch';
+import graphCache, { GraphQuery } from './GraphCache';
+import { projectStore } from './ProjectStore';
+
+
+// configure mobx
+configure({
+    reactionRequiresObservable: true,
+    enforceActions: "observed"
+});
+
+// Update interval for relatively static global data such as projects, schedules, templates
+const updateInterval = 120 * 1000;
+
+export class Backend {
+
+    getProjects(): Promise<ProjectData[]> {
+
+        return new Promise<ProjectData[]>((resolve, reject) => {
+
+            this.backend.get("/api/v1/projects", { params: { categories: true } }).then(async (response) => {
+
+                const projects = response.data as ProjectData[];
+
+                await this.backend.get("/api/v1/streams").then(response => {
+
+                    const streams = (response.data as StreamData[]);
+
+                    projects.forEach(p => {
+                        p.streams = streams.filter(stream => stream.projectId === p.id);
+                        p.streams!.forEach(s => {
+
+                            // create a default columbn if none exist
+                            s.tabs?.forEach(tab => {
+                                if (tab.type === TabType.Jobs) {
+                                    const jobTab = tab as GetJobsTabResponse;
+                                    if (!jobTab.columns) {
+                                        jobTab.columns = [{
+                                            type: JobsTabColumnType.Labels,
+                                            heading: "Other",
+                                            relativeWidth: 1
+                                        }]
+                                    }
+                                }
+                            });
+
+                            s.fullname = s.name;
+                            s.name = s.name.substr(s.name.lastIndexOf("/") + 1);
+                            s.project = p;
+                        })
+                    });
+                });
+
+                resolve(projects);
+
+            }).catch(reason => { reject(reason); });
+
+        });
+
+    }
+
+    getProject(project: ProjectData): Promise<ProjectData> {
+
+        return new Promise<ProjectData>((resolve, reject) => {
+
+            this.backend.get(`/api/v1/projects/${project.id}`).then((response) => {
+
+                const project = response.data as ProjectData;
+                resolve(project);
+
+            }).catch(reason => { reject(reason); });
+
+        });
+
+    }
+
+    getTemplate(streamId: string, templateRefId: string): Promise<TemplateData> {
+        return new Promise<TemplateData>((resolve, reject) => {
+            this.backend.get(`/api/v1/streams/${streamId}/templates/${templateRefId}`)
+                .then(response => resolve(response.data as TemplateData))
+                .catch(reason => reject(reason));
+        });
+    }
+
+    getPools(): Promise<PoolData[]> {
+        return new Promise<PoolData[]>((resolve, reject) => {
+            this.backend.get("/api/v1/pools").then((response) => {
+                const pools = response.data as PoolData[];
+                resolve(pools);
+            }).catch(reason => { reject(reason); });
+        });
+    }
+
+    getAgent(id: string): Promise<AgentData> {
+        return new Promise<AgentData>((resolve, reject) => {
+            this.backend.get(`/api/v1/agents/${id}`).then((response) => {
+                const agent = response.data as AgentData;
+                resolve(agent);
+            }).catch(reason => { reject(reason); });
+        });
+    }
+
+    getAgents(modifiedAfter: Date): Promise<AgentData[]> {
+        return new Promise<AgentData[]>((resolve, reject) => {
+            // for whenever we bring back modifiedDate, swap to this line
+            this.backend.get(`/api/v1/agents?ModifiedAfter=${modifiedAfter.toISOString()}`).then((response) => {
+                //this.backend.get(`/api/v1/agents`).then((response) => {
+                const agents = response.data as AgentData[];
+                resolve(agents);
+            }).catch(reason => { reject(reason); });
+        });
+    }
+
+    getLease(leaseId: string, agentId?: string): Promise<LeaseData> {
+
+        if (agentId) {
+            return new Promise<LeaseData>((resolve, reject) => {
+                this.backend.get(`/api/v1/agents/${agentId}/leases/${leaseId}`).then((response) => {
+                    resolve(response.data as LeaseData);
+                }).catch(reason => { reject(reason); });
+            });    
+        }
+
+        return new Promise<LeaseData>((resolve, reject) => {
+            this.backend.get(`/api/v1/leases/${leaseId}`).then((response) => {
+                resolve(response.data as LeaseData);
+            }).catch(reason => { reject(reason); });
+        });    
+
+    }
+
+    getLeases(agentId: string, index: number, count: number): Promise<LeaseData[]> {
+        const params: string[] = [];
+        params.push("Index=" + index);
+        params.push("Count=" + count);
+        const paramString = params.join("&");
+        return new Promise<LeaseData[]>((resolve, reject) => {
+            this.backend.get(`/api/v1/agents/${agentId}/leases?${paramString}`).then((response) => {
+                const leases = response.data as LeaseData[];
+                leases.forEach(function (lease) {
+                    lease.startTime = new Date(Date.parse(lease.startTime as string));
+                    if (lease.finishTime) {
+                        lease.finishTime = new Date(Date.parse(lease.finishTime as string));
+                    }
+                    if (lease.details) {
+                        lease.batchId = lease.details.BatchId;
+                        lease.jobId = lease.details.JobId;
+                        lease.type = lease.details.Type;
+                    }
+                });
+                resolve(leases);
+            }).catch(reason => { reject(reason); });
+        });
+    }
+
+    getSessions(agentId: string, index: number, count: number): Promise<SessionData[]> {
+        const params: string[] = [];
+        params.push("Index=" + index);
+        params.push("Count=" + count);
+        const paramString = params.join("&");
+        return new Promise<SessionData[]>((resolve, reject) => {
+            this.backend.get(`/api/v1/agents/${agentId}/sessions?${paramString}`).then((response) => {
+                const sessions = response.data as SessionData[];
+                sessions.forEach(function (session) {
+                    session.startTime = new Date(Date.parse(session.startTime as string));
+                    if (session.finishTime) {
+                        session.finishTime = new Date(Date.parse(session.finishTime as string));
+                    }
+                });
+                resolve(sessions);
+            }).catch(reason => { reject(reason); });
+        });
+    }
+
+
+    getJobTemplate(jobId: string): Promise<TemplateData> {
+
+        return new Promise<TemplateData>((resolve, reject) => {
+
+            this.backend.get(`/api/v1/jobs/${jobId}/template`, {
+                params: {}
+            }).then((value) => {
+                resolve(value.data as TemplateData);
+            }).catch(reason => {
+                reject(reason);
+            });
+        });
+
+    }
+
+    retryJobStep(jobId: string, batchId: string, stepId: string): Promise<any> {
+
+        return new Promise<any>((resolve, reject) => {
+            this.backend.put(`/api/v1/jobs/${jobId}/batches/${batchId}/steps/${stepId}`, {
+                Retry: true
+            }).then((value) => {
+                resolve(value.data as JobData[]);
+            }).catch(reason => {
+                reject(reason);
+            });
+        });
+
+    }
+
+    /** Get's a job's graph data, used by caching and shouldn't need to call directly */
+    getGraph(jobId: string): Promise<GetGraphResponse> {
+
+        return new Promise<GetGraphResponse>((resolve, reject) => {
+
+            this.backend.get(`/api/v1/jobs/${jobId}/graph`).then((value) => {
+                const response = value.data as GetGraphResponse;
+                // filter out ugs labels
+                response.labels = response.labels?.filter(label => !!label.dashboardName);
+                resolve(response);
+            }).catch(reason => {
+                reject(reason);
+            });
+        });
+
+    }
+
+    getJob(id: string, query?: JobQuery): Promise<JobData> {
+
+        return new Promise<JobData>((resolve, reject) => {
+
+            this.backend.get(`/api/v1/jobs/${id}`, {
+                params: query
+            }).then((value) => {
+
+                const response = value.data as JobData;
+                if (!response.graphHash) {
+                    return reject(`Job ${id} has undefined graph hash`);
+                }
+
+                graphCache.get({ graphHash: response.graphHash, jobId: id }).then(graph => {
+                    response.graphRef = graph;
+                    resolve(response);
+                }).catch(reason => {
+                    reject(reason);
+                });
+
+
+            }).catch(reason => {
+                reject(reason);
+            });
+        });
+
+    }
+
+    getJobsByIds(id: string[], query?: JobQuery, includeGraphs?: boolean): Promise<JobData[]> {
+
+        if (!query) {
+            (query as any) = {
+                id: id
+            }
+        } else {
+            (query as any).id = id;
+        }
+
+        return new Promise<JobData[]>((resolve, reject) => {
+
+            this.backend.get(`/api/v1/jobs`, {
+                params: query
+            }).then((value) => {
+
+                const response = value.data as JobData[];
+
+                if (includeGraphs === false || includeGraphs === undefined) {
+                    resolve(response);
+                    return;
+                }
+
+                graphCache.getGraphs(response.map(j => {
+                    return {
+                        jobId: j.id!,
+                        graphHash: j.graphHash!
+                    };
+                })).then((graphs => {
+                    response.forEach(j => {
+                        j.graphRef = graphs.find(g => g.hash === j.graphHash)
+                    })
+                    resolve(response)
+                })).catch(reason => {
+                    reject(reason);
+                });
+
+            }).catch(reason => {
+                reject(reason);
+            });
+        });
+
+    }
+
+    getJobs(query: JobQuery, queryGraph: boolean = false): Promise<JobData[]> {
+
+        if (typeof query.index === 'number') {
+            query.index = 0;
+        }
+
+        if (typeof query.count !== 'number') {
+            query.count = 100;
+        }
+
+        return new Promise<JobData[]>((resolve, reject) => {
+
+            this.backend.get("/api/v1/jobs", {
+                params: query
+            }).then((value) => {
+                const jobs = value.data as JobData[];
+
+                if (!queryGraph) {
+                    resolve(jobs);
+                    return;
+                }
+
+                const query: GraphQuery[] = [];
+                jobs.forEach(j => {
+                    if (!j.graphHash) {
+                        console.error(`Job ${j.id} has no graph hash`);
+                        return;
+                    }
+
+                    if (!query.find(q => q.graphHash === j.graphHash)) {
+                        query.push({ graphHash: j.graphHash!, jobId: j.id });
+                    }
+                });
+
+                graphCache.getGraphs(query).then(graphs => {
+                    jobs.forEach(j => {
+                        j.graphRef = graphs.find(g => g.hash === j.graphHash);
+                    });
+                    resolve(jobs);
+                }).catch(reason => reject(reason));
+
+            }).catch((reason) => {
+                reject(reason);
+            });
+        });
+
+    }
+
+    getJobTiming(jobId: string): Promise<GetJobTimingResponse> {
+
+        return new Promise<GetJobTimingResponse>((resolve, reject) => {
+            this.backend.get(`/api/v1/jobs/${jobId}/timing`).then((value) => {
+                resolve(value.data as GetJobTimingResponse);
+            }).catch(reason => {
+                reject(reason);
+            });
+        });
+
+    }
+
+    getLogEvents(logId: string): Promise<EventData[]> {
+
+        return new Promise<EventData[]>((resolve, reject) => {
+
+            this.backend.get(`/api/v1/logs/${logId}/events`).then((value) => {
+                resolve(value.data as EventData[]);
+            }).catch(reason => {
+                reject(reason);
+            });
+        });
+    }
+
+    searchLog(logId: string, text: string, firstLine?: number, count?: number) {
+
+        return new Promise<SearchLogFileResponse>((resolve, reject) => {
+
+            const params: any = {
+                text: text,
+                firstLine: firstLine,
+                count: count
+            };
+
+            this.backend.get(`/api/v1/logs/${logId}/search`, {
+                params: params
+            }).then((value) => {
+                resolve(value.data as SearchLogFileResponse);
+            }).catch(reason => {
+                reject(reason);
+            });
+        });
+
+    }
+
+    getJobStepHistory(streamId: string, stepName: string, count: number = 10, templateId: string) {
+        return new Promise<GetJobStepRefResponse[]>((resolve, reject) => {
+
+            const params: any = {
+                step: stepName,
+                count: count,
+                templateId: templateId
+            };
+
+            this.backend.get(`/api/v1/streams/${streamId}/history`, {
+                params: params
+            }).then((value) => {
+                resolve(value.data as GetJobStepRefResponse[]);
+            }).catch(reason => {
+                reject(reason);
+            });
+        });
+    }
+
+    getJobStepTrace(jobId: string, batchId: string, stepId: string) {
+        return new Promise<GetJobStepTraceResponse>((resolve, reject) => {
+            this.backend.get(`/api/v1/jobs/${jobId}/batches/${batchId}/steps/${stepId}/trace`, {
+                suppress404: true
+            }).then((value) => {
+                resolve(value.data as GetJobStepTraceResponse);
+            }).catch(reason => {
+                reject(reason);
+            });
+        });
+    }
+
+    getJobArtifacts(jobId: string, stepId?: string): Promise<ArtifactData[]> {
+
+        const params: any = {
+            jobId: jobId,
+            code: true,
+        };
+
+        if (stepId) {
+            params.stepId = stepId;
+        }
+
+        return new Promise<ArtifactData[]>((resolve, reject) => {
+            this.backend.get(`/api/v1/artifacts`, { params: params, suppress404: true }).then((value) => {
+                resolve(value.data as ArtifactData[]);
+            }).catch(reason => {
+                resolve([]);
+                if (reason !== "Not Found") {
+                    console.error(reason);
+                }
+            });
+        });
+    }
+
+    downloadJobArtifacts(request: GetArtifactZipRequest): Promise<boolean> {
+        return new Promise<any>((resolve, reject) => {
+            this.backend.post(`/api/v1/artifacts/zip`, request, { responseBlob: true }).then(response => {
+                const url = window.URL.createObjectURL(response.data);
+                const link = document.createElement('a');
+                link.href = url;
+                let name = request.jobId + '-' + request.stepId + ".zip";
+                link.setAttribute('download', name);
+                document.body.appendChild(link);
+                link.click();
+                document.body.removeChild(link);
+                setTimeout(() => {
+                    window.URL.revokeObjectURL(url);
+                }, 100);
+                resolve(true);
+            }).catch(reason => {
+                reject(reason);
+            });
+        });
+    }
+
+    getArtifactDataById(id: string): Promise<object> {
+        return new Promise<object>((resolve, reject) => {
+            this.backend.get(`/api/v1/artifacts/${id}/data`).then((value) => {
+                resolve(value.data as object);
+            }).catch((reason) => {
+                reject(reason);
+            });
+        });
+    }
+
+    getLogData(logId: string): Promise<LogData> {
+
+        return new Promise<LogData>((resolve, reject) => {
+
+            this.backend.get(`/api/v1/logs/${logId}`).then((value) => {
+                resolve(value.data as LogData);
+            }).catch(reason => {
+                reject(reason);
+            });
+        });
+
+    }
+
+    getLogLines(logId: string, index?: number, count?: number): Promise<LogLineData> {
+
+        const params: any = {};
+
+        if (isNumber(index)) {
+            params.index = index;
+        }
+
+        if (isNumber(count)) {
+            params.Count = count;
+        }
+
+        return new Promise<LogLineData>((resolve, reject) => {
+
+            this.backend.get(`/api/v1/logs/${logId}/lines`, {
+                params: params
+            }).then((value) => {
+                resolve(value.data as LogLineData);
+            }).catch(reason => {
+                reject(reason);
+            });
+        });
+
+    }
+
+    getLog(logId: string, offset?: number, count?: number): Promise<string> {
+
+        return new Promise<string>((resolve, reject) => {
+
+            this.backend.get(`/api/v1/logs/${logId}/data`, {
+                params: {
+                    offset: offset,
+                    count: count
+                }
+            }).then((value) => {
+                resolve(value.data as string);
+            }).catch(reason => {
+                reject(reason);
+            });
+        });
+
+    }
+
+    createJob(request: CreateJobRequest): Promise<CreateJobResponse> {
+
+        return new Promise<CreateJobResponse>((resolve, reject) => {
+
+            this.backend.post("/api/v1/jobs", request).then((value) => {
+                resolve(value.data as CreateJobResponse);
+            }).catch(reason => {
+                reject(reason);
+            });
+        });
+    }
+
+    updateJob(jobId: string, request: UpdateJobRequest): Promise<boolean> {
+
+        return new Promise<boolean>((resolve, reject) => {
+            this.backend.put(`/api/v1/jobs/${jobId}`, request).then((value) => {
+                resolve(true);
+            }).catch(reason => {
+                reject(reason);
+            });
+        });
+    }
+
+    updateAgent(agentId: string, request: UpdateAgentRequest): Promise<boolean> {
+        return new Promise<boolean>((resolve, reject) => {
+            this.backend.put(`api/v1/agents/${agentId}`, request).then((value) => {
+                resolve(value.data as boolean);
+            }).catch(reason => {
+                reject(reason);
+            });
+        });
+    }
+
+    deleteAgent(agentId: string): Promise<any> {
+        return new Promise<boolean>((resolve, reject) => {
+            this.backend.delete(`api/v1/agents/${agentId}`).then((value) => {
+                resolve(value.data);
+            }).catch(reason => {
+                reject(reason);
+            });
+        });
+    }
+
+    createPool(request: CreatePoolRequest): Promise<string> {
+        return new Promise<string>((resolve, reject) => {
+            this.backend.post(`api/v1/pools`, request).then((value) => {
+                resolve(value.data as string);
+            }).catch(reason => {
+                reject(reason);
+            });
+        });
+    }
+
+    updatePool(poolId: string, request: UpdatePoolRequest): Promise<boolean> {
+        return new Promise<boolean>((resolve, reject) => {
+            this.backend.put(`api/v1/pools/${poolId}`, request).then((value) => {
+                resolve(value.data as boolean);
+            }).catch(reason => {
+                reject(reason);
+            });
+        });
+    }
+
+    deletePool(poolId: string): Promise<boolean> {
+        return new Promise<boolean>((resolve, reject) => {
+            this.backend.delete(`api/v1/pools/${poolId}`).then((value) => {
+                resolve(value.data as boolean);
+            }).catch(reason => {
+                reject(reason);
+            });
+        });
+    }
+
+    batchUpdatePools(request: BatchUpdatePoolRequest[]): Promise<boolean> {
+        return new Promise<boolean>((resolve, reject) => {
+            this.backend.put(`api/v1/pools`, request).then((value) => {
+                resolve(value.data as boolean);
+            }).catch(reason => {
+                reject(reason);
+            });
+        });
+    }
+
+    uploadSoftware(formData: FormData): Promise<CreateSoftwareResponse> {
+
+        return new Promise<CreateSoftwareResponse>((resolve, reject) => {
+
+            this.backend.post("/api/v1/software", formData, { formData: true }).then((value) => {
+                resolve(value.data as CreateSoftwareResponse);
+            }).catch(reason => {
+                reject(reason);
+            });
+        });
+
+    }
+
+    updateSoftware(id: string, isDefault: boolean): Promise<any> {
+
+        return new Promise<any>((resolve, reject) => {
+            this.backend.put(`/api/v1/software/${id}`, {
+                Default: isDefault
+            }).then((value) => {
+                resolve(value.data as any);
+            }).catch(reason => {
+                reject(reason);
+            });
+        });
+
+    }
+
+
+    getSchedule(id: string): Promise<ScheduleData> {
+
+        return new Promise<ScheduleData>((resolve, reject) => {
+
+            this.backend.get(`/api/v1/schedules/${id}`).then((value) => {
+                resolve(value.data as ScheduleData);
+            }).catch(reason => {
+                reject(reason);
+            });
+        });
+
+    }
+
+    getSchedules(query: ScheduleQuery): Promise<ScheduleData[]> {
+
+        return new Promise<ScheduleData[]>((resolve, reject) => {
+            this.backend.get(`/api/v1/schedules`, {
+                params: query
+            }).then((value) => {
+                resolve(value.data as ScheduleData[]);
+            }).catch(reason => {
+                reject(reason);
+            });
+        });
+
+    }
+
+
+    getSoftware(): Promise<GetSoftwareResponse[]> {
+
+        return new Promise<GetSoftwareResponse[]>((resolve, reject) => {
+
+            this.backend.get(`/api/v1/software`).then((value) => {
+                resolve(value.data as GetSoftwareResponse[]);
+            }).catch(reason => {
+                reject(reason);
+            });
+        });
+
+    }
+
+    deleteSoftware(id: string): Promise<any> {
+
+        return new Promise<any>((resolve, reject) => {
+
+            this.backend.delete(`/api/v1/software/${id}`).then((value) => {
+                resolve(value.data);
+            }).catch(reason => {
+                reject(reason);
+            });
+        });
+
+    }
+
+    getCommit(streamId?: string, changelist?: number): Promise<ChangeSummaryData | undefined> {
+
+        return new Promise<ChangeSummaryData | undefined>((resolve) => {
+
+            if (!streamId || !changelist) {
+                resolve(undefined);
+                return;
+            }
+
+            this.backend.get(`/api/v1/streams/${streamId}/changes/${changelist}`).then((value) => {
+                resolve(value.data as ChangeSummaryData);
+            }).catch(reason => {
+                console.error(reason);
+                resolve(undefined);
+            });
+        });
+    }
+
+    getChangeSummaries(streamId: string, minChange?: number, maxChange?: number, maxResults = 100): Promise<ChangeSummaryData[]> {
+
+        return new Promise<ChangeSummaryData[]>((resolve, reject) => {
+
+            const params: any = {
+                results: maxResults
+            };
+
+            if (minChange) {
+                params.min = minChange;
+            }
+
+            if (maxChange) {
+                params.max = maxChange;
+            }
+
+            this.backend.get(`/api/v1/streams/${streamId}/changes`, {
+                params: params
+            }).then((value) => {
+                resolve(value.data as ChangeSummaryData[]);
+            }).catch(reason => {
+                reject(reason);
+            });
+
+        });
+
+    }
+
+    getAdminToken(): Promise<string> {
+
+        return new Promise<string>((resolve, reject) => {
+
+            this.backend.get("/api/v1/admin/token").then((response) => {
+                resolve(response.data as string);
+            }).catch(reason => reject(reason));
+
+        });
+
+    }
+
+    updateUser(request: UpdateUserRequest): Promise<void> {
+
+        return new Promise<void>((resolve, reject) => {
+            this.backend.put("/api/v1/user", request).then((response) => {
+                resolve();
+            }).catch(reason => reject(reason));
+
+        });
+
+    }
+
+    getUser(): Promise<GetUserResponse> {
+
+        return new Promise<GetUserResponse>((resolve, reject) => {
+
+            this.backend.get("/api/v1/user", { suppress404: true }).then((response) => {
+
+                let data = response.data as GetUserResponse;
+
+                // convert from an object to a map
+                const prefs = Object.values(DashboardPreference);
+                const preferences = new Map<DashboardPreference, string>();
+
+                if (!data.dashboardSettings?.preferences) {
+                    if (!data.dashboardSettings) {
+                        data.dashboardSettings = { preferences: preferences };
+                    } else {
+                        data.dashboardSettings.preferences = preferences;
+                    }
+
+                } else {
+                    for (let key of Object.keys(data.dashboardSettings.preferences)) {
+
+                        if (prefs.indexOf(key as DashboardPreference) === -1) {
+                            continue;
+                        }
+
+                        preferences.set(key as DashboardPreference, (data.dashboardSettings.preferences as any)[key])
+                    }
+
+                    data.dashboardSettings.preferences = preferences;
+                }
+                resolve(data);
+            }).catch(reason => reject(reason));
+
+        });
+
+    }
+
+
+    getIssues(queryIn?: IssueQuery): Promise<IssueData[]> {
+
+        const query = queryIn ?? {};
+
+        return new Promise<IssueData[]>((resolve, reject) => {
+
+            this.backend.get("/api/v1/issues", {
+                params: query
+            }).then((value) => {
+                resolve(value.data as IssueData[]);
+            }).catch((reason) => {
+                reject(reason);
+            });
+        });
+
+    }
+
+    getIssuesByIds(ids: number[]): Promise<IssueData[]> {
+
+        const unique = Array.from(new Set(ids));
+
+        return new Promise<IssueData[]>((resolve, reject) => {
+
+            this.backend.get(`/api/v1/issues`, { params: { id: unique.map(id => id.toString()) } })
+                .then(response => { resolve(response.data); })
+                .catch(reason => reject(reason));
+        })
+    }
+
+    getIssueStreams(issueId: number): Promise<GetIssueStreamResponse[]> {
+        return new Promise<GetIssueStreamResponse[]>((resolve, reject) => {
+            this.backend.get(`/api/v1/issues/${issueId}/streams`).then((value) => {
+                resolve(value.data as GetIssueStreamResponse[]);
+            }).catch((reason) => {
+                reject(reason);
+            });
+        });
+
+    }
+
+    getIssueEvents(issueId: number, jobId?: string, logIds?: string[]): Promise<GetLogEventResponse[]> {
+
+        if (logIds && logIds.length) {
+            jobId = undefined;
+        }
+
+        const params = {
+            jobId: jobId,
+            logIds: logIds
+        };
+
+        return new Promise<GetLogEventResponse[]>((resolve, reject) => {
+            this.backend.get(`/api/v1/issues/${issueId}/events`, { params: params }).then((value) => {
+                resolve(value.data as GetLogEventResponse[]);
+            }).catch((reason) => {
+                reject(reason);
+            });
+        });
+
+    }
+
+    updateIssue(id: number, request: UpdateIssueRequest): Promise<boolean> {
+        return new Promise<boolean>((resolve, reject) => {
+            this.backend.put(`/api/v1/issues/${id}`, request).then((value) => {
+                resolve(true);
+            }).catch((reason) => {
+                reject(reason);
+            });
+        });
+    }
+
+    getJobStepTestDataByKey(jobId: string, stepId: string, key: string): Promise<TestData[]> {
+        return new Promise<TestData[]>((resolve, reject) => {
+            this.backend.get(`/api/v1/testdata?JobId=${jobId}&JobStepId=${stepId}&Key=${key}`).then((value) => {
+                resolve(value.data as TestData[]);
+            }).catch((reason) => {
+                reject(reason);
+            });
+        });
+    }
+
+    getTestData(id: string): Promise<TestData> {
+        return new Promise<TestData>((resolve, reject) => {
+            this.backend.get(`/api/v1/testdata/${id}`).then((value) => {
+                resolve(value.data as TestData);
+            }).catch((reason) => {
+                reject(reason);
+            });
+        });
+    }
+
+    getJobTestData(jobId: string, stepId?: string): Promise<TestData[]> {
+        return new Promise<TestData[]>((resolve, reject) => {
+            const params: any = {
+                jobId: jobId,
+                filter: 'key,id,stepId',
+                count: 200,
+            };
+
+            if (stepId) {
+                params.jobStepId = stepId;
+            }
+
+            this.backend.get(`/api/v1/testdata`, {
+                params: params
+            }).then((value) => {
+                resolve(value.data as TestData[]);
+            }).catch((reason) => {
+                reject(reason);
+            });
+        });
+    }
+
+    getTestDataHistory(streamId: string, key: string, maxChange?: number, count: number = 30, index?: number): Promise<TestData[]> {
+        return new Promise<TestData[]>((resolve, reject) => {
+            const params: any = {
+                streamId: streamId,
+                key: key,
+                count: count,
+                filter: 'id,change,data'
+            };
+
+            if (maxChange) {
+                params.maxChange = maxChange;
+            }
+            if (index) {
+                params.index = index;
+            }
+
+            this.backend.get(`/api/v1/testdata`, {
+                params: params
+            }).then((value) => {
+                resolve(value.data as TestData[]);
+            }).catch((reason) => {
+                reject(reason);
+            });
+        });
+    }
+
+    downloadLog(logId: string, filename: string, json = false) {
+
+        try {
+            const url = `${this.serverUrl}/api/v1/logs/${logId}/data?download=true&format=${json ? 'raw' : 'text'}&filename=${encodeURIComponent(filename)}`;
+            const link = document.createElement('a');
+            link.href = url;
+            document.body.appendChild(link);
+            link.click();
+        } catch (reason) {
+            console.error(reason);
+        }
+
+    }
+
+    updateJobStep(jobId: string, batchId: string, stepId: string, request: UpdateStepRequest): Promise<UpdateStepResponse> {
+        return new Promise<UpdateStepResponse>((resolve, reject) => {
+            this.backend.put(`api/v1/jobs/${jobId}/batches/${batchId}/steps/${stepId}`, request).then((value) => {
+                resolve(value.data as UpdateStepResponse);
+            }).catch(reason => {
+                reject(reason);
+            });
+        });
+    }
+
+    getNotification(type: string, jobId: string, labelIdx?: string | null, batchId?: string, stepId?: string | null): Promise<GetNotificationResponse> {
+
+        let url = `api/v1/jobs/${jobId}`;
+        if (type === "label" && labelIdx != null) {
+            url = `${url}/labels/${labelIdx}`;
+        }
+        else if (type === "step" && batchId && stepId) {
+            url = `${url}/batches/${batchId}/steps/${stepId}`;
+        }
+        return new Promise<GetNotificationResponse>((resolve, reject) => {
+            this.backend.get(`${url}/notifications`).then((value) => {
+                resolve(value.data as GetNotificationResponse);
+            }).catch(reason => {
+                reject(reason);
+            });
+        });
+    }
+
+    updateNotification(request: UpdateNotificationsRequest, type: string, jobId: string, labelIdx?: string | null, batchId?: string, stepId?: string | null): Promise<boolean> {
+
+        let url = `api/v1/jobs/${jobId}`;
+
+        if (type === "label" && labelIdx != null) {
+            url = `${url}/labels/${labelIdx}`;
+        }
+        else if (type === "step" && batchId && stepId) {
+            url = `${url}/batches/${batchId}/steps/${stepId}`;
+        }
+        return new Promise<boolean>((resolve, reject) => {
+            this.backend.put(`${url}/notifications`, request).then((value) => {
+                resolve(true);
+            }).catch(reason => {
+                reject(reason);
+            });
+        });
+    }
+
+    createSubscription(requests: CreateSubscriptionRequest[]): Promise<CreateSubscriptionResponse[]> {
+        requests.forEach(request => {
+            request.userId = dashboard.username ?? "Anonymous";
+        });
+        let url = `api/v1/subscriptions`;
+        return new Promise<CreateSubscriptionResponse[]>((resolve, reject) => {
+            this.backend.post(`${url}`, requests).then((value) => {
+                let created = value.data as CreateSubscriptionResponse[];
+                resolve(created);
+            }).catch(reason => {
+                reject(reason);
+            });
+        });
+    }
+
+    getSubscriptions(): Promise<GetSubscriptionResponse[]> {
+        let url = `api/v1/subscriptions?userId=${dashboard.username ?? "Anonymous"}`;
+        return new Promise<GetSubscriptionResponse[]>((resolve, reject) => {
+            this.backend.get(`${url}`, { suppress404: true }).then((value) => {
+                resolve(value.data as GetSubscriptionResponse[]);
+            }).catch(reason => {
+                reject(reason);
+            });
+        });
+    }
+
+    deleteSubscription(id: string): Promise<boolean> {
+        let url = `api/v1/subscriptions/${id}`;
+        return new Promise<boolean>((resolve, reject) => {
+            this.backend.delete(`${url}`).then((value) => {
+                resolve(true);
+            }).catch(reason => {
+                reject(reason);
+            });
+        });
+    }
+
+    getUtilizationData(endDate: string, range: number, tzOffset: number): Promise<GetUtilizationTelemetryResponse[]> {
+        let url = `api/v1/reports/utilization/${endDate}?range=${range}&tzOffset=${tzOffset}`;
+        return new Promise<GetUtilizationTelemetryResponse[]>((resolve, reject) => {
+            this.backend.get(`${url}`).then((value) => {
+                resolve(value.data as GetUtilizationTelemetryResponse[]);
+            }).catch(reason => {
+                reject(reason);
+            });
+        });
+    }
+
+    getPerforceServerStatus(): Promise<GetPerforceServerStatusResponse[]> {
+
+        return new Promise<GetPerforceServerStatusResponse[]>((resolve, reject) => {
+            this.backend.get(`/api/v1/perforce/status`).then((value) => {
+                resolve(value.data as GetPerforceServerStatusResponse[]);
+            }).catch(reason => {
+                reject(reason);
+            });
+        });
+        
+    }
+
+    private async update() {
+
+        if (this.updateID === "updating") {
+            return;
+        }
+
+        this.updateID = "updating";
+
+        await Promise.all([projectStore.update()]).then(values => {
+
+        }).catch(reason => {
+
+        }).finally(() => {
+            this.updateID = setTimeout(() => { this.update(); }, updateInterval);
+        });
+
+    }
+
+    async serverLogout(redirect: string) {
+        try {
+            this.backend.logout = true;
+            await this.backend.get("/api/v1/dashboard/logout");            
+            window.location.assign(redirect);
+
+        } catch (err) {
+            console.error(err);
+        }
+    }
+
+    // server url if on a separate origin from dashboard, as with local development and a debug token
+    get serverUrl(): string {
+
+        const url = process.env.REACT_APP_HORDE_BACKEND;
+
+        if (!url) {
+            return "";
+        }
+
+        return url;
+    }
+
+    get debugToken(): string {
+
+        const token = process.env.REACT_APP_HORDE_DEBUG_TOKEN;
+
+        if (!token) {
+            return "";
+        }
+
+        return token;
+
+    }
+
+    init() {
+        return new Promise<boolean>(async (resolve, reject) => {
+
+            this.backend.setBaseUrl(this.serverUrl);
+            this.backend.setDebugToken(this.debugToken);
+
+            const challenge = await this.backend.challenge();
+
+            if (challenge === ChallengeStatus.Unauthorized) {
+                this.backend.login(window.location.toString());
+                return;
+            }
+
+            await dashboard.update();
+
+            if (dashboard.localCache) {
+                graphCache.initialize();
+                templateCache.initialize();
+            }
+
+            this.update().then(() => {
+                resolve(true);
+            }).catch((reason) => {
+                reject(reason);
+            });
+
+        });
+    }
+
+    constructor() {
+
+        this.backend = new Fetch();
+    }
+
+    updateID?: any;
+    logout: boolean = false;
+
+
+    private backend: Fetch;
+
+}
diff --git a/Engine/Source/Programs/Horde/HordeDashboard/Dashboard/src/backend/CommitCache.ts b/Engine/Source/Programs/Horde/HordeDashboard/Dashboard/src/backend/CommitCache.ts
new file mode 100644
--- /dev/null
+++ b/Engine/Source/Programs/Horde/HordeDashboard/Dashboard/src/backend/CommitCache.ts
@@ -0,0 +1,97 @@
+// Copyright Epic Games, Inc. All Rights Reserved.
+import { action, observable } from 'mobx';
+import backend from '.';
+import { GetChangeSummaryResponse } from "./Api";
+
+export class CommitCache {
+
+    @action
+    setUpdated() {
+        this.updated++;
+    }
+
+    @observable
+    updated: number = 0;
+
+    getCommit(streamId: string, change: number): GetChangeSummaryResponse | undefined {
+
+        let map = this.summaries.get(streamId); 
+                
+        if (!map) {
+            return undefined;
+        }
+        
+        const commit = map.get(change);
+
+        if (!commit || !commit.author) {
+            return undefined;
+        }
+
+        return commit;
+
+    }
+
+    async set(streamId: string, changes: number[]) {
+
+        let map = this.summaries.get(streamId);
+        if (!map) {
+            map = new Map();
+            this.summaries.set(streamId, map);
+        }
+
+        if (!changes.length) {
+            return;
+        }
+
+        const cminCL = 999999999999;
+        const cmaxCL = -1;
+
+        let minCL = cminCL;
+        let maxCL = cmaxCL;
+
+        changes.forEach(c => {
+            const change = map!.get(c);
+
+            if (!change) {
+
+                if (c < minCL) {
+                    minCL = c;
+                }
+
+                if (c > maxCL) {
+                    maxCL = c;
+                }
+            }
+        })
+
+        if (maxCL === cmaxCL && minCL === cminCL) {
+            return;
+        }
+
+        changes.forEach(c => {
+
+            const change = map!.get(c);
+            if (change) {
+                return;
+            }
+
+            map!.set(c, {
+                number: 0, author: "", description: ""
+            })
+
+        })
+
+        // new query
+        try {
+            const commits = await backend.getChangeSummaries(streamId, minCL, maxCL, maxCL - minCL + 1);
+            commits.forEach(c => map!.set(c.number, c));
+            this.setUpdated();
+        } catch (reason) {
+            console.error(reason);
+            // @todo: we could clear out the place holders on error
+        }
+
+    }
+
+    private summaries: Map<string, Map<number, GetChangeSummaryResponse>> = new Map();
+}
diff --git a/Engine/Source/Programs/Horde/HordeDashboard/Dashboard/src/backend/Config.ts b/Engine/Source/Programs/Horde/HordeDashboard/Dashboard/src/backend/Config.ts
new file mode 100644
--- /dev/null
+++ b/Engine/Source/Programs/Horde/HordeDashboard/Dashboard/src/backend/Config.ts
@@ -0,0 +1,33 @@
+// Copyright Epic Games, Inc. All Rights Reserved.
+
+import { initDatadog} from './Datadog';
+
+export type SiteConfig = {
+    environment: "dev" | "production" | "dev-local" | "";
+    plugins?: string[];
+    datadogClient?: string;
+}
+
+const dev = {
+	"environment": "dev",
+	"plugins": [
+		"UnrealAutomatedTests",
+		"SimpleTestReport"
+	]
+} 
+
+const prod = {
+	"environment": "production",
+	"plugins": [
+		"UnrealAutomatedTests",
+		"SimpleTestReport"
+	]
+} 
+
+export function getSiteConfig() {
+
+    // @todo: fix me
+    const isProd = window.location.hostname === "horde.devtools.epicgames.com";
+    
+    return isProd ? prod : dev;
+}
diff --git a/Engine/Source/Programs/Horde/HordeDashboard/Dashboard/src/backend/Dashboard.ts b/Engine/Source/Programs/Horde/HordeDashboard/Dashboard/src/backend/Dashboard.ts
new file mode 100644
--- /dev/null
+++ b/Engine/Source/Programs/Horde/HordeDashboard/Dashboard/src/backend/Dashboard.ts
@@ -0,0 +1,329 @@
+// Copyright Epic Games, Inc. All Rights Reserved.
+
+import { action, observable } from 'mobx';
+import backend from '.';
+import { UserClaim, DashboardPreference, GetUserResponse } from './Api';
+
+export class Dashboard {
+
+    startPolling() {
+        this.polling = true;
+    }
+
+    stopPolling() {
+        this.polling = false;
+    }
+
+    jobPinned(id: string | undefined) {
+        return !!this.pinnedJobsIds.find(j => j === id);
+    }
+
+    pinJob(id: string) {
+
+        if (this.data.pinnedJobIds.find(j => j === id)) {
+            return;
+        }
+
+        this.data.pinnedJobIds.push(id);
+
+        backend.updateUser({ addPinnedJobIds: [id] });
+
+        this.setUpdated();
+    }
+
+    unpinJob(id: string) {
+
+        if (!this.data.pinnedJobIds.find(j => j === id)) {
+            return;
+        }
+
+        this.data.pinnedJobIds = this.data.pinnedJobIds.filter(j => j !== id);
+
+        backend.updateUser({ removePinnedJobIds: [id] });
+
+        this.setUpdated();
+    }
+
+    get username(): string {
+
+        let email = this.claims.find(c => c.type.endsWith("/ue/horde/user"));
+
+        if (!email) {
+            email = this.claims.find(c => c.type.endsWith("/identity/claims/name"));
+        }
+
+        return email ? email.value : "???";
+    }
+
+
+    get hordeAdmin(): boolean {
+
+        return !!this.roles.find(r => r.value === "app-horde-admins");
+    }
+
+    get internalEmployee(): boolean {
+
+        return !!this.roles.find(r => r.value === "Internal-Employees");
+
+    }
+
+    get email(): string {
+
+        const email = this.claims.find(c => c.type.endsWith("/emailaddress"));
+
+        return email ? email.value : "???";
+    }
+
+    get p4user(): string {
+        const claims = this.claims;
+        const user = claims.filter(c => c.type.endsWith("/perforce-user"));
+        if (!user.length) {
+            return "";
+        }
+        return user[0].value;
+
+    }
+
+    get pinnedJobsIds(): string[] {
+
+        return this.data.pinnedJobIds ?? [];
+    }
+
+    get issueNotifications(): boolean {
+
+        return this.data.enableIssueNotifications;
+    }
+
+    set issueNotifications(value: boolean) {
+
+        backend.updateUser({ enableIssueNotifications: value }).then(() => {
+
+            this.data.enableIssueNotifications = value;
+            this.setUpdated();
+
+        });
+
+    }
+
+    get experimentalFeatures(): boolean {
+
+        return this.data.enableExperimentalFeatures;
+    }
+
+    set experimentalFeatures(value: boolean) {
+
+        backend.updateUser({ enableExperimentalFeatures: value }).then(() => {
+
+            this.data.enableExperimentalFeatures = value;
+            this.setUpdated();
+
+        });
+
+    }
+
+    get roles(): UserClaim[] {
+        const claims = this.claims;
+        return claims.filter(c => c.type.endsWith("/role"));
+    }
+
+
+    get claims(): UserClaim[] {
+        return this.data.claims ? this.data.claims : [];
+    }
+
+    get displayUTC(): boolean {
+
+        return this.preferences.get(DashboardPreference.DisplayUTC) === 'true';
+
+    }
+
+    setDisplayUTC(value: boolean | undefined) {
+        this.setPreference(DashboardPreference.DisplayUTC, value ? "true" : "false");
+    }
+
+    get localCache(): boolean {
+
+        return this.preferences.get(DashboardPreference.LocalCache) === 'true';
+
+    }
+
+    setLocalCache(value: boolean | undefined) {
+        this.setPreference(DashboardPreference.LocalCache, value ? "true" : "false");
+    }
+
+
+    get display24HourClock(): boolean {
+
+        return this.preferences.get(DashboardPreference.DisplayClock) === '24';
+
+    }
+
+    setStatusColor(pref: DashboardPreference, value: string | undefined) {
+
+        if (value && !value.startsWith("#")) {
+            console.error("Status preference color must be in hex format with preceding #")
+        }
+
+        this.setPreference(pref, value);
+    }
+
+
+    setDisplay24HourClock(value: boolean | undefined) {
+        this.setPreference(DashboardPreference.DisplayClock, value ? "24" : "");
+    }
+
+    getPreference(pref: DashboardPreference): string | undefined {
+
+
+        if (!this.available) {
+            return undefined;
+        }
+
+        if (!this.preferences) {
+            return undefined;
+        }
+
+        return this.preferences.get(pref);
+    }
+
+    private get preferences() {
+        return this.data.dashboardSettings.preferences;
+    }
+
+    async update() {
+
+        try {
+
+            if (this.updating) {
+                clearTimeout(this.updateTimeoutId);
+                this.updateTimeoutId = setTimeout(() => { this.update(); }, 4000);
+                return;
+            }
+
+            this.updating = true;
+
+            if (this.polling || !this.available) {
+
+                const cancelId = this.cancelId++;
+
+                const response = await backend.getUser();
+
+                // check for canceled during graph request
+                if (!this.canceled.has(cancelId)) {
+
+                    this.data = response;
+
+                    if (this.data.claims) {
+
+                        const set = new Set<string>();
+                        this.data.claims = this.data.claims.filter(c => {
+                            const key = c.type + c.value;
+                            if (set.has(key)) {
+                                return false;
+                            }
+                            set.add(key);
+                            return true;
+                        })
+
+                    }
+
+                    // @todo: detect changed                
+                    this.setUpdated();
+
+                }
+
+            }
+
+        } catch (reason) {
+            if (!this.available) {
+                // this is being added 1/25/21 for changes to how User's are handled on backend
+                // ie. not created until logged in via Okta
+                // if this is still an error in the future, this may be changed
+                this.requestLogout = true;
+            }
+            console.error("Error updating user dashboard settings", reason)
+        } finally {
+            this.updating = false;
+            clearTimeout(this.updateTimeoutId);
+            this.updateTimeoutId = setTimeout(() => { this.update(); }, 4000);
+        }
+    }
+
+    private setPreference(pref: DashboardPreference, value: string | undefined): void {
+
+        if (!this.available) {
+            return;
+        }
+
+        if (this.preferences.get(pref) === value) {
+            return;
+        }
+
+        if (value === undefined) {
+            this.preferences.delete(pref);
+        } else {
+            this.preferences.set(pref, value);
+        }
+
+        this.postPreferences();
+
+    }
+
+
+    @observable
+    updated: number = 0;
+
+    @action
+    private setUpdated() {
+        this.updated++;
+    }
+
+    get available(): boolean {
+        return this.data.id !== "";
+    }
+
+
+    private async postPreferences(): Promise<boolean> {
+
+        // cancel any pending        
+        for (let i = 0; i < this.cancelId; i++) {
+            this.canceled.add(i);
+        }
+
+        const data: any = {};
+
+        for (const key of Object.keys(DashboardPreference)) {
+            data[key] = this.data.dashboardSettings.preferences?.get(key as DashboardPreference);
+        }
+
+        let success = true;
+        try {
+            await backend.updateUser({ dashboardSettings: { preferences: data } });
+        } catch (reason) {
+            success = false;
+            console.error("Error posting user preferences", reason)
+        }
+
+        return success;
+
+    }
+
+    requestLogout = false;
+
+    private data: GetUserResponse = { id: "", enableIssueNotifications: false, enableExperimentalFeatures: false, claims: [], pinnedJobIds: [], dashboardSettings: { preferences: new Map() } };
+
+    private updateTimeoutId: any = undefined;
+
+    private updating = false;
+
+    private polling = false;
+
+    private canceled = new Set<number>();
+    private cancelId = 0;
+
+}
+
+const dashboard = new Dashboard();
+
+export default dashboard;
+
diff --git a/Engine/Source/Programs/Horde/HordeDashboard/Dashboard/src/backend/Datadog.ts b/Engine/Source/Programs/Horde/HordeDashboard/Dashboard/src/backend/Datadog.ts
new file mode 100644
--- /dev/null
+++ b/Engine/Source/Programs/Horde/HordeDashboard/Dashboard/src/backend/Datadog.ts
@@ -0,0 +1,40 @@
+// Copyright Epic Games, Inc. All Rights Reserved.
+
+import { datadogLogs, StatusType, Datacenter } from '@datadog/browser-logs';
+import { getSiteConfig } from '../backend/Config';
+
+export const initDatadog = () => {
+
+	const siteConfig = getSiteConfig();
+
+	if (!siteConfig.datadogClient) {
+		return;
+	}	
+
+	datadogLogs.init({
+		clientToken: siteConfig.datadogClient,
+		datacenter: Datacenter.US,
+		env: siteConfig.environment === "production" ? "prod" : "dev",
+		forwardErrorsToLogs: true,
+		sampleRate: 100
+
+	});
+
+	datadogLogs.logger.setLevel(StatusType.error);
+	datadogLogs.addLoggerGlobalContext('hordedashboard', { env: siteConfig.environment, username: siteConfig.environment === 'dev-local' ? 'local.developer' : ""});	
+
+};
+
+export const setDatadogUser = (username: string) => {
+
+	const siteConfig = getSiteConfig();
+
+	if (!siteConfig.datadogClient) {
+		return;
+	}	
+
+	datadogLogs.addLoggerGlobalContext('hordedashboard', { username: username, env: siteConfig.environment });
+};
+
+
+
diff --git a/Engine/Source/Programs/Horde/HordeDashboard/Dashboard/src/backend/Fetch.ts b/Engine/Source/Programs/Horde/HordeDashboard/Dashboard/src/backend/Fetch.ts
new file mode 100644
--- /dev/null
+++ b/Engine/Source/Programs/Horde/HordeDashboard/Dashboard/src/backend/Fetch.ts
@@ -0,0 +1,403 @@
+// Copyright Epic Games, Inc. All Rights Reserved.
+
+import { ErrorHandler, ErrorInfo } from "../components/ErrorHandler";
+import { setDatadogUser } from './Datadog';
+
+export enum ChallengeStatus {
+    Ok,
+    Unauthorized,
+    Error
+}
+
+export interface FetchRequestConfig {
+    formData?: boolean;
+    responseBlob?: boolean;
+    params?: Record<string, string | string[] | number | boolean | undefined>;
+    suppress404?: boolean;
+}
+
+export type FetchResponse = {
+    data: any;
+    status: number;
+}
+
+function handleError(errorIn: ErrorInfo, force?:boolean) {
+
+    const e = { ...errorIn, hordeError: true } as ErrorInfo;
+
+    ErrorHandler.set(e, force ?? false);
+
+}
+
+/** Fetch convenience wrapper */
+export class Fetch {
+
+    get(url: string, config?: FetchRequestConfig) {
+
+        url = this.buildUrl(url, config);
+
+        return new Promise<FetchResponse>(async (resolve, reject) => {
+
+            try {
+
+                const response = await fetch(url, this.buildRequest("GET"));
+
+                if (response.ok) {
+                    this.handleResponse(response, url, "GET", resolve, reject, config);
+                    return;
+                }
+
+                if (response.status === 404) {
+                    if (config?.suppress404) {
+                        return;
+                    }
+                }
+
+                const challenge = await this.challenge();
+
+                if (challenge === ChallengeStatus.Unauthorized) {
+                    
+                    this.login(window.location.toString());
+                    return;
+                }
+
+                if (response.status === 502) {
+                    throw new Error(`502 Gateway error connecting to server`);
+                }
+
+                throw response.statusText ? response.statusText : response.status;
+
+            } catch (reason) {
+
+                // Note: the fetch request sets error for redirects, this is opaque as per the spec
+                // so all we get is an error string and can't detect AccessDenied, etc in redirect :/
+                let message = `Possible permissions issue, ${reason}`;
+
+                handleError({
+                    reason: message,
+                    mode: "GET",
+                    url: url
+                });
+
+                reject(message);
+            };
+        });
+    }
+
+    post(url: string, data?: any, config?: FetchRequestConfig) {
+
+        url = this.buildUrl(url, config);
+
+        return new Promise<FetchResponse>((resolve, reject) => {
+
+            fetch(url, this.buildRequest("POST", data, config)).then(response => {
+
+                this.handleResponse(response, url, "POST", resolve, reject, config);
+
+            }).catch(reason => {
+
+                // Note: the fetch request sets error for redirects, this is opaque as per the spec
+                // so all we get is an error string and can't detect AccessDenied, etc in redirect :/
+                let message = `Possible permission issue, ${reason}`;
+
+                handleError({
+                    reason: message,
+                    mode: "POST",
+                    url: url
+                }, true);
+
+                reject(message);
+            });
+        });
+
+    }
+
+    put(url: string, data?: any, config?: FetchRequestConfig) {
+
+        url = this.buildUrl(url, config);
+
+        return new Promise<FetchResponse>((resolve, reject) => {
+
+            fetch(url, this.buildRequest("PUT", data)).then(response => {
+
+                this.handleResponse(response, url, "PUT", resolve, reject, config);
+
+            }).catch(reason => {
+
+                handleError({
+                    reason: reason,
+                    mode: "PUT",
+                    url: url
+                });
+
+                reject(reason);
+            });
+        });
+
+    }
+
+
+    delete(url: string, config?: FetchRequestConfig) {
+
+        url = this.buildUrl(url, config);
+
+        return new Promise<FetchResponse>((resolve, reject) => {
+
+            fetch(url, this.buildRequest("DELETE")).then(response => {
+
+                this.handleResponse(response, url, "DELETE", resolve, reject, config);
+
+            }).catch(reason => {
+
+                handleError({
+                    reason: reason,
+                    mode: "DELETE",
+                    url: url
+                });
+
+                reject(reason);
+            });
+        });
+    }
+
+    login(redirect: string) {
+
+        if (this.debugToken || this.logout) {
+            return;
+        }
+
+        window.location.assign("/api/v1/dashboard/login?redirect=" + encodeURI(redirect));
+    }
+
+
+
+    private async handleResponse(response: Response, url: string, mode: string, resolve: (value: FetchResponse | PromiseLike<FetchResponse>) => void, reject: (reason?: any) => void, config?: FetchRequestConfig) {
+
+
+        if (response.status === 401) {
+
+            return reject(response.statusText);
+        }
+
+        if (!response.ok) {
+
+            if (response.url?.indexOf("AccessDenied") !== -1) {
+
+                handleError({
+                    response: response,
+                    url: url,
+                    title: "Access Denied"
+                });
+
+            } else {
+
+                if (response.status !== 404 || !config?.suppress404) {
+
+                    handleError({
+                        response: response,
+                        url: url
+                    });
+
+                }
+            }
+
+            reject(response.statusText);
+            return;
+        }
+
+        if (config?.responseBlob) {
+            response.blob().then(blob => {
+                resolve({
+                    data: blob,
+                    status: response.status
+                });
+            }).catch(reason => {
+
+                handleError({
+                    response: response,
+                    url: url,
+                    reason: reason
+                });
+
+                reject(reason);
+            });
+        }
+        else {
+            response.clone().json().then(data => {
+                resolve({
+                    data: data,
+                    status: response.status
+                });
+            }).catch((reason) => {
+                response.clone().text().then(text => {
+                    resolve({
+                        data: text,
+                        status: response.status
+                    });
+                }).catch(reason => {
+
+                    handleError({
+                        response: response,
+                        url: url,
+                        reason: reason
+                    });
+
+                    reject(reason);
+                });
+            });
+        }
+
+    }
+
+    private buildRequest(method: "GET" | "POST" | "PUT" | "DELETE", data?: any, config?: FetchRequestConfig): RequestInit {
+
+        const headers = this.buildHeaders();
+
+        let body: any = undefined;
+
+        if (method === "POST" || method === "PUT") {
+            if (data) {
+                if (config?.formData) {
+                    body = data;
+                } else {
+                    headers['Content-Type'] = 'application/json';
+                    body = JSON.stringify(data);
+                }
+            }
+        }
+
+        // no-cors doesn't allow authorization header
+
+        const requestInit: RequestInit = {
+            method: method,
+            mode: "cors",
+            cache: "no-cache",
+            credentials: this.debugToken ? "same-origin" : 'include',
+            redirect: "error",
+            headers: headers,
+            body: body
+        };
+
+        return requestInit;
+    }
+
+    private get withCredentials(): boolean {
+        return !!this.authorization && !!this.authorization.length;
+    }
+
+    private buildHeaders(): Record<string, string> {
+
+        const headers: Record<string, string> = {};
+
+        if (this.withCredentials) {
+            headers["Authorization"] = this.authorization!;
+        }
+
+        return headers;
+    }
+
+
+    private buildUrl(url: string, config?: FetchRequestConfig): string {
+
+        while (url.startsWith("/")) {
+            url = url.slice(1);
+        }
+
+        while (url.endsWith("/")) {
+            url = url.slice(0, url.length);
+        }
+
+        url = `${this.baseUrl}/${url}`;
+
+        if (!config?.params) {
+            return url;
+        }
+
+        const keys = Object.keys(config.params).filter((key) => {
+            return config.params![key] !== undefined;
+        });
+
+        const query = keys.map((key) => {
+            const value = config.params![key]!;
+            if (Array.isArray(value)) {
+                return (value as string[]).map(v => {
+                    return encodeURIComponent(key) + '=' + encodeURIComponent(v);
+                });
+            } else {
+                return encodeURIComponent(key) + '=' + encodeURIComponent(value as any);
+            }
+
+        }).flat().join('&');
+
+        return query.length ? `${url}?${query}` : url;
+
+    }
+
+    async challenge(): Promise<ChallengeStatus> {
+
+        try {
+            const url = this.buildUrl("/api/v1/dashboard/challenge");
+
+            const result = await fetch(url, this.buildRequest("GET"));
+
+            if (result.ok) {
+                return ChallengeStatus.Ok;
+            }
+
+            if (result.status === 401 || result.status === 404) {
+
+                if (this.debugToken) {
+                    // raise error to check debug token expired
+                    handleError({
+                        reason: "Unauthorized challenge with debug token, it may have expired",
+                        mode: "GET",
+                        url: url
+                    });
+                }
+        
+                return ChallengeStatus.Unauthorized;
+            }
+        } catch (reason) {
+            console.error(reason);
+        }
+
+        return ChallengeStatus.Error;
+
+    }
+
+    // debug token
+    private setAuthorization(authorization: string | undefined) {
+        this.authorization = authorization;
+    }
+
+    setDebugToken(token?: string) {
+
+        if (!token) {
+
+            this.debugToken = undefined;
+            this.setAuthorization(undefined);    
+            return;
+        }
+
+        this.debugToken = token;
+        this.setAuthorization(`Bearer ${this.debugToken}`);
+
+    }
+
+    setBaseUrl(url?: string) {
+
+        if (!url) {
+            this.baseUrl = "";
+            return;
+        }
+
+        this.baseUrl = url;        
+    }
+
+    logout:boolean = false;
+
+    private baseUrl = "";
+    private debugToken?: string;
+    private authorization?: string;
+
+}
\ No newline at end of file
diff --git a/Engine/Source/Programs/Horde/HordeDashboard/Dashboard/src/backend/GraphCache.ts b/Engine/Source/Programs/Horde/HordeDashboard/Dashboard/src/backend/GraphCache.ts
new file mode 100644
--- /dev/null
+++ b/Engine/Source/Programs/Horde/HordeDashboard/Dashboard/src/backend/GraphCache.ts
@@ -0,0 +1,88 @@
+// Copyright Epic Games, Inc. All Rights Reserved.
+
+import backend from '.';
+import { GetGraphResponse } from './Api';
+import { LocalCache } from './LocalCache';
+
+export type GraphQuery = {
+	jobId: string;
+	graphHash: string;
+}
+
+export class GraphCache {
+
+	getGraphs(query: GraphQuery[]): Promise<GetGraphResponse[]> {
+
+		return new Promise<GetGraphResponse[]>((resolve, reject) => {
+
+			const results: Map<string, GetGraphResponse> = new Map();
+
+			const unresolved = query.filter(q => {
+
+				const graph = this.cache.get(q.graphHash);
+				if (graph) {
+					results.set(q.graphHash, graph);
+					return false;
+				}
+				return true;
+			});
+
+			Promise.all(unresolved.map(q => {
+				return this.get(q);
+			})).then(values => {
+				resolve(query.map(q => this.cache.get(q.graphHash)!));
+			}).catch(reason => {
+				reject(reason);
+			});
+		});
+	}
+
+	get(query: GraphQuery): Promise<GetGraphResponse> {
+
+		return new Promise<GetGraphResponse>(async (resolve, reject) => {
+
+			let graph = this.cache.get(query.graphHash);
+
+			if (!graph) {
+				graph = await graphDatabase.getItem(query.graphHash);
+				if (graph) {
+					this.cache.set(query.graphHash, graph);
+				}
+			}
+
+			if (graph) {
+				return resolve(graph);
+			}
+
+			backend.getGraph(query.jobId).then(value => {
+				graph = value as GetGraphResponse;
+				this.cache.set(query.graphHash, graph);
+				graphDatabase.storeItem(query.graphHash, graph);
+				resolve(graph);
+			}).catch(reason => {
+				reject(reason);
+			});
+
+		});
+
+	}
+
+	initialize() {
+		graphDatabase.initialize();
+	}
+
+	cache: Map<string, GetGraphResponse> = new Map();
+
+}
+
+class GraphDatabase extends LocalCache<GetGraphResponse> {
+
+	constructor() {
+		super("Horde", "GraphCache", "Cache for Horde graph objects");
+	}
+
+}
+
+const graphDatabase = new GraphDatabase();
+
+export default new GraphCache();
\ No newline at end of file
diff --git a/Engine/Source/Programs/Horde/HordeDashboard/Dashboard/src/backend/IssueStore.ts b/Engine/Source/Programs/Horde/HordeDashboard/Dashboard/src/backend/IssueStore.ts
new file mode 100644
--- /dev/null
+++ b/Engine/Source/Programs/Horde/HordeDashboard/Dashboard/src/backend/IssueStore.ts
@@ -0,0 +1,85 @@
+// Copyright Epic Games, Inc. All Rights Reserved.
+
+import { action, observable } from 'mobx';
+import { IssueData } from './Api';
+import backend from '.';
+
+const updateMS = 5000;
+
+// This class is intended for high level, cross job stream issue views, for job specific issues see JobDetails
+export class IssueStore {
+
+    @observable.ref
+    issues: IssueData[] = [];
+
+    getIssues(jobId: string, stepId: string): IssueData[] {
+        // @todo: refactor for new issue API
+        return [];
+    }
+
+    setStream(streamId: string) {
+
+        if (this.streamId === streamId) {
+            return;
+        }
+
+        this.streamId = streamId;
+
+        if (this.issues.length) {
+            this.setData([]);
+        }
+
+        clearTimeout(this.timeoutID);
+        this.timeoutID = undefined;
+        this.updating = false;
+
+        if (this.streamId) {
+            this.update();
+        }
+    }
+
+    @action
+    private setData(data: IssueData[]) {
+        // @todo: do a deep compare here to avoid redundant updates
+        this.issues = data;
+    }
+
+    private update(): void {
+
+        if (!this.streamId) {
+            return;
+        }
+
+        clearTimeout(this.timeoutID);
+        this.timeoutID = setTimeout(() => { this.update(); }, updateMS);
+
+        if (this.updating) {
+            return;
+        }
+
+        this.updating = true;
+
+        const updateStream = this.streamId;
+
+        new Promise<void>(() => {
+            backend.getIssues().then(data => {
+                if (this.streamId !== updateStream) {
+                    return;
+                }
+                this.setData(data);
+            }).catch(reason => {
+                console.error(`Error getting issues: ${reason}`);
+            }).finally(() => {
+                this.updating = false;
+            });
+        });
+    }
+
+    private streamId = '';
+
+    private timeoutID?: any;
+
+    private updating = false;
+}
+
+export const issueStore = new IssueStore();
\ No newline at end of file
diff --git a/Engine/Source/Programs/Horde/HordeDashboard/Dashboard/src/backend/Issues.ts b/Engine/Source/Programs/Horde/HordeDashboard/Dashboard/src/backend/Issues.ts
new file mode 100644
--- /dev/null
+++ b/Engine/Source/Programs/Horde/HordeDashboard/Dashboard/src/backend/Issues.ts
@@ -0,0 +1,84 @@
+// Copyright Epic Games, Inc. All Rights Reserved.
+
+import { GetIssueResponse, GetLogEventResponse, IssueSeverity } from "./Api";
+
+enum IssueType {
+    Unknown,
+    Compiler
+}
+
+const getCompilerIssueSummary = (issue: GetIssueResponse, events: GetLogEventResponse[]): string => {
+
+    let summary = issue.summary;
+
+    let sourceFiles: string[] = [];
+
+    events.forEach(event => {
+
+        event.lines?.forEach(line => {
+
+            let props = line.properties;
+
+            if (props && props["file"]) {
+
+                const file = props["file"] as Record<string, string>;
+
+                if (file["type"] === "SourceFile") {
+
+                    let relativePath = file["relativePath"];
+
+                    if (relativePath) {
+
+                        const sourceFile = relativePath.replace(/^.*[\\\/]/, '');
+                        if (!sourceFiles.find(s => s === sourceFile)) {
+                            sourceFiles.push(sourceFile);
+                        }
+                    }
+                }
+            }
+        })
+    })
+
+    if (!sourceFiles.length) {
+        return summary;
+    }
+
+    summary = (issue.severity === IssueSeverity.Warning) ? "Compile warnings" : "Compile errors";
+
+    if (sourceFiles.length === 1) {
+        return summary.slice(0, summary.length - 1) + " in " + sourceFiles[0];
+    }
+
+    if (sourceFiles.length === 2) {
+        return summary + " in " + sourceFiles[0] + " and " + sourceFiles[1];
+    }
+
+    return `${summary} in ${sourceFiles[0]} and ${sourceFiles.length - 1} others`;
+
+}
+
+export const getIssueSummary = (issue?: GetIssueResponse, events?: GetLogEventResponse[]): string => {
+
+    if (!issue) {
+        return "Issue is undefined";
+    }
+
+    let summary = issue.summary;
+
+    let type = IssueType.Unknown;
+
+    // @todo: would be better to have type meta data in issue
+    if (summary.startsWith("Compile warning") || summary.startsWith("Compile error")) {
+        type = IssueType.Compiler;    
+    }
+
+    if (type === IssueType.Unknown || !summary || !events || !events.length) {
+        return summary ? summary : issue.details ?? "Issue has no summary or details";
+    }
+
+    if (type === IssueType.Compiler) {
+        return getCompilerIssueSummary(issue, events);
+    }
+
+    return "";
+}
\ No newline at end of file
diff --git a/Engine/Source/Programs/Horde/HordeDashboard/Dashboard/src/backend/JobDetails.ts b/Engine/Source/Programs/Horde/HordeDashboard/Dashboard/src/backend/JobDetails.ts
new file mode 100644
--- /dev/null
+++ b/Engine/Source/Programs/Horde/HordeDashboard/Dashboard/src/backend/JobDetails.ts
@@ -0,0 +1,1142 @@
+// Copyright Epic Games, Inc. All Rights Reserved.
+
+import { action, observable } from 'mobx';
+import { getTheme, mergeStyles, mergeStyleSets } from 'office-ui-fabric-react/lib/Styling';
+import backend from '.';
+import { getBatchInitElapsed, getNiceTime, getStepElapsed, getStepETA, getStepFinishTime, getStepTimingDelta } from '../base/utilities/timeUtils';
+import { getBatchText } from '../components/JobDetailCommon';
+import { AgentData, ArtifactData, BatchData, EventData, GetGroupResponse, GetJobStepRefResponse, GetJobTimingResponse, GetLabelResponse, GetLabelStateResponse, GetLabelTimingInfoResponse, GetTemplateResponse, GroupData, IssueData, JobData, JobStepBatchState, JobStepOutcome, JobStepState, LabelState, NodeData, ReportPlacement, ReportScope, StepData, StreamData, TestData } from './Api';
+import { projectStore } from './ProjectStore';
+
+const theme = getTheme();
+
+export type JobLabel = GetLabelResponse & {
+    stateResponse: GetLabelStateResponse;
+    timing?: GetLabelTimingInfoResponse;
+    default: boolean;
+}
+
+const defaultUpdateMS = 5000;
+
+export class JobDetails {
+
+    constructor(id?: string, logId?: string, stepId?: string, isLogView?: boolean) {
+
+        if (id) {
+            this.set(id, logId, stepId, undefined);
+        }
+
+        this.isLogView = isLogView;
+    }
+
+    get name(): string | undefined {
+        return this.jobdata?.name;
+    }
+
+    async set(id: string, logId: string | undefined = undefined, stepId: string | undefined = undefined, labelIdx: number | undefined, updateMS = defaultUpdateMS, callback?: (details: JobDetails) => void, eventsCallback?: (details: JobDetails) => void) {
+
+        if (this.id === id) {
+
+            this.updateMS = updateMS;
+            this.updateCallback = callback;
+            this.eventsCallback = eventsCallback;
+
+            // step id always defines log id when provided
+            let logId = this.logId;
+            if (stepId) {
+                logId = this.stepById(stepId)?.logId;
+            }
+
+            if (typeof (labelIdx) === 'number') {
+                stepId = logId = undefined;
+            }
+
+            if (this.stepId === stepId && this.logId === logId && this.labelIdx === labelIdx) {
+                return;
+            }
+
+            const requests = [];
+
+            let historyIdx = -1;
+
+            if (this.stepId !== stepId || this.labelIdx !== labelIdx) {
+
+                this.stepId = stepId;
+                this.labelIdx = labelIdx;
+                requests.push(this.getIssues());
+
+                this.history = undefined;
+
+                if (stepId) {
+                    historyIdx = requests.length;
+                    requests.push(backend.getJobStepHistory(this.jobdata!.streamId, this.getStepName(stepId), 1024, this.jobdata!.templateId!));
+                }
+            }
+
+            if (logId && logId !== this.logId) {
+                requests.push(this.getLogEvents(logId));
+            }
+
+            if (requests.length) {
+                await Promise.all(requests as any).then(async (values) => {
+
+                    if (historyIdx !== -1) {
+                        this.history = values[historyIdx] as GetJobStepRefResponse[];
+                    }
+                }).catch(reason => {
+                    console.error(reason);
+                });
+            }
+
+            this.logId = logId;
+
+            if (eventsCallback) {
+                eventsCallback(this);
+            }
+
+            this.externalUpdate();
+
+            return;
+
+        }
+
+        this.clear();
+
+        this.id = id;
+        this.updateMS = updateMS;
+        this.logId = logId;
+        this.stepId = stepId;
+        this.labelIdx = labelIdx;
+        this.updateTimingInfo = 0;
+        this.updateCallback = callback;
+        this.eventsCallback = eventsCallback;
+        this.update();
+    }
+
+    setJobData(data: JobData) {
+        this.jobdata = data;
+        this.events = [];
+        this.process();
+    }
+
+    stepsByLabel(label: GetLabelResponse): StepData[] {
+
+        const steps: StepData[] = [];
+
+        this.getSteps().forEach(step => {
+
+            const node = this.nodeByStepId(step.id);
+
+            if (!node) {
+                return;
+            }
+
+            if (label.includedNodes.indexOf(node.name) !== -1 && !steps.find(s => s.id === step.id)) {
+                steps.push(step);
+            }
+
+        })
+
+        return steps;
+    }
+
+    stepById(id?: string): StepData | undefined {
+
+        if (!id || !this.id || !this.batches.length) {
+            return undefined;
+        }
+
+        let step: StepData | undefined;
+
+        this.batches.forEach(b => {
+
+            if (step) {
+                return;
+            }
+
+            step = b.steps.find(s => s.id === id);
+
+        });
+
+        return step;
+    }
+
+    stepByLogId(logId?: string): StepData | undefined {
+
+        if (!logId || !this.id || !this.batches.length) {
+            return undefined;
+        }
+
+        let step: StepData | undefined;
+
+        this.batches.forEach(b => {
+
+            if (step) {
+                return;
+            }
+
+            step = b.steps.find(s => s.logId === logId);
+
+        });
+
+        return step;
+    }
+
+
+    get aborted(): boolean {
+
+        return this.jobdata?.abortedByUser ? true : false;
+    }
+
+    stepByName(name: string): StepData | undefined {
+
+        if (!this.id || !this.batches.length) {
+            return undefined;
+        }
+
+        let step: StepData | undefined;
+
+        this.batches.forEach(b => {
+
+            if (step) {
+                return;
+            }
+
+            step = b.steps.find(s => this.nodeByStepId(s.id)?.name === name);
+
+        });
+
+        return step;
+    }
+
+    getRetryStepId(stepId: string): string | undefined {
+
+        const steps = this.getStepRetries(stepId);
+
+        if (steps.length <= 1) {
+            return undefined;
+        }
+
+        const step = this.stepById(stepId);
+
+        if (!step) {
+            return undefined;
+        }
+
+        const idx = steps.indexOf(step);
+
+        if (idx === -1 || idx === steps.length - 1) {
+            return undefined;
+        }
+
+        return steps[idx + 1]?.id;
+
+    }
+
+    getStepRetries(stepId: string): StepData[] {
+
+        const node = this.nodeByStepId(stepId);
+
+        if (!node) {
+            return [];
+        }
+
+        let steps: StepData[] = [];
+
+        this.batches.forEach(b => {
+            b.steps.filter(s => this.nodeByStepId(s.id) === node).forEach(s => {
+                steps.push(s);
+            });
+        })
+
+        return steps;
+
+    }
+
+    getStepRetryNumber(stepId: string): number {
+
+        const steps = this.getStepRetries(stepId);
+
+        if (!steps.length) {
+            return 0;
+        }
+
+        const step = this.stepById(stepId);
+
+        if (!step) {
+            return 0;
+        }
+
+        const idx = steps.indexOf(step);
+
+        return idx === -1 ? 0 : idx;
+
+    }
+
+
+    getNodeGroupIdx(node: NodeData): number {
+        const group = this.groups.find(g => g.nodes.indexOf(node) !== -1);
+        if (!group) {
+            return -1;
+        }
+        return this.groups.indexOf(group);
+    }
+
+
+    getTargets(): NodeData[] {
+
+        return this.nodes.filter(n => n.target);
+
+    }
+
+    batchByStepId(stepId: string | undefined): BatchData | undefined {
+        if (!stepId) {
+            return undefined;
+        }
+        return this.batches.find(batch => {
+            return batch.steps.find(step => step.id === stepId) ? true : false;
+        });
+
+    }
+
+    batchById(batchId: string | undefined): BatchData | undefined {
+        if (!batchId) {
+            return undefined;
+        }
+        return this.batches.find(batch => batch.id === batchId);
+    }
+
+    batchesByAgentId(agentId: string): BatchData[] {
+        return this.batches.filter(b => b.agentId === agentId);
+    }
+
+
+    nodeByStepId(stepId: string | undefined): NodeData | undefined {
+
+        if (!stepId) {
+            return undefined;
+        }
+
+        const batch = this.batchByStepId(stepId);
+        const step = this.stepById(stepId);
+
+        if (!batch || !step) return undefined;
+
+        return this.groups[batch.groupIdx].nodes[step.nodeIdx];
+
+    }
+
+    getStepName(stepId: string | undefined): string {
+
+        if (!stepId) {
+            return "";
+        }
+
+        const node = this.nodeByStepId(stepId);
+        if (!node) {
+            return "";
+        }
+
+        const idx = this.getStepRetryNumber(stepId);
+        if (!idx) {
+            return node.name;
+        }
+
+        return `${node.name} (${idx + 1})`;
+
+    }
+
+    getStepTestData(stepId: string | undefined): TestData[] {
+        if (!stepId || !this.testdata) {
+            return [];
+        }
+
+        return this.testdata.filter(t => t.stepId === stepId);
+    }
+
+    getLabelIndex(labelIn: JobLabel): number {
+        const label = this.findLabel(labelIn.name, labelIn.category);
+        if (!label) {
+            return -1;
+        }
+        return this.labels.indexOf(label);
+    }
+
+    labelByIndex(idxIn: number | string | undefined | null): JobLabel | undefined {
+
+        const idx = parseInt((idxIn ?? "").toString());
+        if (isNaN(idx)) {
+            return undefined;
+        }
+
+        const labels = this.labels;
+
+        if (!labels || idx >= labels.length) {
+            return undefined;
+        }
+
+        return labels[idx];
+    }
+
+    findLabel(name: string, category?: string): JobLabel | undefined {
+
+        const labels = this.labels.filter(label => label.stateResponse.state !== LabelState.Unspecified);
+
+        const label = labels.find(label => {
+            if (category) {
+                return category === label.category && name === label.name;
+            }
+            return name === label.name;
+        });
+
+        return label;
+
+    }
+
+    getLogIds(): string[] {
+        return this.getSteps().filter(s => !!s.logId).map(s => s.logId!);
+    }
+
+    getSteps(): StepData[] {
+        return this.batches.map(b => b.steps).flat();
+    }
+
+    eventsByStep(stepId: string): EventData[] {
+
+        const step = this.stepById(stepId);
+        if (!step) {
+            return [];
+        }
+
+        return this.events.filter(e => e.logId === step.logId!);
+
+    }
+
+
+    groupByNodeName(name: string | undefined): GetGroupResponse | undefined {
+
+        if (!name) {
+            return undefined;
+        }
+
+        const groups = this.jobdata?.graphRef?.groups;
+
+        return groups?.find(g => {
+            return g?.nodes.find(n => n.name.toLowerCase() === name.toLowerCase());
+        });
+
+    }
+
+
+    resumePolling() {
+
+        if (!this.id) {
+            return;
+        }
+
+        clearTimeout(this.timeoutID);
+        this.updating = false;
+
+        if (this.updateMS) {
+            this.timeoutID = setTimeout(() => { this.update(); }, this.updateMS);
+        }
+
+    }
+
+    stopPolling() {
+        clearTimeout(this.timeoutID);
+        this.updating = false;
+    }
+
+    cancel() {
+        this.updateTime = undefined;
+        for (let i = 0; i < this.cancelID; i++) {
+            this.canceled.add(i);
+        }
+    }
+
+    clear() {
+        this.stopPolling();
+        this.updateTime = undefined;
+        this.updateTimingInfo = 0;
+        this.jobdata = undefined;
+        this.id = undefined;
+        this.logId = undefined;
+        this.stepId = undefined;
+        this.labelIdx = undefined;
+        this.groups = [];
+        this.batches = [];
+        this.events = [];
+        this.issues = [];
+        this.history = undefined;
+        this.labels = [];
+        this.template = undefined;
+        this.fatalError = undefined;
+        this.state = JobStepState.Waiting;
+        this.reportData = new Map<string, string>();
+        this.cancel();
+        this.eventsCallback = undefined;
+        this.updateCallback = undefined;
+
+    }
+
+    private updateTimimg() {
+
+        const timing = this.timing;
+
+        if (!timing) {
+            return;
+        }
+
+        this.labels.forEach((label, index) => label.timing = timing.labels[index] ? timing.labels[index] : undefined);
+
+        this.getSteps().forEach(step => step.timing = timing.steps[step.id] ? timing.steps[step.id] : undefined);
+
+    }
+
+    private process() {
+
+        const jobdata = this.jobdata;
+
+        if (!jobdata) {
+            return;
+        }
+
+        this.id = jobdata.id;
+        this.stream = projectStore.streamById(jobdata.streamId);
+        this.batches = jobdata.batches ?? [];
+        this.groups = jobdata.graphRef?.groups ?? [];
+        this.nodes = [];
+        if (this.groups) {
+            this.nodes = this.groups.map(g => g.nodes).flat();
+        }
+
+        let labels: GetLabelResponse[] = [];
+        if (jobdata.graphRef?.labels && jobdata.labels) {
+            labels = jobdata.graphRef.labels;
+        }
+
+        this.labels = labels.map((label, index) => {
+
+            return {
+                category: label.category,
+                name: label.name,
+                includedNodes: label.includedNodes,
+                requiredNodes: label.requiredNodes,
+                default: false,
+                stateResponse: jobdata.labels![index],
+                internal: label,
+                timing: undefined
+            }
+        });
+
+        const defaultLabel = jobdata.defaultLabel;
+        if (defaultLabel) {
+            this.labels.push({
+                category: "Other",
+                name: "Other",
+                includedNodes: defaultLabel.nodes,
+                requiredNodes: [],
+                stateResponse: defaultLabel,
+                default: true
+            })
+        }
+
+        this.updateTimimg();
+
+    }
+
+    getLogActive(logId: string): boolean {
+
+        const batch = this.batches.find((b) => b.logId === logId);
+        if (batch) {
+            return batch.state === JobStepBatchState.Running;
+        }
+
+        const step = this.getSteps().find(s => s.logId === logId);
+        if (step) {
+            return step.state === JobStepState.Running;
+        }
+
+        return false;
+
+    }
+
+    private getLogEvents(logId: string) {
+
+        return new Promise<EventData[]>((resolve, reject) => {
+
+            this.events = this.events.filter(event => event.logId !== logId);
+
+            backend.getLogEvents(logId).then(async (results) => {
+
+                this.events.push(...results);
+
+                resolve(results)
+
+            });
+
+        })
+    }
+
+    private getIssues() {
+
+        return new Promise<void>((resolve, reject) => {
+
+            let stepId = this.stepId;
+            let labelIdx = this.labelIdx;
+
+            if (typeof (labelIdx) === 'number') {
+                stepId = undefined;
+            }
+
+            this.issues = [];
+            backend.getIssues({ jobId: this.id, stepId: stepId, label: this.labelIdx, count: 50 }).then(async (issues) => {
+
+                let nstepId = this.stepId;
+                let nlabelIdx = this.labelIdx;
+                if (typeof (nlabelIdx) === 'number') {
+                    nstepId = undefined;
+                }
+
+                let reqIssues = [...issues];
+
+                while (reqIssues.length) {
+
+                    await Promise.all(reqIssues.slice(0, 4).map(async (issue) => {
+                        issue.events = await backend.getIssueEvents(issue.id, this.id);
+                    }))
+
+                    reqIssues = reqIssues.slice(4);
+                }
+
+
+                // if hasn't changed since request
+                if (stepId === nstepId && labelIdx === nlabelIdx) {
+                    // sort descending
+                    this.issues = issues.sort((a, b) => b.id - a.id);
+
+                }
+
+                resolve();
+            }).catch(reason => {
+                console.error(reason);
+                resolve();
+            });
+        })
+    }    
+
+    getReportData(scope: ReportScope, placement: ReportPlacement): string | undefined {
+        
+        // placement and scope not currently returned for reports
+
+        if (scope === ReportScope.Job && placement === ReportPlacement.Summary) {
+            if (this.jobdata?.reports?.length) {
+                return this.reportData.get(this.jobdata.reports[0].artifactId);
+            }
+        }
+
+        return undefined;
+
+        /*
+        const report = this.jobdata?.reports?.find(r => r.placement === placement && r.scope === scope);
+
+        if (!report) {
+            return undefined;
+        }
+        return this.reportData.get(report.artifactId);
+        */
+    }
+
+    private async queryReports():Promise<void> {
+
+        if (!this.jobdata?.reports?.length) {
+            return;
+        }
+
+        for (let i = 0; i < this.jobdata.reports.length; i++) {
+
+            const report = this.jobdata.reports[i];
+
+            if (!this.reportData.has(report.artifactId)) {
+
+                const r = await backend.getArtifactDataById(report.artifactId) as unknown as string;
+                
+                this.reportData.set(report.artifactId, r);                                
+            }
+        }
+    }
+
+    private async update(force: boolean = false) {
+
+        if (!this.id) {
+            return;
+        }
+
+        clearTimeout(this.timeoutID);
+        if (this.updateMS) {
+            this.timeoutID = setTimeout(() => { this.update(); }, this.updateMS);
+        }
+
+        if (!force && (this.updating /*|| this.jobdata?.state === JobState.Complete*/)) {
+            return;
+        }
+
+        try {
+
+            if (typeof (this.labelIdx) === 'number') {
+                this.stepId = this.logId = undefined;
+            }
+
+            this.updating = true;
+            const cancelID = this.cancelID++;
+
+            let requests = [];
+
+            requests.push(backend.getJob(this.id));
+
+            const doTiming = !this.updateTimingInfo;
+            this.updateTimingInfo++;
+            this.updateTimingInfo %= 3;
+
+            if (doTiming) {
+                requests.push(backend.getJobTiming(this.id));
+            }
+
+            let results: any;
+
+            await Promise.all(requests as any).then(r => results = r).catch(reason => {
+                console.error("Error getting job " + this.id)
+                this.updating = false;
+                reason += ` : ${new Error().stack}}`
+                reason += "\n\nCheck browser console for more details."
+
+                this.fatalError = reason;
+                this.stopPolling();
+                if (this.updateCallback) {
+                    this.updateCallback(this);
+                }
+                this.updateReady();
+            });
+
+            if (this.fatalError) {
+                return;
+            }
+
+            if (this.canceled.has(cancelID)) {
+                this.updating = false;
+                return;
+            }
+
+            const jobdata = this.jobdata = results[0] as JobData;
+
+            if (doTiming) {
+                this.timing = results[1] as GetJobTimingResponse;
+            }
+
+            // clear
+            requests = [];
+
+            if (jobdata.updateTime) {
+                const updateTime = new Date(jobdata.updateTime);
+                if (!this.updateTime || this.updateTime < updateTime) {
+                    this.updateTime = updateTime;
+                } else {
+                    this.updating = false;
+                    if (doTiming) {
+                        this.updateTimimg();
+                        this.updateReady();
+                    }
+                    return;
+                }
+            }
+
+            let stepName = "";
+            const batch = jobdata.batches?.find(b => b.steps.find(s => s.id === this.stepId));
+            const stepNode = batch?.steps.find(s => s.id === this.stepId);
+            const groups = this.jobdata?.graphRef?.groups;
+            if (groups && stepNode && batch) {
+                stepName = groups[batch.groupIdx]?.nodes[stepNode.nodeIdx]?.name;
+            }
+
+
+            if (!this.isLogView) {
+                requests.push(backend.getJobArtifacts(this.id!));
+                requests.push(backend.getJobTestData(this.id!));
+                requests.push(this.getIssues());
+                requests.push(this.queryReports());
+            }
+
+            // sync step/log
+            let logId = this.logId;
+            if (stepNode) {
+                logId = stepNode.logId;
+            }
+
+            if (logId) {
+                this.logId = logId;
+                requests.push(this.getLogEvents(logId));
+            }
+
+            this.template = JobDetails.templates.get(this.jobdata!.templateHash!);
+
+            let templateIdx = -1;
+            if (!this.template && !this.isLogView) {
+                templateIdx = requests.length;
+                requests.push(backend.getJobTemplate(this.jobdata!.id));
+            }
+
+            let historyIdx = -1;
+            if (stepName) {
+                historyIdx = requests.length;
+                requests.push(backend.getJobStepHistory(jobdata.streamId, stepName, 1024, jobdata.templateId!));
+            }
+
+            Promise.all(requests as any).then(async (values) => {
+
+                if (this.canceled.has(cancelID)) {
+                    return;
+                }
+
+                this.artifacts = values[0] as any;
+
+                if (!this.isLogView) {
+                    this.testdata = values[1] as any;
+                }
+
+                if (historyIdx !== -1) {
+                    this.history = values[historyIdx] as any;
+                }
+
+                if (templateIdx !== -1) {
+                    const template = values[templateIdx] as GetTemplateResponse;
+                    JobDetails.templates.set(this.jobdata!.id, template);
+                    this.template = template;
+
+                }
+
+                this.process();
+
+                const agentIds: string[] = [];
+                this.batches.map(b => b.agentId).forEach(id => {
+
+                    if (!id) {
+                        return;
+                    }
+                    if (!this.agents.has(id) && agentIds.indexOf(id) === -1) {
+                        agentIds.push(id);
+                    }
+                });
+
+                if (logId) {
+                    if (this.eventsCallback) {
+                        this.eventsCallback(this);
+                    }
+                }
+
+                if (this.updateCallback) {
+                    this.updateCallback(this);
+                }
+
+                // hm, we should detect dirty
+                this.updateReady();
+
+
+            }).catch(reason => {
+                console.log(`Error updating job detail:\n${reason}`);
+            }).finally(() => {
+                this.updating = false;
+            });
+        }
+        catch (reason) {
+            console.error(reason);
+            this.updating = false;
+        }
+    }
+
+    @observable
+    updated = 0
+
+    @action
+    externalUpdate() {
+        this.updated++;
+    }
+
+    @action
+    private updateReady() {
+        this.updated++;
+    }
+
+    updating = false;
+    updateMS = defaultUpdateMS;
+
+    id?: string
+    logId?: string
+    stepId?: string
+    labelIdx?: number
+
+    jobdata?: JobData
+    stream?: StreamData
+    labels: JobLabel[] = []
+    groups: GroupData[] = []
+    nodes: NodeData[] = []
+    batches: BatchData[] = []
+    events: EventData[] = []
+    artifacts: ArtifactData[] = []
+    testdata: TestData[] = []
+    issues: IssueData[] = [];
+    history?: GetJobStepRefResponse[] = undefined;
+    private timing?: GetJobTimingResponse;
+    agents: Map<string, AgentData> = new Map();
+    template?: GetTemplateResponse;
+
+    outcome: JobStepOutcome = JobStepOutcome.Failure;
+    state: JobStepState = JobStepState.Waiting;
+
+    // artifact id => report contents (markdown)
+    private reportData = new Map<string, string>();
+
+    timeoutID?: any;
+
+    updateTime?: Date;
+
+    updateCallback?: (details: JobDetails) => void;
+    eventsCallback?: (details: JobDetails) => void;
+
+    canceled = new Set<number>();
+    cancelID = 0;
+
+    // temporary hack for issue events on prod
+    isLogView?: boolean;
+
+    updateTimingInfo = 0;
+
+    fatalError?: string;
+
+    private static templates = new Map<string, GetTemplateResponse>();
+}
+
+
+export type DetailStyle = {
+    iconName: string;
+    className: string;
+}
+
+const iconClass = mergeStyles({
+    fontSize: 13
+});
+
+
+export const detailClassNames = mergeStyleSets({
+    success: [{ color: theme.palette.green, userSelect: "none" }, iconClass],
+    warnings: [{ color: "#F7D154", userSelect: "none" }, iconClass],
+    failure: [{ color: "#EC4C47", userSelect: "none" }, iconClass],
+    waiting: [{ color: "#A19F9D", userSelect: "none" }, iconClass],
+    ready: [{ color: "#A19F9D", userSelect: "none" }, iconClass],
+    skipped: [{ color: "#F3F2F1", userSelect: "none" }, iconClass],
+    aborted: [{ color: "#F3F2F1", userSelect: "none" }, iconClass],
+    running: [{ color: theme.palette.blueLight, userSelect: "none" }, iconClass]
+});
+
+export const getDetailStyle = (state: JobStepState, outcome: JobStepOutcome): DetailStyle => {
+
+    let className = "";
+    let iconName = "";
+
+    if (state === JobStepState.Running) {
+        className = outcome === JobStepOutcome.Warnings ? detailClassNames.warnings : outcome === JobStepOutcome.Failure ? detailClassNames.failure : detailClassNames.running;
+        iconName = "Square";
+    }
+
+    if (state === JobStepState.Waiting || state === JobStepState.Ready) {
+        className = detailClassNames.ready;
+        iconName = "Square";
+    }
+
+    if (state === JobStepState.Skipped) {
+        className = detailClassNames.skipped;
+        iconName = "Square";
+    }
+
+    if (state === JobStepState.Aborted) {
+        className = detailClassNames.aborted;
+        iconName = "Square";
+    }
+
+
+    if (state === JobStepState.Completed) {
+
+        if (outcome === JobStepOutcome.Success) {
+            iconName = "Square";
+            className = detailClassNames.success;
+        }
+        else if (outcome === JobStepOutcome.Warnings) {
+            iconName = "Square";
+            className = detailClassNames.warnings;
+        } else {
+            iconName = "Square";
+            className = detailClassNames.failure;
+        }
+
+    }
+
+    return {
+        className: className,
+        iconName: iconName
+    };
+
+};
+
+export const getBatchSummaryMarkdown = (jobDetails:JobDetails, batchId: string): string => {
+
+	const batch = jobDetails.batchById(batchId);
+
+    if (!batch) {
+        return "";
+    }
+
+
+    const group = jobDetails.groups[batch.groupIdx];
+    const agentType = group?.agentType ?? "";
+    const agentPool = jobDetails.stream?.agentTypes[agentType!]?.pool ?? "";
+
+	let summaryText = "";
+	if (batch && batch.agentId) {
+
+		if (batch.state === JobStepBatchState.Running || batch.state === JobStepBatchState.Complete) {
+			let duration = "";
+			if (batch.startTime) {
+				duration = getBatchInitElapsed(batch);
+				summaryText = `Batch started ${getNiceTime(batch.startTime)} and ${batch.finishTime ? "ran" : "is running"} on [${batch.agentId}](?batch=${batchId}&agentId=${encodeURIComponent(batch.agentId)}) with a setup time of ${duration}.`;
+			}
+		}
+	} 
+
+	if (!summaryText) {
+		summaryText = getBatchText({ batch: batch, agentId: batch?.agentId, agentType: agentType, agentPool: agentPool }) ?? "";		
+		if (!summaryText) {
+			summaryText = batch?.agentId ?? (batch?.state ?? "Batch is unassigned");			
+		}
+	}
+
+    return summaryText;
+
+}
+
+export const getStepSummaryMarkdown = (jobDetails:JobDetails, stepId: string): string => {
+
+    const step = jobDetails.stepById(stepId)!;
+    const batch = jobDetails.batchByStepId(stepId);
+
+    if (!step || !batch) {
+        return "";
+    }
+
+    const duration = getStepElapsed(step);
+    let eta = getStepETA(step, jobDetails.jobdata!);
+
+    const text: string[] = [];
+
+    if (jobDetails.jobdata) {
+        text.push(`Job created by ${jobDetails.jobdata.startedByUser ? jobDetails.jobdata.startedByUser : "scheduler"}`);
+    }
+
+    const batchText = () => {
+
+        if (!batch) {
+            return undefined;
+        }
+
+        const group = jobDetails.groups[batch!.groupIdx];
+        const agentType = group?.agentType;
+        const agentPool = jobDetails.stream?.agentTypes[agentType!]?.pool;
+        return getBatchText({ batch: batch, agentType: agentType, agentPool: agentPool });
+
+    };
+
+    if (step.retryByUser) {
+        const retryId = jobDetails.getRetryStepId(step.id);
+        if (retryId) {
+            text.push(`[Step was retried by ${step.retryByUser}](/job/${jobDetails.id!}?step=${retryId})`);
+        } else {
+            text.push(`Step was retried by ${step.retryByUser}`);
+        }
+
+    }
+
+    if (step.state === JobStepState.Skipped) {
+        eta.display = eta.server = "";
+        text.push("The step was skipped");
+    } else if (step.state === JobStepState.Aborted) {
+        eta.display = eta.server = "";
+        let aborted = "";
+        if (step.abortByUser) {
+            aborted = "This step was aborted";
+            aborted += ` by ${step.abortByUser}.`;
+        }
+        else if (jobDetails.jobdata?.abortedByUser) {
+            aborted = "The job was aborted";
+            aborted += ` by ${jobDetails.jobdata?.abortedByUser}.`;
+        }
+        text.push(aborted);
+    } else if (step.state === JobStepState.Ready || step.state === JobStepState.Waiting) {
+
+        text.push(batchText() ?? `The step is pending in ${step.state} state`);
+    } else if (batch?.agentId) {
+
+        if (step.startTime) {
+            const str = getNiceTime(step.startTime);
+            text.push(`Step started on ${str}`);
+        }
+
+        let runningText = `${step.finishTime ? "Ran" : "Running"} on [${batch.agentId}](?step=${stepId}&agentId=${encodeURIComponent(batch.agentId)})`;
+
+        if (duration) {
+            runningText += ` for ${duration.trim()}`;
+        }
+
+        if (step.finishTime && (step.outcome === JobStepOutcome.Success || step.outcome === JobStepOutcome.Warnings)) {
+            const delta = getStepTimingDelta(step);
+            if (delta) {
+                runningText += " " + delta;
+            }
+        }
+
+        if (runningText) {
+            text.push(runningText);
+        }
+
+    } else {
+        text.push(batchText() ?? "Step does not have a batch.");
+    }
+
+    if (eta.display) {
+        text.push(`Expected to finish around ${eta.display}.`);
+    }
+
+    const finish = getStepFinishTime(step).display;
+
+    if (finish && step.state !== JobStepState.Aborted) {
+
+        let outcome = "";
+        if (step.outcome === JobStepOutcome.Success) {
+            outcome += `Completed at ${finish}`;
+        }
+        if (step.outcome === JobStepOutcome.Failure)
+            outcome += `Completed with errors at ${finish}.`;
+        if (step.outcome === JobStepOutcome.Warnings)
+            outcome += `Completed with warnings at ${finish}.`;
+
+        if (outcome) {
+            text.push(outcome);
+        }
+    }
+
+    if (!text.length) {
+        text.push(`Step is in ${step.state} state.`);
+    }
+
+    return text.join(".&nbsp;&nbsp;");
+
+}
+
diff --git a/Engine/Source/Programs/Horde/HordeDashboard/Dashboard/src/backend/JobEventHandler.ts b/Engine/Source/Programs/Horde/HordeDashboard/Dashboard/src/backend/JobEventHandler.ts
new file mode 100644
--- /dev/null
+++ b/Engine/Source/Programs/Horde/HordeDashboard/Dashboard/src/backend/JobEventHandler.ts
@@ -0,0 +1,111 @@
+
+// Copyright Epic Games, Inc. All Rights Reserved.
+
+import backend from '../backend';
+import { EventData, JobStepOutcome, JobStepState } from '../backend/Api';
+import { JobDetails } from './JobDetails';
+import { PollBase } from './PollBase';
+
+
+type StepEvents = {
+    stepId: string;
+    events: EventData[];
+    complete?: boolean;
+}
+
+// Helper class that can poll log events globally for a job
+// intended for preflights otherwise the health panel is preferred
+export class JobEventHandler extends PollBase {
+
+    constructor(pollTime = 5000) {
+
+        super(pollTime);
+
+    }
+
+    set(details:JobDetails) {
+        if (this.details === details) {
+            return;
+        }
+        this.clear();
+        this.details = details;
+        this.update();
+    }
+
+    clear() {
+        super.clear();
+        this.details = undefined;
+        this.stepEvents.clear();
+    }
+
+    async poll(): Promise<void> {
+
+        try {
+
+            const details = this.details;
+
+            if (!details) {
+                return;
+            }
+
+            // need to query any steps which are running or have run and are not marked complete
+            let queries = details.getSteps().filter(step => {
+
+                if (!step.logId || this.stepEvents.get(step.id)?.complete) {
+                    // already have completed log events for step
+                    return false;
+                }
+
+                if (step.outcome === JobStepOutcome.Success) {
+                    return false;
+                }
+
+                return true;
+
+            })
+
+            let anyQueried = false;
+
+            while (true) {
+
+                if (!queries.length) {
+                    break;
+                }
+
+                let querySteps = queries.slice(0, 3);
+                queries = queries.slice(3);
+                
+                const events = await Promise.all(querySteps.map(s => backend.getLogEvents(s.logId!)));
+
+                events.forEach((e, index) => {
+
+                    const queryStep = querySteps[index];
+
+                    this.stepEvents.set(queryStep.id, {
+                        stepId: queryStep.id,
+                        events: e,
+                        complete: (queryStep.state === JobStepState.Aborted || queryStep.state === JobStepState.Completed) ? true : false
+                    });
+    
+                });
+
+
+                anyQueried = true;
+
+            }
+
+            if (anyQueried) {
+                this.setUpdated();
+            }
+
+        } catch (err) {
+            console.error(err);
+        }
+
+    }
+
+    details?: JobDetails;
+
+    stepEvents: Map<string, StepEvents> = new Map();
+
+}
\ No newline at end of file
diff --git a/Engine/Source/Programs/Horde/HordeDashboard/Dashboard/src/backend/JobHandler.ts b/Engine/Source/Programs/Horde/HordeDashboard/Dashboard/src/backend/JobHandler.ts
new file mode 100644
--- /dev/null
+++ b/Engine/Source/Programs/Horde/HordeDashboard/Dashboard/src/backend/JobHandler.ts
@@ -0,0 +1,331 @@
+// Copyright Epic Games, Inc. All Rights Reserved.
+
+import { action, observable } from 'mobx';
+import moment from 'moment-timezone';
+import backend from '../backend';
+import { JobData, JobQuery, JobState, JobStepOutcome, StreamData } from '../backend/Api';
+import graphCache, { GraphQuery } from '../backend/GraphCache';
+
+
+export type FilterStatus = "Running" | "Complete" | "Succeeded" | "Failed" | "Waiting";
+
+const jobPageSize = 50;
+const jobsRefreshTime = 3000;
+
+export class JobHandler {
+
+    constructor(includeBatches: boolean = false, jobLimit?:number) {
+        this.includeBatches = includeBatches;
+        this.jobLimit = jobLimit;
+    }
+
+
+    filter(stream: StreamData, templateNamesIn?: string[], preflightStartedByUser?: string): boolean {
+
+        const names = templateNamesIn ?? [];
+        const sameNames = (this.templateNames.length === names.length && this.templateNames.every((value, index) => value === names[index]));
+
+        if (stream.id === this.stream?.id && sameNames && this.preflightStartedByUser === preflightStartedByUser) {
+            return false;
+        }
+
+        this.clear();
+        this.stream = stream;
+        this.templateNames = names;
+        this.preflightStartedByUser = preflightStartedByUser;
+        this.setUpdated();
+        this.update();
+        return true;
+    }
+
+    clear() {
+        clearTimeout(this.timeoutId);
+        this.timeoutId = undefined;
+        this.modifiedAfter = undefined;
+        this.stream = undefined;
+        this.templateNames = [];
+        this.jobs = [];
+        this.initial = true;
+        this.bumpCount = false;
+        this.count = jobPageSize * 2;
+        this.includePreflights = true;
+        this.preflightStartedByUser = undefined;
+
+        // cancel any pending        
+        for (let i = 0; i < this.cancelId; i++) {
+            this.canceled.add(i);
+        }
+    }
+
+    async update() {
+
+        clearTimeout(this.timeoutId);
+        this.timeoutId = setTimeout(() => { this.update(); }, jobsRefreshTime);
+
+        if (this.updating) {
+            return;
+        }
+
+        this.updating = true;
+
+        // cancel any pending        
+        for (let i = 0; i < this.cancelId; i++) {
+            this.canceled.add(i);
+        }
+
+        let wasUpdated = false;
+
+        try {            
+
+            // discover new/updated jobs
+            let filter = "id,streamId,name,change,preflightChange,templateId,templateHash,graphHash,startedByUser,abortedByUser,createTime,state,arguments,updateTime,labels,defaultLabel";
+
+            if (this.includeBatches) {
+                filter += ",batches";
+            }
+
+            // needs modified time
+            const query: JobQuery = {
+                filter: filter,
+                count: this.count,
+                streamId: this.stream?.id,
+                template: this.templateNames,
+                includePreflight: this.includePreflights,
+                preflightStartedByUser: this.preflightStartedByUser
+            };
+
+            // instance property can mutate externally as we're async
+            const bumpCount = this.bumpCount;
+
+            if (bumpCount) {
+
+                let maxCreate: Date | undefined;
+
+                // find max create time
+                this.jobs.forEach(j => {
+                    const time = new Date(j.createTime);
+                    if (!maxCreate || (time < maxCreate)) {
+                        maxCreate = time;
+                    }
+                })
+
+                if (maxCreate) {
+                    maxCreate = new Date(maxCreate.getTime() - 1);
+                    query.maxCreateTime = maxCreate.toISOString();
+                    query.count = jobPageSize;
+                }
+
+            } else {
+
+                query.modifiedAfter = this.modifiedAfter;
+
+            }
+
+            const cancelId = this.cancelId++;
+
+            const queryTime = moment.utc().toISOString();
+
+            const mjobs = await backend.getJobs(query, false);
+
+            // check for canceled after modified test
+            if (this.canceled.has(cancelId)) {
+                return;
+            }
+
+            if (bumpCount) {
+
+                let njob = false;
+                mjobs.forEach(j => {
+                    if (!this.jobs.find(j2 => j.id === j2.id)) {
+                        njob = true;
+                    }
+                })
+
+                // add page size if we found any new jobs
+                if (njob) {
+                    this.count += jobPageSize
+                }
+
+            }
+
+            const jobs: JobData[] = [];
+
+            mjobs.forEach(j1 => {
+                const existing = this.jobs.find(j2 => j1.id === j2.id);
+                if (existing) {
+                    j1.graphRef = existing.graphRef;
+                }
+                jobs.push(j1);
+            })
+
+            jobs.push(...this.jobs.filter(j1 => !jobs.find(j2 => j1.id === j2.id)));
+
+            this.jobs = jobs.sort((a, b) => {
+                const timeA = new Date(a.createTime).getTime();
+                const timeB = new Date(b.createTime).getTime();
+                if (timeA === timeB) return 0;
+                return timeA < timeB ? 1 : -1;
+            });
+
+            const graphHashes = new Set<string>();
+
+            this.jobs.forEach(j => {
+
+                if (j.graphRef?.hash !== j.graphHash) {
+                    j.graphRef = undefined;
+                }
+
+                if (graphHashes.size > 5) {
+                    return;
+                }
+
+                if (j.graphHash && !j.graphRef) {
+                    j.graphRef = graphCache.cache.get(j.graphHash);
+                    if (!j.graphRef) {
+                        graphHashes.add(j.graphHash);
+                    }
+                }
+            })
+
+            if (graphHashes.size) {
+
+                const graphQuery: GraphQuery[] = [];
+                Array.from(graphHashes.values()).forEach(h => {
+                    const jobId = this.jobs.find(j => j.graphHash === h)!.id;
+                    graphQuery.push({
+                        jobId: jobId,
+                        graphHash: h
+                    })
+                })
+
+                const graphs = await graphCache.getGraphs(graphQuery);
+
+                graphs.forEach(graph => {
+
+                    this.jobs.forEach(j => {
+                        if (graph.hash === j.graphHash) {
+                            j.graphRef = graph;
+                        }
+                    })
+                })
+
+            }
+
+            // check for canceled during graph request
+            if (this.canceled.has(cancelId)) {
+                return;
+            }
+
+            if (graphHashes.size || mjobs.length || this.initial) {
+                this.initial = false;
+                wasUpdated = true;
+            }
+
+            if (!this.bumpCount) {
+                this.modifiedAfter = queryTime;
+            }
+
+        } catch (reason) {
+            console.log(reason);
+        } finally {
+
+            this.bumpCount = false;
+
+            if (this.jobLimit) {
+
+                if (this.count === this.jobs.length && this.count < this.jobLimit) {
+                    this.bumpCount = true;
+                }    
+            }
+
+            if (wasUpdated) {
+                this.setUpdated();
+            }
+            
+            this.updating = false;
+        }
+
+    }
+
+    getFilteredJobs(status?: FilterStatus[]) {
+        if (!status || !status.length) {
+            return this.jobs;
+        }
+
+        return this.jobs.filter(job => {
+
+            return status.every(status => {
+
+                if ((status === "Complete" || status === "Succeeded") && job.state !== JobState.Complete) {
+                    return false;
+                }
+
+                if (status === "Running" && job.state !== JobState.Running) {
+                    return false;
+                }
+
+                if (status === "Waiting" && job.state !== JobState.Waiting) {
+                    return false;
+                }
+
+                if (status === "Succeeded") {
+                    if (!job.batches?.every(b => b.steps.every(s => s.outcome === JobStepOutcome.Success || s.outcome === JobStepOutcome.Warnings))) {
+                        return false;
+                    }
+                }
+
+                if (status === "Failed") {
+                    if (job.batches?.every(b => b.steps.every(s => s.outcome === JobStepOutcome.Success || s.outcome === JobStepOutcome.Warnings))) {
+                        return false;
+                    }
+                }
+
+                return true;
+            })
+
+        });
+    }
+
+
+    @action
+    setUpdated() {
+        this.updated++;
+    }
+
+    @observable
+    updated: number = 0;
+
+    updating = false;
+
+    jobs: JobData[] = [];
+
+    addJobs() {
+
+        if (this.count < 1000) {
+            this.bumpCount = true;
+        }
+
+    }
+
+    private modifiedAfter?: string;
+    private stream?: StreamData;
+    private templateNames: string[] = [];
+
+    count = jobPageSize * 2;
+    bumpCount = false;
+
+    private timeoutId: any;
+
+    private canceled = new Set<number>();
+    private cancelId = 0;
+
+    includeBatches = false;
+
+    includePreflights = true;
+    private preflightStartedByUser?: string;
+
+    private jobLimit?:number;
+
+    initial = true;
+
+}
diff --git a/Engine/Source/Programs/Horde/HordeDashboard/Dashboard/src/backend/LocalCache.ts b/Engine/Source/Programs/Horde/HordeDashboard/Dashboard/src/backend/LocalCache.ts
new file mode 100644
--- /dev/null
+++ b/Engine/Source/Programs/Horde/HordeDashboard/Dashboard/src/backend/LocalCache.ts
@@ -0,0 +1,211 @@
+// Copyright Epic Games, Inc. All Rights Reserved.
+
+import * as CryptoJS from "crypto-js";
+import { openDB, DBSchema, IDBPDatabase } from 'idb';
+
+type CacheItem = {
+    storeTime: number;
+    encodedData: string;
+}
+
+interface CacheSchema extends DBSchema {
+    'cache-item': {
+        key: string;
+        value: {
+            key: string;
+            encodedData: string;
+            storeTime: number;
+        };
+        indexes: { 'by-key': string };
+    };
+}
+
+const DB_VERSION = 8;
+// @todo: make this tuneable per cache type
+const DB_MAX_ITEMS = 2048;
+
+export abstract class LocalCache<ValueType> {
+
+    async getItem(hash: string): Promise<ValueType | undefined> {
+
+        const cached = this.cache.get(hash);
+
+        if (cached) {
+            return cached;
+        }
+
+        if (!this.initialized) {
+            return undefined;
+        }
+
+        try {
+
+            const cipherHash = CryptoJS.SHA256(hash).toString();
+
+            const item = await this.db!.get("cache-item", cipherHash);
+
+            if (!item) {
+                //console.log("cache miss", cipherHash);
+                return undefined;
+            }
+
+            const bytes = CryptoJS.AES.decrypt(item.encodedData, hash);
+            const value = JSON.parse((bytes.toString(CryptoJS.enc.Utf8))) as ValueType;
+
+            this.cache.set(hash, value);
+
+            //console.log("cache hit");
+
+            return value;
+
+        } catch (err) {
+            console.error(err);
+        }
+
+        return undefined;
+
+    }
+
+    async storeItem(hash: string, value: ValueType) {
+
+        if (!this.initialized) {
+            return undefined;
+        }
+
+        try {
+
+            this.cache.set(hash, value);
+
+            const cipherHash = CryptoJS.SHA256(hash).toString();
+            const encodedData = CryptoJS.AES.encrypt(JSON.stringify(value), hash).toString();
+
+            //console.log("cache store", cipherHash);
+
+            await this.db!.put("cache-item", { key: cipherHash, storeTime: Date.now(), encodedData: encodedData }, cipherHash)
+
+        } catch (err) {
+            console.error(err);
+        }
+
+    }
+
+    async evict(): Promise<boolean> {
+
+        try {
+
+            if (this.verbose) {
+                console.log(`Cache ${this.name} : ${this.storeName}, checking for eviction`);
+            }
+
+            const count = await this.db!.count("cache-item");
+
+            if (count >= DB_MAX_ITEMS) {
+
+                console.log(`Cache: ${this.storeName} evicting data`)
+
+                await this.db!.clear("cache-item");
+
+            }
+
+        } catch (err) {
+            console.error(err)
+            return false;
+        }
+
+        return true;
+
+    }
+
+    initialize() {
+
+        if (this.failure || this.attempted) {
+            return false;
+        }
+
+        this.attempted = true;
+
+        console.log(`Initializing local cache, ${this.name + " : " + this.storeName}`);
+
+        const innerInit = async () => {
+
+            try {
+                let success = true;
+
+                this.db = await openDB<CacheSchema>(this.storeName, DB_VERSION, {
+
+                    upgrade: (db) => {
+
+                        console.log("upgrade");
+                        try {
+                            db.deleteObjectStore('cache-item');
+                        } catch (reason) {
+                            console.log("Did not delete existing cache", reason);
+                        }
+                        
+                        const store = db.createObjectStore('cache-item');
+                        store.createIndex('by-key', 'key');
+
+                    },
+                    blocked: () => {
+                        success = false;
+                        console.log("blocked");
+                    },
+                    blocking: () => {
+                        this.initialized = false;
+                        this.db?.close();
+                        console.log("blocking");
+                    },
+                    terminated: () => {
+                        this.initialized = false;
+                        console.log("terminated");                                                
+                    }
+                });
+
+                if (success) {
+                    success = await this.evict();
+                }
+                
+                if (!success) {
+                    this.failure = true;
+                    return;
+                }
+            } catch (reason) {
+
+                this.failure = true;
+                console.error(reason);
+                return;
+            }
+
+            this.initialized = true;
+
+        };
+
+        innerInit();
+
+    }
+
+    constructor(name: string, storeName: string, description: string) {
+
+        this.name = name;
+        this.storeName = storeName;
+        this.description = description;
+
+    }
+
+    name: string;
+    storeName: string;
+    description: string;
+
+    initialized = false;
+
+    attempted = false;
+
+    verbose = true;
+
+    failure = false;
+
+    db?: IDBPDatabase<CacheSchema>;
+
+    cache: Map<string, ValueType> = new Map();
+
+}
\ No newline at end of file
diff --git a/Engine/Source/Programs/Horde/HordeDashboard/Dashboard/src/backend/PollBase.ts b/Engine/Source/Programs/Horde/HordeDashboard/Dashboard/src/backend/PollBase.ts
new file mode 100644
--- /dev/null
+++ b/Engine/Source/Programs/Horde/HordeDashboard/Dashboard/src/backend/PollBase.ts
@@ -0,0 +1,95 @@
+// Copyright Epic Games, Inc. All Rights Reserved.
+
+import { action, observable } from 'mobx';
+
+// base class for polling behavior
+export abstract class PollBase {
+
+    constructor(pollTime = 5000) {
+
+        this.pollTime = pollTime;
+    }
+
+    clear() {
+
+        clearTimeout(this.timeoutId);
+        this.timeoutId = undefined;
+
+        // cancel any pending        
+        for (let i = 0; i < this.cancelId; i++) {
+            this.canceled.add(i);
+        }
+
+        this.updating = false;
+    }
+
+    start() {
+
+        if (this.timeoutId || this.updating) {
+            return;
+        }
+
+        this.update();
+    }
+
+    abstract poll(): Promise<void>;
+
+    async update() {
+
+        clearTimeout(this.timeoutId);
+        this.timeoutId = setTimeout(() => { this.update(); }, this.pollTime);
+
+        if (this.updating) {
+            return;
+        }
+
+        this.updating = true;
+
+        // cancel any pending        
+        for (let i = 0; i < this.cancelId; i++) {
+            this.canceled.add(i);
+        }
+
+        try {
+
+            const cancelId = this.cancelId++;
+
+            await this.poll();
+
+            // check for canceled during graph request
+            if (this.canceled.has(cancelId)) {
+                return;
+            }
+
+
+        } catch (reason) {
+
+            console.log(reason);
+
+        } finally {
+
+            this.updating = false;
+        }
+
+    }
+
+
+
+    @action
+    setUpdated() {
+        this.updated++;
+    }
+
+    @observable
+    updated: number = 0;
+
+    updating = false;
+
+    private timeoutId: any;
+
+    private pollTime: number;
+
+    private canceled = new Set<number>();
+    private cancelId = 0;
+
+}
\ No newline at end of file
diff --git a/Engine/Source/Programs/Horde/HordeDashboard/Dashboard/src/backend/ProjectStore.ts b/Engine/Source/Programs/Horde/HordeDashboard/Dashboard/src/backend/ProjectStore.ts
new file mode 100644
--- /dev/null
+++ b/Engine/Source/Programs/Horde/HordeDashboard/Dashboard/src/backend/ProjectStore.ts
@@ -0,0 +1,109 @@
+// Copyright Epic Games, Inc. All Rights Reserved.
+
+import { action, observable } from 'mobx';
+import { ProjectData, StreamData } from './Api';
+import backend from '.';
+
+export class ProjectStore {
+
+    @observable.ref
+    projects: ProjectData[] = [];
+
+    @observable
+    activeId?: string;
+
+    @action
+    setProjects(data: ProjectData[]) {
+        this.projects = data;
+    }
+
+    @action
+    setActive(projectId: string | undefined): ProjectData | undefined {
+        if (this.activeId === projectId) {
+            return this.byId(this.activeId);
+        }
+
+        this.activeId = projectId;
+        return this.byId(this.activeId);
+    }
+
+    getActive(): ProjectData | undefined {
+        return this.projects.find(p => p.id === this.activeId);
+    }
+
+    byId(id: string | undefined): ProjectData | undefined {
+        return this.projects.find(p => p.id === id);
+    }
+
+    byName(name: string | undefined): ProjectData | undefined {
+        return this.projects.find(p => p.name === name);
+    }
+
+    firstProject(): ProjectData | undefined {
+        return this.projects.length ? this.projects[0] : undefined;
+    }
+
+    projectByStreamId(streamId: string): ProjectData | undefined {
+        return this.projects.find(p => p.streams?.find(stream => stream.id === streamId))
+    }
+
+    streamById(id: string | undefined): StreamData | undefined {
+        if (!id) {
+            return undefined;
+        }
+        let stream: StreamData | undefined;
+        this.projects.forEach(p => {
+            if (stream) {
+                return;
+            }
+            stream = p.streams?.find(s => s.id === id);
+
+        });
+        return stream;
+    }
+
+    streamByFullname(name: string): StreamData | undefined {
+
+        let stream: StreamData | undefined;
+
+        this.projects.forEach(p => {
+            if (stream) {
+                return;
+            }
+            stream = p.streams?.find(s => s.fullname?.toLowerCase() === name.toLowerCase());
+
+        });
+        return stream;
+    }
+
+
+    streamByName(project: ProjectData, streamName: string | undefined): StreamData | undefined {
+        if (!project || !project.streams) {
+            return undefined;
+        }
+        return project.streams.find(s => s.name.toLowerCase() === streamName?.toLowerCase());
+    }
+
+    firstStream(project: ProjectData): StreamData | undefined {
+        if (!project.streams || !project.streams.length) {
+            return undefined;
+        }
+        return project.streams[0];
+    }
+
+    async update(): Promise<void> {
+        return new Promise<void>((resolve, reject) => {
+            backend.getProjects().then(projects => {
+
+                this.setProjects(projects);
+                this.setActive(this.byId(this.activeId)?.id);
+
+                resolve();
+            }).catch(reason => {
+                reject(`Error getting projects: ${reason}`);
+            });
+        });
+    }
+}
+
+export const projectStore = new ProjectStore();
\ No newline at end of file
diff --git a/Engine/Source/Programs/Horde/HordeDashboard/Dashboard/src/backend/TemplateCache.ts b/Engine/Source/Programs/Horde/HordeDashboard/Dashboard/src/backend/TemplateCache.ts
new file mode 100644
--- /dev/null
+++ b/Engine/Source/Programs/Horde/HordeDashboard/Dashboard/src/backend/TemplateCache.ts
@@ -0,0 +1,85 @@
+// Copyright Epic Games, Inc. All Rights Reserved.
+
+import backend from '.';
+import { GetTemplateRefResponse, StreamData, TemplateData } from './Api';
+import { LocalCache } from './LocalCache';
+
+export class TemplateCache {
+
+	get(ref: GetTemplateRefResponse): TemplateData | undefined {
+
+		const key = ref.id + ref.hash;
+		return this.cache.get(key);
+	}
+
+	set(ref: GetTemplateRefResponse, template: TemplateData) {
+
+		const key = ref.id + ref.hash;
+		this.cache.set(key, template);
+
+	}
+
+	getStreamTemplates(stream: StreamData): Promise<TemplateData[]> {
+
+		return new Promise<TemplateData[]>(async (resolve, reject) => {
+
+			const templates: TemplateData[] = [];
+
+			const unresolved: GetTemplateRefResponse[] = [];
+
+			for (let i = 0; i < stream.templates.length; i++) {
+
+				const tref = stream.templates[i];
+
+				let template = this.get(tref);
+
+				if (!template) {
+					template = await templateDatabase.getItem(tref.id + tref.hash);
+					if (template) {						
+						this.set(tref, template);
+					}
+				}
+
+				if (template) {
+					templates.push(template);
+				} else {
+					unresolved.push(tref);
+				}
+
+			}
+
+			// get stream templates
+			Promise.all(unresolved.map(async (tref) => {
+				return backend.getTemplate(stream.id, tref.id).then((response) => {
+					response.ref = tref;
+					this.set(tref, response);
+
+					templateDatabase.storeItem(tref.id + tref.hash, response);
+
+					templates.push(response);
+				});
+			})).then(() => {
+				resolve(templates);
+			}).catch(reason => reject(reason));
+		});
+	}
+
+	initialize() {
+		templateDatabase.initialize();
+	}
+
+	cache: Map<string, TemplateData> = new Map();
+
+}
+
+class TemplateDatabase extends LocalCache<TemplateData> {
+
+	constructor() {
+		super("Horde", "TemplateCache", "Cache for Horde template objects");
+	}
+
+}
+
+const templateDatabase = new TemplateDatabase();
+
+export default new TemplateCache();
\ No newline at end of file
diff --git a/Engine/Source/Programs/Horde/HordeDashboard/Dashboard/src/backend/TestDataHandler.ts b/Engine/Source/Programs/Horde/HordeDashboard/Dashboard/src/backend/TestDataHandler.ts
new file mode 100644
--- /dev/null
+++ b/Engine/Source/Programs/Horde/HordeDashboard/Dashboard/src/backend/TestDataHandler.ts
@@ -0,0 +1,252 @@
+// Copyright Epic Games, Inc. All Rights Reserved.
+
+import { observable, action } from "mobx";
+import backend from '../backend';
+import { StreamData, JobData, ArtifactData, TestData } from './Api';
+import { projectStore } from './ProjectStore';
+import { BreadcrumbItem } from '../components/Breadcrumbs';
+
+class JobInfo {
+    stream?: StreamData;
+    jobdata?: JobData;
+    artifacts?: ArtifactData[];
+    stepName : string = "UnknownJob";
+    id?: string;
+    stepId?: string;
+
+    @action
+    setActive(active: boolean) {
+        this.active = active;
+    }
+
+    @observable
+    active = true;
+
+    async set(targetJobId : string, targetStepId : string) {
+        if (this.id === targetJobId && this.stepId === targetStepId) {
+            return;
+        }
+
+        if (!targetJobId || !targetStepId) {
+            return;
+        }
+
+        this.stepId = targetStepId;
+        this.stepName = "UnknownJob";
+
+        this.setActive(true);
+
+        try {
+            if (this.jobdata === undefined || this.id !== targetJobId) {
+                this.jobdata = await backend.getJob(targetJobId);
+                this.id = targetJobId;
+                this.stream = projectStore.streamById(this.jobdata.streamId);
+            }
+
+            let stepName = "";
+            const batch = this.jobdata.batches?.find(b => b.steps.find(s => s.id === this.stepId));
+            const stepNode = batch?.steps.find(s => s.id === this.stepId);
+            const groups = this.jobdata?.graphRef?.groups;
+            if (groups && stepNode && batch) {
+                stepName = groups[batch.groupIdx]?.nodes[stepNode.nodeIdx]?.name;
+            }
+            this.stepName = stepName;
+
+            this.artifacts = await backend.getJobArtifacts(this.id!, this.stepId);
+        }
+        catch(reason) {
+            console.error('Failed to load job info!');
+            throw reason;
+        }
+        finally {
+            this.setActive(false);
+        }
+    }   
+
+}
+
+export class TestDataHandler {
+    stream?: StreamData;
+    id?: string;
+
+    jobInfo: JobInfo = new JobInfo();
+    testdata?: TestData;
+    name?: string;
+    type?: string;
+
+    history?: TestData[];
+
+    artifactMap: Map<string, ArtifactData | undefined> = new Map();
+
+    @action
+    setActive(active: boolean) {
+        this.active = active;
+    }
+
+    @observable
+    active = true;
+
+    @action
+    setHistoryLoaded(loaded: boolean) {
+        this.historyLoaded = loaded;
+    }
+
+    @observable
+    historyLoaded = false
+
+    async set(targetId: string) {
+
+        if (this.id === targetId) {
+            return;
+        }
+
+        if (!targetId) {
+            return;
+        }
+
+        this.id = targetId;
+        this.testdata = undefined;
+        this.name = undefined;
+        this.type = undefined;
+        this.history = undefined;
+        if (this.artifactMap.size > 0) {
+            this.artifactMap = new Map();
+        }
+
+        this.setActive(true);
+
+        try {
+            this.testdata = await backend.getTestData(targetId);
+            const key = this.testdata.key.split('::');
+            if (key.length > 1) {
+                [ this.type, this.name ] = key;
+            }
+            else {
+                this.name = key[0];
+            }
+
+            await this.jobInfo.set(this.testdata.jobId, this.testdata.stepId);
+            this.stream = this.jobInfo.stream;
+        }
+        catch(reason) {
+            console.error('Failed to load test data from job!');
+            throw reason;
+        }
+        finally {
+            this.setActive(false);
+        }
+
+    }
+
+    async getHistory(fromChange? : number) {
+        if (!this.id || !this.stream || !this.testdata) {
+            return;
+        }
+
+        if (this.history !== undefined) {
+            return;
+        }
+
+        this.setHistoryLoaded(false);
+
+        try {
+            this.history = await backend.getTestDataHistory(this.stream.id, this.testdata.key, fromChange);
+        }
+        catch(reason) {
+            console.error('Failed to load test data history!');
+            throw reason;
+        }
+        finally {
+            this.setHistoryLoaded(true);
+        }
+    }
+
+    cleanHistory() {
+        if (this.history !== undefined) {
+            this.history = undefined;
+            this.setHistoryLoaded(false);
+        }
+    }
+
+    get clText(): string {
+
+        const data = this.jobInfo.jobdata;
+        if (!data) {
+            return "";
+        }
+
+        let clText = "";
+        if (data.preflightChange) {
+            clText = `Preflight ${data.preflightChange} `;
+            clText += ` - Base ${data.change ? data.change : "Latest"}`;
+
+        } else {
+            clText = `${data.change ? "CL " + data.change : "Latest CL"}`;
+        }
+        
+        return clText;
+    }
+
+    get crumbs(): BreadcrumbItem[] {
+
+        const jobdata = this.jobInfo.jobdata!;
+
+        if (!this.stream) {
+            return [];
+        }
+
+        let projectName = this.stream.project?.name;
+        if (projectName === "Engine") {
+            projectName = "UE4";
+        }
+
+        const crumbItems: BreadcrumbItem[] = [
+            {
+                text: projectName ?? "Unknown Project",
+                link: `/project/${this.stream?.project?.id}`
+            },
+            {
+                text: this.stream.name,
+                link: `/stream/${this.stream.id}`
+            },
+            {
+                text: `${this.clText}: ${jobdata?.name ?? ""}`,
+                link: `/job/${jobdata?.id}`
+            },
+            {
+                text: this.jobInfo.stepName,
+                link: this.jobInfo.stepId ? `/job/${jobdata?.id}?step=${this.jobInfo.stepId}` : `/job/${jobdata?.id}`
+            },
+            {
+                text: this.name + " (Test Report)" 
+            }
+        ];
+
+        return crumbItems;
+    }
+
+    get crumbTitle(): string | undefined {
+
+        if (!this.stream) {
+            return undefined;
+        }
+        return `Horde - ${this.clText}: ${this.jobInfo.jobdata?.name ?? ""} - ${this.jobInfo.stepName} (Test Report)`
+    }
+
+    findArtifactData(artifactName: string): ArtifactData | undefined {
+        if (!artifactName) {
+            return undefined;
+        }
+
+        artifactName = artifactName.replace(/\\/g,  '/');
+
+        if (this.artifactMap.has(artifactName)) {
+            return this.artifactMap.get(artifactName);
+        }
+
+        const found = this.jobInfo.artifacts?.find((value) => value.name.indexOf(artifactName) > -1);
+        this.artifactMap.set(artifactName, found);
+
+        return found;
+    }
+}
diff --git a/Engine/Source/Programs/Horde/HordeDashboard/Dashboard/src/backend/index.ts b/Engine/Source/Programs/Horde/HordeDashboard/Dashboard/src/backend/index.ts
new file mode 100644
--- /dev/null
+++ b/Engine/Source/Programs/Horde/HordeDashboard/Dashboard/src/backend/index.ts
@@ -0,0 +1,15 @@
+// Copyright Epic Games, Inc. All Rights Reserved.
+
+import React from 'react';
+import { Backend } from './Backend';
+import { projectStore } from './ProjectStore';
+import { issueStore } from './IssueStore';
+
+export default new Backend();
+
+const backendContext = React.createContext({
+  projectStore: projectStore,
+  issueStore: issueStore  
+});
+
+export const useBackend = () => React.useContext(backendContext);
diff --git a/Engine/Source/Programs/Horde/HordeDashboard/Dashboard/src/base/utilities/extractAnchorLink.ts b/Engine/Source/Programs/Horde/HordeDashboard/Dashboard/src/base/utilities/extractAnchorLink.ts
new file mode 100644
--- /dev/null
+++ b/Engine/Source/Programs/Horde/HordeDashboard/Dashboard/src/base/utilities/extractAnchorLink.ts
@@ -0,0 +1,24 @@
+/**
+ * Given a URL containing a route path (first hash) and possibly an anchor (second hash),
+ * returns only the anchor part, if it exists. (Does not include the query string, if any.)
+ * If `url` has no hash, or only has a single hash (for the route path), returns an empty string.
+ * @param url - Full or partial URL. Just the hash section is valid, as long as it's prepended with #.
+ */
+export function extractAnchorLink(url: string): string {
+  // URLs containing anchors:
+  // #/components/checkbox#Overview
+  // http://whatever#/components/checkbox#Overview
+
+  // URLs NOT containing anchors, by this function's definition:
+  // #/components/checkbox
+  // http://whatever#/components/checkbox
+  // #Overview
+  // http://whatever#Overview
+  const split = url.split('#');
+  if (split.length === 3) {
+    // Also remove the query if present
+    // (technically the query can't be after the hash, but this is likely with hash routing)
+    return split[2].split('?')[0];
+  }
+  return '';
+}
diff --git a/Engine/Source/Programs/Horde/HordeDashboard/Dashboard/src/base/utilities/filter.ts b/Engine/Source/Programs/Horde/HordeDashboard/Dashboard/src/base/utilities/filter.ts
new file mode 100644
--- /dev/null
+++ b/Engine/Source/Programs/Horde/HordeDashboard/Dashboard/src/base/utilities/filter.ts
@@ -0,0 +1,52 @@
+import { GetJobResponse } from "../../backend/Api";
+
+
+export type JobFilterSimple = {
+    filterKeyword?: string;
+    showOthersPreflights: boolean;
+}
+
+const getJobKeywords = (job: GetJobResponse): string[] => {
+
+    let keywords: string[] = [];
+
+    if (job.change) {
+        keywords.push(job.change.toString());
+    }
+
+    if (job.preflightChange) {
+        keywords.push(job.preflightChange.toString());
+    }
+
+    if (job.startedByUser) {
+        keywords.push(job.startedByUser);
+    } else {
+        keywords.push("Scheduler");
+    }
+
+
+    keywords.push(...job.arguments);
+
+    return keywords;
+
+}
+
+export const filterJob = (job: GetJobResponse, keywordIn?: string, additionalKeywords?:string[]): boolean => {
+
+    const keyword = keywordIn?.toLowerCase();
+
+    if (!keyword) {
+        return true;
+    }
+
+    let keywords = getJobKeywords(job);
+
+
+    if (additionalKeywords) {
+        keywords.push(...additionalKeywords);
+    }
+
+    keywords = keywords.map(k => k.toLowerCase());
+
+    return !!keywords.find(k => k.indexOf(keyword) !== -1);
+}
\ No newline at end of file
diff --git a/Engine/Source/Programs/Horde/HordeDashboard/Dashboard/src/base/utilities/hooks.ts b/Engine/Source/Programs/Horde/HordeDashboard/Dashboard/src/base/utilities/hooks.ts
new file mode 100644
--- /dev/null
+++ b/Engine/Source/Programs/Horde/HordeDashboard/Dashboard/src/base/utilities/hooks.ts
@@ -0,0 +1,28 @@
+import { useState, useEffect } from 'react';
+
+
+export function useWindowSize():{width:number, height:number} {
+
+	function getSize() {
+		return {
+			width: window.innerWidth,
+			height: window.innerHeight
+		};
+	}
+
+	const [windowSize, setWindowSize] = useState(getSize);
+
+	useEffect(() => {
+
+		function handleResize() {
+			setWindowSize(getSize());
+		}
+
+		window.addEventListener('resize', handleResize);
+		return () => window.removeEventListener('resize', handleResize);
+	}, []);
+
+
+	return windowSize;
+}
+
diff --git a/Engine/Source/Programs/Horde/HordeDashboard/Dashboard/src/base/utilities/index.ts b/Engine/Source/Programs/Horde/HordeDashboard/Dashboard/src/base/utilities/index.ts
new file mode 100644
--- /dev/null
+++ b/Engine/Source/Programs/Horde/HordeDashboard/Dashboard/src/base/utilities/index.ts
@@ -0,0 +1,3 @@
+export * from './jumpToAnchor';
+export * from './removeAnchorLink';
+export * from './extractAnchorLink';
diff --git a/Engine/Source/Programs/Horde/HordeDashboard/Dashboard/src/base/utilities/issueUtils.ts b/Engine/Source/Programs/Horde/HordeDashboard/Dashboard/src/base/utilities/issueUtils.ts
new file mode 100644
--- /dev/null
+++ b/Engine/Source/Programs/Horde/HordeDashboard/Dashboard/src/base/utilities/issueUtils.ts
@@ -0,0 +1,95 @@
+import moment from "moment";
+import { GetIssueResponse, StreamData } from "../../backend/Api";
+import dashboard from "../../backend/Dashboard";
+import { projectStore } from "../../backend/ProjectStore";
+import { displayTimeZone } from "./timeUtils";
+
+
+export function getIssueStatus(issue: GetIssueResponse, showResolveTime?:boolean): string {
+
+	let text = "";	
+
+	if (issue.resolvedAt) {
+
+		if (!showResolveTime) {
+
+			if (!issue.fixChange) {
+				return `Resolved`;
+			}
+			return `Resolved in CL ${issue.fixChange}`;			
+		}
+
+		const displayTime = moment(issue.resolvedAt).tz(displayTimeZone());
+
+		const format = dashboard.display24HourClock ? "MMM Do, HH:mm z" : "MMM Do, h:mma z";
+
+		let text = `Resolved at ${displayTime.format(format)}`;
+
+		if (issue.fixChange) {
+			text += ` in CL ${issue.fixChange}`;
+		}
+
+		return text;
+
+	}
+
+	if (!issue.owner) {
+		text = "Currently unassigned.";
+	} else {
+		if (issue.owner.toLowerCase() === dashboard.username?.toLowerCase()) {
+			if (issue.nominatedBy) {
+				text = `You have been nominated to fix this issue by ${issue.nominatedBy}.`;
+			} else {
+				text = `Assigned to ${issue.owner}.`;
+			}
+		} else {
+			text = `Assigned to ${issue.owner}`;
+			if (issue.nominatedBy) {
+				text += ` by ${issue.nominatedBy}`;
+			}
+			if (!issue.acknowledgedAt) {
+				text += ` (Unacknowledged)`;
+			} else {
+				text += ` (Acknowledged)`;
+			}
+			text += ".";
+		}
+	}
+
+	return text;
+
+}
+
+export function generateStreamSummary(issue: GetIssueResponse): string {
+
+	const streams: StreamData[] = [];
+
+	projectStore.projects.forEach(p => {
+		if (!p.streams) {
+			return;
+		}
+
+		p.streams.filter(s => issue.unresolvedStreams.indexOf(s.id) !== -1).forEach(stream => {
+			if (!streams.find(s => s.id === stream.id)) {
+				streams.push(stream);
+			}
+		});
+	});
+
+	if (streams.length === 0) {
+		return "";
+	}
+
+	const present = streams.slice(0, 3);
+	const others = streams.slice(3);
+
+	let text = "Affects " + present.map(s => `//${s.project?.name}/${s.name}`).join(", ");
+
+	if (others.length) {
+		text += ` and ${others.length} other streams.`;
+	}
+
+	return text;
+
+}
+
diff --git a/Engine/Source/Programs/Horde/HordeDashboard/Dashboard/src/base/utilities/jumpToAnchor.ts b/Engine/Source/Programs/Horde/HordeDashboard/Dashboard/src/base/utilities/jumpToAnchor.ts
new file mode 100644
--- /dev/null
+++ b/Engine/Source/Programs/Horde/HordeDashboard/Dashboard/src/base/utilities/jumpToAnchor.ts
@@ -0,0 +1,25 @@
+import { extractAnchorLink } from './extractAnchorLink';
+
+const SCROLL_DISTANCE = 52;
+
+export function jumpToAnchor(anchor?: string, scrollDistance: number = SCROLL_DISTANCE): void {
+  const hash = anchor || extractAnchorLink(window.location.hash);
+  const el = hash && document.getElementById(hash);
+  if (hash && el) {
+    const elRect = el.getBoundingClientRect();
+    const windowY = window.scrollY || window.pageYOffset;
+    const currentScrollPosition = windowY + elRect.top;
+    const top = currentScrollPosition - scrollDistance;
+    if (window.scrollTo) {
+      if (window.navigator.userAgent.indexOf('rv:11.0') > -1 || window.navigator.userAgent.indexOf('Edge') > -1) {
+        // Edge currently has a bug that jumps to the top of the page if window.scrollTo is passed an oject.
+        window.scrollTo(0, top);
+      } else {
+        window.scrollTo({
+          top,
+          behavior: 'smooth'
+        });
+      }
+    }
+  }
+}
diff --git a/Engine/Source/Programs/Horde/HordeDashboard/Dashboard/src/base/utilities/removeAnchorLink.ts b/Engine/Source/Programs/Horde/HordeDashboard/Dashboard/src/base/utilities/removeAnchorLink.ts
new file mode 100644
--- /dev/null
+++ b/Engine/Source/Programs/Horde/HordeDashboard/Dashboard/src/base/utilities/removeAnchorLink.ts
@@ -0,0 +1,17 @@
+/**
+ * Remove the anchor link from a full page URL or a hash.
+ * Preserves any route path specified in the URL.
+ * Behavior with more than two # in the URL is unspecified.
+ * If there's a query string after the hash (not technically valid but possible with hash routing),
+ * also removes that.
+ */
+export function removeAnchorLink(url: string): string {
+  // First group: most of the URL
+  // Second group: optional last hash with only word or dash characters following (an anchor)
+  // Third group: optional post-hash query string
+  const match = url.match(/^(.*?)((#[\w-]*)(\?.*)?)$/);
+  if (match) {
+    return match[1];
+  }
+  return url;
+}
diff --git a/Engine/Source/Programs/Horde/HordeDashboard/Dashboard/src/base/utilities/timeUtils.ts b/Engine/Source/Programs/Horde/HordeDashboard/Dashboard/src/base/utilities/timeUtils.ts
new file mode 100644
--- /dev/null
+++ b/Engine/Source/Programs/Horde/HordeDashboard/Dashboard/src/base/utilities/timeUtils.ts
@@ -0,0 +1,370 @@
+import moment, { Moment } from "moment-timezone";
+import { BatchData, GetAgentLeaseResponse, JobData, LabelState, StepData, TimingInfo } from "../../backend/Api";
+import { JobLabel } from "../../backend/JobDetails";
+import dashboard from "../../backend/Dashboard";
+
+export const serverTimeZone = "UTC";
+const localTimeZone = moment.tz.guess();
+
+export const displayTimeZone = () => {
+    return dashboard.displayUTC ? "UTC" : localTimeZone;
+}
+
+type HordeTime = {
+    display: string;
+    server: string;
+}
+
+export const msecToElapsed = (millisec: number, includeMinutes: boolean = true): string => {
+
+    let duration = "";
+    const d = moment.duration(millisec);
+
+    if (d.days()) {
+        duration += `${d.days()}d `;
+    }
+
+    if (d.hours()) {
+        duration += `${d.hours()}h `;
+    }
+
+    if (includeMinutes) {
+
+        if (d.minutes()) {
+            duration += `${d.minutes()}m `;
+        }
+
+        if (d.seconds()) {
+            duration += `${d.seconds()}s `;
+        }
+    }
+
+
+    if (duration === "") {
+        duration = "0s";
+    }
+
+    return duration;
+};
+
+function roundTime(date: moment.Moment, duration: moment.Duration) {
+    return moment(Math.ceil((+date) / (+duration)) * (+duration));
+}
+
+export const getTargetTimingDelta = (timing: TimingInfo | undefined): string => {
+
+    if (!timing) {
+        return "";
+    }
+
+    const avg = timing.averageTotalTimeToComplete;
+    const time = timing.totalTimeToComplete;
+    if (avg && time) {
+
+        const seconds = time - avg;
+        const tm = msecToElapsed(Math.abs(seconds) * 1000);
+
+        if (seconds > 5) {
+            return `(${tm} longer than normal)`;
+        } else if (seconds < -5) {
+            return `(${tm} shorter than normal)`;
+        }
+
+        return `(taking the normal time).`;
+    }
+
+    return "";
+
+};
+
+export const getStepTimingDelta = (step: StepData): string => {
+
+    const avg = step.timing?.averageStepDuration;
+
+    if (avg && step.startTime && step.finishTime) {
+
+        const start = moment(step.startTime);
+        const finish = moment(step.finishTime);
+
+        const seconds = finish.diff(start, "seconds") - avg;
+        const tm = msecToElapsed(Math.abs(seconds) * 1000);
+
+
+        if (seconds > 5) {
+            return `(${tm} longer than normal)`;
+        } else if (seconds < -5) {
+            return `(${tm} shorter than normal)`;
+        }
+
+        return `(taking the normal time)`;
+    }
+
+    return "";
+
+};
+
+
+export const getStepETA = (step: StepData, job: JobData): HordeTime => {
+
+    if (step.finishTime || !step.timing || !step.timing.totalTimeToComplete) {
+        return {
+            display: "",
+            server: ""
+        };
+    }
+
+    const format = dashboard.display24HourClock ? "HH:mm:ss z" : "LT z";
+
+    let end = moment(job.createTime).add(step.timing.totalTimeToComplete, "seconds");
+
+    end = roundTime(end, moment.duration(5, "minutes"));
+
+    const serverETA = moment.utc(end).tz(serverTimeZone).format(format);
+    const displayETA = moment.utc(end).tz(displayTimeZone()).format(format);
+
+    return {
+        display: displayETA,
+        server: serverETA
+    };
+
+};
+
+export const getStepFinishTime = (step: StepData): HordeTime => {
+
+    let displayFinish = "";    
+    let serverFinish = "";    
+
+    const format = dashboard.display24HourClock ? "HH:mm:ss z" : "LT z";
+
+    if (step.finishTime) {
+        const end = moment(step.finishTime);
+        displayFinish = moment.utc(end).tz(displayTimeZone()).format(format);
+        serverFinish = moment.utc(end).tz(serverTimeZone).format(format);
+    }
+
+    return {
+        display: displayFinish,
+        server: serverFinish
+    };
+};
+
+
+export const getLabelFinishTime = (label: JobLabel, job: JobData): HordeTime => {
+
+    let displayFinish = "";    
+    let serverFinish = "";    
+
+    const format = dashboard.display24HourClock ? "HH:mm:ss z" : "LT z";
+
+    if (label.stateResponse.state === LabelState.Complete && label.timing) {
+        const end = moment(job.createTime).add(label.timing.totalTimeToComplete, "seconds");
+        displayFinish = moment.utc(end).tz(displayTimeZone()).format(format);
+        serverFinish = moment.utc(end).tz(serverTimeZone).format(format);
+    }
+
+    return {
+        display: displayFinish,
+        server: serverFinish
+    };
+};
+
+export const getLabelETA = (label: JobLabel, job: JobData): HordeTime => {
+
+    let displayEta = "";
+    let serverEta = "";
+
+    const format = dashboard.display24HourClock ? "HH:mm:ss z" : "LT z";
+
+    if (label.timing && label.timing.totalTimeToComplete && label.stateResponse.state === LabelState.Running) {
+        const end = moment(job.createTime).add(label.timing.totalTimeToComplete, "seconds");
+        serverEta = moment.utc(end).tz(serverTimeZone).format(format);
+        displayEta = moment.utc(end).tz(displayTimeZone()).format(format);
+    }
+
+    return {
+        display: displayEta,
+        server: serverEta
+    };
+
+};
+
+
+export const getStepPercent = (step: StepData): number | undefined => {
+
+    const timing = step.timing;
+
+    if (!step.startTime || !timing?.averageStepDuration) {
+        return undefined;
+    }
+
+    if (step.finishTime) {
+        return 1;
+    }
+
+    const start = moment(step.startTime);
+    const end = moment(Date.now());
+    const duration = moment.duration(end.diff(start)).asSeconds();
+
+    // 90% weighted from estimate, meaning if finish on time will jump 90% => done
+    let percent = duration / timing.averageStepDuration * 0.9;
+    if (percent >= 0.8) {
+
+        // sigmoid that hits 99% at 1.5x duration
+        percent = Math.min((2 / (1 + Math.pow(40, -percent)) - 1), 0.99);
+    }
+
+    return percent;
+};
+
+export const getElapsedString = (start: Moment, end: Moment, includeSeconds:boolean = true): string => {
+
+    let duration = "";
+    const d = moment.duration(end.diff(start));
+
+    if (d.months()) {
+        duration = `${d.months()} month`;
+        if (d.months() > 1) {
+            duration += "s";
+        }
+        return duration;
+    }
+
+    if (d.days()) {
+        duration += `${d.days()}d `;
+    }
+
+    if (d.hours()) {
+        duration += `${d.hours()}h `;
+    }
+
+    if (d.minutes()) {
+        duration += `${d.minutes()}m `;
+    }
+
+    if (includeSeconds) {
+        if (d.seconds()) {
+            duration += `${d.seconds()}s `;
+        }    
+    }
+
+    return duration;
+
+};
+
+export const getBatchInitElapsed = (batch: BatchData | undefined): string => {
+
+    if (!batch || !batch.startTime) {
+        return "";
+    }
+
+    const start = moment(batch.startTime);
+    let end = moment(Date.now());
+
+    if (!batch.steps.find(step => step.startTime) && batch.finishTime) {
+        end = moment(batch.finishTime);
+    } else {
+        batch.steps.forEach(step => {
+
+            if (!step.startTime) {
+                return;
+            }
+
+            const time = moment(step.startTime);
+
+            if (time.unix() < end.unix()) {
+                end = time;
+            }
+        });
+    }
+
+
+    return getElapsedString(start, end);
+};
+
+export const getStepElapsed = (step: StepData | BatchData | undefined): string => {
+
+    if (!step) {
+        return "";
+    }
+
+    const start = moment(step.startTime);
+    let end = moment(Date.now());
+
+    if (step.finishTime) {
+        end = moment(step.finishTime);
+    }
+
+    return getElapsedString(start, end);
+
+};
+
+export const getLeaseElapsed = (lease: GetAgentLeaseResponse | undefined): string => {
+
+    if (!lease) {
+        return "";
+    }
+
+    const start = moment(lease.startTime);
+    let end = moment(Date.now());
+
+    if (lease.finishTime) {
+        end = moment(lease.finishTime);
+    }
+
+    return getElapsedString(start, end);
+
+};
+
+
+
+export const getShortNiceTime = (timeIn: Date | string | undefined, relative: boolean = false): string => {
+
+    if (!timeIn) {
+        return "";
+    }
+
+    const now = moment.utc().tz(displayTimeZone());
+    const nowTimeStr = now.format('MMM Do');
+
+    const time = moment(timeIn).tz(displayTimeZone());
+
+    let timeStr = time.format('MMM Do');
+    if (relative && timeStr === nowTimeStr) {
+        timeStr += " (Today)";
+    } else if (relative && time.calendar().toLowerCase().indexOf("yesterday") !== -1) {
+        timeStr += " (Yesterday)";
+    } else {
+        timeStr = time.format('MMM Do');
+    }
+
+    return timeStr;
+
+};
+
+export const getNiceTime = (timeIn: Date | string | undefined, relative: boolean = true): string => {
+
+    if (!timeIn) {
+        return "";
+    }
+
+    const now = moment.utc().tz(displayTimeZone());
+    const nowTimeStr = now.format('dddd, MMMM Do');
+
+    const time = moment(timeIn).tz(displayTimeZone());
+
+    let timeStr = time.format('dddd, MMMM Do');
+    if (relative && timeStr === nowTimeStr) {
+        timeStr += " (Today)";
+    } else if (relative && time.calendar().toLowerCase().indexOf("yesterday") !== -1) {
+        timeStr += " (Yesterday)";
+    } else {
+        timeStr = time.format('dddd, MMMM Do');
+    }
+
+    const format = dashboard.display24HourClock ? "HH:mm:ss z" : "LT z";
+
+    timeStr += ` at ${time.format(format)}`;
+
+    return timeStr;
+
+};
\ No newline at end of file
diff --git a/Engine/Source/Programs/Horde/HordeDashboard/Dashboard/src/index.css b/Engine/Source/Programs/Horde/HordeDashboard/Dashboard/src/index.css
new file mode 100644
--- /dev/null
+++ b/Engine/Source/Programs/Horde/HordeDashboard/Dashboard/src/index.css
@@ -0,0 +1,89 @@
+body {
+  margin: 0;
+  font-family: -apple-system, BlinkMacSystemFont, 'Segoe UI', 'Roboto', 'Oxygen',
+    'Ubuntu', 'Cantarell', 'Fira Sans', 'Droid Sans', 'Helvetica Neue',
+    sans-serif;
+  -webkit-font-smoothing: antialiased;
+  -moz-osx-font-smoothing: grayscale; 
+  color:#2D3F5F;
+  overflow: hidden;
+}
+
+a, a:visited, a:active {
+  text-decoration: none;
+  color:rgb(0, 120, 212);
+}
+
+:focus {
+  outline:none;
+}
+
+a:hover {
+  text-decoration: none;
+  color:rgb(0, 120, 212);
+}
+
+code {
+  font-family: source-code-pro, Menlo, Monaco, Consolas, 'Courier New',
+    monospace;
+}
+
+.log-link, .log-link:visited, .log-link:hover {
+  text-decoration: none;
+  color:#000000;
+}
+
+.job-item:hover {
+  background:rgb(243, 242, 241);
+  user-select: none;
+  cursor: pointer;
+}
+
+.job-item.selected {
+  background:rgb(243, 242, 241);
+  user-select: none;
+  cursor: pointer;
+}
+
+
+.view-log-link, .view-log-link:visited, .view-log-link:hover {
+  text-decoration: none;
+  color: rgb(48, 66, 97);
+  user-select: none;
+  cursor: pointer;
+}
+
+.view-log-link-disabled, .view-log-link-disabled:visited, .view-log-link-disabled:hover {
+  text-decoration: none;
+  color: #8f8f8f;
+  user-select: none;
+}
+
+.cl-callout-button,.cl-callout-button:hover  {
+  text-decoration: none;
+  white-space: nowrap;
+  background: #035ca1;
+  color: rgb(255, 255, 255);
+  user-select: none;
+  font-family: "Horde Open Sans SemiBold";  
+  font-size: 10px;
+  border-width: 0;
+}
+
+.cl-callout-button-user,.cl-callout-button-user:hover  {
+  text-decoration: none;
+  white-space: nowrap;
+  background:#0288ee;
+  color: rgb(255, 255, 255);
+  user-select: none;
+  font-family: "Horde Open Sans SemiBold";  
+  font-size: 10px;
+  border-width: 0;
+}
+
+
+div[data-automation-key="ViewLogColumn"] {
+  padding: 0
+}
+
+
diff --git a/Engine/Source/Programs/Horde/HordeDashboard/Dashboard/src/react-app-env.d.ts b/Engine/Source/Programs/Horde/HordeDashboard/Dashboard/src/react-app-env.d.ts
new file mode 100644
--- /dev/null
+++ b/Engine/Source/Programs/Horde/HordeDashboard/Dashboard/src/react-app-env.d.ts
@@ -0,0 +1,3 @@
+// Copyright Epic Games, Inc. All Rights Reserved.
+
+/// <reference types="react-scripts" />
diff --git a/Engine/Source/Programs/Horde/HordeDashboard/Dashboard/src/serviceWorker.ts b/Engine/Source/Programs/Horde/HordeDashboard/Dashboard/src/serviceWorker.ts
new file mode 100644
--- /dev/null
+++ b/Engine/Source/Programs/Horde/HordeDashboard/Dashboard/src/serviceWorker.ts
@@ -0,0 +1,145 @@
+// This optional code is used to register a service worker.
+// register() is not called by default.
+
+// This lets the app load faster on subsequent visits in production, and gives
+// it offline capabilities. However, it also means that developers (and users)
+// will only see deployed updates on subsequent visits to a page, after all the
+// existing tabs open on the page have been closed, since previously cached
+// resources are updated in the background.
+
+// To learn more about the benefits of this model and instructions on how to
+// opt-in, read https://bit.ly/CRA-PWA
+
+const isLocalhost = Boolean(
+  window.location.hostname === 'localhost' ||
+    // [::1] is the IPv6 localhost address.
+    window.location.hostname === '[::1]' ||
+    // 127.0.0.0/8 are considered localhost for IPv4.
+    window.location.hostname.match(
+      /^127(?:\.(?:25[0-5]|2[0-4][0-9]|[01]?[0-9][0-9]?)){3}$/
+    )
+);
+
+type Config = {
+  onSuccess?: (registration: ServiceWorkerRegistration) => void;
+  onUpdate?: (registration: ServiceWorkerRegistration) => void;
+};
+
+export function register(config?: Config) {
+  if (process.env.NODE_ENV === 'production' && 'serviceWorker' in navigator) {
+    // The URL constructor is available in all browsers that support SW.
+    const publicUrl = new URL(
+      process.env.PUBLIC_URL,
+      window.location.href
+    );
+    if (publicUrl.origin !== window.location.origin) {
+      // Our service worker won't work if PUBLIC_URL is on a different origin
+      // from what our page is served on. This might happen if a CDN is used to
+      // serve assets; see [REDACTED_PROFILE_URL]/create-react-app/issues/2374
+      return;
+    }
+
+    window.addEventListener('load', () => {
+      const swUrl = `${process.env.PUBLIC_URL}/service-worker.js`;
+
+      if (isLocalhost) {
+        // This is running on localhost. Let's check if a service worker still exists or not.
+        checkValidServiceWorker(swUrl, config);
+
+        // Add some additional logging to localhost, pointing developers to the
+        // service worker/PWA documentation.
+        navigator.serviceWorker.ready.then(() => {
+          console.log(
+            'This web app is being served cache-first by a service ' +
+              'worker. To learn more, visit https://bit.ly/CRA-PWA'
+          );
+        });
+      } else {
+        // Is not localhost. Just register service worker
+        registerValidSW(swUrl, config);
+      }
+    });
+  }
+}
+
+function registerValidSW(swUrl: string, config?: Config) {
+  navigator.serviceWorker
+    .register(swUrl)
+    .then(registration => {
+      registration.onupdatefound = () => {
+        const installingWorker = registration.installing;
+        if (installingWorker == null) {
+          return;
+        }
+        installingWorker.onstatechange = () => {
+          if (installingWorker.state === 'installed') {
+            if (navigator.serviceWorker.controller) {
+              // At this point, the updated precached content has been fetched,
+              // but the previous service worker will still serve the older
+              // content until all client tabs are closed.
+              console.log(
+                'New content is available and will be used when all ' +
+                  'tabs for this page are closed. See https://bit.ly/CRA-PWA.'
+              );
+
+              // Execute callback
+              if (config && config.onUpdate) {
+                config.onUpdate(registration);
+              }
+            } else {
+              // At this point, everything has been precached.
+              // It's the perfect time to display a
+              // "Content is cached for offline use." message.
+              console.log('Content is cached for offline use.');
+
+              // Execute callback
+              if (config && config.onSuccess) {
+                config.onSuccess(registration);
+              }
+            }
+          }
+        };
+      };
+    })
+    .catch(error => {
+      console.error('Error during service worker registration:', error);
+    });
+}
+
+function checkValidServiceWorker(swUrl: string, config?: Config) {
+  // Check if the service worker can be found. If it can't reload the page.
+  fetch(swUrl, {
+    headers: { 'Service-Worker': 'script' }
+  })
+    .then(response => {
+      // Ensure service worker exists, and that we really are getting a JS file.
+      const contentType = response.headers.get('content-type');
+      if (
+        response.status === 404 ||
+        (contentType != null && contentType.indexOf('javascript') === -1)
+      ) {
+        // No service worker found. Probably a different app. Reload the page.
+        navigator.serviceWorker.ready.then(registration => {
+          registration.unregister().then(() => {
+            window.location.reload();
+          });
+        });
+      } else {
+        // Service worker found. Proceed as normal.
+        registerValidSW(swUrl, config);
+      }
+    })
+    .catch(() => {
+      console.log(
+        'No internet connection found. App is running in offline mode.'
+      );
+    });
+}
+
+export function unregister() {
+  if ('serviceWorker' in navigator) {
+    navigator.serviceWorker.ready.then(registration => {
+      registration.unregister();
+    });
+  }
+}
diff --git a/Engine/Source/Programs/Horde/HordeDashboard/Dashboard/src/setupTests.ts b/Engine/Source/Programs/Horde/HordeDashboard/Dashboard/src/setupTests.ts
new file mode 100644
--- /dev/null
+++ b/Engine/Source/Programs/Horde/HordeDashboard/Dashboard/src/setupTests.ts
@@ -0,0 +1,5 @@
+// jest-dom adds custom jest matchers for asserting on DOM nodes.
+// allows you to do things like:
+// expect(element).toHaveTextContent(/react/i)
+// learn more: [REDACTED_PROFILE_URL]/jest-dom
+import '@testing-library/jest-dom/extend-expect';
diff --git a/Engine/Source/Programs/Horde/HordeDashboard/Dashboard/src/styles/colors.ts b/Engine/Source/Programs/Horde/HordeDashboard/Dashboard/src/styles/colors.ts
new file mode 100644
--- /dev/null
+++ b/Engine/Source/Programs/Horde/HordeDashboard/Dashboard/src/styles/colors.ts
@@ -0,0 +1,106 @@
+// Copyright Epic Games, Inc. All Rights Reserved.
+
+import { getTheme } from "@fluentui/react";
+import { DashboardPreference, JobState, LabelOutcome, LabelState } from "../backend/Api";
+import dashboard from "../backend/Dashboard";
+
+const theme = getTheme();
+
+export enum StatusColor {
+    Success,
+    Warnings,
+    Failure,
+    Waiting,
+    Ready,
+    Skipped,
+    Aborted,
+    Running,
+    Unspecified
+}
+
+export const defaultStatusColors = new Map<StatusColor, string>([
+    [StatusColor.Success, "#52C705"],
+    [StatusColor.Warnings, "#EDC74A"],
+    [StatusColor.Failure, "#DE4522"],
+    [StatusColor.Running, theme.palette.blueLight],
+    [StatusColor.Waiting, "#A19F9D"],
+    [StatusColor.Ready, "#A19F9D"],
+    [StatusColor.Skipped, "#F3F2F1"],
+    [StatusColor.Aborted, "#F3F2F1"],
+    [StatusColor.Unspecified, "#637087"]
+]);
+
+
+export const getStatusColors = () => {
+
+    const success = dashboard.getPreference(DashboardPreference.ColorSuccess);
+    const warning = dashboard.getPreference(DashboardPreference.ColorWarning);
+    const error = dashboard.getPreference(DashboardPreference.ColorError);
+    const running = dashboard.getPreference(DashboardPreference.ColorRunning);
+
+    return new Map<StatusColor, string>([
+        [StatusColor.Success, success ? success : defaultStatusColors.get(StatusColor.Success)!],
+        [StatusColor.Warnings, warning ? warning : defaultStatusColors.get(StatusColor.Warnings)!],
+        [StatusColor.Failure, error ? error : defaultStatusColors.get(StatusColor.Failure)!],
+        [StatusColor.Running, running ? running : defaultStatusColors.get(StatusColor.Running)!],
+        [StatusColor.Waiting, defaultStatusColors.get(StatusColor.Waiting)!],
+        [StatusColor.Ready, defaultStatusColors.get(StatusColor.Ready)!],
+        [StatusColor.Skipped, defaultStatusColors.get(StatusColor.Skipped)!],
+        [StatusColor.Aborted, defaultStatusColors.get(StatusColor.Aborted)!],
+        [StatusColor.Unspecified, defaultStatusColors.get(StatusColor.Unspecified)!]
+    ]);
+}
+
+export const getJobStateColor = (state: JobState): string => {
+
+    const colors = getStatusColors();    
+
+    if (state === JobState.Waiting) {
+        return colors.get(StatusColor.Waiting)!;
+    }
+
+    if (state === JobState.Running) {
+        return colors.get(StatusColor.Running)!;
+    }
+
+
+    return colors.get(StatusColor.Unspecified)!;
+
+}
+
+export const getLabelColor = (state: LabelState | undefined, outcome: LabelOutcome | undefined):{primaryColor: string, secondaryColor?: string} => {
+
+    if (!state || !outcome) {
+        return {
+            primaryColor: "#000000" 
+        }
+    }
+
+    const colors = getStatusColors();
+
+    if (state === LabelState.Complete) {
+        switch (outcome!) {
+            case LabelOutcome.Failure:
+                return { primaryColor: colors.get(StatusColor.Failure)!};        
+            case LabelOutcome.Success:
+                return { primaryColor: colors.get(StatusColor.Success)!};        
+            case LabelOutcome.Warnings:
+                return { primaryColor: colors.get(StatusColor.Warnings)!};            
+            default:
+                return { primaryColor: colors.get(StatusColor.Unspecified)!};        
+        }
+    }
+
+    if (outcome === LabelOutcome.Failure) {
+        return { primaryColor: colors.get(StatusColor.Running)!, secondaryColor: colors.get(StatusColor.Failure)!};        
+    }
+
+    if (outcome === LabelOutcome.Warnings) {
+        return { primaryColor: colors.get(StatusColor.Running)!, secondaryColor: colors.get(StatusColor.Warnings)!};        
+    }
+
+    return {
+        primaryColor: colors.get(StatusColor.Running)!
+    }
+};
+
diff --git a/Engine/Source/Programs/Horde/HordeDashboard/Dashboard/src/types/react-oidc-context.d.ts b/Engine/Source/Programs/Horde/HordeDashboard/Dashboard/src/types/react-oidc-context.d.ts
new file mode 100644
--- /dev/null
+++ b/Engine/Source/Programs/Horde/HordeDashboard/Dashboard/src/types/react-oidc-context.d.ts
@@ -0,0 +1,3 @@
+// Copyright Epic Games, Inc. All Rights Reserved.
+
+declare module "@axa-fr/react-oidc-context";
\ No newline at end of file
diff --git a/Engine/Source/Programs/Horde/HordeDashboard/Dashboard/tsconfig.json b/Engine/Source/Programs/Horde/HordeDashboard/Dashboard/tsconfig.json
new file mode 100644
--- /dev/null
+++ b/Engine/Source/Programs/Horde/HordeDashboard/Dashboard/tsconfig.json
@@ -0,0 +1,28 @@
+{
+  "compilerOptions": {
+    "target": "es5",
+    "lib": [
+      "dom",
+      "dom.iterable",
+      "esnext"
+    ],
+    "allowJs": true,
+    "skipLibCheck": true,
+    "esModuleInterop": true,
+    "allowSyntheticDefaultImports": true,
+    "strict": true,
+    "forceConsistentCasingInFileNames": true,
+    "module": "esnext",
+    "moduleResolution": "node",
+    "resolveJsonModule": true,
+    "isolatedModules": true,
+    "noEmit": true,
+    "jsx": "react-jsx",
+    "experimentalDecorators": true,
+    "noFallthroughCasesInSwitch": true
+  },
+  "include": [
+    "src",
+    "types"
+  ]
+}
diff --git a/Engine/Source/Programs/Horde/HordeDashboard/Dashboard/types/react-flame-graph/index.d.ts b/Engine/Source/Programs/Horde/HordeDashboard/Dashboard/types/react-flame-graph/index.d.ts
new file mode 100644
--- /dev/null
+++ b/Engine/Source/Programs/Horde/HordeDashboard/Dashboard/types/react-flame-graph/index.d.ts
@@ -0,0 +1 @@
+declare module 'react-flame-graph'
\ No newline at end of file
diff --git a/Engine/Source/Programs/Horde/HordeDashboard/Dashboard/types/react-text-marquee/index.d.ts b/Engine/Source/Programs/Horde/HordeDashboard/Dashboard/types/react-text-marquee/index.d.ts
new file mode 100644
--- /dev/null
+++ b/Engine/Source/Programs/Horde/HordeDashboard/Dashboard/types/react-text-marquee/index.d.ts
@@ -0,0 +1 @@
+declare module 'react-text-marquee'
\ No newline at end of file
