Update Horde from latest.

diff --git a/Engine/Source/Programs/Horde/HordeAgent/Commands/Execution/ExecuteCommand.cs b/Engine/Source/Programs/Horde/HordeAgent/Commands/Execution/ExecuteCommand.cs
--- a/Engine/Source/Programs/Horde/HordeAgent/Commands/Execution/ExecuteCommand.cs
+++ b/Engine/Source/Programs/Horde/HordeAgent/Commands/Execution/ExecuteCommand.cs
@@ -20,6 +20,7 @@ using System.Text.Json;
 using HordeAgent.Services;
 using Grpc.Net.Client;
 using System.Linq;
+using System.Net.Http.Headers;
 using Grpc.Core;
 using Google.LongRunning;
 using Microsoft.Extensions.Hosting;
@@ -162,6 +163,29 @@ namespace HordeAgent.Commands
 		/// </summary>
 		[CommandLine("-Salt=")]
 		public string? Salt = null;
+		
+		/// <summary>
+		/// gRPC URL for the content-addressable storage service.
+		///
+		/// If empty, Horde's built-in storage will be used.
+		/// The URL must match what has been configured for the instance name in Horde server settings.
+		/// </summary>
+		[CommandLine("-CasUrl=")]
+		public string? CasUrl = null;
+		
+		/// <summary>
+		/// Auth token to use for accessing CAS service
+		/// </summary>
+		[CommandLine("-CasAuthToken=")]
+		public string? CasAuthToken = null;
+		
+		/// <summary>
+		/// Instance name to use
+		///
+		/// See RE API specification for description.
+		/// </summary>
+		[CommandLine("-InstanceName=")]
+		public string? InstanceName = null;
 
 		/// <inheritdoc/>
 		public override async Task<int> ExecuteAsync(ILogger Logger)
@@ -208,10 +232,19 @@ namespace HordeAgent.Commands
 
 				using (GrpcChannel Channel = GrpcService.CreateGrpcChannel())
 				{
-					ContentAddressableStorage.ContentAddressableStorageClient StorageClient = new ContentAddressableStorage.ContentAddressableStorageClient(Channel);
+					GrpcChannel GetChannel()
+					{
+						if (CasUrl == null) return Channel;
+						if (CasAuthToken != null) return GrpcService.CreateGrpcChannel(CasUrl, new AuthenticationHeaderValue("ServiceAccount", CasAuthToken));
+						return GrpcService.CreateGrpcChannel(CasUrl, null);
+					}
+
+					using GrpcChannel CasChannel = GetChannel();
+					ContentAddressableStorage.ContentAddressableStorageClient StorageClient = new ContentAddressableStorage.ContentAddressableStorageClient(CasChannel);
 
 					// Find the missing blobs and upload them
 					FindMissingBlobsRequest MissingBlobsRequest = new FindMissingBlobsRequest();
+					MissingBlobsRequest.InstanceName = InstanceName;
 					MissingBlobsRequest.BlobDigests.Add(UploadList.Blobs.Values.Select(x => x.Digest));
 
 					FindMissingBlobsResponse MissingBlobs = await StorageClient.FindMissingBlobsAsync(MissingBlobsRequest);
@@ -220,8 +253,24 @@ namespace HordeAgent.Commands
 						foreach (Digest MissingBlobDigest in MissingBlobs.MissingBlobDigests)
 						{
 							BatchUpdateBlobsRequest UpdateRequest = new BatchUpdateBlobsRequest();
+							UpdateRequest.InstanceName = InstanceName;
 							UpdateRequest.Requests.Add(UploadList.Blobs[MissingBlobDigest.Hash]);
-							await StorageClient.BatchUpdateBlobsAsync(UpdateRequest);
+							BatchUpdateBlobsResponse UpdateResponse = await StorageClient.BatchUpdateBlobsAsync(UpdateRequest);
+							
+							bool UpdateFailed = false;
+							foreach (var Response in UpdateResponse.Responses)
+							{
+								if (Response.Status.Code != (int) Google.Rpc.Code.Ok)
+								{
+									Console.WriteLine($"Upload failed for {Response.Digest}. Code: {Response.Status.Code} Message: {Response.Status.Message}");
+									UpdateFailed = true;
+								}
+							}
+
+							if (UpdateFailed)
+							{
+								throw new Exception("Failed updating blobs in CAS service");
+							}
 						}
 					}
 
@@ -230,6 +279,10 @@ namespace HordeAgent.Commands
 
 					ExecuteRequest ExecuteRequest = new ExecuteRequest();
 					ExecuteRequest.ActionDigest = ActionDigest;
+					if (InstanceName != null)
+					{
+						ExecuteRequest.InstanceName = InstanceName;
+					}
 
 					using (AsyncServerStreamingCall<Operation> Call = ExecutionClient.Execute(ExecuteRequest))
 					{
@@ -266,6 +319,7 @@ namespace HordeAgent.Commands
 
 								// Print the result
 								BatchReadBlobsRequest BatchReadRequest = new BatchReadBlobsRequest();
+								BatchReadRequest.InstanceName = InstanceName;
 								if (ActionResult.StdoutDigest != null && ActionResult.StdoutDigest.SizeBytes > 0)
 								{
 									BatchReadRequest.Digests.Add(ActionResult.StdoutDigest);
diff --git a/Engine/Source/Programs/Horde/HordeAgent/Commands/Service/UpgradeCommand.cs b/Engine/Source/Programs/Horde/HordeAgent/Commands/Service/UpgradeCommand.cs
--- a/Engine/Source/Programs/Horde/HordeAgent/Commands/Service/UpgradeCommand.cs
+++ b/Engine/Source/Programs/Horde/HordeAgent/Commands/Service/UpgradeCommand.cs
@@ -101,6 +101,15 @@ namespace HordeAgent.Commands
 				{
 					UpgradeMacService(Logger, OtherProcess, TargetFiles, RenameFiles);
 				}
+				else if (RuntimeInformation.IsOSPlatform(OSPlatform.Linux))
+				{
+					UpgradeLinuxService(Logger, OtherProcess, TargetFiles, RenameFiles);
+				}
+				else
+				{
+					Logger.LogError("Agent is not running a platform that supports Upgrades. Platform: {0}", RuntimeInformation.OSDescription);
+					return Task.FromResult(-1);
+				}
 			}
 			Logger.LogInformation("Upgrade complete");
 			return Task.FromResult(0);
@@ -131,6 +140,13 @@ namespace HordeAgent.Commands
 			OtherProcess.Kill();
 		}
 
+		void UpgradeLinuxService(ILogger Logger, Process OtherProcess, HashSet<string> TargetFiles, List<Tuple<string, string>> RenameFiles)
+		{
+			UpgradeFilesInPlace(Logger, TargetFiles, RenameFiles);
+			Logger.LogDebug("Upgrade completed, restarting...");
+			OtherProcess.Kill();
+		}
+
 		void UpgradeWindowsService(ILogger Logger, Process OtherProcess, HashSet<string> TargetFiles, List<Tuple<string, string>> RenameFiles)
 		{
 			// Try to get the service associated with the passed-in process id
diff --git a/Engine/Source/Programs/Horde/HordeAgent/Execution/ActionExecutor.cs b/Engine/Source/Programs/Horde/HordeAgent/Execution/ActionExecutor.cs
--- a/Engine/Source/Programs/Horde/HordeAgent/Execution/ActionExecutor.cs
+++ b/Engine/Source/Programs/Horde/HordeAgent/Execution/ActionExecutor.cs
@@ -34,6 +34,11 @@ namespace HordeAgent
 	/// </summary>
 	class ActionExecutor
 	{
+		/// <summary>
+		/// The instance name
+		/// </summary>
+		string InstanceName;
+		
 		/// <summary>
 		/// The storage wrapper
 		/// </summary>
@@ -57,13 +62,17 @@ namespace HordeAgent
 		/// <summary>
 		/// Constructor
 		/// </summary>
-		/// <param name="Channel"></param>
+		/// <param name="InstanceName"></param>
+		/// <param name="CasChannel"></param>
+		/// <param name="ActionCacheChannel"></param>
+		/// <param name="ActionRpcChannel"></param>
 		/// <param name="Logger"></param>
-		public ActionExecutor(GrpcChannel Channel, ILogger Logger)
+		public ActionExecutor(string InstanceName, GrpcChannel CasChannel, GrpcChannel ActionCacheChannel, GrpcChannel ActionRpcChannel, ILogger Logger)
 		{
-			this.Storage = new ContentAddressableStorage.ContentAddressableStorageClient(Channel);
-			this.Cache = new ActionCache.ActionCacheClient(Channel);
-			this.ActionRpc = new ActionRpc.ActionRpcClient(Channel);
+			this.InstanceName = InstanceName;
+			this.Storage = new ContentAddressableStorage.ContentAddressableStorageClient(CasChannel);
+			this.Cache = new ActionCache.ActionCacheClient(ActionCacheChannel);
+			this.ActionRpc = new ActionRpc.ActionRpcClient(ActionRpcChannel);
 			this.Logger = Logger;
 		}
 
@@ -76,7 +85,7 @@ namespace HordeAgent
 		/// <returns>The action result</returns>
 		public async Task<ActionResult> ExecuteActionAsync(string LeaseId, ActionTask ActionTask, DirectoryReference SandboxDir)
 		{
-			Action? Action = await Storage.GetProtoMessageAsync<Action>(ActionTask.Digest);
+			Action? Action = await Storage.GetProtoMessageAsync<Action>(InstanceName, ActionTask.Digest);
 			if (Action == null)
 			{
 				throw new RpcException(new Grpc.Core.Status(StatusCode.NotFound, "Unable to find action message"));
@@ -85,10 +94,10 @@ namespace HordeAgent
 			DirectoryReference.CreateDirectory(SandboxDir);
 			FileUtils.ForceDeleteDirectoryContents(SandboxDir);
 
-			Directory InputDirectory = await Storage.GetProtoMessageAsync<Directory>(Action.InputRootDigest);
+			Directory InputDirectory = await Storage.GetProtoMessageAsync<Directory>(InstanceName, Action.InputRootDigest);
 			await SetupSandboxAsync(InputDirectory, SandboxDir);
 
-			RpcCommand Command = await Storage.GetProtoMessageAsync<RpcCommand>(Action.CommandDigest);
+			RpcCommand Command = await Storage.GetProtoMessageAsync<RpcCommand>(InstanceName, Action.CommandDigest);
 
 			foreach (string OutputFile in Command.OutputFiles)
 			{
@@ -132,8 +141,8 @@ namespace HordeAgent
 					}
 
 					Result = new ActionResult();
-					Result.StdoutDigest = await Storage.PutBulkDataAsync(StdOutData);
-					Result.StderrDigest = await Storage.PutBulkDataAsync(StdErrData);
+					Result.StdoutDigest = await Storage.PutBulkDataAsync(InstanceName, StdOutData);
+					Result.StderrDigest = await Storage.PutBulkDataAsync(InstanceName, StdErrData);
 					Result.ExitCode = Process.ExitCode;
 
 					foreach (string Line in Encoding.UTF8.GetString(StdOutData).Split('\n'))
@@ -153,7 +162,7 @@ namespace HordeAgent
 						if (FileReference.Exists(FileRef))
 						{
 							byte[] Bytes = await FileReference.ReadAllBytesAsync(FileRef);
-							Digest Digest = await Storage.PutBulkDataAsync(Bytes);
+							Digest Digest = await Storage.PutBulkDataAsync(InstanceName, Bytes);
 
 							OutputFile OutputFileInfo = new OutputFile();
 							OutputFileInfo.Path = OutputFile;
@@ -173,6 +182,7 @@ namespace HordeAgent
 					if (!Action.DoNotCache)
 					{
 						UpdateActionResultRequest Request = new UpdateActionResultRequest();
+						Request.InstanceName = InstanceName;
 						Request.ActionDigest = ActionTask.Digest;
 						Request.ActionResult = Result;
 						await Cache.UpdateActionResultAsync(Request);
@@ -197,7 +207,7 @@ namespace HordeAgent
 
 			foreach (FileNode FileNode in InputDirectory.Files)
 			{
-				ReadOnlyMemory<byte> Data = await Storage.GetBulkDataAsync(FileNode.Digest);
+				ReadOnlyMemory<byte> Data = await Storage.GetBulkDataAsync(InstanceName, FileNode.Digest);
 				FileReference File = FileReference.Combine(OutputDir, FileNode.Name);
 				Logger.LogInformation("Writing {File} (digest: {Digest}, size: {Size})", File, FileNode.Digest.Hash, FileNode.Digest.SizeBytes);
 				await FileReference.WriteAllBytesAsync(File, Data.ToArray());
@@ -205,7 +215,7 @@ namespace HordeAgent
 
 			foreach (DirectoryNode DirectoryNode in InputDirectory.Directories)
 			{
-				Directory InputSubDirectory = await Storage.GetProtoMessageAsync<Directory>(DirectoryNode.Digest);
+				Directory InputSubDirectory = await Storage.GetProtoMessageAsync<Directory>(InstanceName, DirectoryNode.Digest);
 				DirectoryReference OutputSubDirectory = DirectoryReference.Combine(OutputDir, DirectoryNode.Name);
 				await SetupSandboxAsync(InputSubDirectory, OutputSubDirectory);
 			}
diff --git a/Engine/Source/Programs/Horde/HordeAgent/HordeAgent.csproj b/Engine/Source/Programs/Horde/HordeAgent/HordeAgent.csproj
--- a/Engine/Source/Programs/Horde/HordeAgent/HordeAgent.csproj
+++ b/Engine/Source/Programs/Horde/HordeAgent/HordeAgent.csproj
@@ -21,6 +21,7 @@
   <ItemGroup>
     <PackageReference Include="AWSSDK.Core" Version="3.5.1.49" />
     <PackageReference Include="AWSSDK.EC2" Version="3.5.1.2" />
+    <PackageReference Include="Blake3" Version="0.3.0" />
     <PackageReference Include="Datadog.Trace" Version="1.18.0" />
     <PackageReference Include="Google.Protobuf" Version="3.12.3" />
     <PackageReference Include="Grpc.Net.Client" Version="2.29.0" />
@@ -101,12 +102,6 @@
     </EmbeddedResource>
   </ItemGroup>
 
-  <ItemGroup>
-    <None Update="startagent.sh">
-      <CopyToOutputDirectory>PreserveNewest</CopyToOutputDirectory>
-    </None>
-  </ItemGroup>
-
   <PropertyGroup>
 	<Pdb2PdbPath>$(NuGetPackageRoot)microsoft.diasymreader.pdb2pdb\1.1.0-beta2-20370-01\tools\Pdb2Pdb.exe</Pdb2PdbPath>
 	<SymStorePath>C:\Program Files (x86)\Windows Kits\10\Debuggers\x64\symstore.exe</SymStorePath>
diff --git a/Engine/Source/Programs/Horde/HordeAgent/Parser/Matchers/LinkEventMatcher.cs b/Engine/Source/Programs/Horde/HordeAgent/Parser/Matchers/LinkEventMatcher.cs
--- a/Engine/Source/Programs/Horde/HordeAgent/Parser/Matchers/LinkEventMatcher.cs
+++ b/Engine/Source/Programs/Horde/HordeAgent/Parser/Matchers/LinkEventMatcher.cs
@@ -25,7 +25,7 @@ namespace HordeAgent.Parser.Matchers
 
 			int LineCount = 0;
 			bool bIsWarning = true;
-			while (Cursor.IsMatch(LineCount, @": In function |: undefined reference to |undefined symbol|^\s*(ld|ld.lld):|[^a-zA-Z]ld: |^\s*>>>"))
+			while (Cursor.IsMatch(LineCount, @": undefined reference to |undefined symbol|^\s*(ld|ld.lld):|[^a-zA-Z]ld: |^\s*>>>"))
 			{
 				bIsWarning &= Cursor.IsMatch("warning:");
 				LineCount++;
diff --git a/Engine/Source/Programs/Horde/HordeAgent/Program.cs b/Engine/Source/Programs/Horde/HordeAgent/Program.cs
--- a/Engine/Source/Programs/Horde/HordeAgent/Program.cs
+++ b/Engine/Source/Programs/Horde/HordeAgent/Program.cs
@@ -118,6 +118,9 @@ namespace HordeAgent
 		/// <returns>Exit code</returns>
 		static async Task<int> GuardedMain(string[] Args, ILogger Logger)
 		{
+			// Enable unencrypted HTTP/2 for gRPC channel without TLS
+			AppContext.SetSwitch("System.Net.Http.SocketsHttpHandler.Http2UnencryptedSupport", true);
+			
 			// Find all the command types
 			List<(CommandAttribute, Type)> Commands = new List<(CommandAttribute, Type)>();
 			foreach (Type Type in Assembly.GetExecutingAssembly().GetTypes())
diff --git a/Engine/Source/Programs/Horde/HordeAgent/Services/GrpcService.cs b/Engine/Source/Programs/Horde/HordeAgent/Services/GrpcService.cs
--- a/Engine/Source/Programs/Horde/HordeAgent/Services/GrpcService.cs
+++ b/Engine/Source/Programs/Horde/HordeAgent/Services/GrpcService.cs
@@ -55,12 +55,21 @@ namespace HordeAgent.Services
 		{
 			return CreateGrpcChannel(ServerProfile.Token);
 		}
-
+		
 		/// <summary>
 		/// Create a GRPC channel with the given bearer token
 		/// </summary>
 		/// <returns>New grpc channel</returns>
 		public GrpcChannel CreateGrpcChannel(string BearerToken)
+		{
+			return CreateGrpcChannel(ServerProfile.Url, new AuthenticationHeaderValue("Bearer", BearerToken));
+		}
+
+		/// <summary>
+		/// Create a GRPC channel with the given auth header value
+		/// </summary>
+		/// <returns>New grpc channel</returns>
+		public GrpcChannel CreateGrpcChannel(string Address, AuthenticationHeaderValue? AuthHeaderValue)
 		{
 			HttpClientHandler CustomCertHandler = new HttpClientHandler();
 			CustomCertHandler.ServerCertificateCustomValidationCallback += (Sender, Cert, Chain, Errors) => CertificateHelper.CertificateValidationCallBack(Logger, Sender, Cert, Chain, Errors, ServerProfile);
@@ -72,11 +81,15 @@ namespace HordeAgent.Services
 
 			HttpClient HttpClient = new HttpClient(RetryHandler);
 			HttpClient.DefaultRequestHeaders.Add("Accept", "application/json");
-			HttpClient.DefaultRequestHeaders.Authorization = new AuthenticationHeaderValue("Bearer", BearerToken);
+			if (AuthHeaderValue != null)
+			{
+				HttpClient.DefaultRequestHeaders.Authorization = AuthHeaderValue;
+			}
+
 			HttpClient.Timeout = TimeSpan.FromSeconds(210); // Need to make sure this doesn't cancel any long running gRPC streaming calls (eg. session update)
 
 			Logger.LogInformation("Connecting to rpc server {BaseUrl}", ServerProfile.Url);
-			return GrpcChannel.ForAddress(ServerProfile.Url, new GrpcChannelOptions { HttpClient = HttpClient, DisposeHttpClient = true });
+			return GrpcChannel.ForAddress(Address, new GrpcChannelOptions { HttpClient = HttpClient, DisposeHttpClient = true });
 		}
 	}
 }
diff --git a/Engine/Source/Programs/Horde/HordeAgent/Services/WorkerService.cs b/Engine/Source/Programs/Horde/HordeAgent/Services/WorkerService.cs
--- a/Engine/Source/Programs/Horde/HordeAgent/Services/WorkerService.cs
+++ b/Engine/Source/Programs/Horde/HordeAgent/Services/WorkerService.cs
@@ -771,9 +771,26 @@ namespace HordeAgent.Services
 				DirectoryReference LeaseDir = DirectoryReference.Combine(WorkingDir, "Remote", LeaseId);
 				DirectoryReference.CreateDirectory(LeaseDir);
 
+				GrpcChannel GetChannel(string? Url, string? Token)
+				{
+					if (string.IsNullOrEmpty(Url)) return Client.Channel;
+					if (!string.IsNullOrEmpty(Token))
+						return GrpcService.CreateGrpcChannel(Url, new AuthenticationHeaderValue("ServiceAccount", Token));
+					
+					return GrpcChannel.ForAddress(Url);
+				}
+
+				GrpcChannel CasChannel = GetChannel(ActionTask.CasUrl, ActionTask.ServiceAccountToken);// == null ? Client.Channel : GrpcChannel.ForAddress(ActionTask.CasUrl);
+				GrpcChannel ActionCacheChannel = GetChannel(ActionTask.ActionCacheUrl, ActionTask.ServiceAccountToken);//ActionTask.ActionCacheUrl == null ? Client.Channel : GrpcChannel.ForAddress(ActionTask.ActionCacheUrl);
+				GrpcChannel ActionRpcChannel = Client.Channel;
+
+				Logger.LogDebug("CasChannel={CasChannel}", CasChannel.Target);
+				Logger.LogDebug("ActionCacheChannel={ActionCacheChannel}", ActionCacheChannel.Target);
+				Logger.LogDebug("ActionRpcChannel={ActionRpcChannel}", ActionRpcChannel.Target);
+
 				try
 				{
-					ActionExecutor Executor = new ActionExecutor(Client.Channel, Logger);
+					ActionExecutor Executor = new ActionExecutor(ActionTask.InstanceName, CasChannel, ActionCacheChannel, ActionRpcChannel, Logger);
 					await Executor.ExecuteActionAsync(LeaseId, ActionTask, LeaseDir);
 				}
 				finally
@@ -1374,21 +1391,21 @@ namespace HordeAgent.Services
 					}
 				}
 
-				// Add a flag for EC2 instances
-				using (ManagementObjectSearcher Searcher = new ManagementObjectSearcher("select uuid from Win32_ComputerSystemProduct"))
-				{
-					foreach (ManagementObject Row in Searcher.Get())
-					{
-						string? Uuid = Row.GetPropertyValue("uuid") as string;
-						if (Uuid != null && Uuid.StartsWith("EC2", StringComparison.OrdinalIgnoreCase))
-						{
-							PrimaryDevice.Properties.Add("EC2=1");
-							break;
-						}
-					}
-				}
+				// Add EC2 properties if needed
+				await AddAwsProperties(PrimaryDevice.Properties, Logger);
+
+				// Add session information
+				PrimaryDevice.Properties.Add($"User={Environment.UserName}");
+				PrimaryDevice.Properties.Add($"Domain={Environment.UserDomainName}");
+				PrimaryDevice.Properties.Add($"Interactive={Environment.UserInteractive}");
+				PrimaryDevice.Properties.Add($"Elevated={BuildGraphExecutor.IsUserAdministrator()}");
+			}
+			else if (RuntimeInformation.IsOSPlatform(OSPlatform.Linux))
+			{
+				PrimaryDevice.Properties.Add("OSFamily=Linux");
+				PrimaryDevice.Properties.Add("OSVersion=Linux");
 
-				// Add other EC2 properties
+				// Add EC2 properties if needed
 				await AddAwsProperties(PrimaryDevice.Properties, Logger);
 
 				// Add session information
@@ -1573,6 +1590,7 @@ namespace HordeAgent.Services
 		{
 			if (EC2InstanceMetadata.IdentityDocument != null)
 			{
+				Properties.Add("EC2=1");
 				AddAwsProperty("aws-instance-id", "/instance-id", Properties);
 				AddAwsProperty("aws-instance-type", "/instance-type", Properties);
 				AddAwsProperty("aws-region", "/region", Properties);
diff --git a/Engine/Source/Programs/Horde/HordeAgent/Utility/StorageExtensions.cs b/Engine/Source/Programs/Horde/HordeAgent/Utility/StorageExtensions.cs
--- a/Engine/Source/Programs/Horde/HordeAgent/Utility/StorageExtensions.cs
+++ b/Engine/Source/Programs/Horde/HordeAgent/Utility/StorageExtensions.cs
@@ -5,6 +5,7 @@ using System.Collections.Generic;
 using System.Security.Cryptography;
 using System.Text;
 using System.Threading.Tasks;
+using Blake3;
 using Build.Bazel.Remote.Execution.V2;
 using EpicGames.Core;
 using Google.Protobuf;
@@ -14,22 +15,24 @@ namespace HordeAgent.Utility
 {
 	static class StorageExtensions
 	{
-		public static async Task<byte[]> GetBulkDataAsync(this ContentAddressableStorageClient Storage, Digest Digest)
+		public static async Task<byte[]> GetBulkDataAsync(this ContentAddressableStorageClient Storage, string InstanceName, Digest Digest)
 		{
 			BatchReadBlobsRequest Request = new BatchReadBlobsRequest();
+			Request.InstanceName = InstanceName;
 			Request.Digests.Add(Digest);
 
 			BatchReadBlobsResponse Response = await Storage.BatchReadBlobsAsync(Request);
 			return Response.Responses[0].Data.ToByteArray();
 		}
 
-		public static async Task<Digest> PutBulkDataAsync(this ContentAddressableStorageClient Storage, byte[] Data)
+		public static async Task<Digest> PutBulkDataAsync(this ContentAddressableStorageClient Storage, string InstanceName, byte[] Data)
 		{
 			BatchUpdateBlobsRequest.Types.Request BlobRequest = new BatchUpdateBlobsRequest.Types.Request();
 			BlobRequest.Data = ByteString.CopyFrom(Data);
 			BlobRequest.Digest = GetDigest(Data);
 
 			BatchUpdateBlobsRequest Request = new BatchUpdateBlobsRequest();
+			Request.InstanceName = InstanceName;
 			Request.Requests.Add(BlobRequest);
 
 			BatchUpdateBlobsResponse Response = await Storage.BatchUpdateBlobsAsync(Request);
@@ -37,38 +40,64 @@ namespace HordeAgent.Utility
 			return BlobRequest.Digest;
 		}
 
-		public static async Task<T> GetProtoMessageAsync<T>(this ContentAddressableStorageClient Storage, Digest Digest) where T : class, IMessage<T>, new()
+		public static async Task<T> GetProtoMessageAsync<T>(this ContentAddressableStorageClient Storage, string InstanceName, Digest Digest) where T : class, IMessage<T>, new()
 		{
 			BatchReadBlobsRequest Request = new BatchReadBlobsRequest();
+			Request.InstanceName = InstanceName;
 			Request.Digests.Add(Digest);
 
 			BatchReadBlobsResponse Response = await Storage.BatchReadBlobsAsync(Request);
+			if (Response.Responses[0].Status.Code != (int) Google.Rpc.Code.Ok)
+			{
+				throw new Exception($"Failed getting digest {Digest}. Code={Response.Responses[0].Status.Code} Message={Response.Responses[0].Status.Message}");
+			}
 
 			T Item = new T();
 			Item.MergeFrom(Response.Responses[0].Data);
 			return Item;
 		}
 
-		public static async Task<Digest> PutProtoMessageAsync<T>(this ContentAddressableStorageClient Storage, IMessage<T> Message) where T : class, IMessage<T>
+		public static async Task<Digest> PutProtoMessageAsync<T>(this ContentAddressableStorageClient Storage, string InstanceName, IMessage<T> Message) where T : class, IMessage<T>
 		{
 			BatchUpdateBlobsRequest.Types.Request BlobRequest = new BatchUpdateBlobsRequest.Types.Request();
 			BlobRequest.Data = Message.ToByteString();
 			BlobRequest.Digest = GetDigest(BlobRequest.Data.ToByteArray());
 
 			BatchUpdateBlobsRequest Request = new BatchUpdateBlobsRequest();
+			Request.InstanceName = InstanceName;
 			Request.Requests.Add(BlobRequest);
 
 			BatchUpdateBlobsResponse Response = await Storage.BatchUpdateBlobsAsync(Request);
 			//??
 			return BlobRequest.Digest;
 		}
-
+		
 		public static Digest GetDigest(ByteString Data)
 		{
-			return GetDigest(Data.ToByteArray());
+			return GetDigest(Data.ToByteArray(), DigestFunction.Types.Value.Epiciohash);
 		}
 
 		public static Digest GetDigest(byte[] Data)
+		{
+			return GetDigest(Data, DigestFunction.Types.Value.Epiciohash);
+		}
+
+		public static Digest GetDigest(ByteString Data, DigestFunction.Types.Value DigestFunction)
+		{
+			return GetDigest(Data.ToByteArray(), DigestFunction);
+		}
+
+		public static Digest GetDigest(byte[] Data, DigestFunction.Types.Value DigestFunction)
+		{
+			switch (DigestFunction)
+			{
+				case Build.Bazel.Remote.Execution.V2.DigestFunction.Types.Value.Sha256: return GetDigestSha256(Data);
+				case Build.Bazel.Remote.Execution.V2.DigestFunction.Types.Value.Epiciohash: return GetDigestIoHash(Data);
+				default: throw new Exception("Unknown digest function " + DigestFunction);
+			}
+		}
+		
+		public static Digest GetDigestSha256(byte[] Data)
 		{
 			using (SHA256 Algorithm = SHA256.Create())
 			{
@@ -76,5 +105,18 @@ namespace HordeAgent.Utility
 				return new Digest { Hash = StringUtils.FormatHexString(Hash), SizeBytes = Data.Length };
 			}
 		}
+		
+		public static Digest GetDigestIoHash(byte[] Data)
+		{
+			using Hasher Hasher = Hasher.New();
+			const int HashLength = 20;
+			Hasher.UpdateWithJoin(Data);
+			Hash Blake3Hash = Hasher.Finalize();
+            
+			// we only keep the first 20 bytes of the Blake3 hash
+			Span<byte> Hash = Blake3Hash.AsSpan().Slice(0, HashLength);
+			//return new Digest { Hash = StringUtils.FormatHexString(Hash), SizeBytes = HashLength };
+			return new Digest { Hash = StringUtils.FormatHexString(Hash).ToUpper(), SizeBytes = HashLength };
+		}
 	}
 }
diff --git a/Engine/Source/Programs/Horde/HordeAgentTests/HordeAgentTests.csproj b/Engine/Source/Programs/Horde/HordeAgentTests/HordeAgentTests.csproj
--- a/Engine/Source/Programs/Horde/HordeAgentTests/HordeAgentTests.csproj
+++ b/Engine/Source/Programs/Horde/HordeAgentTests/HordeAgentTests.csproj
@@ -1,4 +1,4 @@
-<Project Sdk="Microsoft.NET.Sdk">
+﻿<Project Sdk="Microsoft.NET.Sdk">
 
   <PropertyGroup>
     <TargetFramework>netcoreapp3.1</TargetFramework>
@@ -7,7 +7,7 @@
   </PropertyGroup>
 
   <ItemGroup>
-    <PackageReference Include="Microsoft.Extensions.Logging.Abstractions" Version="5.0.0" />
+    <PackageReference Include="Microsoft.Extensions.Logging.Abstractions" Version="3.1.6" />
     <PackageReference Include="Microsoft.NET.Test.Sdk" Version="16.5.0" />
     <PackageReference Include="Moq" Version="4.14.5" />
     <PackageReference Include="MSTest.TestAdapter" Version="2.1.0" />
diff --git a/Engine/Source/Programs/Horde/HordeAgentTests/LogParserTests.cs b/Engine/Source/Programs/Horde/HordeAgentTests/LogParserTests.cs
--- a/Engine/Source/Programs/Horde/HordeAgentTests/LogParserTests.cs
+++ b/Engine/Source/Programs/Horde/HordeAgentTests/LogParserTests.cs
@@ -484,8 +484,8 @@ namespace HordeAgentTests
 		{
 			List<CapturedEvent> Events = Parse(@"  C:\Program Files (x86)\Microsoft Visual Studio\2019\BuildTools\MSBuild\Current\Bin\Microsoft.Common.CurrentVersion.targets(4207,5): warning MSB3026: Could not copy ""obj\Development\DotNETUtilities.dll"" to ""..\..\..\..\Binaries\DotNET\DotNETUtilities.dll"". Beginning retry 2 in 1000ms. The process cannot access the file '..\..\..\..\Binaries\DotNET\DotNETUtilities.dll' because it is being used by another process. The file is locked by: ""UnrealAutomationTool(13236)"" [C:\Horde\Engine\Source\Programs\DotNETCommon\DotNETUtilities\DotNETUtilities.csproj]");
 			Assert.AreEqual(1, Events.Count);
-			Assert.AreEqual(LogLevel.Warning, Events[0].Level);
-			Assert.AreEqual(KnownLogEvents.MSBuild, Events[0].Id);
+			Assert.AreEqual(LogLevel.Information, Events[0].Level);
+			Assert.AreEqual(KnownLogEvents.Systemic_MSBuild, Events[0].Id);
 			Assert.AreEqual("warning", Events[0].Properties["severity"].ToString());
 			
 			// FIXME: Fails on Linux. Properties dict is empty
@@ -803,6 +803,31 @@ namespace HordeAgentTests
 			Assert.AreEqual(KnownLogEvents.Gauntlet, Event.Id);
 		}
 
+		[TestMethod]
+		public void DockerWarningMatcher()
+		{
+			string[] Lines =
+			{
+				@"#14 8.477 cc -O2 -Wall -DLUA_ANSI -DENABLE_CJSON_GLOBAL -DREDIS_STATIC=''    -c -o lauxlib.o lauxlib.c",
+				@"#14 8.499 lauxlib.c: In function 'luaL_loadfile':",
+				@"#14 8.499 lauxlib.c:577:4: warning: this 'while' clause does not guard... [-Wmisleading-indentation]",
+				@"#14 8.499     while ((c = getc(lf.f)) != EOF && c != LUA_SIGNATURE[0]) ;",
+				@"#14 8.499     ^~~~~",
+				@"#14 8.499 lauxlib.c:578:5: note: ...this statement, but the latter is misleadingly indented as if it were guarded by the 'while'",
+				@"#14 8.499      lf.extraline = 0;",
+				@"#14 8.499      ^~",
+				@"#14 8.643 cc -O2 -Wall -DLUA_ANSI -DENABLE_CJSON_GLOBAL -DREDIS_STATIC=''    -c -o lbaselib.o lbaselib.c",
+			};
+
+			List<CapturedEvent> Events = Parse(Lines);
+			Assert.AreEqual(1, Events.Count);
+
+			Assert.AreEqual(LogLevel.Warning, Events[0].Level);
+			Assert.AreEqual(KnownLogEvents.Compiler, Events[0].Id);
+			Assert.AreEqual(2, Events[0].LineIndex);
+			Assert.AreEqual(1, Events[0].LineCount);
+		}
+
 		[TestMethod]
 		public void GauntletErrorMatcher()
 		{
diff --git a/Engine/Source/Programs/Horde/HordeServer/Api/Issues.cs b/Engine/Source/Programs/Horde/HordeServer/Api/Issues.cs
--- a/Engine/Source/Programs/Horde/HordeServer/Api/Issues.cs
+++ b/Engine/Source/Programs/Horde/HordeServer/Api/Issues.cs
@@ -26,6 +26,11 @@ namespace HordeServer.Api
 		/// </summary>
 		public int Change { get; set; }
 
+		/// <summary>
+		/// Severity of the issue in this step
+		/// </summary>
+		public IssueSeverity Severity { get; set; }
+
 		/// <summary>
 		/// Name of the job containing this step
 		/// </summary>
@@ -63,6 +68,7 @@ namespace HordeServer.Api
 		public GetIssueStepResponse(IIssueStep IssueStep)
 		{
 			this.Change = IssueStep.Change;
+			this.Severity = IssueStep.Severity;
 			this.JobName = IssueStep.JobName;
 			this.JobId = IssueStep.JobId.ToString();
 			this.BatchId = IssueStep.BatchId.ToString();
@@ -416,6 +422,11 @@ namespace HordeServer.Api
 		/// </summary>
 		public string? Owner { get; set; }
 
+		/// <summary>
+		/// User id of the owner
+		/// </summary>
+		public string? OwnerId { get; set; }
+
 		/// <summary>
 		/// User that nominated the current owner
 		/// </summary>
@@ -436,6 +447,26 @@ namespace HordeServer.Api
 		/// </summary>
 		public DateTime? ResolvedAt { get; set; }
 
+		/// <summary>
+		/// Name of the user that resolved the issue
+		/// </summary>
+		public string? ResolvedBy { get; set; }
+
+		/// <summary>
+		/// User id of the person that resolved the issue
+		/// </summary>
+		public string? ResolvedById { get; set; }
+
+		/// <summary>
+		/// Time at which the issue was verified
+		/// </summary>
+		public DateTime? VerifiedAt { get; set; }
+
+		/// <summary>
+		/// Time that the issue was last seen
+		/// </summary>
+		public DateTime LastSeenAt { get; set; }
+
 		/// <summary>
 		/// List of stream paths affected by this issue
 		/// </summary>
@@ -480,11 +511,16 @@ namespace HordeServer.Api
 			this.RetrievedAt = DateTime.UtcNow;
 			this.Summary = String.IsNullOrEmpty(Issue.UserSummary)? Issue.Summary : Issue.UserSummary;
 			this.Severity = Issue.Severity;
-			this.Owner = Issue.Owner;
-			this.NominatedBy = Issue.NominatedBy;
+			this.Owner = Details.Owner?.Login;
+			this.OwnerId = (Details.Owner == null)? null : Details.Owner.Id.ToString();
+			this.NominatedBy = Details.NominatedBy?.Login;
 			this.AcknowledgedAt = Issue.AcknowledgedAt;
 			this.FixChange = Issue.FixChange;
 			this.ResolvedAt = Issue.ResolvedAt;
+			this.ResolvedBy = Details.ResolvedBy?.Login;
+			this.ResolvedById = (Details.ResolvedBy == null) ? null : Details.ResolvedBy.Id.ToString();
+			this.VerifiedAt = Issue.VerifiedAt;
+			this.LastSeenAt = Issue.LastSeenAt;
 			this.Streams = Details.Spans.Select(x => x.StreamName).Distinct().ToList()!;
 			this.ResolvedStreams = new List<string>();
 			this.UnresolvedStreams = new List<string>();
@@ -500,7 +536,7 @@ namespace HordeServer.Api
 					this.UnresolvedStreams.Add(Stream.Key.ToString());
 				}
 			}
-			this.PrimarySuspects = Issue.Suspects.Select(x => x.Author).Distinct(StringComparer.OrdinalIgnoreCase).ToList();
+			this.PrimarySuspects = Details.SuspectUsers.Where(x => x.Login != null).Select(x => x.Login).ToList();
 			this.ShowDesktopAlerts = ShowDesktopAlerts;
 		}
 	}
diff --git a/Engine/Source/Programs/Horde/HordeServer/Api/UgsMetadata.cs b/Engine/Source/Programs/Horde/HordeServer/Api/UgsMetadata.cs
--- a/Engine/Source/Programs/Horde/HordeServer/Api/UgsMetadata.cs
+++ b/Engine/Source/Programs/Horde/HordeServer/Api/UgsMetadata.cs
@@ -499,17 +499,19 @@ namespace HordeServer.Api
 		/// Constructs a new issue
 		/// </summary>
 		/// <param name="Details">Issue to construct from</param>
+		/// <param name="Owner"></param>
+		/// <param name="NominatedBy"></param>
 		/// <param name="bNotify">Whether to notify the user about this issue</param>
 		/// <param name="BuildUrl">Link to the last build for this issue</param>
-		public GetUgsIssueResponse(IIssueDetails Details, bool bNotify, Uri? BuildUrl)
+		public GetUgsIssueResponse(IIssueDetails Details, IUser? Owner, IUser? NominatedBy, bool bNotify, Uri? BuildUrl)
 		{
 			IIssue Issue = Details.Issue;
 			this.Id = Issue.Id;
 			this.CreatedAt = Issue.CreatedAt;
 			this.RetrievedAt = DateTime.UtcNow;
 			this.Summary = Issue.Summary;
-			this.Owner = Issue.Owner;
-			this.NominatedBy = Issue.NominatedBy;
+			this.Owner = Owner?.Login;
+			this.NominatedBy = NominatedBy?.Login;
 			this.AcknowledgedAt = Issue.AcknowledgedAt;
 			this.FixChange = Issue.FixChange;
 			this.ResolvedAt = Issue.ResolvedAt;
@@ -554,5 +556,10 @@ namespace HordeServer.Api
 		/// Whether the issue should be marked as resolved
 		/// </summary>
 		public bool? Resolved { get; set; }
+
+		/// <summary>
+		/// Name of the user that resolved the issue
+		/// </summary>
+		public string? ResolvedBy { get; set; }
 	}
 }
\ No newline at end of file
diff --git a/Engine/Source/Programs/Horde/HordeServer/Authentication/ServiceAccountAuthHandler.cs b/Engine/Source/Programs/Horde/HordeServer/Authentication/ServiceAccountAuthHandler.cs
--- a/Engine/Source/Programs/Horde/HordeServer/Authentication/ServiceAccountAuthHandler.cs
+++ b/Engine/Source/Programs/Horde/HordeServer/Authentication/ServiceAccountAuthHandler.cs
@@ -60,7 +60,7 @@ namespace HordeServer.Authentication
 
 			List<Claim> Claims = new List<Claim>(10);
 			Claims.Add(new Claim(ClaimTypes.Name, AuthenticationScheme));
-			Claims.AddRange(ServiceAccount.Claims.Select(ClaimPair => new Claim(ClaimPair.Key, ClaimPair.Value)));
+			Claims.AddRange(ServiceAccount.GetClaims().Select(ClaimPair => new Claim(ClaimPair.Type, ClaimPair.Value)));
 
 			ClaimsIdentity Identity = new ClaimsIdentity(Claims, Scheme.Name);
 			ClaimsPrincipal Principal = new ClaimsPrincipal(Identity);
diff --git a/Engine/Source/Programs/Horde/HordeServer/Collections/IIssueCollection.cs b/Engine/Source/Programs/Horde/HordeServer/Collections/IIssueCollection.cs
--- a/Engine/Source/Programs/Horde/HordeServer/Collections/IIssueCollection.cs
+++ b/Engine/Source/Programs/Horde/HordeServer/Collections/IIssueCollection.cs
@@ -182,6 +182,11 @@ namespace HordeServer.Collections
 		/// </summary>
 		public int Change { get; set; }
 
+		/// <summary>
+		/// Severity of the issue
+		/// </summary>
+		public IssueSeverity Severity { get; set; }
+
 		/// <summary>
 		/// Name of the job
 		/// </summary>
@@ -221,6 +226,7 @@ namespace HordeServer.Collections
 		/// Constructor
 		/// </summary>
 		/// <param name="Change">The changelist number for this job</param>
+		/// <param name="Severity">Severity of the issue in this step</param>
 		/// <param name="JobName">The job name</param>
 		/// <param name="JobId">The unique job id</param>
 		/// <param name="BatchId">The batch id</param>
@@ -228,9 +234,10 @@ namespace HordeServer.Collections
 		/// <param name="StepTime">Time that the step started</param>
 		/// <param name="LogId">Unique id of the log file for this step</param>
 		/// <param name="NotifySuspects">Whether to notify suspects for this step</param>
-		public NewIssueStepData(int Change, string JobName, ObjectId JobId, SubResourceId BatchId, SubResourceId StepId, DateTime StepTime, ObjectId? LogId, bool NotifySuspects)
+		public NewIssueStepData(int Change, IssueSeverity Severity, string JobName, ObjectId JobId, SubResourceId BatchId, SubResourceId StepId, DateTime StepTime, ObjectId? LogId, bool NotifySuspects)
 		{
 			this.Change = Change;
+			this.Severity = Severity;
 			this.JobName = JobName;
 			this.JobId = JobId;
 			this.BatchId = BatchId;
@@ -246,9 +253,10 @@ namespace HordeServer.Collections
 		/// <param name="Job">The job being built</param>
 		/// <param name="Batch">Batch of the job for the step</param>
 		/// <param name="Step">The step being built</param>
+		/// <param name="Severity">Severity of the issue in this step</param>
 		/// <param name="NotifySuspects">Whether to notify suspects for this step failing</param>
-		public NewIssueStepData(IJob Job, IJobStepBatch Batch, IJobStep Step, bool NotifySuspects)
-			: this(Job.Change, Job.Name, Job.Id, Batch.Id, Step.Id, Step.StartTimeUtc ?? default, Step.LogId, NotifySuspects)
+		public NewIssueStepData(IJob Job, IJobStepBatch Batch, IJobStep Step, IssueSeverity Severity, bool NotifySuspects)
+			: this(Job.Change, Severity, Job.Name, Job.Id, Batch.Id, Step.Id, Step.StartTimeUtc ?? default, Step.LogId, NotifySuspects)
 		{
 		}
 
@@ -257,7 +265,7 @@ namespace HordeServer.Collections
 		/// </summary>
 		/// <param name="JobStepRef">The jobstep reference</param>
 		public NewIssueStepData(IJobStepRef JobStepRef)
-			: this(JobStepRef.Change, JobStepRef.JobName, JobStepRef.Id.JobId, JobStepRef.Id.BatchId, JobStepRef.Id.StepId, JobStepRef.StartTimeUtc, JobStepRef.LogId, false)
+			: this(JobStepRef.Change, IssueSeverity.Unspecified, JobStepRef.JobName, JobStepRef.Id.JobId, JobStepRef.Id.BatchId, JobStepRef.Id.StepId, JobStepRef.StartTimeUtc, JobStepRef.LogId, false)
 		{
 		}
 	}
@@ -270,7 +278,7 @@ namespace HordeServer.Collections
 		/// <summary>
 		/// Author of the changelist
 		/// </summary>
-		public string Author { get; set; }
+		public ObjectId AuthorId { get; set; }
 
 		/// <summary>
 		/// The submitted changelist
@@ -285,12 +293,12 @@ namespace HordeServer.Collections
 		/// <summary>
 		/// Constructor
 		/// </summary>
-		/// <param name="Author">Author of the change</param>
+		/// <param name="AuthorId">Author of the change</param>
 		/// <param name="Change">The changelist number</param>
 		/// <param name="DeclinedAt">The time that the user declined this issue</param>
-		public NewIssueSuspectData(string Author, int Change, DateTime? DeclinedAt)
+		public NewIssueSuspectData(ObjectId AuthorId, int Change, DateTime? DeclinedAt)
 		{
-			this.Author = Author;
+			this.AuthorId = AuthorId;
 			this.Change = Change;
 			this.DeclinedAt = DeclinedAt;
 		}
@@ -377,10 +385,18 @@ namespace HordeServer.Collections
 		/// <returns>The issue matching the given id, or null</returns>
 		Task<IIssue?> GetIssueAsync(int IssueId);
 
+		/// <summary>
+		/// Gets the suspects for an issue
+		/// </summary>
+		/// <param name="Issue">The issue to retrieve suspects for</param>
+		/// <returns>List of suspects</returns>
+		Task<List<IIssueSuspect>> GetSuspectsAsync(IIssue Issue);
+
 		/// <summary>
 		/// Searches for open issues
 		/// </summary>
 		/// <param name="Ids">Set of issue ids to find</param>
+		/// <param name="UserId">The user to find issues for</param>
 		/// <param name="StreamId">The stream affected by the issue</param>
 		/// <param name="MinChange">Minimum changelist affected by the issue</param>
 		/// <param name="MaxChange">Maximum changelist affected by the issue</param>
@@ -388,7 +404,7 @@ namespace HordeServer.Collections
 		/// <param name="Index">Index within the results to return</param>
 		/// <param name="Count">Number of results</param>
 		/// <returns>List of streams open in the given stream at the given changelist</returns>
-		Task<List<IIssue>> FindIssuesAsync(IEnumerable<int>? Ids = null, StreamId? StreamId = null, int? MinChange = null, int? MaxChange = null, bool? Resolved = null, int? Index = null, int? Count = null);
+		Task<List<IIssue>> FindIssuesAsync(IEnumerable<int>? Ids = null, ObjectId? UserId = null, StreamId? StreamId = null, int? MinChange = null, int? MaxChange = null, bool? Resolved = null, int? Index = null, int? Count = null);
 
 		/// <summary>
 		/// Searches for open issues
@@ -404,16 +420,24 @@ namespace HordeServer.Collections
 		/// <param name="NewSeverity">New severity for the issue</param>
 		/// <param name="NewSummary">New summary for the issue</param>
 		/// <param name="NewUserSummary">New user summary for the issue</param>
-		/// <param name="NewOwner">New owner of the issue</param>
-		/// <param name="NewNominatedBy">Person that nominated the new owner</param>
+		/// <param name="NewOwnerId">New owner of the issue</param>
+		/// <param name="NewNominatedById">Person that nominated the new owner</param>
 		/// <param name="NewAcknowledged">Whether the issue has been acknowledged</param>
-		/// <param name="NewDeclinedBy">Name of a user that has declined the issue</param>
+		/// <param name="NewDeclinedById">Name of a user that has declined the issue</param>
 		/// <param name="NewFixChange">Fix changelist for the issue. Pass 0 to clear the fix changelist, -1 for systemic issue.</param>
-		/// <param name="NewResolved">Whether the issue has been resolved</param>
+		/// <param name="NewFixStreamIds">List of stream ids with a bool indicating whether the fix changelist is in this stream</param>
+		/// <param name="NewResolvedById">User that resolved the issue (may be ObjectId.Empty to clear)</param>
+		/// <param name="NewLastSeenAt"></param>
 		/// <param name="NewNotifySuspects">Whether all suspects should be notified about this issue</param>
-		/// <param name="NewSuspects">New list of suspects for this issue</param>
 		/// <returns>True if the issue was updated</returns>
-		Task<IIssue?> UpdateIssueAsync(IIssue Issue, IssueSeverity? NewSeverity = null, string? NewSummary = null, string? NewUserSummary = null, string? NewOwner = null, string? NewNominatedBy = null, bool? NewAcknowledged = null, string? NewDeclinedBy = null, int? NewFixChange = null, bool? NewResolved = null, bool? NewNotifySuspects = null, List<NewIssueSuspectData>? NewSuspects = null);
+		Task<IIssue?> UpdateIssueAsync(IIssue Issue, IssueSeverity? NewSeverity = null, string? NewSummary = null, string? NewUserSummary = null, ObjectId? NewOwnerId = null, ObjectId? NewNominatedById = null, bool? NewAcknowledged = null, ObjectId? NewDeclinedById = null, int? NewFixChange = null, Dictionary<StreamId, bool>? NewFixStreamIds = null, ObjectId? NewResolvedById = null, DateTime? NewLastSeenAt = null, bool? NewNotifySuspects = null);
+
+		/// <summary>
+		/// Updates the list of suspects for an issue, and optionally sets the resolved state
+		/// </summary>
+		/// <param name="Issue">The issue to update</param>
+		/// <returns>True if the issue was updated</returns>
+		Task<IIssue?> UpdateIssueDerivedDataAsync(IIssue Issue);
 
 		#endregion
 
@@ -426,14 +450,13 @@ namespace HordeServer.Collections
 		/// <param name="Stream">The stream containing the span</param>
 		/// <param name="TemplateRefId">The template being executed</param>
 		/// <param name="NodeName">Name of the step</param>
-		/// <param name="Severity">Severity of this span</param>
 		/// <param name="Fingerprint">Fingerprint for this issue</param>
 		/// <param name="LastSuccess">Last time the build succeeded before the span</param>
 		/// <param name="Failure">List of failing steps</param>
 		/// <param name="NextSuccess">First time the build succeeded after the span</param>
 		/// <param name="Suspects">Suspects for this span</param>
 		/// <returns>New span, or null if the sequence token is not valid</returns>
-		Task<IIssueSpan?> AddSpanAsync(IIssueSequenceToken Token, IStream Stream, TemplateRefId TemplateRefId, string NodeName, IssueSeverity Severity, NewIssueFingerprint Fingerprint, NewIssueStepData? LastSuccess, NewIssueStepData Failure, NewIssueStepData? NextSuccess, List<NewIssueSpanSuspectData> Suspects);
+		Task<IIssueSpan?> AddSpanAsync(IIssueSequenceToken Token, IStream Stream, TemplateRefId TemplateRefId, string NodeName, NewIssueFingerprint Fingerprint, NewIssueStepData? LastSuccess, NewIssueStepData Failure, NewIssueStepData? NextSuccess, List<NewIssueSpanSuspectData> Suspects);
 
 		/// <summary>
 		/// Gets a particular span
diff --git a/Engine/Source/Programs/Horde/HordeServer/Collections/IServiceAccountCollection.cs b/Engine/Source/Programs/Horde/HordeServer/Collections/IServiceAccountCollection.cs
--- a/Engine/Source/Programs/Horde/HordeServer/Collections/IServiceAccountCollection.cs
+++ b/Engine/Source/Programs/Horde/HordeServer/Collections/IServiceAccountCollection.cs
@@ -18,7 +18,7 @@ namespace HordeServer.Collections
 		/// <param name="SecretToken">Secret token used for identifying this service account</param>
 		/// <param name="Claims">Dictionary of claims to assign</param>
 		/// <param name="Description">Description of the account</param>
-		Task<IServiceAccount> AddAsync(string SecretToken, Dictionary<string, string> Claims, string Description);
+		Task<IServiceAccount> AddAsync(string SecretToken, List<string> Claims, string Description);
 
 		/// <summary>
 		/// Get service account via ID
@@ -43,7 +43,7 @@ namespace HordeServer.Collections
 		/// <param name="Enabled">If set, enabled flag will be set</param>
 		/// <param name="Description">If set, description will be set</param>
 		/// <returns>Async task</returns>
-		Task UpdateAsync(ObjectId Id, string? SecretToken, Dictionary<string, string>? Claims, bool? Enabled, string? Description);
+		Task UpdateAsync(ObjectId Id, string? SecretToken, List<string>? Claims, bool? Enabled, string? Description);
 
 		/// <summary>
 		/// Delete a service account from the collection
diff --git a/Engine/Source/Programs/Horde/HordeServer/Collections/IUserCollection.cs b/Engine/Source/Programs/Horde/HordeServer/Collections/IUserCollection.cs
--- a/Engine/Source/Programs/Horde/HordeServer/Collections/IUserCollection.cs
+++ b/Engine/Source/Programs/Horde/HordeServer/Collections/IUserCollection.cs
@@ -24,6 +24,13 @@ namespace HordeServer.Collections
 		/// <returns>The user information</returns>
 		Task<IUser?> GetUserAsync(ObjectId Id);
 
+		/// <summary>
+		/// Gets a cached user by unique id
+		/// </summary>
+		/// <param name="Id">Id of the user</param>
+		/// <returns>The user information</returns>
+		Task<IUser?> GetCachedUserAsync(ObjectId Id);
+
 		/// <summary>
 		/// Gets a user by unique id
 		/// </summary>
diff --git a/Engine/Source/Programs/Horde/HordeServer/Collections/Impl/IssueCollection.cs b/Engine/Source/Programs/Horde/HordeServer/Collections/Impl/IssueCollection.cs
--- a/Engine/Source/Programs/Horde/HordeServer/Collections/Impl/IssueCollection.cs
+++ b/Engine/Source/Programs/Horde/HordeServer/Collections/Impl/IssueCollection.cs
@@ -77,13 +77,16 @@ namespace HordeServer.Collections.Impl
 			[BsonIgnoreIfNull]
 			public IssueSpan? InsertIssueSpan { get; set; }
 
+			[BsonIgnoreIfNull]
+			public List<IssueSuspect>? InsertIssueSuspects { get; set; }
+
 			[BsonIgnoreIfNull]
 			public UpdateTransaction? UpdateIssue { get; set; }
 
 			[BsonIgnoreIfNull]
 			public UpdateTransaction? UpdateIssueSpan { get; set; }
 
-			public bool HasPendingTransaction => InsertIssue != null || InsertIssueSpan != null || UpdateIssue != null || UpdateIssueSpan != null;
+			public bool HasPendingTransaction => InsertIssue != null || InsertIssueSpan != null || InsertIssueSuspects != null || UpdateIssue != null || UpdateIssueSpan != null;
 		}
 
 		class IssueSequenceToken : IIssueSequenceToken
@@ -142,14 +145,22 @@ namespace HordeServer.Collections.Impl
 			[BsonIgnoreIfNull]
 			public DateTime? AcknowledgedAt { get; set; }
 
-			public bool Resolved { get; set; }
-
 			[BsonIgnoreIfNull]
 			public DateTime? ResolvedAt { get; set; }
 
+			[BsonIgnoreIfNull]
+			public ObjectId? ResolvedById { get; set; }
+
+			[BsonIgnoreIfNull]
+			public DateTime? VerifiedAt { get; set; }
+
+			public DateTime LastSeenAt { get; set; }
+
 			[BsonIgnoreIfNull]
 			public int? FixChange { get; set; }
 
+			public List<IssueStream> Streams { get; set; } = new List<IssueStream>();
+
 			public int MinSuspectChange { get; set; }
 			public int MaxSuspectChange { get; set; }
 
@@ -161,9 +172,9 @@ namespace HordeServer.Collections.Impl
 			public int UpdateIndex { get; set; }
 
 			IIssueFingerprint IIssue.Fingerprint => Fingerprint;
-			IReadOnlyList<IIssueSuspect> IIssue.Suspects => Suspects;
-			string? IIssue.Owner => Owner ?? GetDefaultOwner()?.Author;
 			ObjectId? IIssue.OwnerId => OwnerId ?? GetDefaultOwner()?.AuthorId;
+			IReadOnlyList<IIssueStream> IIssue.Streams => Streams;
+			DateTime IIssue.LastSeenAt => (LastSeenAt == default) ? DateTime.UtcNow : LastSeenAt;
 
 			[BsonConstructor]
 			private Issue()
@@ -178,10 +189,13 @@ namespace HordeServer.Collections.Impl
 				this.Id = Id;
 				this.Summary = Summary;
 				this.CreatedAt = DateTime.UtcNow;
+				this.LastSeenAt = DateTime.UtcNow;
 				this.Fingerprint = new IssueFingerprint(Span.Fingerprint);
 				this.Severity = Span.Severity;
 				this.NotifySuspects = Span.NotifySuspects;
-				this.Suspects = Span.Suspects.ConvertAll(x => new IssueSuspect(x.Author, x.AuthorId, x.OriginatingChange ?? x.Change, null));
+				this.ResolvedAt = Span.NextSuccess?.StepTime;
+				this.VerifiedAt = this.ResolvedAt;
+				this.Suspects = Span.Suspects.ConvertAll(x => new IssueSuspect(Id, x.AuthorId, x.OriginatingChange ?? x.Change, null, this.ResolvedAt));
 
 				if (Suspects.Count > 0)
 				{
@@ -194,8 +208,8 @@ namespace HordeServer.Collections.Impl
 			{
 				if (Suspects.Count > 0)
 				{
-					string PossibleOwner = Suspects[0].Author;
-					if (Suspects.All(x => x.Author.Equals(PossibleOwner, StringComparison.OrdinalIgnoreCase)) && Suspects.Any(x => x.DeclinedAt == null))
+					ObjectId PossibleOwner = Suspects[0].AuthorId;
+					if (Suspects.All(x => x.AuthorId == PossibleOwner) && Suspects.Any(x => x.DeclinedAt == null))
 					{
 						return Suspects[0];
 					}
@@ -204,33 +218,48 @@ namespace HordeServer.Collections.Impl
 			}
 		}
 
+		class IssueStream : IIssueStream
+		{
+			public StreamId StreamId { get; set; }
+			public bool? ContainsFix { get; set; }
+
+			public IssueStream()
+			{
+			}
+
+			public IssueStream(IIssueStream Other)
+			{
+				this.StreamId = Other.StreamId;
+				this.ContainsFix = Other.ContainsFix;
+			}
+		}
+
 		class IssueSuspect : IIssueSuspect
 		{
-			[BsonRequired]
-			public string Author { get; set; }
+			public ObjectId Id { get; set; }
+			public int IssueId { get; set; }
 			public ObjectId AuthorId { get; set; }
 			public int Change { get; set; }
 			public DateTime? DeclinedAt { get; set; }
+			public DateTime? ResolvedAt { get; set; } // Degenerate
 
 			private IssueSuspect()
 			{
-				Author = String.Empty;
 			}
 
-			public IssueSuspect(string Author, ObjectId AuthorId, int Change, DateTime? DeclinedAt)
+			public IssueSuspect(int IssueId, IIssueSpanSuspect Suspect)
+				: this(IssueId, Suspect.AuthorId, Suspect.OriginatingChange ?? Suspect.Change, null, null)
 			{
-				this.Author = Author;
-				this.AuthorId = AuthorId;
-				this.Change = Change;
-				this.DeclinedAt = DeclinedAt;
 			}
 
-			public IssueSuspect(IIssueSuspect Other)
+			public IssueSuspect(int IssueId, ObjectId AuthorId, int Change, DateTime? DeclinedAt, DateTime? ResolvedAt)
 			{
-				this.Author = Other.Author;
-				this.AuthorId = Other.AuthorId;
-				this.Change = Other.Change;
-				this.DeclinedAt = Other.DeclinedAt;
+				this.Id = ObjectId.GenerateNewId();
+				this.IssueId = IssueId;
+				this.AuthorId = AuthorId;
+				this.Change = Change;
+				this.DeclinedAt = DeclinedAt;
+				this.ResolvedAt = ResolvedAt;
 			}
 		}
 
@@ -271,6 +300,7 @@ namespace HordeServer.Collections.Impl
 			[BsonRequired]
 			public string NodeName { get; set; }
 			public IssueSeverity Severity { get; set; }
+			public DateTime? ResolvedAt { get; set; } // Propagated from the owning issue
 
 			[BsonRequired]
 			public IssueFingerprint Fingerprint { get; set; }
@@ -311,14 +341,14 @@ namespace HordeServer.Collections.Impl
 				this.Suspects = new List<IssueSpanSuspect>();
 			}
 
-			public IssueSpan(StreamId StreamId, string StreamName, TemplateRefId TemplateRefId, string NodeName, IssueSeverity Severity, IIssueFingerprint Fingerprint, NewIssueStepData? LastSuccess, NewIssueStepData Failure, NewIssueStepData? NextSuccess, List<IssueSpanSuspect> Suspects)
+			public IssueSpan(StreamId StreamId, string StreamName, TemplateRefId TemplateRefId, string NodeName, IIssueFingerprint Fingerprint, NewIssueStepData? LastSuccess, NewIssueStepData Failure, NewIssueStepData? NextSuccess, List<IssueSpanSuspect> Suspects)
 			{
 				this.Id = ObjectId.GenerateNewId();
 				this.StreamId = StreamId;
 				this.StreamName = StreamName;
 				this.TemplateRefId = TemplateRefId;
 				this.NodeName = NodeName;
-				this.Severity = Severity;
+				this.Severity = Failure.Severity;
 				this.Fingerprint = new IssueFingerprint(Fingerprint);
 				if (LastSuccess != null)
 				{
@@ -366,6 +396,7 @@ namespace HordeServer.Collections.Impl
 			public ObjectId SpanId { get; set; }
 
 			public int Change { get; set; }
+			public IssueSeverity Severity { get; set; }
 
 			[BsonRequired]
 			public string JobName { get; set; }
@@ -396,6 +427,7 @@ namespace HordeServer.Collections.Impl
 				this.Id = ObjectId.GenerateNewId();
 				this.SpanId = SpanId;
 				this.Change = StepData.Change;
+				this.Severity = StepData.Severity;
 				this.JobName = StepData.JobName;
 				this.JobId = StepData.JobId;
 				this.BatchId = StepData.BatchId;
@@ -410,6 +442,7 @@ namespace HordeServer.Collections.Impl
 		IMongoCollection<Issue> Issues;
 		IMongoCollection<IssueSpan> IssueSpans;
 		IMongoCollection<IssueStep> IssueSteps;
+		IMongoCollection<IssueSuspect> IssueSuspects;
 		ILogger Logger;
 
 		// TODO: Temp upgrade path for switching user references to ids. Remove after resaving all issues.
@@ -428,11 +461,13 @@ namespace HordeServer.Collections.Impl
 			Issues = DatabaseService.GetCollection<Issue>("IssuesV2");
 			IssueSpans = DatabaseService.GetCollection<IssueSpan>("IssuesV2.Spans");
 			IssueSteps = DatabaseService.GetCollection<IssueStep>("IssuesV2.Steps");
+			IssueSuspects = DatabaseService.GetCollection<IssueSuspect>("IssuesV2.Suspects");
 			
 			if (!DatabaseService.ReadOnlyMode)
 			{
-				Issues.Indexes.CreateOne(new CreateIndexModel<Issue>(Builders<Issue>.IndexKeys.Ascending(x => x.Resolved)));
-				
+				Issues.Indexes.CreateOne(new CreateIndexModel<Issue>(Builders<Issue>.IndexKeys.Ascending(x => x.ResolvedAt)));
+				Issues.Indexes.CreateOne(new CreateIndexModel<Issue>(Builders<Issue>.IndexKeys.Ascending(x => x.VerifiedAt)));
+
 				IssueSpans.Indexes.CreateOne(new CreateIndexModel<IssueSpan>(Builders<IssueSpan>.IndexKeys.Ascending(x => x.Modified)));
 				IssueSpans.Indexes.CreateOne(new CreateIndexModel<IssueSpan>(Builders<IssueSpan>.IndexKeys.Ascending(x => x.IssueId)));
 				IssueSpans.Indexes.CreateOne(new CreateIndexModel<IssueSpan>(Builders<IssueSpan>.IndexKeys.Ascending(x => x.StreamId).Ascending(x => x.MinChange).Ascending(x => x.MaxChange)));
@@ -440,6 +475,9 @@ namespace HordeServer.Collections.Impl
 				
 				IssueSteps.Indexes.CreateOne(new CreateIndexModel<IssueStep>(Builders<IssueStep>.IndexKeys.Ascending(x => x.SpanId)));
 				IssueSteps.Indexes.CreateOne(new CreateIndexModel<IssueStep>(Builders<IssueStep>.IndexKeys.Ascending(x => x.JobId).Ascending(x => x.BatchId).Ascending(x => x.StepId)));
+
+				IssueSuspects.Indexes.CreateOne(new CreateIndexModel<IssueSuspect>(Builders<IssueSuspect>.IndexKeys.Ascending(x => x.AuthorId).Ascending(x => x.ResolvedAt)));
+				IssueSuspects.Indexes.CreateOne(new CreateIndexModel<IssueSuspect>(Builders<IssueSuspect>.IndexKeys.Ascending(x => x.IssueId).Ascending(x => x.Change), new CreateIndexOptions { Unique = true }));
 			}
 		}
 
@@ -504,6 +542,11 @@ namespace HordeServer.Collections.Impl
 					await IssueSpans.InsertOneIgnoreDuplicatesAsync(Ledger.InsertIssueSpan);
 					Ledger.InsertIssueSpan = null;
 				}
+				if (Ledger.InsertIssueSuspects != null)
+				{
+					await IssueSuspects.InsertManyIgnoreDuplicatesAsync(Ledger.InsertIssueSuspects);
+					Ledger.InsertIssueSuspects = null;
+				}
 				if (Ledger.UpdateIssue != null)
 				{
 					await Issues.UpdateOneAsync(Ledger.UpdateIssue.Filter, Ledger.UpdateIssue.Update);
@@ -547,6 +590,7 @@ namespace HordeServer.Collections.Impl
 			IssueSequenceToken TokenValue = ((IssueSequenceToken)Token);
 			Issue NewIssue = new Issue(++TokenValue.Ledger.NextId, Summary, Span);
 			TokenValue.Ledger.InsertIssue = NewIssue;
+			TokenValue.Ledger.InsertIssueSuspects = NewIssue.Suspects;
 			TokenValue.Ledger.UpdateIssueSpan = UpdateTransaction.Create<IssueSpan>(x => x.Id == Span.Id, Builders<IssueSpan>.Update.Set(x => x.IssueId, NewIssue.Id));
 
 			if (!await LedgerSingleton.TryUpdateAsync(TokenValue.Ledger))
@@ -566,83 +610,159 @@ namespace HordeServer.Collections.Impl
 			return Issue;
 		}
 
+		/// <inheritdoc/>
+		public Task<List<IIssueSuspect>> GetSuspectsAsync(IIssue Issue)
+		{
+			Issue IssueDocument = (Issue)Issue;
+			return Task.FromResult(IssueDocument.Suspects.ConvertAll<IIssueSuspect>(x => x));
+		}
 
 		class ProjectedIssueId
 		{
-			public int? IssueId { get; set; }
+			public int? _id { get; set; }
 		}
 
-		class ProjectedIssueIdWithIssue
+		/// <inheritdoc/>
+		public async Task<List<IIssue>> FindIssuesAsync(IEnumerable<int>? Ids = null, ObjectId? UserId = null, StreamId? StreamId = null, int? MinChange = null, int? MaxChange = null, bool? Resolved = null, int? Index = null, int? Count = null)
 		{
-			public int? IssueId { get; set; }
-			public IEnumerable<Issue> Issues { get; set; } = null!;
+			List<Issue> Results = await FilterIssuesByUserIdAsync(Ids, UserId, StreamId, MinChange, MaxChange, Resolved ?? false, Index ?? 0, Count);
+			return await UpgradeIssuesAsync(Results);
 		}
 
-		/// <inheritdoc/>
-		public async Task<List<IIssue>> FindIssuesAsync(IEnumerable<int>? Ids = null, StreamId? StreamId = null, int? MinChange = null, int? MaxChange = null, bool? Resolved = null, int? Index = null, int? Count = null)
+		async Task<List<Issue>> FilterIssuesByUserIdAsync(IEnumerable<int>? Ids, ObjectId? UserId, StreamId? StreamId, int? MinChange, int? MaxChange, bool? Resolved, int Index, int? Count)
 		{
-			// Filtering against a stream ID needs to be done via a join on the spans collection
-			if (StreamId != null)
+			if (UserId == null)
+			{
+				return await FilterIssuesByStreamIdAsync(Ids, StreamId, MinChange, MaxChange, Resolved, Index, Count);
+			}
+			else
 			{
-				ProjectionDefinition<IssueSpan> Projection = Builders<IssueSpan>.Projection.Include(x => x.IssueId);
+				FilterDefinition<IssueSuspect> Filter = Builders<IssueSuspect>.Filter.Eq(x => x.AuthorId, UserId);
+				if (Ids != null)
+				{
+					Filter &= Builders<IssueSuspect>.Filter.In(x => x.IssueId, Ids);
+				}
+				if (Resolved != null)
+				{
+					if (Resolved.Value)
+					{
+						Filter &= Builders<IssueSuspect>.Filter.Ne(x => x.ResolvedAt, null);
+					}
+					else
+					{
+						Filter &= Builders<IssueSuspect>.Filter.Eq(x => x.ResolvedAt, null);
+					}
+				}
 
-				FilterDefinition<IssueSpan> SpanFilter = Builders<IssueSpan>.Filter.Eq(x => x.StreamId, StreamId.Value);
-				if (Ids != null && Ids.Any())
+				using (IAsyncCursor<ProjectedIssueId> Cursor = await IssueSuspects.Aggregate().Match(Filter).Group(x => x.IssueId, x => new ProjectedIssueId { _id = x.Key }).SortByDescending(x => x._id).ToCursorAsync())
+				{
+					return await PaginatedJoinAsync(Cursor, (NextIds, NextIndex, NextCount) => FilterIssuesByStreamIdAsync(NextIds, StreamId, MinChange, MaxChange, null, NextIndex, NextCount), Index, Count);
+				}
+			}
+		}
+
+		async Task<List<Issue>> FilterIssuesByStreamIdAsync(IEnumerable<int>? Ids, StreamId? StreamId, int? MinChange, int? MaxChange, bool? Resolved, int Index, int? Count)
+		{
+			if (StreamId == null)
+			{
+				return await FilterIssuesByOtherFieldsAsync(Ids, MinChange, MaxChange, Resolved, Index, Count);
+			}
+			else
+			{
+				FilterDefinition<IssueSpan> Filter = Builders<IssueSpan>.Filter.Eq(x => x.StreamId, StreamId.Value);
+				if (Ids != null)
 				{
-					SpanFilter &= Builders<IssueSpan>.Filter.In(x => x.IssueId, Ids.Select<int, int?>(x => x));
+					Filter &= Builders<IssueSpan>.Filter.In(x => x.IssueId, Ids.Select<int, int?>(x => x));
 				}
 				else
 				{
-					SpanFilter &= Builders<IssueSpan>.Filter.Exists(x => x.IssueId);
+					Filter &= Builders<IssueSpan>.Filter.Exists(x => x.IssueId);
 				}
 
 				if (MinChange != null)
 				{
-					SpanFilter &= Builders<IssueSpan>.Filter.Not(Builders<IssueSpan>.Filter.Lt(x => x.MaxChange, MinChange.Value));
+					Filter &= Builders<IssueSpan>.Filter.Not(Builders<IssueSpan>.Filter.Lt(x => x.MaxChange, MinChange.Value));
 				}
 				if (MaxChange != null)
 				{
-					SpanFilter &= Builders<IssueSpan>.Filter.Not(Builders<IssueSpan>.Filter.Gt(x => x.MinChange, MaxChange.Value));
+					Filter &= Builders<IssueSpan>.Filter.Not(Builders<IssueSpan>.Filter.Gt(x => x.MinChange, MaxChange.Value));
 				}
+
 				if (Resolved != null)
 				{
 					if (Resolved.Value)
 					{
-						SpanFilter &= Builders<IssueSpan>.Filter.Ne(x => x.MaxChange, int.MaxValue);
+						Filter &= Builders<IssueSpan>.Filter.Ne(x => x.ResolvedAt, null);
 					}
 					else
 					{
-						SpanFilter &= Builders<IssueSpan>.Filter.Eq(x => x.MaxChange, int.MaxValue);
+						Filter &= Builders<IssueSpan>.Filter.Eq(x => x.ResolvedAt, null);
 					}
 				}
 
-				IAggregateFluent<ProjectedIssueId> IssueIds = IssueSpans.Aggregate().Match(SpanFilter).Group(x => x.IssueId, x => new ProjectedIssueId { IssueId = x.First().IssueId }).SortByDescending(x => x.IssueId);
-				if (Index != null)
+				using (IAsyncCursor<ProjectedIssueId> Cursor = await IssueSpans.Aggregate().Match(Filter).Group(x => x.IssueId, x => new ProjectedIssueId { _id = x.Key }).SortByDescending(x => x._id).ToCursorAsync())
 				{
-					IssueIds = IssueIds.Skip(Index.Value);
-				}
-				if (Count != null)
-				{
-					IssueIds = IssueIds.Limit(Count.Value);
+					return await PaginatedJoinAsync(Cursor, (NextIds, NextIndex, NextCount) => FilterIssuesByOtherFieldsAsync(NextIds, null, null, null, NextIndex, NextCount), Index, Count);
 				}
+			}
+		}
 
-				List<ProjectedIssueIdWithIssue> ProjectedResults = await IssueIds.Lookup(Issues, x => x.IssueId, x => x.Id, (ProjectedIssueIdWithIssue x) => x.Issues).ToListAsync();
-				return await UpgradeIssuesAsync(ProjectedResults.Where(x => x.Issues.Any()).Select(x => x.Issues.First()).ToList());
+		async Task<List<Issue>> FilterIssuesByOtherFieldsAsync(IEnumerable<int>? Ids, int? MinChange, int? MaxChange, bool? Resolved, int Index, int? Count)
+		{
+			FilterDefinition<Issue> Filter = FilterDefinition<Issue>.Empty;
+			if (Ids != null)
+			{
+				Filter &= Builders<Issue>.Filter.In(x => x.Id, Ids);
 			}
-			else
+			if (Resolved != null)
 			{
-				FilterDefinition<Issue> Filter = FilterDefinition<Issue>.Empty;
-				if (Ids != null)
+				if (Resolved.Value)
 				{
-					Filter &= Builders<Issue>.Filter.In(x => x.Id, Ids);
+					Filter &= Builders<Issue>.Filter.Ne(x => x.ResolvedAt, null);
 				}
-				if (Resolved != null)
+				else
 				{
-					Filter &= Builders<Issue>.Filter.Eq(x => x.Resolved, Resolved.Value);
+					Filter &= Builders<Issue>.Filter.Eq(x => x.ResolvedAt, null);
 				}
+			}
+			if (MinChange != null)
+			{
+				Filter &= Builders<Issue>.Filter.Not(Builders<Issue>.Filter.Lt(x => x.MaxSuspectChange, MinChange.Value));
+			}
+			if (MaxChange != null)
+			{
+				Filter &= Builders<Issue>.Filter.Not(Builders<Issue>.Filter.Gt(x => x.MinSuspectChange, MaxChange.Value));
+			}
+			return await Issues.Find(Filter).SortByDescending(x => x.Id).Range(Index, Count).ToListAsync();
+		}
 
-				List<Issue> Results = await Issues.Find(Filter).Range(Index, Count).ToListAsync();
-				return await UpgradeIssuesAsync(Results);
+		/// <summary>
+		/// Performs a client-side join of a filtered set of issues against another query
+		/// </summary>
+		/// <param name="Cursor"></param>
+		/// <param name="NextStageFunc"></param>
+		/// <param name="Index"></param>
+		/// <param name="Count"></param>
+		/// <returns></returns>
+		static async Task<List<Issue>> PaginatedJoinAsync(IAsyncCursor<ProjectedIssueId> Cursor, Func<IEnumerable<int>, int, int?, Task<List<Issue>>> NextStageFunc, int Index, int? Count)
+		{
+			if (Count == null)
+			{
+				List<ProjectedIssueId> IssueIds = await Cursor.ToListAsync();
+				return await NextStageFunc(IssueIds.Where(x => x._id != null).Select(x => x._id.Value), Index, null);
+			}
+			else
+			{
+				List<Issue> Results = new List<Issue>();
+				while (await Cursor.MoveNextAsync() && Results.Count < Count.Value)
+				{
+					List<Issue> NextResults = await NextStageFunc(Cursor.Current.Where(x => x._id != null).Select(x => x._id.Value), 0, Count.Value - Results.Count);
+					int RemoveCount = Math.Min(Index, NextResults.Count);
+					NextResults.RemoveRange(0, RemoveCount);
+					Index -= RemoveCount;
+					Results.AddRange(NextResults);
+				}
+				return Results;
 			}
 		}
 
@@ -683,21 +803,23 @@ namespace HordeServer.Collections.Impl
 
 		async Task<IIssue> UpgradeIssueAsync(Issue Issue)
 		{
+			bool bModified = false;
 			if (Issue.Owner != null && Issue.OwnerId == null)
 			{
 				Issue.OwnerId = await FindOrAddUserIdAsync(Issue.Owner);
+				bModified = true;
 			}
 			if (Issue.NominatedBy != null && Issue.NominatedById == null)
 			{
 				Issue.NominatedById = await FindOrAddUserIdAsync(Issue.NominatedBy);
+				bModified = true;
 			}
-			foreach (IssueSuspect Suspect in Issue.Suspects)
+
+			if (bModified)
 			{
-				if (Suspect.AuthorId == ObjectId.Empty)
-				{
-					Suspect.AuthorId = await FindOrAddUserIdAsync(Suspect.Author);
-				}
+				await Issues.ReplaceOneAsync(x => x.Id == Issue.Id && x.UpdateIndex == Issue.UpdateIndex, Issue);
 			}
+
 			return Issue;
 		}
 
@@ -713,13 +835,21 @@ namespace HordeServer.Collections.Impl
 
 		async Task<IIssueSpan> UpgradeIssueSpanAsync(IssueSpan IssueSpan)
 		{
+			bool bModified = false;
 			foreach (IssueSpanSuspect Suspect in IssueSpan.Suspects)
 			{
 				if (Suspect.AuthorId == ObjectId.Empty)
 				{
 					Suspect.AuthorId = await FindOrAddUserIdAsync(Suspect.Author);
+					bModified = true;
 				}
 			}
+
+			if (bModified)
+			{
+				await IssueSpans.ReplaceOneAsync(x => x.Id == IssueSpan.Id && x.UpdateIndex == IssueSpan.UpdateIndex, IssueSpan);
+			}
+
 			return IssueSpan;
 		}
 
@@ -745,7 +875,7 @@ namespace HordeServer.Collections.Impl
 		}
 
 		/// <inheritdoc/>
-		public async Task<IIssue?> UpdateIssueAsync(IIssue Issue, IssueSeverity? NewSeverity = null, string? NewSummary = null, string? NewUserSummary = null, string? NewOwner = null, string? NewNominatedBy = null, bool? NewAcknowledged = null, string? NewDeclinedBy = null, int? NewFixChange = null, bool? NewResolved = null, bool? NewNotifySuspects = null, List<NewIssueSuspectData>? NewSuspects = null)
+		public async Task<IIssue?> UpdateIssueAsync(IIssue Issue, IssueSeverity? NewSeverity = null, string? NewSummary = null, string? NewUserSummary = null, ObjectId? NewOwnerId = null, ObjectId? NewNominatedById = null, bool? NewAcknowledged = null, ObjectId? NewDeclinedById = null, int? NewFixChange = null, Dictionary<StreamId, bool>? NewFixStreamIds = null, ObjectId? NewResolvedById = null, DateTime? NewLastSeenAt = null, bool? NewNotifySuspects = null)
 		{
 			Issue IssueDocument = (Issue)Issue;
 
@@ -771,9 +901,9 @@ namespace HordeServer.Collections.Impl
 					Transaction.Set(x => x.UserSummary, NewUserSummary);
 				}
 			}
-			if (NewOwner != null)
+			if (NewOwnerId != null)
 			{
-				if (NewOwner.Length == 0)
+				if (NewOwnerId.Value == ObjectId.Empty)
 				{
 					Transaction.Unset(x => x.Owner!);
 					Transaction.Unset(x => x.OwnerId!);
@@ -783,18 +913,21 @@ namespace HordeServer.Collections.Impl
 				}
 				else
 				{
-					Transaction.Set(x => x.Owner!, NewOwner);
-					Transaction.Set(x => x.OwnerId!, await FindOrAddUserIdAsync(NewOwner));
+					IUser? NewOwner = await UserCollection.GetCachedUserAsync(NewOwnerId.Value);
+					Transaction.Set(x => x.Owner!, NewOwner?.Login);
+					Transaction.Set(x => x.OwnerId!, NewOwnerId.Value);
+
 					Transaction.Set(x => x.NominatedAt, DateTime.UtcNow);
-					if (String.IsNullOrEmpty(NewNominatedBy))
+					if (NewNominatedById == null)
 					{
 						Transaction.Unset(x => x.NominatedBy!);
 						Transaction.Unset(x => x.NominatedById!);
 					}
 					else
 					{
-						Transaction.Set(x => x.NominatedBy, NewNominatedBy);
-						Transaction.Set(x => x.NominatedById, await FindOrAddUserIdAsync(NewNominatedBy));
+						IUser? NewNominatedBy = await UserCollection.GetCachedUserAsync(NewNominatedById.Value);
+						Transaction.Set(x => x.NominatedBy, NewNominatedBy?.Login);
+						Transaction.Set(x => x.NominatedById, NewNominatedById.Value);
 					}
 					NewAcknowledged = NewAcknowledged ?? false;
 				}
@@ -816,12 +949,16 @@ namespace HordeServer.Collections.Impl
 					}
 				}
 			}
-			if (NewDeclinedBy != null)
+			if (NewDeclinedById != null)
 			{
-				int SuspectIdx = IssueDocument.Suspects.FindIndex(x => x.Author.Equals(NewDeclinedBy, StringComparison.OrdinalIgnoreCase));
-				if (SuspectIdx != -1 && IssueDocument.Suspects[SuspectIdx].DeclinedAt == null)
+				for (int SuspectIdx = 0; SuspectIdx < IssueDocument.Suspects.Count; SuspectIdx++)
 				{
-					Transaction.Set(x => x.Suspects[SuspectIdx].DeclinedAt, UtcNow);
+					IssueSuspect Suspect = IssueDocument.Suspects[SuspectIdx];
+					if (Suspect.AuthorId == NewDeclinedById.Value && Suspect.DeclinedAt == null)
+					{
+						int SuspectIdxCopy = SuspectIdx;
+						Transaction.Set(x => x.Suspects[SuspectIdxCopy].DeclinedAt, UtcNow);
+					}
 				}
 			}
 			if (NewFixChange != null)
@@ -835,15 +972,31 @@ namespace HordeServer.Collections.Impl
 					Transaction.Set(x => x.FixChange, NewFixChange);
 				}
 			}
-			if (NewResolved != null)
+			if (NewFixStreamIds != null)
 			{
-				Transaction.Set(x => x.Resolved, NewResolved.Value);
-
-				if (NewResolved.Value)
+				List<IssueStream> NewStreams = new List<IssueStream>();
+				foreach (IssueStream Stream in Issue.Streams)
+				{
+					IssueStream NewStream = new IssueStream(Stream);
+					if (NewFixStreamIds.TryGetValue(Stream.StreamId, out bool bContainsFix))
+					{
+						NewStream.ContainsFix = bContainsFix;
+					}
+					else
+					{
+						NewStream.ContainsFix = null;
+					}
+				}
+				Transaction.Set(x => x.Streams, NewStreams);
+			}
+			if (NewResolvedById != null)
+			{
+				if (NewResolvedById.Value != ObjectId.Empty)
 				{
-					if (IssueDocument.ResolvedAt == null)
+					if (IssueDocument.ResolvedAt == null || IssueDocument.ResolvedById != NewResolvedById)
 					{
 						Transaction.Set(x => x.ResolvedAt, UtcNow);
+						Transaction.Set(x => x.ResolvedById, NewResolvedById.Value);
 					}
 				}
 				else
@@ -852,28 +1005,168 @@ namespace HordeServer.Collections.Impl
 					{
 						Transaction.Unset(x => x.ResolvedAt!);
 					}
+					if (IssueDocument.ResolvedById != null)
+					{
+						Transaction.Unset(x => x.ResolvedById!);
+					}
 				}
 			}
+			if (NewLastSeenAt != null)
+			{
+				Transaction.Set(x => x.LastSeenAt, NewLastSeenAt.Value);
+			}
 			if (NewNotifySuspects != null)
 			{
 				Transaction.Set(x => x.NotifySuspects, NewNotifySuspects.Value);
 			}
-			if (NewSuspects != null)
+
+			IIssue? NewIssue = await ApplyTransactionAsync(Issue, Transaction);
+			if (NewIssue != null)
 			{
-				Transaction.Set(x => x.MinSuspectChange, (NewSuspects.Count > 0) ? NewSuspects.Min(x => x.Change) : 0);
-				Transaction.Set(x => x.MaxSuspectChange, (NewSuspects.Count > 0) ? NewSuspects.Max(x => x.Change) : 0);
+				if (NewResolvedById != null)
+				{
+					await UpdateIssueDerivedDataAsync(Issue);
+				}
+			}
+			return NewIssue;
+		}
+
+		public async Task<IIssue?> UpdateIssueDerivedDataAsync(IIssue Issue)
+		{
+			for (; ; )
+			{
+				Issue IssueDocument = (Issue)Issue;
+
+				// Find all the spans that are attached to the issue
+				List<IssueSpan> Spans = await IssueSpans.Find(x => x.IssueId == Issue.Id).ToListAsync();
+
+				// Find the list of suspects for this issue. If the issue has an assigned owner, the other suspects can be eliminated.
+				List<IssueSuspect> NewSuspects = new List<IssueSuspect>();
+				if (Spans.Count > 0)
+				{
+					HashSet<int> SuspectChanges = new HashSet<int>(Spans[0].Suspects.Where(x => Issue.OwnerId == null || x.AuthorId == Issue.OwnerId).Select(x => x.OriginatingChange ?? x.Change));
+					for (int SpanIdx = 1; SpanIdx < Spans.Count; SpanIdx++)
+					{
+						SuspectChanges.IntersectWith(Spans[SpanIdx].Suspects.Select(x => x.OriginatingChange ?? x.Change));
+					}
+					NewSuspects = Spans[0].Suspects.Where(x => SuspectChanges.Contains(x.OriginatingChange ?? x.Change)).Select(x => new IssueSuspect(Issue.Id, x)).ToList();
+				}
+				if (Issue.OwnerId != null)
+				{
+					NewSuspects.RemoveAll(x => x.AuthorId != Issue.OwnerId);
+					if (NewSuspects.Count == 0)
+					{
+						NewSuspects.Add(new IssueSuspect(Issue.Id, Issue.OwnerId.Value, 0, null, null));
+					}
+				}
+
+				// Find the current list of suspects
+				List<IssueSuspect> CurSuspects = await IssueSuspects.Find(x => x.IssueId == Issue.Id).ToListAsync();
+
+				HashSet<(ObjectId, int)> CurSuspectKeys = new HashSet<(ObjectId, int)>(CurSuspects.Select(x => (x.AuthorId, x.Change)));
+				List<IssueSuspect> CreateSuspects = NewSuspects.Where(x => !CurSuspectKeys.Contains((x.AuthorId, x.Change))).ToList();
 
-				List<IssueSuspect> Suspects = new List<IssueSuspect>();
-				foreach (NewIssueSuspectData NewSuspect in NewSuspects)
+				HashSet<(ObjectId, int)> NewSuspectKeys = new HashSet<(ObjectId, int)>(NewSuspects.Select(x => (x.AuthorId, x.Change)));
+				List<IssueSuspect> DeleteSuspects = CurSuspects.Where(x => !NewSuspectKeys.Contains((x.AuthorId, x.Change))).ToList();
+
+				// Apply the suspect changes
+				if (CreateSuspects.Count > 0)
+				{
+					await IssueSuspects.InsertManyIgnoreDuplicatesAsync(CreateSuspects);
+					CurSuspects.AddRange(CreateSuspects);
+				}
+				if (DeleteSuspects.Count > 0)
 				{
-					ObjectId AuthorId = await FindOrAddUserIdAsync(NewSuspect.Author);
-					Suspects.Add(new IssueSuspect(NewSuspect.Author, AuthorId, NewSuspect.Change, NewSuspect.DeclinedAt));
+					await IssueSuspects.DeleteManyAsync(Builders<IssueSuspect>.Filter.In(x => x.Id, DeleteSuspects.Select(y => y.Id)));
+					CurSuspects.RemoveAll(x => !NewSuspectKeys.Contains((x.AuthorId, x.Change)));
 				}
 
-				Transaction.Set(x => x.Suspects, Suspects);
-			}
+				// Calculate the last time that this issue was seen
+				DateTime NewLastSeenAt = Issue.CreatedAt;
+				foreach (IIssueSpan Span in Spans)
+				{
+					if (Span.LastFailure.StepTime > NewLastSeenAt)
+					{
+						NewLastSeenAt = Span.LastFailure.StepTime;
+					}
+				}
+
+				// Calculate the time at which the issue is verified fixed
+				DateTime? VerifiedAt = null;
+				foreach (IIssueSpan Span in Spans)
+				{
+					if(Span.NextSuccess == null)
+					{
+						break;
+					}
+					else if(VerifiedAt == null || Span.NextSuccess.StepTime < VerifiedAt.Value)
+					{
+						VerifiedAt = Span.NextSuccess.StepTime;
+					}
+				}
+				if (Spans.Count == 0)
+				{
+					VerifiedAt = Issue.VerifiedAt ?? DateTime.UtcNow;
+				}
 
-			return await ApplyTransactionAsync(Issue, Transaction);
+				// Update the spans with the correct resolved time
+				DateTime? ResolvedAt = Issue.ResolvedAt ?? VerifiedAt;
+
+				List<ObjectId> UpdateSpanIds = Spans.Where(x => x.ResolvedAt != ResolvedAt).Select(x => x.Id).ToList();
+				if (UpdateSpanIds.Count > 0)
+				{
+					await IssueSpans.UpdateManyAsync(Builders<IssueSpan>.Filter.In(x => x.Id, UpdateSpanIds), Builders<IssueSpan>.Update.Set(x => x.ResolvedAt, ResolvedAt));
+				}
+
+				List<ObjectId> UpdateSuspectIds = CurSuspects.Where(x => x.ResolvedAt != ResolvedAt).Select(x => x.Id).ToList();
+				if (UpdateSuspectIds.Count > 0)
+				{
+					await IssueSuspects.UpdateManyAsync(Builders<IssueSuspect>.Filter.In(x => x.Id, UpdateSuspectIds), Builders<IssueSuspect>.Update.Set(x => x.ResolvedAt, ResolvedAt));
+				}
+
+				// Build a list of streams
+				List<IssueStream> Streams = new List<IssueStream>();
+				foreach (IGrouping<StreamId, IIssueSpan> Group in Spans.GroupBy(x => x.StreamId))
+				{
+					IssueStream? Stream = IssueDocument.Streams.FirstOrDefault(x => x.StreamId == Group.Key);
+					if (Stream == null)
+					{
+						Stream = new IssueStream { StreamId = Group.Key };
+					}
+					else
+					{
+						Stream = new IssueStream(Stream);
+					}
+					Streams.Add(Stream);
+				}
+
+				// Update the issue document itself
+				TransactionBuilder<Issue> Transaction = new TransactionBuilder<Issue>();
+				Transaction.Set(x => x.Severity, Spans.Any(x => x.Severity == IssueSeverity.Error) ? IssueSeverity.Error : IssueSeverity.Warning);
+				Transaction.Set(x => x.NotifySuspects, Spans.Any(x => x.NotifySuspects));
+				Transaction.Set(x => x.LastSeenAt, NewLastSeenAt);
+				Transaction.Set(x => x.ResolvedAt, ResolvedAt);
+				Transaction.Set(x => x.VerifiedAt, VerifiedAt);
+				Transaction.Set(x => x.Suspects, CurSuspects);
+				Transaction.Set(x => x.Streams, Streams);
+				Transaction.Set(x => x.MinSuspectChange, (CurSuspects.Count > 0) ? CurSuspects.Min(x => x.Change) : 0);
+				Transaction.Set(x => x.MaxSuspectChange, (CurSuspects.Count > 0) ? CurSuspects.Max(x => x.Change) : 0);
+
+				// Apply the transaction
+				IIssue? NewIssue = await ApplyTransactionAsync(Issue, Transaction);
+				if(NewIssue != null)
+				{
+					return NewIssue;
+				}
+
+				// Update the issue
+				NewIssue = await GetIssueAsync(Issue.Id);
+				if(NewIssue == null)
+				{
+					return null;
+				}
+				Issue = NewIssue;
+			}
 		}
 
 		#endregion
@@ -881,10 +1174,10 @@ namespace HordeServer.Collections.Impl
 		#region Spans
 
 		/// <inheritdoc/>
-		public async Task<IIssueSpan?> AddSpanAsync(IIssueSequenceToken Token, IStream Stream, TemplateRefId TemplateRefId, string NodeName, IssueSeverity Severity, NewIssueFingerprint Fingerprint, NewIssueStepData? LastSuccess, NewIssueStepData Failure, NewIssueStepData? NextSuccess, List<NewIssueSpanSuspectData> Suspects)
+		public async Task<IIssueSpan?> AddSpanAsync(IIssueSequenceToken Token, IStream Stream, TemplateRefId TemplateRefId, string NodeName, NewIssueFingerprint Fingerprint, NewIssueStepData? LastSuccess, NewIssueStepData Failure, NewIssueStepData? NextSuccess, List<NewIssueSpanSuspectData> Suspects)
 		{
 			List<IssueSpanSuspect> NewSuspects = await CreateSpanSuspectsAsync(Suspects);
-			IssueSpan NewSpan = new IssueSpan(Stream.Id, Stream.Name, TemplateRefId, NodeName, Severity, Fingerprint, LastSuccess, Failure, NextSuccess, NewSuspects);
+			IssueSpan NewSpan = new IssueSpan(Stream.Id, Stream.Name, TemplateRefId, NodeName, Fingerprint, LastSuccess, Failure, NextSuccess, NewSuspects);
 
 			IssueLedger Ledger = ((IssueSequenceToken)Token).Ledger;
 			Ledger.InsertIssueSpan = NewSpan;
@@ -924,7 +1217,7 @@ namespace HordeServer.Collections.Impl
 				{
 					Transaction.Set(x => x.FirstFailure, new IssueStep(Span.Id, NewFailure));
 				}
-				if (NewFailure.Change > Span.LastFailure.Change)
+				if (NewFailure.Change >= Span.LastFailure.Change)
 				{
 					Transaction.Set(x => x.LastFailure, new IssueStep(Span.Id, NewFailure));
 				}
@@ -990,6 +1283,10 @@ namespace HordeServer.Collections.Impl
 			{
 				IssueUpdates.Add(Builders<Issue>.Update.Set(x => x.Summary, NewSummary));
 			}
+			if (Span.LastFailure.StepTime > Issue.LastSeenAt)
+			{
+				IssueUpdates.Add(Builders<Issue>.Update.Set(x => x.LastSeenAt, Span.LastFailure.StepTime));
+			}
 
 			IssueSequenceToken TokenValue = ((IssueSequenceToken)Token);
 			if (IssueUpdates.Count > 0)
diff --git a/Engine/Source/Programs/Horde/HordeServer/Collections/Impl/JobCollection.cs b/Engine/Source/Programs/Horde/HordeServer/Collections/Impl/JobCollection.cs
--- a/Engine/Source/Programs/Horde/HordeServer/Collections/Impl/JobCollection.cs
+++ b/Engine/Source/Programs/Horde/HordeServer/Collections/Impl/JobCollection.cs
@@ -188,7 +188,7 @@ namespace HordeServer.Collections.Impl
 			public StreamId StreamId { get; set; }
 			public TemplateRefId TemplateId { get; set; }
 			public ContentHash? TemplateHash { get; set; }
-			public ContentHash GraphHash { get; set; } = ContentHash.Empty;
+			public ContentHash GraphHash { get; set; }
 
 			[BsonIgnoreIfNull]
 			public ObjectId? StartedByUserId { get; set; }
@@ -265,6 +265,7 @@ namespace HordeServer.Collections.Impl
 			private JobDocument()
 			{
 				Name = null!;
+				GraphHash = null!;
 			}
 
 			public JobDocument(ObjectId Id, StreamId StreamId, TemplateRefId TemplateId, ContentHash TemplateHash, ContentHash GraphHash, string Name, int Change, int CodeChange, int PreflightChange, int ClonedPreflightChange, ObjectId? StartedByUserId, string? StartedByUserName, Priority? Priority, bool? AutoSubmit, DateTime CreateTimeUtc, List<ChainedJobDocument> ChainedJobs, bool ShowUgsBadges, bool ShowUgsAlerts, string? NotificationChannel, string? NotificationChannelFilter, string? HelixSwarmCallbackUrl, List<string>? Arguments)
diff --git a/Engine/Source/Programs/Horde/HordeServer/Collections/Impl/ServiceAccountCollection.cs b/Engine/Source/Programs/Horde/HordeServer/Collections/Impl/ServiceAccountCollection.cs
--- a/Engine/Source/Programs/Horde/HordeServer/Collections/Impl/ServiceAccountCollection.cs
+++ b/Engine/Source/Programs/Horde/HordeServer/Collections/Impl/ServiceAccountCollection.cs
@@ -22,6 +22,8 @@ namespace HordeServer.Collections.Impl
 		/// </summary>
 		class ServiceAccountDocument : IServiceAccount
 		{
+			public const string ClaimSeparator = "###";
+			
 			[BsonRequired, BsonId]
 			public ObjectId Id { get; set; }
 
@@ -29,7 +31,7 @@ namespace HordeServer.Collections.Impl
 			public string SecretToken { get; set; } = "empty";
 
 			[BsonRequired]
-			public IReadOnlyDictionary<string, string> Claims { get; set; } = new Dictionary<string, string>();
+			public List<string> Claims { get; set; } = new List<string>();
 			
 			public bool Enabled { get; set; }
 			public string Description { get; set; } = "empty";
@@ -39,7 +41,7 @@ namespace HordeServer.Collections.Impl
 			{
 			}
 
-			public ServiceAccountDocument(ObjectId Id, string SecretToken, IReadOnlyDictionary<string, string> Claims, bool Enabled, string Description)
+			public ServiceAccountDocument(ObjectId Id, string SecretToken, List<string> Claims, bool Enabled, string Description)
 			{
 				this.Id = Id;
 				this.SecretToken = SecretToken;
@@ -47,6 +49,22 @@ namespace HordeServer.Collections.Impl
 				this.Enabled = Enabled;
 				this.Description = Description;
 			}
+			
+			/// <inheritdoc/>
+			public void AddClaim(string Type, string Value)
+			{
+				Claims.Add(Type + ClaimSeparator + Value);
+			}
+
+			/// <inheritdoc/>
+			public IReadOnlyList<(string Type, string Value)> GetClaims()
+			{
+				return Claims.Select(x =>
+				{
+					string[] Split = x.Split(ClaimSeparator);
+					return (Split[0], Split[1]);
+				}).ToList();
+			}
 
 			protected bool Equals(ServiceAccountDocument other)
 			{
@@ -89,7 +107,7 @@ namespace HordeServer.Collections.Impl
 		}
 
 		/// <inheritdoc/>
-		public async Task<IServiceAccount> AddAsync(string SecretToken, Dictionary<string, string> Claims, string Description)
+		public async Task<IServiceAccount> AddAsync(string SecretToken, List<string> Claims, string Description)
 		{
 			ServiceAccountDocument NewSession = new ServiceAccountDocument(ObjectId.GenerateNewId(), SecretToken, Claims, true, Description);
 			await ServiceAccounts.InsertOneAsync(NewSession);
@@ -109,7 +127,7 @@ namespace HordeServer.Collections.Impl
 		}
 
 		/// <inheritdoc/>
-		public Task UpdateAsync(ObjectId Id, string? SecretToken, Dictionary<string, string>? Claims, bool? Enabled, string? Description)
+		public Task UpdateAsync(ObjectId Id, string? SecretToken, List<string>? Claims, bool? Enabled, string? Description)
 		{
 			UpdateDefinitionBuilder<ServiceAccountDocument> Update = Builders<ServiceAccountDocument>.Update;
 			List<UpdateDefinition<ServiceAccountDocument>> Updates = new List<UpdateDefinition<ServiceAccountDocument>>();
diff --git a/Engine/Source/Programs/Horde/HordeServer/Collections/Impl/TemplateCollection.cs b/Engine/Source/Programs/Horde/HordeServer/Collections/Impl/TemplateCollection.cs
--- a/Engine/Source/Programs/Horde/HordeServer/Collections/Impl/TemplateCollection.cs
+++ b/Engine/Source/Programs/Horde/HordeServer/Collections/Impl/TemplateCollection.cs
@@ -50,7 +50,7 @@ namespace HordeServer.Collections.Impl
 		class TemplateDocument : ITemplate
 		{
 			[BsonRequired, BsonId]
-			public ContentHash Id { get; private set; }
+			public ContentHash Id { get; private set; } = ContentHash.Empty;
 
 			[BsonRequired]
 			public string Name { get; private set; }
@@ -73,7 +73,6 @@ namespace HordeServer.Collections.Impl
 			[BsonConstructor]
 			private TemplateDocument()
 			{
-				Id = ContentHash.Empty;
 				Name = null!;
 			}
 
diff --git a/Engine/Source/Programs/Horde/HordeServer/Collections/Impl/UserCollection.cs b/Engine/Source/Programs/Horde/HordeServer/Collections/Impl/UserCollection.cs
--- a/Engine/Source/Programs/Horde/HordeServer/Collections/Impl/UserCollection.cs
+++ b/Engine/Source/Programs/Horde/HordeServer/Collections/Impl/UserCollection.cs
@@ -102,6 +102,12 @@ namespace HordeServer.Collections.Impl
 			return await Users.Find(x => x.Id == Id).FirstOrDefaultAsync();
 		}
 
+		/// <inheritdoc/>
+		public Task<IUser?> GetCachedUserAsync(ObjectId Id)
+		{
+			return GetUserAsync(Id);
+		}
+
 		/// <inheritdoc/>
 		public async Task<List<IUser>> FindUsersAsync(IEnumerable<ObjectId>? Ids, int? Index, int? Count)
 		{
diff --git a/Engine/Source/Programs/Horde/HordeServer/Collections/Impl/UserCollectionV2.cs b/Engine/Source/Programs/Horde/HordeServer/Collections/Impl/UserCollectionV2.cs
--- a/Engine/Source/Programs/Horde/HordeServer/Collections/Impl/UserCollectionV2.cs
+++ b/Engine/Source/Programs/Horde/HordeServer/Collections/Impl/UserCollectionV2.cs
@@ -3,6 +3,7 @@
 using HordeServer.Models;
 using HordeServer.Services;
 using HordeServer.Utilities;
+using Microsoft.Extensions.Caching.Memory;
 using Microsoft.Extensions.Logging;
 using MongoDB.Bson;
 using MongoDB.Bson.Serialization;
@@ -19,7 +20,7 @@ namespace HordeServer.Collections.Impl
 	/// <summary>
 	/// Manages user documents
 	/// </summary>
-	class UserCollectionV2 : IUserCollection
+	class UserCollectionV2 : IUserCollection, IDisposable
 	{
 		class UserDocument : IUser
 		{
@@ -146,6 +147,8 @@ namespace HordeServer.Collections.Impl
 
 		ILogger<UserCollectionV2> Logger;
 
+		MemoryCache UserCache;
+
 		/// <summary>
 		/// Static constructor
 		/// </summary>
@@ -176,12 +179,41 @@ namespace HordeServer.Collections.Impl
 				Users.Indexes.CreateOne(new CreateIndexModel<UserDocument>(Builders<UserDocument>.IndexKeys.Ascending(x => x.LoginUpper), new CreateIndexOptions { Unique = true }));
 				Users.Indexes.CreateOne(new CreateIndexModel<UserDocument>(Builders<UserDocument>.IndexKeys.Ascending(x => x.EmailUpper)));
 			}
+
+			MemoryCacheOptions Options = new MemoryCacheOptions();
+			UserCache = new MemoryCache(Options);
+		}
+
+		/// <inheritdoc/>
+		public void Dispose()
+		{
+			UserCache.Dispose();
 		}
 
 		/// <inheritdoc/>
 		public async Task<IUser?> GetUserAsync(ObjectId Id)
 		{
-			return await Users.Find(x => x.Id == Id).FirstOrDefaultAsync();
+			IUser? User = await Users.Find(x => x.Id == Id).FirstOrDefaultAsync();
+			using (ICacheEntry Entry = UserCache.CreateEntry(Id))
+			{
+				Entry.SetValue(User);
+				Entry.SetSlidingExpiration(TimeSpan.FromMinutes(30.0));
+			}
+			return User;
+		}
+
+		/// <inheritdoc/>
+		public Task<IUser?> GetCachedUserAsync(ObjectId Id)
+		{
+			IUser? User;
+			if (UserCache.TryGetValue(Id, out User))
+			{
+				return Task.FromResult(User);
+			}
+			else
+			{
+				return GetUserAsync(Id);
+			}
 		}
 
 		/// <inheritdoc/>
diff --git a/Engine/Source/Programs/Horde/HordeServer/Controllers/DebugController.cs b/Engine/Source/Programs/Horde/HordeServer/Controllers/DebugController.cs
--- a/Engine/Source/Programs/Horde/HordeServer/Controllers/DebugController.cs
+++ b/Engine/Source/Programs/Horde/HordeServer/Controllers/DebugController.cs
@@ -38,6 +38,7 @@ using HordeServer.Services.Impl;
 using Microsoft.Net.Http.Headers;
 using HordeServer.Storage;
 using Microsoft.Extensions.Hosting;
+using Microsoft.Extensions.Options;
 using StatsdClient;
 
 namespace HordeServer.Controllers
@@ -53,7 +54,7 @@ namespace HordeServer.Controllers
 		/// The connection tracker service singleton
 		/// </summary>
 		RequestTrackerService RequestTrackerService;
-		
+
 		IHostApplicationLifetime ApplicationLifetime;
 
 		IDogStatsd DogStatsd;
@@ -91,7 +92,7 @@ namespace HordeServer.Controllers
 			Content.Append("</pre></body></html>");
 			return new ContentResult { ContentType = "text/html", StatusCode = (int)HttpStatusCode.OK, Content = Content.ToString() };
 		}
-		
+
 		/// <summary>
 		/// Waits specified number of milliseconds and then returns a response
 		/// Used for testing timeouts proxy settings.
@@ -105,7 +106,7 @@ namespace HordeServer.Controllers
 			string Content = $"Waited {WaitTimeMs} ms. " + new Random().Next(0, 10000000);
 			return new ContentResult { ContentType = "text/plain", StatusCode = (int)HttpStatusCode.OK, Content = Content };
 		}
-		
+
 		/// <summary>
 		/// Waits specified number of milliseconds and then throws an exception
 		/// Used for testing graceful shutdown and interruption of outstanding requests.
@@ -118,7 +119,7 @@ namespace HordeServer.Controllers
 			await Task.Delay(WaitTimeMs);
 			throw new Exception("Test exception triggered by debug controller!");
 		}
-		
+
 		/// <summary>
 		/// Trigger an increment of a DogStatsd metric
 		/// </summary>
@@ -130,7 +131,7 @@ namespace HordeServer.Controllers
 			DogStatsd.Increment("hordeMetricTest", Value);
 			return Ok("Incremented metric 'hordeMetricTest' Type: " + DogStatsd.GetType());
 		}
-		
+
 		/// <summary>
 		/// Display metrics related to the .NET runtime
 		/// </summary>
@@ -145,14 +146,14 @@ namespace HordeServer.Controllers
 
 			int BusyIoThreads = MaxIoThreads - FreeIoThreads;
 			int BusyWorkerThreads = MaxWorkerThreads - FreeWorkerThreads;
-			
+
 			StringBuilder Content = new StringBuilder();
 			Content.AppendLine("Threads:");
 			Content.AppendLine("-------------------------------------------------------------");
 			Content.AppendLine("Worker busy={0,-5} free={1,-5} min={2,-5} max={3,-5}", BusyWorkerThreads, FreeWorkerThreads, MinWorkerThreads, MaxWorkerThreads);
 			Content.AppendLine("  IOCP busy={0,-5} free={1,-5} min={2,-5} max={3,-5}", BusyIoThreads, FreeIoThreads, MinIoThreads, MaxWorkerThreads);
 
-			
+
 			NumberFormatInfo Nfi = (NumberFormatInfo)CultureInfo.InvariantCulture.NumberFormat.Clone();
 			Nfi.NumberGroupSeparator = " ";
 
@@ -160,7 +161,7 @@ namespace HordeServer.Controllers
 			{
 				return (Number / 1024 / 1024).ToString("#,0", Nfi) + " MB";
 			}
-			
+
 			GCMemoryInfo GcMemoryInfo = GC.GetGCMemoryInfo();
 			Content.AppendLine("");
 			Content.AppendLine("");
@@ -175,7 +176,7 @@ namespace HordeServer.Controllers
 			Content.AppendLine("               Memory Load: " + FormatBytes(GcMemoryInfo.MemoryLoadBytes));
 			Content.AppendLine("    Total available memory: " + FormatBytes(GcMemoryInfo.TotalAvailableMemoryBytes));
 			Content.AppendLine("High memory load threshold: " + FormatBytes(GcMemoryInfo.HighMemoryLoadThresholdBytes));
-			
+
 			return Ok(Content.ToString());
 		}
 
@@ -213,9 +214,9 @@ namespace HordeServer.Controllers
 			Content.AppendLine("<th>Age</th>");
 			Content.AppendLine("</tr>");
 
-			List<KeyValuePair<string,TrackedRequest>> Requests = RequestTrackerService.GetRequestsInProgress().ToList();
+			List<KeyValuePair<string, TrackedRequest>> Requests = RequestTrackerService.GetRequestsInProgress().ToList();
 			Requests.Sort((A, B) => A.Value.StartedAt.CompareTo(B.Value.StartedAt));
-			
+
 			foreach (KeyValuePair<string, TrackedRequest> Entry in Requests)
 			{
 				Content.Append("<tr>");
@@ -324,17 +325,22 @@ namespace HordeServer.Controllers
 		/// The storage provider
 		/// </summary>
 		IStorageBackend StorageProvider;
-		
+
 		/// <summary>
 		/// Perforce client
 		/// </summary>
 		IPerforceService Perforce;
-		
+
 		/// <summary>
 		/// Fleet manager
 		/// </summary>
 		IFleetManager FleetManager;
 
+		/// <summary>
+		/// Settings
+		/// </summary>
+		private IOptionsMonitor<ServerSettings> Settings;
+
 		/// <summary>
 		/// Constructor
 		/// </summary>
@@ -354,7 +360,8 @@ namespace HordeServer.Controllers
 		/// <param name="StorageProvider">The storage provider</param>
 		/// <param name="Perforce">Perforce client</param>
 		/// <param name="FleetManager">The default fleet manager</param>
-		public SecureDebugController(AclService AclService, DatabaseService DatabaseService, JobTaskSource JobTaskSource, ISingletonDocument<Globals> GlobalsSingleton, IPoolCollection PoolCollection, IProjectCollection ProjectCollection, IAgentCollection AgentCollection, ISessionCollection SessionCollection, ILeaseCollection LeaseCollection, ITemplateCollection TemplateCollection, IStreamCollection StreamCollection, IGraphCollection GraphCollection, ILogFileCollection LogFileCollection, IStorageBackend StorageProvider, IPerforceService Perforce, IFleetManager FleetManager)
+		/// <param name="Settings">Settings</param>
+		public SecureDebugController(AclService AclService, DatabaseService DatabaseService, JobTaskSource JobTaskSource, ISingletonDocument<Globals> GlobalsSingleton, IPoolCollection PoolCollection, IProjectCollection ProjectCollection, IAgentCollection AgentCollection, ISessionCollection SessionCollection, ILeaseCollection LeaseCollection, ITemplateCollection TemplateCollection, IStreamCollection StreamCollection, IGraphCollection GraphCollection, ILogFileCollection LogFileCollection, IStorageBackend StorageProvider, IPerforceService Perforce, IFleetManager FleetManager, IOptionsMonitor<ServerSettings> Settings)
 		{
 			this.AclService = AclService;
 			this.DatabaseService = DatabaseService;
@@ -372,6 +379,7 @@ namespace HordeServer.Controllers
 			this.StorageProvider = StorageProvider;
 			this.Perforce = Perforce;
 			this.FleetManager = FleetManager;
+			this.Settings = Settings;
 		}
 
 		/// <summary>
@@ -603,6 +611,27 @@ namespace HordeServer.Controllers
 
 			return await JobTaskSource.GetStatus();
 		}
+		
+		/// <summary>
+		/// Returns the complete config Horde uses
+		/// </summary>
+		/// <returns>Information about the config</returns>
+		[HttpGet]
+		[Route("/api/v1/debug/config")]
+		public async Task<ActionResult> GetConfig()
+		{
+			if (!await AclService.AuthorizeAsync(AclAction.AdminRead, User))
+			{
+				return Forbid();
+			}
+
+			JsonSerializerOptions Options = new JsonSerializerOptions
+			{
+				WriteIndented = true
+			};
+			
+			return Ok(JsonSerializer.Serialize(Settings.CurrentValue, Options));
+		}
 
 		/// <summary>
 		/// Queries for all graphs
@@ -657,7 +686,7 @@ namespace HordeServer.Controllers
 			}
 
 			ReadOnlyMemory<byte>? Data = await StorageProvider.ReadAsync(Path);
-			if(Data == null)
+			if (Data == null)
 			{
 				return NotFound();
 			}
@@ -752,7 +781,7 @@ namespace HordeServer.Controllers
 			List<Dictionary<string, object>> Documents = await Collection.Find(Filter ?? "{}").Skip(Index).Limit(Count).ToListAsync();
 			return Documents;
 		}
-		
+
 		/// <summary>
 		/// Create P4 login ticket for the specified username
 		/// </summary>
@@ -773,7 +802,162 @@ namespace HordeServer.Controllers
 
 			return await Perforce.CreateTicket(Username);
 		}
-				
+
+
+		/// <summary>
+		/// Debugs perforce service commands
+		/// </summary>
+		/// <returns>Perforce ticket</returns>
+		[HttpGet]
+		[Route("/api/v1/debug/perforce/commands")]
+		public async Task<ActionResult<string>> Get()
+		{
+			if (!await AclService.AuthorizeAsync(AclAction.AdminWrite, User))
+			{
+				return Forbid();
+			}
+
+			const string DevBuild = "//UE4/Dev-Build";
+			// a test CL I have setup in Dev-Build
+			const int TestCL = 16242697;
+
+			string? PerforceUser = User.GetPerforceUser();
+
+			if (string.IsNullOrEmpty(PerforceUser))
+			{
+				throw new Exception("Perforce user required");
+			}
+
+			Dictionary<string, string> Results = new Dictionary<string, string>();
+
+			int NewChangeId = await Perforce.CreateNewChangeAsync("//UE4/Private-Build", "Counter.txt");
+
+			List<ChangeDetails> ChangeDetails = await Perforce.GetChangeDetailsAsync("//UE4/Private-Build", new int[]{NewChangeId}, PerforceUser);
+
+			Results.Add("CreateNewChangeAsync", $"{NewChangeId} : {ChangeDetails[0].Description}"); 
+
+			await Perforce.UpdateChangelistDescription(TestCL, $"Updated Description: New Change CL was {NewChangeId}");			
+
+			ChangeDetails = await Perforce.GetChangeDetailsAsync(DevBuild, new int[]{TestCL}, PerforceUser);
+
+			Results.Add("UpdateChangelistDescription", $"{TestCL} : {ChangeDetails[0].Description}"); 
+
+			int LatestChange = await Perforce.GetLatestChangeAsync(DevBuild, PerforceUser);
+			Results.Add("GetLatestChangeAsync", LatestChange.ToString(CultureInfo.InvariantCulture));
+
+			int CodeChange = await Perforce.GetCodeChangeAsync(DevBuild, LatestChange);
+			Results.Add("GetCodeChangeAsync", CodeChange.ToString(CultureInfo.InvariantCulture));
+			
+			PerforceUserInfo? UserInfo =  await Perforce.GetUserInfoAsync(PerforceUser);
+
+			if (UserInfo == null)
+			{
+				Results.Add("GetUserInfoAsync", "null");
+			}
+			else
+			{
+				Results.Add("GetUserInfoAsync", $"{UserInfo.Name} : {UserInfo.Email}");
+			}
+
+			// changes
+			List<ChangeSummary> Changes = await Perforce.GetChangesAsync(DevBuild, null, LatestChange, 10, PerforceUser);
+			int Count = 0;
+			foreach (ChangeSummary Summary in Changes)
+			{
+				Results.Add($"GetChangesAsync Result {Count++}", $"{Summary.Number} : {Summary.Author} - {Summary.Description}");
+			}
+
+			ChangeDetails = await Perforce.GetChangeDetailsAsync(DevBuild, Changes.Select(Change => Change.Number).ToArray(), PerforceUser);			
+			Count = 0;
+			foreach (ChangeDetails Details in ChangeDetails)
+			{
+				Results.Add($"GetChangeDetailsAsync Result {Count++}", $"{Details.Number} : {Details.Author} - {Details.Files.ToJson()}");
+			}
+			
+			String Ticket = "Failed (Expected if not running service account)";
+			try 
+			{
+				await Perforce.CreateTicket(PerforceUser);
+				Ticket = "Success";
+			}
+			catch 
+			{
+
+			}
+
+			Results.Add($"CreateTicket", $"{Ticket}");
+			List<FileSummary> FileSummaries = await Perforce.FindFilesAsync(ChangeDetails.SelectMany(Details => Details.Files).Select(File => $"{DevBuild}{File}"));
+			Count = 0;
+			foreach (FileSummary File in FileSummaries)
+			{
+				Results.Add($"FindFilesAsync Result {Count++}", $"{File.Change} : {File.DepotPath} : {File.Exists}");
+			}
+
+			// print async
+
+			byte[] Bytes = await Perforce.PrintAsync($"{DevBuild}/RunUAT.bat");
+			Results.Add($"PrintAsync: {DevBuild}/RunUAT.bat ",  $"{System.Text.Encoding.Default.GetString(Bytes)}");
+
+			Bytes = await Perforce.PrintAsync($"{DevBuild}/RunUAT.bat@13916380");
+			Results.Add($"PrintAsync: {DevBuild}/RunUAT.bat@13916380",  $"{System.Text.Encoding.Default.GetString(Bytes)}");
+
+
+			// DuplicateShelvedChangeAsync is disabled due to issue with p4 librarian
+			/*
+			// need to be running as service account for these
+
+			int? DuplicateCL = null;
+
+			int ShelvedChangeOnDevBuild = 16158333;
+
+			try 
+			{
+				DuplicateCL = await Perforce.DuplicateShelvedChangeAsync(ShelvedChangeOnDevBuild);
+			}
+			catch 
+			{
+				Results.Add($"DuplicateShelvedChangeAsync",  "Failed - expected unless using service account");
+			}
+
+			if (DuplicateCL.HasValue)
+			{
+				// Note change client much exist for the update
+				ChangeDetails = await Perforce.GetChangeDetailsAsync(DevBuild, new int[]{NewChangeId}, PerforceUser);
+
+				Results.Add("UpdateChangelistDescription - PreUpdate", $"{DuplicateCL} : {ChangeDetails[0].Description}"); 
+
+				try 
+				{
+					await Perforce.UpdateChangelistDescription(DuplicateCL.Value, "Updated Description from Horde");
+
+					ChangeDetails = await Perforce.GetChangeDetailsAsync(DevBuild, new int[]{NewChangeId}, PerforceUser);
+
+					Results.Add("UpdateChangelistDescription - PostUpdate", $"{DuplicateCL} : {ChangeDetails[0].Description}"); 
+				}
+				catch 
+				{
+					Results.Add("UpdateChangelistDescription - Faile", $"Client from change {DuplicateCL} must exist"); 
+				}
+
+				await Perforce.DeleteShelvedChangeAsync(DuplicateCL.Value);
+
+				ChangeDetails = await Perforce.GetChangeDetailsAsync(DevBuild,new int[] {DuplicateCL.Value} , PerforceUser);
+
+				Results.Add("GetChangeDetailsAsync - After Delete", $"ChangeDetails.Count: {ChangeDetails.Count}"); 
+
+			}
+			*/			
+			StringBuilder Content = new StringBuilder();
+			Content.AppendLine("<html><body><pre>");
+			foreach (KeyValuePair<string, string> Pair in Results)
+			{
+				Content.AppendLine(HttpUtility.HtmlEncode($"{Pair.Key}={Pair.Value}"));
+			}
+			Content.Append("</pre></body></html>");
+			return new ContentResult { ContentType = "text/html", StatusCode = (int)HttpStatusCode.OK, Content = Content.ToString() };
+
+		}
+
 		/// <summary>
 		/// Debugging fleet managers
 		/// </summary>
@@ -786,7 +970,7 @@ namespace HordeServer.Controllers
 			{
 				return Forbid();
 			}
-			
+
 			if (PoolId == null)
 			{
 				return BadRequest("Missing 'PoolId' query parameter");
@@ -797,12 +981,12 @@ namespace HordeServer.Controllers
 			{
 				return BadRequest($"Pool with ID '{PoolId}' not found");
 			}
-			
+
 			if (!(ChangeType == "expand" || ChangeType == "shrink"))
 			{
 				return BadRequest("Missing 'ChangeType' query parameter and/or must be set to 'expand' or 'shrink'");
 			}
-			
+
 			if (Count == null || Count.Value <= 0)
 			{
 				return BadRequest("Missing 'Count' query parameter and/or must be greater than 0");
@@ -810,7 +994,7 @@ namespace HordeServer.Controllers
 
 			List<IAgent> Agents = new List<IAgent>();
 			string Message = "No operation";
-				
+
 			if (ChangeType == "expand")
 			{
 				await FleetManager.ExpandPool(Pool, Agents, Count.Value);
@@ -823,7 +1007,7 @@ namespace HordeServer.Controllers
 				await FleetManager.ShrinkPool(Pool, Agents, Count.Value);
 				Message = $"Shrank pool {Pool.Name} by {Count}";
 			}
-			
+
 			return new ContentResult { ContentType = "text/plain", StatusCode = (int)HttpStatusCode.OK, Content = Message };
 		}
 	}
diff --git a/Engine/Source/Programs/Horde/HordeServer/Controllers/IssuesController.cs b/Engine/Source/Programs/Horde/HordeServer/Controllers/IssuesController.cs
--- a/Engine/Source/Programs/Horde/HordeServer/Controllers/IssuesController.cs
+++ b/Engine/Source/Programs/Horde/HordeServer/Controllers/IssuesController.cs
@@ -46,6 +46,11 @@ namespace HordeServer.Controllers
 		/// </summary>
 		private readonly StreamService StreamService;
 
+		/// <summary>
+		/// 
+		/// </summary>
+		private readonly IUserCollection UserCollection;
+
 		/// <summary>
 		/// Collection of events
 		/// </summary>
@@ -62,13 +67,15 @@ namespace HordeServer.Controllers
 		/// <param name="IssueService">The issue service</param>
 		/// <param name="JobService">The job service</param>
 		/// <param name="StreamService">The stream service</param>
+		/// <param name="UserCollection"></param>
 		/// <param name="LogEventCollection">The event collection</param>
 		/// <param name="LogFileService">The log file service</param>
-		public IssuesController(IIssueService IssueService, JobService JobService, StreamService StreamService, ILogEventCollection LogEventCollection, ILogFileService LogFileService)
+		public IssuesController(IIssueService IssueService, JobService JobService, StreamService StreamService, IUserCollection UserCollection, ILogEventCollection LogEventCollection, ILogFileService LogFileService)
 		{
 			this.IssueService = IssueService;
 			this.JobService = JobService;
 			this.StreamService = StreamService;
+			this.UserCollection = UserCollection;
 			this.LogEventCollection = LogEventCollection;
 			this.LogFileService = LogFileService;
 		}
@@ -96,6 +103,21 @@ namespace HordeServer.Controllers
 		[ProducesResponseType(typeof(List<GetIssueResponse>), 200)]
 		public async Task<ActionResult<object>> FindIssuesAsync([FromQuery(Name = "Id")] int[]? Ids = null, [FromQuery] string? StreamId = null, [FromQuery] int? Change = null, [FromQuery] int? MinChange = null, [FromQuery] int? MaxChange = null, [FromQuery] string? JobId = null, [FromQuery] string? BatchId = null, [FromQuery] string? StepId = null, [FromQuery(Name = "label")] int? LabelIdx = null, [FromQuery] string? User = null, [FromQuery] bool? Resolved = null, [FromQuery] int Index = 0, [FromQuery] int Count = 10, [FromQuery] PropertyFilter? Filter = null)
 		{
+			if(Ids != null && Ids.Length == 0)
+			{
+				Ids = null;
+			}
+
+			ObjectId? UserIdValue = null;
+			if (User != null)
+			{
+				UserIdValue = (await UserCollection.FindUserByLoginAsync(User))?.Id;
+				if (UserIdValue == null)
+				{
+					return new List<IIssue>();
+				}
+			}
+
 			List<IIssue> Issues;
 			if (JobId == null)
 			{
@@ -105,7 +127,7 @@ namespace HordeServer.Controllers
 					StreamIdValue = new StreamId(StreamId);
 				}
 
-				Issues = await IssueService.FindIssuesAsync(Ids, StreamIdValue, MinChange ?? Change, MaxChange ?? Change, Resolved, Index, Count);
+				Issues = await IssueService.FindIssuesAsync(Ids, UserIdValue, StreamIdValue, MinChange ?? Change, MaxChange ?? Change, Resolved, Index, Count);
 			}
 			else
 			{
@@ -122,7 +144,7 @@ namespace HordeServer.Controllers
 				}
 
 				IGraph Graph = await JobService.GetGraphAsync(Job);
-				Issues = await IssueService.FindIssuesForJobAsync(Ids, Job, Graph, StepId?.ToSubResourceId(), BatchId?.ToSubResourceId(), LabelIdx, Resolved, Index, Count);
+				Issues = await IssueService.FindIssuesForJobAsync(Ids, Job, Graph, StepId?.ToSubResourceId(), BatchId?.ToSubResourceId(), LabelIdx, UserIdValue, Resolved, Index, Count);
 			}
 
 			StreamPermissionsCache PermissionsCache = new StreamPermissionsCache();
@@ -130,15 +152,12 @@ namespace HordeServer.Controllers
 			List<object> Responses = new List<object>();
 			foreach (IIssue Issue in Issues)
 			{
-				if (User == null || IncludeIssueForUser(Issue, User))
+				IIssueDetails Details = await IssueService.GetIssueDetailsAsync(Issue);
+				if (await AuthorizeIssue(Details, PermissionsCache))
 				{
-					IIssueDetails Details = await IssueService.GetIssueDetailsAsync(Issue);
-					if (await AuthorizeIssue(Details, PermissionsCache))
-					{
-						bool bShowDesktopAlerts = IssueService.ShowDesktopAlertsForIssue(Issue, Details.Spans);
-						GetIssueResponse Response = await CreateIssueResponseAsync(Details, bShowDesktopAlerts);
-						Responses.Add(PropertyFilter.Apply(Response, Filter));
-					}
+					bool bShowDesktopAlerts = IssueService.ShowDesktopAlertsForIssue(Issue, Details.Spans);
+					GetIssueResponse Response = await CreateIssueResponseAsync(Details, bShowDesktopAlerts);
+					Responses.Add(PropertyFilter.Apply(Response, Filter));
 				}
 			}
 			return Responses;
@@ -371,24 +390,6 @@ namespace HordeServer.Controllers
 			return false;
 		}
 
-		/// <summary>
-		/// Determines whether to include an issue for a particular user
-		/// </summary>
-		/// <param name="Issue"></param>
-		/// <param name="User"></param>
-		/// <returns></returns>
-		static bool IncludeIssueForUser(IIssue Issue, string User)
-		{
-			if (Issue.Owner != null)
-			{
-				return Issue.Owner.Equals(User, StringComparison.OrdinalIgnoreCase);
-			}
-			else
-			{
-				return Issue.Suspects.Any(x => x.Author.Equals(User, StringComparison.OrdinalIgnoreCase));
-			}
-		}
-
 		/// <summary>
 		/// Update an issue
 		/// </summary>
@@ -399,12 +400,31 @@ namespace HordeServer.Controllers
 		[Route("/api/v1/issues/{IssueId}")]
 		public async Task<ActionResult> UpdateIssueAsync(int IssueId, [FromBody] UpdateIssueRequest Request)
 		{
-			string? DeclinedBy = null;
-			if(Request.Declined ?? false)
+			ObjectId? NewOwnerId = null;
+			if (Request.Owner != null)
 			{
-				DeclinedBy = User.GetPerforceUser();
+				NewOwnerId = (await UserCollection.FindOrAddUserByLoginAsync(Request.Owner))?.Id;
 			}
-			if (!await IssueService.UpdateIssueAsync(IssueId, Request.Summary, Request.Owner, Request.NominatedBy, Request.Acknowledged, DeclinedBy, Request.FixChange, Request.Resolved))
+
+			ObjectId? NewNominatedById = null;
+			if (Request.NominatedBy != null)
+			{
+				NewNominatedById = (await UserCollection.FindOrAddUserByLoginAsync(Request.NominatedBy))?.Id;
+			}
+
+			ObjectId? NewDeclinedById = null;
+			if (Request.Declined ?? false)
+			{
+				NewDeclinedById = User.GetUserId();
+			}
+
+			ObjectId? NewResolvedById = null;
+			if (Request.Resolved.HasValue)
+			{
+				NewResolvedById = Request.Resolved.Value ? User.GetUserId() : ObjectId.Empty;
+			}
+
+			if (!await IssueService.UpdateIssueAsync(IssueId, Request.Summary, NewOwnerId, NewNominatedById, Request.Acknowledged, NewDeclinedById, Request.FixChange, NewResolvedById))
 			{
 				return NotFound();
 			}
diff --git a/Engine/Source/Programs/Horde/HordeServer/Controllers/UgsController.cs b/Engine/Source/Programs/Horde/HordeServer/Controllers/UgsController.cs
--- a/Engine/Source/Programs/Horde/HordeServer/Controllers/UgsController.cs
+++ b/Engine/Source/Programs/Horde/HordeServer/Controllers/UgsController.cs
@@ -57,6 +57,11 @@ namespace HordeServer.Controllers
 		/// </summary>
 		private readonly ILogEventCollection LogEventCollection;
 
+		/// <summary>
+		/// Collection of users
+		/// </summary>
+		private readonly IUserCollection UserCollection;
+
 		/// <summary>
 		/// The log file service
 		/// </summary>
@@ -80,16 +85,18 @@ namespace HordeServer.Controllers
 		/// <param name="StreamService">Collection of stream documents</param>
 		/// <param name="UgsMetadataCollection">Collection of UGS metadata documents</param>
 		/// <param name="LogEventCollection">Collection of log event documents</param>
+		/// <param name="UserCollection">Collection of user documents</param>
 		/// <param name="LogFileService">The log file service</param>
 		/// <param name="OptionsMonitor">The server settings</param>
 		/// <param name="Logger">Logger</param>
-		public UgsController(IIssueService IssueService, JobService JobService, StreamService StreamService, IUgsMetadataCollection UgsMetadataCollection, ILogEventCollection LogEventCollection, ILogFileService LogFileService, IOptionsMonitor<ServerSettings> OptionsMonitor, ILogger<UgsController> Logger)
+		public UgsController(IIssueService IssueService, JobService JobService, StreamService StreamService, IUgsMetadataCollection UgsMetadataCollection, ILogEventCollection LogEventCollection, IUserCollection UserCollection, ILogFileService LogFileService, IOptionsMonitor<ServerSettings> OptionsMonitor, ILogger<UgsController> Logger)
 		{
 			this.IssueService = IssueService;
 			this.JobService = JobService;
 			this.StreamService = StreamService;
 			this.UgsMetadataCollection = UgsMetadataCollection;
 			this.LogEventCollection = LogEventCollection;
+			this.UserCollection = UserCollection;
 			this.LogFileService = LogFileService;
 			this.Settings = OptionsMonitor.CurrentValue;
 			this.Logger = Logger;
@@ -183,6 +190,8 @@ namespace HordeServer.Controllers
 		[System.Diagnostics.CodeAnalysis.SuppressMessage("Style", "IDE0060:Remove unused parameter", Justification = "<Pending>")]
 		public async Task<ActionResult<List<GetUgsIssueResponse>>> GetIssuesAsync([FromQuery] string? User = null, [FromQuery] bool IncludeResolved = false, [FromQuery] int MaxResults = 100)
 		{
+			IUser? UserInfo = (User != null) ? await UserCollection.FindUserByLoginAsync(User) : null;
+
 			List<GetUgsIssueResponse> Responses = new List<GetUgsIssueResponse>();
 			if (IncludeResolved)
 			{
@@ -190,8 +199,8 @@ namespace HordeServer.Controllers
 				foreach(IIssue Issue in Issues)
 				{
 					IIssueDetails Details = await IssueService.GetIssueDetailsAsync(Issue);
-					bool bNotify = User != null && Issue.Suspects.Any(x => x.Author.Equals(User, StringComparison.OrdinalIgnoreCase));
-					Responses.Add(CreateIssueResponse(Details, bNotify));
+					bool bNotify = UserInfo != null && Details.Suspects.Any(x => x.AuthorId == UserInfo.Id);
+					Responses.Add(await CreateIssueResponseAsync(Details, bNotify));
 				}
 			}
 			else
@@ -205,8 +214,8 @@ namespace HordeServer.Controllers
 
 					if (CachedOpenIssue.ShowNotifications())
 					{
-						bool bNotify = User != null && CachedOpenIssue.IncludeForUser(User);
-						Responses.Add(CreateIssueResponse(CachedOpenIssue, bNotify));
+						bool bNotify = UserInfo != null && CachedOpenIssue.IncludeForUser(UserInfo.Id);
+						Responses.Add(await CreateIssueResponseAsync(CachedOpenIssue, bNotify));
 					}
 				}
 			}
@@ -230,7 +239,7 @@ namespace HordeServer.Controllers
 				return NotFound();
 			}
 
-			return PropertyFilter.Apply(CreateIssueResponse(Issue, false), Filter);
+			return PropertyFilter.Apply(await CreateIssueResponseAsync(Issue, false), Filter);
 		}
 
 		/// <summary>
@@ -312,13 +321,17 @@ namespace HordeServer.Controllers
 		/// <summary>
 		/// Gets the URL for a failing step in the
 		/// </summary>
-		/// <param name="Issue">The issue to get a URL for</param>
+		/// <param name="Details">The issue to get a URL for</param>
 		/// <param name="bNotify">Whether to show notifications for this issue</param>
 		/// <returns>The issue response</returns>
-		GetUgsIssueResponse CreateIssueResponse(IIssueDetails Issue, bool bNotify)
+		async Task<GetUgsIssueResponse> CreateIssueResponseAsync(IIssueDetails Details, bool bNotify)
 		{
-			Uri? BuildUrl = GetIssueBuildUrl(Issue);
-			return new GetUgsIssueResponse(Issue, bNotify, BuildUrl);
+			Uri? BuildUrl = GetIssueBuildUrl(Details);
+
+			IUser? Owner = Details.Issue.OwnerId.HasValue ? await UserCollection.GetCachedUserAsync(Details.Issue.OwnerId.Value) : null;
+			IUser? NominatedBy = Details.Issue.NominatedById.HasValue ? await UserCollection.GetCachedUserAsync(Details.Issue.NominatedById.Value) : null;
+
+			return new GetUgsIssueResponse(Details, Owner, NominatedBy, bNotify, BuildUrl);
 		}
 
 		/// <summary>
@@ -368,7 +381,35 @@ namespace HordeServer.Controllers
 		[Route("/ugs/api/issues/{IssueId}")]
 		public async Task<ActionResult> UpdateIssueAsync(int IssueId, [FromBody] UpdateUgsIssueRequest Request)
 		{
-			if (!await IssueService.UpdateIssueAsync(IssueId, null, Request.Owner, Request.NominatedBy, Request.Acknowledged, Request.DeclinedBy, Request.FixChange, Request.Resolved))
+			ObjectId? NewOwnerId = null;
+			if (!String.IsNullOrEmpty(Request.Owner))
+			{
+				NewOwnerId = (await UserCollection.FindOrAddUserByLoginAsync(Request.Owner))?.Id;
+			}
+
+			ObjectId? NewNominatedById = null;
+			if (!String.IsNullOrEmpty(Request.NominatedBy))
+			{
+				NewNominatedById = (await UserCollection.FindOrAddUserByLoginAsync(Request.NominatedBy))?.Id;
+			}
+
+			ObjectId? NewDeclinedById = null;
+			if (!String.IsNullOrEmpty(Request.DeclinedBy))
+			{
+				NewDeclinedById = (await UserCollection.FindOrAddUserByLoginAsync(Request.DeclinedBy))?.Id;
+			}
+
+			ObjectId? NewResolvedById = null;
+			if (!String.IsNullOrEmpty(Request.ResolvedBy))
+			{
+				NewResolvedById = (await UserCollection.FindOrAddUserByLoginAsync(Request.ResolvedBy))?.Id;
+			}
+			if (NewResolvedById == null && Request.Resolved.HasValue)
+			{
+				NewResolvedById = Request.Resolved.Value ? IIssue.ResolvedByUnknownId : ObjectId.Empty;
+			}
+
+			if (!await IssueService.UpdateIssueAsync(IssueId, null, NewOwnerId, NewNominatedById, Request.Acknowledged, NewDeclinedById, Request.FixChange, NewResolvedById))
 			{
 				return NotFound();
 			}
diff --git a/Engine/Source/Programs/Horde/HordeServer/HordeServer.csproj b/Engine/Source/Programs/Horde/HordeServer/HordeServer.csproj
--- a/Engine/Source/Programs/Horde/HordeServer/HordeServer.csproj
+++ b/Engine/Source/Programs/Horde/HordeServer/HordeServer.csproj
@@ -73,6 +73,16 @@
     <ProjectReference Include="..\..\Shared\EpicGames.Core\EpicGames.Core.csproj" />
     <ProjectReference Include="..\..\Shared\EpicGames.Perforce\EpicGames.Perforce.csproj" />
     <ProjectReference Include="..\HordeCommon\HordeCommon.csproj" />
+    <ProjectReference Include="..\ThirdParty\P4Api.Net\P4Api.Net.csproj" />
+  </ItemGroup>
+
+  <ItemGroup>
+    <Content Include="..\ThirdParty\P4Api.Net\Native\Windows\Release\p4bridge.*" Condition=" '$(OS)' == 'Windows_NT' ">
+      <CopyToOutputDirectory>PreserveNewest</CopyToOutputDirectory>
+    </Content>
+    <Content Include="..\ThirdParty\P4Api.Net\Native\Linux\Release\p4bridge.*" Condition="$([MSBuild]::IsOSPlatform('Linux'))" >
+      <CopyToOutputDirectory>PreserveNewest</CopyToOutputDirectory>
+    </Content>	
   </ItemGroup>
 
 </Project>
diff --git a/Engine/Source/Programs/Horde/HordeServer/Models/IIssue.cs b/Engine/Source/Programs/Horde/HordeServer/Models/IIssue.cs
--- a/Engine/Source/Programs/Horde/HordeServer/Models/IIssue.cs
+++ b/Engine/Source/Programs/Horde/HordeServer/Models/IIssue.cs
@@ -46,6 +46,16 @@ namespace HordeServer.Models
 	/// </summary>
 	public interface IIssue
 	{
+		/// <summary>
+		/// Id for resolved by system
+		/// </summary>
+		static ObjectId ResolvedByTimeoutId { get; } = ObjectId.Parse("609592712b5c90b5bcf88c48");
+
+		/// <summary>
+		/// Id for resolved by unknown user
+		/// </summary>
+		static ObjectId ResolvedByUnknownId { get; } = ObjectId.Parse("609593b83e9b0b6dde620cf3");
+
 		/// <summary>
 		/// The unique object id
 		/// </summary>
@@ -71,21 +81,11 @@ namespace HordeServer.Models
 		/// </summary>
 		public IssueSeverity Severity { get; }
 
-		/// <summary>
-		/// Current owner of this issue
-		/// </summary>
-		public string? Owner { get; }
-
 		/// <summary>
 		/// User id of the owner
 		/// </summary>
 		public ObjectId? OwnerId { get; }
 
-		/// <summary>
-		/// User that nominated the current owner
-		/// </summary>
-		public string? NominatedBy { get; }
-
 		/// <summary>
 		/// User id of the person that nominated the owner
 		/// </summary>
@@ -111,20 +111,35 @@ namespace HordeServer.Models
 		/// </summary>
 		public DateTime? ResolvedAt { get; }
 
+		/// <summary>
+		/// User that resolved the issue
+		/// </summary>
+		public ObjectId? ResolvedById { get; }
+
+		/// <summary>
+		/// Time at which the issue was verified fixed
+		/// </summary>
+		public DateTime? VerifiedAt { get; }
+
+		/// <summary>
+		/// Time at which the issue was last seen.
+		/// </summary>
+		public DateTime LastSeenAt { get; }
+
 		/// <summary>
 		/// Fix changelist for this issue
 		/// </summary>
 		public int? FixChange { get; }
 
 		/// <summary>
-		/// Whether all suspects should be notified about this issue
+		/// List of streams affected by this issue
 		/// </summary>
-		public bool NotifySuspects { get; }
+		public IReadOnlyList<IIssueStream> Streams { get; }
 
 		/// <summary>
-		/// Suspects for this issue
+		/// Whether all suspects should be notified about this issue
 		/// </summary>
-		public IReadOnlyList<IIssueSuspect> Suspects { get; }
+		public bool NotifySuspects { get; }
 
 		/// <summary>
 		/// Update index for this instance
@@ -132,15 +147,36 @@ namespace HordeServer.Models
 		public int UpdateIndex { get; }
 	}
 
+	/// <summary>
+	/// Information about a stream affected by an issue
+	/// </summary>
+	public interface IIssueStream
+	{
+		/// <summary>
+		/// The stream id
+		/// </summary>
+		StreamId StreamId { get; }
+
+		/// <summary>
+		/// Whether this stream contains the fix change
+		/// </summary>
+		bool? ContainsFix { get; }
+	}
+
 	/// <summary>
 	/// Suspect for an issue
 	/// </summary>
 	public interface IIssueSuspect
 	{
 		/// <summary>
-		/// Author of the change
+		/// Unique id for this suspect
+		/// </summary>
+		ObjectId Id { get; }
+
+		/// <summary>
+		/// Issue that this suspect belongs to
 		/// </summary>
-		string Author { get; }
+		int IssueId { get; }
 
 		/// <summary>
 		/// The user id of the change's author
@@ -157,4 +193,52 @@ namespace HordeServer.Models
 		/// </summary>
 		DateTime? DeclinedAt { get; }
 	}
+
+	/// <summary>
+	/// Extension methods for issues
+	/// </summary>
+	static class IssueExtensions
+	{
+		/// <summary>
+		/// Creates a lookup from stream id to whether it's fixed
+		/// </summary>
+		/// <param name="Issue"></param>
+		/// <returns></returns>
+		public static Dictionary<StreamId, bool> GetFixStreamIds(this IIssue Issue)
+		{
+			Dictionary<StreamId, bool> FixStreamIds = new Dictionary<StreamId, bool>();
+			foreach (IIssueStream Stream in Issue.Streams)
+			{
+				if (Stream.ContainsFix.HasValue)
+				{
+					FixStreamIds[Stream.StreamId] = Stream.ContainsFix.Value;
+				}
+			}
+			return FixStreamIds;
+		}
+
+		/// <summary>
+		/// Find the first fix-failed step for this issue
+		/// </summary>
+		/// <param name="Issue"></param>
+		/// <param name="Spans"></param>
+		/// <returns></returns>
+		public static IIssueStep? FindFixFailedStep(this IIssue Issue, IEnumerable<IIssueSpan> Spans)
+		{
+			IIssueStep? FixFailedStep = null;
+			if (Issue.FixChange != null && Issue.FixChange.Value >= 0)
+			{
+				foreach (IIssueSpan Span in Spans)
+				{
+					IIssueStream? Stream = Issue.Streams.FirstOrDefault(x => x.StreamId == Span.StreamId);
+					if(Stream != null && (Stream.ContainsFix ?? false) && Span.LastFailure.Change >= Issue.FixChange.Value)
+					{
+						FixFailedStep = Span.LastFailure;
+						break;
+					}
+				}
+			}
+			return FixFailedStep;
+		}
+	}
 }
\ No newline at end of file
diff --git a/Engine/Source/Programs/Horde/HordeServer/Models/IIssueDetails.cs b/Engine/Source/Programs/Horde/HordeServer/Models/IIssueDetails.cs
--- a/Engine/Source/Programs/Horde/HordeServer/Models/IIssueDetails.cs
+++ b/Engine/Source/Programs/Horde/HordeServer/Models/IIssueDetails.cs
@@ -1,6 +1,7 @@
 // Copyright Epic Games, Inc. All Rights Reserved.
 
 using HordeServer.Services;
+using MongoDB.Bson;
 using System;
 using System.Collections.Generic;
 using System.Linq;
@@ -18,6 +19,21 @@ namespace HordeServer.Models
 		/// </summary>
 		IIssue Issue { get; }
 
+		/// <summary>
+		/// The user that currently owns this issue
+		/// </summary>
+		IUser? Owner { get; }
+
+		/// <summary>
+		/// The user that nominated the current owner
+		/// </summary>
+		IUser? NominatedBy { get; }
+
+		/// <summary>
+		/// The user that resolved the issue
+		/// </summary>
+		IUser? ResolvedBy { get; }
+
 		/// <summary>
 		/// List of spans for the issue
 		/// </summary>
@@ -28,6 +44,16 @@ namespace HordeServer.Models
 		/// </summary>
 		IReadOnlyList<IIssueStep> Steps { get; }
 
+		/// <summary>
+		/// List of suspects for the issue
+		/// </summary>
+		IReadOnlyList<IIssueSuspect> Suspects { get; }
+
+		/// <summary>
+		/// List of users that are suspects for this issue
+		/// </summary>
+		IReadOnlyList<IUser> SuspectUsers { get; }
+
 		/// <summary>
 		/// Determines whether the given user should be notified about the given issue
 		/// </summary>
@@ -37,9 +63,9 @@ namespace HordeServer.Models
 		/// <summary>
 		/// Determines if the issue is relevant to the given user
 		/// </summary>
-		/// <param name="User">The user to query</param>
+		/// <param name="UserId">The user to query</param>
 		/// <returns>True if the issue is relevant to the given user</returns>
-		bool IncludeForUser(string User);
+		bool IncludeForUser(ObjectId UserId);
 	}
 
 	/// <summary>
diff --git a/Engine/Source/Programs/Horde/HordeServer/Models/IIssueStep.cs b/Engine/Source/Programs/Horde/HordeServer/Models/IIssueStep.cs
--- a/Engine/Source/Programs/Horde/HordeServer/Models/IIssueStep.cs
+++ b/Engine/Source/Programs/Horde/HordeServer/Models/IIssueStep.cs
@@ -25,6 +25,11 @@ namespace HordeServer.Models
 		/// </summary>
 		public int Change { get; }
 
+		/// <summary>
+		/// Severity of the issue in this step
+		/// </summary>
+		public IssueSeverity Severity { get; }
+
 		/// <summary>
 		/// Name of the job
 		/// </summary>
diff --git a/Engine/Source/Programs/Horde/HordeServer/Models/IServiceAccount.cs b/Engine/Source/Programs/Horde/HordeServer/Models/IServiceAccount.cs
--- a/Engine/Source/Programs/Horde/HordeServer/Models/IServiceAccount.cs
+++ b/Engine/Source/Programs/Horde/HordeServer/Models/IServiceAccount.cs
@@ -21,11 +21,6 @@ namespace HordeServer.Models
 		/// </summary>
 		public string SecretToken { get; }
 		
-		/// <summary>
-		/// Claims
-		/// </summary>
-		public IReadOnlyDictionary<string, string> Claims { get; }
-
 		/// <summary>
 		/// If the service account is active
 		/// </summary>
@@ -35,5 +30,19 @@ namespace HordeServer.Models
 		/// Description of the service account (who is it for, is there an owner etc)
 		/// </summary>
 		public string Description { get; }
+
+		/// <summary>
+		/// Add a claim
+		/// </summary>
+		/// <param name="Type">Type of claim</param>
+		/// <param name="Value">Value of claim</param>
+		/// <returns></returns>
+		public void AddClaim(string Type, string Value);
+		
+		/// <summary>
+		/// Get list of claims
+		/// </summary>
+		/// <returns>List of claims</returns>
+		public IReadOnlyList<(string Type, string Value)> GetClaims();
 	}
 }
diff --git a/Engine/Source/Programs/Horde/HordeServer/Models/Stream.cs b/Engine/Source/Programs/Horde/HordeServer/Models/Stream.cs
--- a/Engine/Source/Programs/Horde/HordeServer/Models/Stream.cs
+++ b/Engine/Source/Programs/Horde/HordeServer/Models/Stream.cs
@@ -297,7 +297,7 @@ namespace HordeServer.Models
 		/// <summary>
 		/// Hash of the template definition
 		/// </summary>
-		public ContentHash Hash { get; set; } = ContentHash.Empty;
+		public ContentHash Hash { get; set; }
 
 		/// <summary>
 		/// Whether to show badges in UGS for this schedule
@@ -347,6 +347,7 @@ namespace HordeServer.Models
 		private TemplateRef()
 		{
 			this.Name = null!;
+			this.Hash = null!;
 		}
 
 		/// <summary>
diff --git a/Engine/Source/Programs/Horde/HordeServer/Notifications/Impl/SlackNotificationSink.cs b/Engine/Source/Programs/Horde/HordeServer/Notifications/Impl/SlackNotificationSink.cs
--- a/Engine/Source/Programs/Horde/HordeServer/Notifications/Impl/SlackNotificationSink.cs
+++ b/Engine/Source/Programs/Horde/HordeServer/Notifications/Impl/SlackNotificationSink.cs
@@ -278,7 +278,14 @@ namespace HordeServer.Notifications.Impl
 				else
 				{
 					Attachment.Color = BlockKitAttachmentColors.Warning;
-					Attachment.Blocks.Add(new SectionBlock($"Files in CL {Job.PreflightChange} were *not submitted*. Please resolve the following issues and submit manually.\n\n```{Job.AutoSubmitMessage}```"));
+
+					string AutoSubmitMessage = String.Empty;
+					if (!String.IsNullOrEmpty(Job.AutoSubmitMessage))
+					{
+						AutoSubmitMessage = $"\n\n```{Job.AutoSubmitMessage}```";
+					}
+
+					Attachment.Blocks.Add(new SectionBlock($"Files in CL {Job.PreflightChange} were *not submitted*. Please resolve the following issues and submit manually.{Job.AutoSubmitMessage}"));
 				}
 			}
 
@@ -423,10 +430,12 @@ namespace HordeServer.Notifications.Impl
 		/// <inheritdoc/>
 		public async Task NotifyIssueUpdatedAsync(IIssue Issue)
 		{
+			IIssueDetails Details = await IssueService.GetIssueDetailsAsync(Issue);
+
 			HashSet<ObjectId> UserIds = new HashSet<ObjectId>();
 			if (Issue.NotifySuspects)
 			{
-				UserIds.UnionWith(Issue.Suspects.Select(x => x.AuthorId));
+				UserIds.UnionWith(Details.Suspects.Select(x => x.AuthorId));
 			}
 			if (Issue.OwnerId.HasValue)
 			{
@@ -448,9 +457,7 @@ namespace HordeServer.Notifications.Impl
 				}
 			}
 
-			IIssueDetails Details = await IssueService.GetIssueDetailsAsync(Issue);
-
-			if (Issue.OwnerId == null && (Issue.Suspects.Count == 0 || Issue.Suspects.All(x => x.DeclinedAt != null) || Issue.CreatedAt < DateTime.UtcNow - TimeSpan.FromHours(1.0)))
+			if (Issue.OwnerId == null && (Details.Suspects.Count == 0 || Details.Suspects.All(x => x.DeclinedAt != null) || Issue.CreatedAt < DateTime.UtcNow - TimeSpan.FromHours(1.0)))
 			{
 				foreach (IIssueSpan Span in Details.Spans)
 				{
@@ -558,11 +565,18 @@ namespace HordeServer.Notifications.Impl
 
 			if (Issue.FixChange != null)
 			{
+				IIssueStep? FixFailedStep = Issue.FindFixFailedStep(Details.Spans);
+
 				string Text;
 				if (Issue.FixChange.Value < 0)
 				{
 					Text = ":tick: Marked as a systemic issue.";
 				}
+				else if (FixFailedStep != null)
+				{
+					Uri FixFailedUrl = new Uri(Settings.DashboardUrl, $"job/{FixFailedStep.JobId}?step={FixFailedStep.StepId}&issue={Issue.Id}");
+					Text = $":cross: Marked fixed in *CL {Issue.FixChange.Value}*, but seen again at *<{FixFailedUrl}|CL {FixFailedStep.Change}>*";
+				}
 				else
 				{
 					Text = $":tick: Marked fixed in *CL {Issue.FixChange.Value}*.";
@@ -590,7 +604,7 @@ namespace HordeServer.Notifications.Impl
 					}
 					else
 					{
-						List<int> Changes = Issue.Suspects.Where(x => x.AuthorId == UserId).Select(x => x.Change).OrderBy(x => x).ToList();
+						List<int> Changes = Details.Suspects.Where(x => x.AuthorId == UserId).Select(x => x.Change).OrderBy(x => x).ToList();
 						if (Changes.Count > 0)
 						{
 							string Text = $"Horde has determined that {StringUtils.FormatList(Changes.Select(x => $"CL {x}"), "or")} is the most likely cause for this issue.";
@@ -626,7 +640,7 @@ namespace HordeServer.Notifications.Impl
 			}
 			else if (UserId != null)
 			{
-				IIssueSuspect? Suspect = Issue.Suspects.FirstOrDefault(x => x.AuthorId == UserId);
+				IIssueSuspect? Suspect = Details.Suspects.FirstOrDefault(x => x.AuthorId == UserId);
 				if (Suspect != null)
 				{
 					if (Suspect.DeclinedAt != null)
@@ -644,10 +658,10 @@ namespace HordeServer.Notifications.Impl
 					}
 				}
 			}
-			else if (Issue.Suspects.Count > 0)
+			else if (Details.Suspects.Count > 0)
 			{
 				List<string> DeclinedLines = new List<string>();
-				foreach (IIssueSuspect Suspect in Issue.Suspects)
+				foreach (IIssueSuspect Suspect in Details.Suspects)
 				{
 					if (Suspect.DeclinedAt == null)
 					{
@@ -718,11 +732,11 @@ namespace HordeServer.Notifications.Impl
 			}
 			else if (String.Equals(Verb, "accept", StringComparison.Ordinal))
 			{
-				await IssueService.UpdateIssueAsync(IssueId, Owner: UserName, Acknowledged: true);
+				await IssueService.UpdateIssueAsync(IssueId, OwnerId: UserId, Acknowledged: true);
 			}
 			else if (String.Equals(Verb, "decline", StringComparison.Ordinal))
 			{
-				await IssueService.UpdateIssueAsync(IssueId, DeclinedBy: UserName);
+				await IssueService.UpdateIssueAsync(IssueId, DeclinedById: UserId);
 			}
 
 			IIssue? NewIssue = await IssueService.GetIssueAsync(IssueId);
diff --git a/Engine/Source/Programs/Horde/HordeServer/Program.cs b/Engine/Source/Programs/Horde/HordeServer/Program.cs
--- a/Engine/Source/Programs/Horde/HordeServer/Program.cs
+++ b/Engine/Source/Programs/Horde/HordeServer/Program.cs
@@ -30,7 +30,7 @@ namespace HordeServer
 {
 	static class LoggerExtensions
 	{
-		public static Serilog.LoggerConfiguration Console(this LoggerSinkConfiguration SinkConfig, ServerSettings Settings)
+		public static LoggerConfiguration Console(this LoggerSinkConfiguration SinkConfig, ServerSettings Settings)
 		{
 			if (Settings.LogJsonToStdOut)
 			{
@@ -62,13 +62,14 @@ namespace HordeServer
 				.SetBasePath(AppDir.FullName)
 				.AddJsonFile("appsettings.json", optional: false)
 				.AddJsonFile($"appsettings.{Environment.GetEnvironmentVariable("ASPNETCORE_ENVIRONMENT")}.json", optional: true)
+				.AddJsonFile("appsettings.User.json", optional: true)
 				.AddEnvironmentVariables()
 				.Build();
 
 			ServerSettings HordeSettings = new ServerSettings();
 			Config.GetSection("Horde").Bind(HordeSettings);
 
-			Serilog.Log.Logger = new Serilog.LoggerConfiguration()
+			Log.Logger = new LoggerConfiguration()
 				.MinimumLevel.Debug()
 //				.MinimumLevel.Override("HordeServer.Services.DatabaseService", LogEventLevel.Verbose) // For MongoDB query tracing
 				.MinimumLevel.Override("Microsoft", LogEventLevel.Warning)
@@ -88,27 +89,28 @@ namespace HordeServer
 			{
 				using (X509Certificate2? GrpcCertificate = ReadGrpcCertificate(HordeSettings))
 				{
-					CreateHostBuilderWithCert(Args, GrpcCertificate).Build().Run();
+					CreateHostBuilderWithCert(Args, Config, GrpcCertificate).Build().Run();
 				}
 			}
 #pragma warning disable CA1031 // Do not catch general exception types
 			catch (Exception Ex)
 #pragma warning restore CA1031 // Do not catch general exception types
 			{
-				Serilog.Log.Logger.Error(Ex, "Unhandled exception");
+				Log.Logger.Error(Ex, "Unhandled exception");
 			}
 		}
 
 		// Used by WebApplicationFactory in controller tests. Uses reflection to call this exact function signature.
 		public static IHostBuilder CreateHostBuilder(string[] Args)
 		{
-			return CreateHostBuilderWithCert(Args, null);
+			return CreateHostBuilderWithCert(Args, new ConfigurationBuilder().Build(), null);
 		}
 
-		public static IHostBuilder CreateHostBuilderWithCert(string[] Args, X509Certificate2? SslCert)
+		public static IHostBuilder CreateHostBuilderWithCert(string[] Args, IConfiguration Config, X509Certificate2? SslCert)
 		{
 			return Host.CreateDefaultBuilder(Args)
 				.UseSerilog()
+				.ConfigureAppConfiguration(Builder => Builder.AddConfiguration(Config))
 				.ConfigureWebHostDefaults(WebBuilder => 
 				{
 					WebBuilder.ConfigureKestrel(Options =>
diff --git a/Engine/Source/Programs/Horde/HordeServer/Rpc/CapabilitiesService.cs b/Engine/Source/Programs/Horde/HordeServer/Rpc/CapabilitiesService.cs
--- a/Engine/Source/Programs/Horde/HordeServer/Rpc/CapabilitiesService.cs
+++ b/Engine/Source/Programs/Horde/HordeServer/Rpc/CapabilitiesService.cs
@@ -24,9 +24,11 @@ namespace HordeServer.Rpc
 
 			Capabilities.CacheCapabilities = new CacheCapabilities();
 			Capabilities.CacheCapabilities.DigestFunction.Add(DigestFunction.Types.Value.Sha256);
+			Capabilities.CacheCapabilities.DigestFunction.Add(DigestFunction.Types.Value.Epiciohash);
 
 			Capabilities.ExecutionCapabilities = new ExecutionCapabilities();
 			Capabilities.ExecutionCapabilities.DigestFunction = DigestFunction.Types.Value.Sha256;
+			Capabilities.ExecutionCapabilities.DigestFunction = DigestFunction.Types.Value.Epiciohash;
 			Capabilities.ExecutionCapabilities.ExecEnabled = true;
 
 			return Task.FromResult(Capabilities);
diff --git a/Engine/Source/Programs/Horde/HordeServer/ServerSettings.cs b/Engine/Source/Programs/Horde/HordeServer/ServerSettings.cs
--- a/Engine/Source/Programs/Horde/HordeServer/ServerSettings.cs
+++ b/Engine/Source/Programs/Horde/HordeServer/ServerSettings.cs
@@ -52,6 +52,45 @@ namespace HordeServer
 		Aws,
 	}
 
+	
+	/// <summary>
+	/// Settings for a remote execution instance
+	///
+	/// Instances are a concept from Google's remote execution API that Horde implements.
+	/// They are used for separating different type of executions into groups.
+	/// </summary>
+	public class RemoteExecInstanceSettings
+	{
+		/// <summary>
+		/// gRPC URL for the content-addressable storage to be used for this instance
+		/// </summary>
+		[Required]
+		public Uri? CasUrl { get; set; }
+		
+		/// <summary>
+		/// gRPC URL for the action cache to be used for this instance
+		/// </summary>
+		[Required]
+		public Uri? ActionCacheUrl { get; set; }
+		
+		/// <summary>
+		/// Service account token for accessing the gRPC URLs
+		/// </summary>
+		[Required]
+		public string? ServiceAccountToken { get; set; }
+	}
+	
+	/// <summary>
+	/// Settings for remote execution
+	/// </summary>
+	public class RemoteExecSettings
+	{
+		/// <summary>
+		/// Mapping instance names to instance settings
+		/// </summary>
+		public Dictionary<string, RemoteExecInstanceSettings> Instances { get; set; } = new Dictionary<string, RemoteExecInstanceSettings>();
+	}
+
 	/// <summary>
 	/// Global settings for the application
 	/// </summary>
@@ -315,7 +354,22 @@ namespace HordeServer
 		/// Multiple projects can be specified, separated by comma.
 		/// </summary>
 		public string? HelixSwarmProjects { get; set; }
-		
+
+		/// <summary>
+		/// The p4 bridge server
+		/// </summary>
+		public string? P4BridgeServer { get; set; }
+
+		/// <summary>
+		/// The p4 bridge service username
+		/// </summary>
+		public string? P4BridgeServiceUsername { get; set; }
+
+		/// <summary>
+		/// The p4 bridge service password
+		/// </summary>
+		public string? P4BridgeServicePassword { get; set; }
+
 		/// <summary>
 		/// Set the minimum size of the global thread pool
 		/// This value has been found in need of tweaking to avoid timeouts with the Redis client during bursts
@@ -323,6 +377,11 @@ namespace HordeServer
 		/// is receiving. For Epic's internal deployment, this is set to 40.
 		/// </summary>
 		public int? GlobalThreadPoolMinSize { get; set; }
+
+		/// <summary>
+		/// Settings for remote execution
+		/// </summary>
+		public RemoteExecSettings RemoteExecSettings { get; set; } = new RemoteExecSettings();
 		
 		/// <summary>
 		/// Lazily computed timezone value
diff --git a/Engine/Source/Programs/Horde/HordeServer/Services/IIssueService.cs b/Engine/Source/Programs/Horde/HordeServer/Services/IIssueService.cs
--- a/Engine/Source/Programs/Horde/HordeServer/Services/IIssueService.cs
+++ b/Engine/Source/Programs/Horde/HordeServer/Services/IIssueService.cs
@@ -76,10 +76,18 @@ namespace HordeServer.Services
 		/// <returns>List of steps for the issue</returns>
 		Task<List<IIssueStep>> GetIssueStepsAsync(IIssueSpan Span);
 
+		/// <summary>
+		/// Gets the suspects for an issue
+		/// </summary>
+		/// <param name="Issue"></param>
+		/// <returns></returns>
+		Task<List<IIssueSuspect>> GetIssueSuspectsAsync(IIssue Issue);
+
 		/// <summary>
 		/// Searches for open issues
 		/// </summary>
 		/// <param name="Ids">List of issue ids to include</param>
+		/// <param name="UserId">User to filter issues returned by</param>
 		/// <param name="StreamId">The stream to return issues for</param>
 		/// <param name="MinChange">The minimum changelist number to include</param>
 		/// <param name="MaxChange">The maximum changelist number to include</param>
@@ -87,7 +95,7 @@ namespace HordeServer.Services
 		/// <param name="Index">Index within the results to return</param>
 		/// <param name="Count">Number of results</param>
 		/// <returns>List of streams open in the given stream at the given changelist</returns>
-		Task<List<IIssue>> FindIssuesAsync(IEnumerable<int>? Ids = null, StreamId? StreamId = null, int? MinChange = null, int? MaxChange = null, bool? Resolved = null, int? Index = null, int? Count = null);
+		Task<List<IIssue>> FindIssuesAsync(IEnumerable<int>? Ids = null, ObjectId? UserId = null, StreamId? StreamId = null, int? MinChange = null, int? MaxChange = null, bool? Resolved = null, int? Index = null, int? Count = null);
 
 		/// <summary>
 		/// Searches for open issues affecting a job
@@ -98,11 +106,12 @@ namespace HordeServer.Services
 		/// <param name="StepId"></param>
 		/// <param name="BatchId"></param>
 		/// <param name="LabelIdx"></param>
+		/// <param name="UserId"></param>
 		/// <param name="Resolved">Whether to include results that are resolved</param>
 		/// <param name="Index">Index within the results to return</param>
 		/// <param name="Count">Number of results</param>
 		/// <returns></returns>
-		Task<List<IIssue>> FindIssuesForJobAsync(int[]? Ids, IJob Job, IGraph Graph, SubResourceId? StepId, SubResourceId? BatchId, int? LabelIdx, bool? Resolved, int Index, int Count);
+		Task<List<IIssue>> FindIssuesForJobAsync(int[]? Ids, IJob Job, IGraph Graph, SubResourceId? StepId = null, SubResourceId? BatchId = null, int? LabelIdx = null, ObjectId? UserId = null, bool? Resolved = null, int? Index = null, int? Count = null);
 
 		/// <summary>
 		/// Find the events for a particular issue
@@ -119,14 +128,14 @@ namespace HordeServer.Services
 		/// </summary>
 		/// <param name="Id">The current issue id</param>
 		/// <param name="Summary">New summary for the issue</param>
-		/// <param name="Owner">New owner of the issue</param>
-		/// <param name="NominatedBy">Person that nominated the new owner</param>
+		/// <param name="OwnerId">New owner of the issue</param>
+		/// <param name="NominatedById">Person that nominated the new owner</param>
 		/// <param name="Acknowledged">Whether the issue has been acknowledged</param>
-		/// <param name="DeclinedBy">Name of the user that has declined this issue</param>
+		/// <param name="DeclinedById">Name of the user that has declined this issue</param>
 		/// <param name="FixChange">Fix changelist for the issue</param>
-		/// <param name="Resolved">Whether the issue has been resolved</param>
+		/// <param name="ResolvedById">Whether the issue has been resolved</param>
 		/// <returns>True if the issue was updated</returns>
-		Task<bool> UpdateIssueAsync(int Id, string? Summary = null, string? Owner = null, string? NominatedBy = null, bool? Acknowledged = null, string? DeclinedBy = null, int? FixChange = null, bool? Resolved = null);
+		Task<bool> UpdateIssueAsync(int Id, string? Summary = null, ObjectId? OwnerId = null, ObjectId? NominatedById = null, bool? Acknowledged = null, ObjectId? DeclinedById = null, int? FixChange = null, ObjectId? ResolvedById = null);
 
 		/// <summary>
 		/// Marks a step as complete
diff --git a/Engine/Source/Programs/Horde/HordeServer/Services/Impl/IssueService.cs b/Engine/Source/Programs/Horde/HordeServer/Services/Impl/IssueService.cs
--- a/Engine/Source/Programs/Horde/HordeServer/Services/Impl/IssueService.cs
+++ b/Engine/Source/Programs/Horde/HordeServer/Services/Impl/IssueService.cs
@@ -39,52 +39,67 @@ namespace HordeServer.Services.Impl
 	/// </summary>
 	class IssueDetails : IIssueDetails
 	{
-		/// <summary>
-		/// The issue instance
-		/// </summary>
+		/// <inheritdoc/>
 		public IIssue Issue { get; }
 
-		/// <summary>
-		/// List of spans for this issue
-		/// </summary>
+		/// <inheritdoc/>
+		public IUser? Owner { get; }
+
+		/// <inheritdoc/>
+		public IUser? NominatedBy { get; }
+
+		/// <inheritdoc/>
+		public IUser? ResolvedBy { get; }
+
+		/// <inheritdoc/>
 		public IReadOnlyList<IIssueSpan> Spans { get; }
 
-		/// <summary>
-		/// List of steps for the spans
-		/// </summary>
+		/// <inheritdoc/>
 		public IReadOnlyList<IIssueStep> Steps { get; }
 
-		/// <summary>
-		/// Whether to show desktop alerts for this issue
-		/// </summary>
+		/// <inheritdoc/>
+		public IReadOnlyList<IIssueSuspect> Suspects { get; }
+
+		/// <inheritdoc/>
+		public IReadOnlyList<IUser> SuspectUsers { get; }
+
+		/// <inheritdoc/>
 		bool ShowDesktopAlerts;
 
-		/// <summary>
-		/// Set of primary suspects
-		/// </summary>
-		HashSet<string> NotifyUsers;
+		/// <inheritdoc/>
+		HashSet<ObjectId> NotifyUsers;
 
 		/// <summary>
 		/// Constructor
 		/// </summary>
 		/// <param name="Issue"></param>
+		/// <param name="Owner"></param>
+		/// <param name="NominatedBy"></param>
+		/// <param name="ResolvedBy"></param>
 		/// <param name="Spans">Spans for this issue</param>
 		/// <param name="Steps">Steps for this issue</param>
+		/// <param name="Suspects"></param>
+		/// <param name="SuspectUsers"></param>
 		/// <param name="ShowDesktopAlerts"></param>
-		public IssueDetails(IIssue Issue, IReadOnlyList<IIssueSpan> Spans, IReadOnlyList<IIssueStep> Steps, bool ShowDesktopAlerts)
+		public IssueDetails(IIssue Issue, IUser? Owner, IUser? NominatedBy, IUser? ResolvedBy, IReadOnlyList<IIssueSpan> Spans, IReadOnlyList<IIssueStep> Steps, IReadOnlyList<IIssueSuspect> Suspects, IReadOnlyList<IUser> SuspectUsers, bool ShowDesktopAlerts)
 		{
 			this.Issue = Issue;
+			this.Owner = Owner;
+			this.NominatedBy = NominatedBy;
+			this.ResolvedBy = ResolvedBy;
 			this.Spans = Spans;
 			this.Steps = Steps;
+			this.Suspects = Suspects;
+			this.SuspectUsers = SuspectUsers;
 			this.ShowDesktopAlerts = ShowDesktopAlerts;
 
-			if (Issue.Owner == null)
+			if (Issue.OwnerId == null)
 			{
-				NotifyUsers = new HashSet<string>(Issue.Suspects.Select(x => x.Author), StringComparer.OrdinalIgnoreCase);
+				NotifyUsers = new HashSet<ObjectId>(Suspects.Select(x => x.AuthorId));
 			}
 			else
 			{
-				NotifyUsers = new HashSet<string> { Issue.Owner };
+				NotifyUsers = new HashSet<ObjectId> { Issue.OwnerId.Value };
 			}
 		}
 
@@ -100,18 +115,18 @@ namespace HordeServer.Services.Impl
 		/// <summary>
 		/// Determines if the issue is relevant to the given user
 		/// </summary>
-		/// <param name="User">The user to query</param>
+		/// <param name="UserId">The user to query</param>
 		/// <returns>True if the issue is relevant to the given user</returns>
-		public bool IncludeForUser(string User)
+		public bool IncludeForUser(ObjectId UserId)
 		{
-			return NotifyUsers.Contains(User);
+			return NotifyUsers.Contains(UserId);
 		}
 	}
 
 	/// <summary>
 	/// Wraps funtionality for manipulating build health issues
 	/// </summary>
-	public class IssueService : TickedBackgroundService, IIssueService
+	public class IssueService : ElectedBackgroundService, IIssueService
 	{
 		/// <summary>
 		/// Maximum number of changes to query from Perforce in one go
@@ -138,6 +153,11 @@ namespace HordeServer.Services.Impl
 		/// </summary>
 		IStreamCollection Streams;
 
+		/// <summary>
+		/// Collection of user documents
+		/// </summary>
+		IUserCollection UserCollection;
+
 		/// <summary>
 		/// Collection of commit documents
 		/// </summary>
@@ -166,7 +186,7 @@ namespace HordeServer.Services.Impl
 		/// <summary>
 		/// Cached list of currently open issues
 		/// </summary>
-		IReadOnlyList<IssueDetails> CachedIssues = new List<IssueDetails>();
+		IReadOnlyList<IIssueDetails> CachedIssues = new List<IIssueDetails>();
 
 		/// <summary>
 		/// Cache of templates to show desktop alerts for
@@ -186,15 +206,17 @@ namespace HordeServer.Services.Impl
 		/// <summary>
 		/// Constructor
 		/// </summary>
+		/// <param name="DatabaseService"></param>
 		/// <param name="IssueCollection">Collection of issue documents</param>
 		/// <param name="Jobs">Collection of job documents</param>
 		/// <param name="JobStepRefs">Collection of jobstepref documents</param>
 		/// <param name="Streams">Collection of stream documents</param>
+		/// <param name="UserCollection"></param>
 		/// <param name="Perforce">The perforce service instance</param>
 		/// <param name="LogFileService">Collection of event documents</param>
 		/// <param name="Logger">The logger instance</param>
-		public IssueService(IIssueCollection IssueCollection, IJobCollection Jobs, IJobStepRefCollection JobStepRefs, IStreamCollection Streams, IPerforceService Perforce, ILogFileService LogFileService, ILogger<IssueService> Logger)
-			: base(TimeSpan.FromSeconds(20.0), Logger)
+		public IssueService(DatabaseService DatabaseService, IIssueCollection IssueCollection, IJobCollection Jobs, IJobStepRefCollection JobStepRefs, IStreamCollection Streams, IUserCollection UserCollection, IPerforceService Perforce, ILogFileService LogFileService, ILogger<IssueService> Logger)
+			: base(DatabaseService, ObjectId.Parse("609542152fb0794700a6c3df"), Logger)
 		{
 			Type[] IssueTypes = Assembly.GetExecutingAssembly().GetTypes().Where(x => !x.IsAbstract && typeof(IIssue).IsAssignableFrom(x)).ToArray();
 			foreach (Type IssueType in IssueTypes)
@@ -207,6 +229,7 @@ namespace HordeServer.Services.Impl
 			this.Jobs = Jobs;
 			this.JobStepRefs = JobStepRefs;
 			this.Streams = Streams;
+			this.UserCollection = UserCollection;
 			this.Perforce = Perforce;
 			this.LogFileService = LogFileService;
 			this.Logger = Logger;
@@ -232,8 +255,10 @@ namespace HordeServer.Services.Impl
 		/// </summary>
 		/// <param name="StoppingToken">Token to indicate that the service should stop</param>
 		/// <returns>Async task</returns>
-		protected override async Task TickAsync(CancellationToken StoppingToken)
+		protected override async Task<DateTime> TickLeaderAsync(CancellationToken StoppingToken)
 		{
+			DateTime UtcNow = DateTime.UtcNow;
+
 			Dictionary<StreamId, HashSet<TemplateRefId>> NewCachedDesktopAlerts = new Dictionary<StreamId, HashSet<TemplateRefId>>();
 
 			List<IStream> CachedStreams = await Streams.FindAllAsync();
@@ -248,16 +273,27 @@ namespace HordeServer.Services.Impl
 
 			CachedDesktopAlerts = NewCachedDesktopAlerts;
 
+			// Resolve any issues that haven't been seen in a week
 			List<IIssue> OpenIssues = await FindIssuesAsync(Resolved: false);
+			for (int Idx = 0; Idx < OpenIssues.Count; Idx++)
+			{
+				IIssue OpenIssue = OpenIssues[Idx];
+				if (OpenIssue.LastSeenAt < UtcNow - TimeSpan.FromDays(7.0))
+				{
+					await IssueCollection.UpdateIssueAsync(OpenIssue, NewResolvedById: IIssue.ResolvedByTimeoutId);
+					OpenIssues.RemoveAt(Idx--);
+				}
+			}
 
-			List<IssueDetails> NewCachedOpenIssues = new List<IssueDetails>();
+			// Cache the details for any issues that are still open
+			List<IIssueDetails> NewCachedOpenIssues = new List<IIssueDetails>();
 			foreach(IIssue OpenIssue in OpenIssues)
 			{
-				List<IIssueSpan> Spans = await IssueCollection.FindSpansAsync(OpenIssue.Id);
-				List<IIssueStep> Steps = await IssueCollection.FindStepsAsync(Spans.Select(x => x.Id));
-				NewCachedOpenIssues.Add(new IssueDetails(OpenIssue, Spans, Steps, ShowDesktopAlertsForIssue(OpenIssue, Spans)));
+				NewCachedOpenIssues.Add(await GetIssueDetailsAsync(OpenIssue));
 			}
 			CachedIssues = NewCachedOpenIssues;
+
+			return UtcNow + TimeSpan.FromMinutes(1.0);
 		}
 
 		/// <inheritdoc/>
@@ -278,9 +314,25 @@ namespace HordeServer.Services.Impl
 		/// <inheritdoc/>
 		public async Task<IIssueDetails> GetIssueDetailsAsync(IIssue Issue)
 		{
+			IUser? Owner = Issue.OwnerId.HasValue ? await UserCollection.GetCachedUserAsync(Issue.OwnerId.Value) : null;
+			IUser? NominatedBy = Issue.NominatedById.HasValue ? await UserCollection.GetCachedUserAsync(Issue.NominatedById.Value) : null;
+			IUser? ResolvedBy = (Issue.ResolvedById.HasValue && Issue.ResolvedById != IIssue.ResolvedByTimeoutId && Issue.ResolvedById != IIssue.ResolvedByUnknownId)? await UserCollection.GetCachedUserAsync(Issue.ResolvedById.Value) : null;
+
 			List<IIssueSpan> Spans = await IssueCollection.FindSpansAsync(Issue.Id);
 			List<IIssueStep> Steps = await IssueCollection.FindStepsAsync(Spans.Select(x => x.Id));
-			return new IssueDetails(Issue, Spans, Steps, ShowDesktopAlertsForIssue(Issue, Spans));
+			List<IIssueSuspect> Suspects = await IssueCollection.GetSuspectsAsync(Issue);
+
+			List<IUser> SuspectUsers = new List<IUser>();
+			foreach (ObjectId SuspectUserId in Suspects.Select(x => x.AuthorId).Distinct())
+			{
+				IUser? SuspectUser = await UserCollection.GetCachedUserAsync(SuspectUserId);
+				if (SuspectUser != null)
+				{
+					SuspectUsers.Add(SuspectUser);
+				}
+			}
+
+			return new IssueDetails(Issue, Owner, NominatedBy, ResolvedBy, Spans, Steps, Suspects, SuspectUsers, ShowDesktopAlertsForIssue(Issue, Spans));
 		}
 
 		/// <inheritdoc/>
@@ -296,13 +348,19 @@ namespace HordeServer.Services.Impl
 		}
 
 		/// <inheritdoc/>
-		public async Task<List<IIssue>> FindIssuesAsync(IEnumerable<int>? Ids = null, StreamId? StreamId = null, int? MinChange = null, int? MaxChange = null, bool? Resolved = null, int? Index = null, int? Count = null)
+		public Task<List<IIssueSuspect>> GetIssueSuspectsAsync(IIssue Issue)
 		{
-			return await IssueCollection.FindIssuesAsync(Ids, StreamId, MinChange, MaxChange, Resolved, Index, Count);
+			return IssueCollection.GetSuspectsAsync(Issue);
 		}
 
 		/// <inheritdoc/>
-		public async Task<List<IIssue>> FindIssuesForJobAsync(int[]? Ids, IJob Job, IGraph Graph, SubResourceId? StepId, SubResourceId? BatchId, int? LabelIdx, bool? Resolved, int Index, int Count)
+		public async Task<List<IIssue>> FindIssuesAsync(IEnumerable<int>? Ids = null, ObjectId? UserId = null, StreamId? StreamId = null, int? MinChange = null, int? MaxChange = null, bool? Resolved = null, int? Index = null, int? Count = null)
+		{
+			return await IssueCollection.FindIssuesAsync(Ids, UserId, StreamId, MinChange, MaxChange, Resolved, Index, Count);
+		}
+
+		/// <inheritdoc/>
+		public async Task<List<IIssue>> FindIssuesForJobAsync(int[]? Ids, IJob Job, IGraph Graph, SubResourceId? StepId = null, SubResourceId? BatchId = null, int? LabelIdx = null, ObjectId? UserId = null, bool? Resolved = null, int? Index = null, int? Count = null)
 		{
 			List<IIssueStep> Steps = await IssueCollection.FindStepsAsync(Job.Id, BatchId, StepId);
 			List<IIssueSpan> Spans = await IssueCollection.FindSpansAsync(Steps.Select(x => x.SpanId));
@@ -322,7 +380,12 @@ namespace HordeServer.Services.Impl
 				}
 			}
 
-			return await FindIssuesAsync(Ids: IssueIds, Index: Index, Count: Count);
+			if(IssueIds.Count == 0)
+			{
+				return new List<IIssue>();
+			}
+
+			return await FindIssuesAsync(Ids: IssueIds, UserId: UserId, Resolved: Resolved, Index: Index, Count: Count);
 		}
 
 		/// <inheritdoc/>
@@ -358,8 +421,9 @@ namespace HordeServer.Services.Impl
 		}
 
 		/// <inheritdoc/>
-		public async Task<bool> UpdateIssueAsync(int Id, string? UserSummary = null, string? Owner = null, string? NominatedBy = null, bool? Acknowledged = null, string? DeclinedBy = null, int? FixChange = null, bool? Resolved = null)
+		public async Task<bool> UpdateIssueAsync(int Id, string? UserSummary = null, ObjectId? OwnerId = null, ObjectId? NominatedById = null, bool? Acknowledged = null, ObjectId? DeclinedById = null, int? FixChange = null, ObjectId? ResolvedById = null)
 		{
+			Dictionary<StreamId, bool>? NewFixStreamIds = null;
 			for (; ; )
 			{
 				IIssue? Issue = await IssueCollection.GetIssueAsync(Id);
@@ -368,8 +432,12 @@ namespace HordeServer.Services.Impl
 					return false;
 				}
 
-				string? OldOwner = Issue.Owner;
-				Issue = await IssueCollection.UpdateIssueAsync(Issue, NewUserSummary: UserSummary, NewOwner: Owner, NewNominatedBy: NominatedBy, NewAcknowledged: Acknowledged, NewDeclinedBy: DeclinedBy, NewFixChange: FixChange, NewResolved: Resolved);
+				if (FixChange != null && FixChange.Value > 0)
+				{
+					NewFixStreamIds = await GetFixStreamsAsync(Issue.Id, FixChange.Value, NewFixStreamIds ?? ((FixChange == Issue.FixChange)? Issue.GetFixStreamIds() : null));
+				}
+
+				Issue = await IssueCollection.UpdateIssueAsync(Issue, NewUserSummary: UserSummary, NewOwnerId: OwnerId ?? ResolvedById, NewNominatedById: NominatedById, NewAcknowledged: Acknowledged, NewDeclinedById: DeclinedById, NewFixChange: FixChange, NewFixStreamIds: NewFixStreamIds, NewResolvedById: ResolvedById);
 
 				if(Issue != null)
 				{
@@ -640,18 +708,32 @@ namespace HordeServer.Services.Impl
 					List<NewEventGroup> MatchEventGroups = NewEventGroups.Where(x => OpenSpan.Fingerprint.IsMatch(x.Fingerprint)).ToList();
 					if (MatchEventGroups.Count > 0)
 					{
+						// Get the severity of these events
+						IssueSeverity Severity = MatchEventGroups.SelectMany(x => x.Events).Any(x => x.Event.Severity == EventSeverity.Error) ? IssueSeverity.Error : IssueSeverity.Warning;
+
 						// Get the new severity
 						IssueSeverity? NewSeverity = null;
-						if (OpenSpan.Severity == IssueSeverity.Warning && MatchEventGroups.SelectMany(x => x.Events).Any(x => x.Event.Severity == EventSeverity.Error))
+						if (OpenSpan.Severity == IssueSeverity.Warning && Severity == IssueSeverity.Error)
 						{
 							NewSeverity = IssueSeverity.Error;
 						}
 
 						// Get the new step data
-						NewIssueStepData NewFailure = new NewIssueStepData(Job, Batch, Step, NotifySuspects);
+						NewIssueStepData NewFailure = new NewIssueStepData(Job, Batch, Step, Severity, NotifySuspects);
+
+						// Figure out whether we need to update the issue
+						bool bMarkAsModified = NewSeverity != null;
+						if (OpenSpan.IssueId != null)
+						{
+							IIssue? Issue = await IssueCollection.GetIssueAsync(OpenSpan.IssueId.Value);
+							if (Issue != null)
+							{
+								bMarkAsModified |= Issue.ResolvedAt != null;
+							}
+						}
 
 						// Update the span
-						IIssueSpan? NewSpan = await IssueCollection.UpdateSpanAsync(OpenSpan, NewSeverity: NewSeverity, NewFailure: NewFailure, NewModified: NewSeverity != null);
+						IIssueSpan? NewSpan = await IssueCollection.UpdateSpanAsync(OpenSpan, NewSeverity: NewSeverity, NewFailure: NewFailure, NewModified: bMarkAsModified);
 						if (NewSpan == null)
 						{
 							return false;
@@ -735,10 +817,10 @@ namespace HordeServer.Services.Impl
 				}
 
 				// Get the step data
-				NewIssueStepData StepData = new NewIssueStepData(Job, Batch, Step, NotifySuspects);
+				NewIssueStepData StepData = new NewIssueStepData(Job, Batch, Step, GetIssueSeverity(EventGroup.Events), NotifySuspects);
 
 				// Create the new span
-				IIssueSpan? NewSpan = await AddEventToNewSpanAsync(Token, Stream, Job, Node, GetIssueSeverity(EventGroup.Events), EventGroup.Fingerprint, StepData);
+				IIssueSpan? NewSpan = await AddEventToNewSpanAsync(Token, Stream, Job, Node, EventGroup.Fingerprint, StepData);
 				if(NewSpan == null)
 				{
 					return false;
@@ -762,11 +844,10 @@ namespace HordeServer.Services.Impl
 		/// <param name="Stream">Stream containing the job</param>
 		/// <param name="Job">Job that spawned the event</param>
 		/// <param name="Node">Node that was running</param>
-		/// <param name="Severity">Severity of the issue</param>
 		/// <param name="Fingerprint">Fingerprint for the event</param>
 		/// <param name="NewFailure">Information about the step that failed</param>
 		/// <returns>The new span, or null if the insert failed</returns>
-		async Task<IIssueSpan?> AddEventToNewSpanAsync(IIssueSequenceToken Token, IStream Stream, IJob Job, INode Node, IssueSeverity Severity, NewIssueFingerprint Fingerprint, NewIssueStepData NewFailure)
+		async Task<IIssueSpan?> AddEventToNewSpanAsync(IIssueSequenceToken Token, IStream Stream, IJob Job, INode Node, NewIssueFingerprint Fingerprint, NewIssueStepData NewFailure)
 		{
 			// Get the previous job in this span
 			IJobStepRef? PrevJob = await JobStepRefs.GetPrevStepForNodeAsync(Job.StreamId, Job.TemplateId, Node.Name, Job.Change);
@@ -794,7 +875,7 @@ namespace HordeServer.Services.Impl
 			}
 
 			// Add the span
-			return await IssueCollection.AddSpanAsync(Token, Stream, Job.TemplateId, Node.Name, Severity, Fingerprint, LastSuccess, NewFailure, NextSuccess, Suspects);
+			return await IssueCollection.AddSpanAsync(Token, Stream, Job.TemplateId, Node.Name, Fingerprint, LastSuccess, NewFailure, NextSuccess, Suspects);
 		}
 
 		/// <summary>
@@ -814,26 +895,26 @@ namespace HordeServer.Services.Impl
 				}
 
 				// Try to find an issue for it
+				IIssue? Issue;
 				if (Span.IssueId == null)
 				{
-					IIssue? ExistingIssue = await FindIssueForSpanAsync(Span);
-					if (ExistingIssue == null)
+					Issue = await FindIssueForSpanAsync(Span);
+					if (Issue == null)
 					{
-						IIssue? NewIssue = await IssueCollection.AddIssueAsync(Token, GetSummary(Span.Fingerprint, Span.Severity), Span);
-						if (NewIssue == null)
+						Issue = await IssueCollection.AddIssueAsync(Token, GetSummary(Span.Fingerprint, Span.Severity), Span);
+						if (Issue == null)
 						{
 							await Token.ResetAsync();
 							continue;
 						}
-						OnIssueUpdated?.Invoke(NewIssue);
 					}
 					else
 					{
-						NewIssueFingerprint NewFingerprint = NewIssueFingerprint.Merge(ExistingIssue.Fingerprint, Span.Fingerprint);
-						IssueSeverity NewSeverity = (ExistingIssue.Severity > Span.Severity) ? ExistingIssue.Severity : Span.Severity;
+						NewIssueFingerprint NewFingerprint = NewIssueFingerprint.Merge(Issue.Fingerprint, Span.Fingerprint);
+						IssueSeverity NewSeverity = (Issue.Severity > Span.Severity) ? Issue.Severity : Span.Severity;
 						string NewSummary = GetSummary(NewFingerprint, NewSeverity);
 
-						if (!await IssueCollection.AttachSpanToIssueAsync(Token, Span, ExistingIssue, NewSeverity, NewFingerprint, NewSummary))
+						if (!await IssueCollection.AttachSpanToIssueAsync(Token, Span, Issue, NewSeverity, NewFingerprint, NewSummary))
 						{
 							await Token.ResetAsync();
 							continue;
@@ -843,7 +924,7 @@ namespace HordeServer.Services.Impl
 				else
 				{
 					// Update the issue that it belongs to
-					IIssue? Issue = await GetIssueAsync(Span.IssueId.Value);
+					Issue = await GetIssueAsync(Span.IssueId.Value);
 					if (Issue != null)
 					{
 						Issue = await UpdateIssueDerivedDataAsync(Issue);
@@ -854,11 +935,103 @@ namespace HordeServer.Services.Impl
 					}
 				}
 
+				// Perform any shared issue updates
+				if (Issue != null)
+				{
+					Issue = await UpdateResolvedStateAsync(Issue, Span);
+					if(Issue == null)
+					{
+						continue;
+					}
+					OnIssueUpdated?.Invoke(Issue);
+				}
+
 				// Clear the modified flag
 				await IssueCollection.UpdateSpanAsync(Span, NewModified: false);
 			}
 		}
 
+		async Task<IIssue?> UpdateResolvedStateAsync(IIssue? Issue, IIssueSpan Span)
+		{
+			// Update the fix streams
+			if (Issue != null && Issue.FixChange != null && Issue.FixChange.Value > 0)
+			{
+				Dictionary<StreamId, bool>? FixStreamIds = await GetFixStreamsAsync(Issue.Id, Issue.FixChange.Value, Issue.GetFixStreamIds());
+				if (FixStreamIds != null)
+				{
+					Issue = await IssueCollection.UpdateIssueAsync(Issue, NewFixStreamIds: FixStreamIds);
+				}
+			}
+
+			// Update the resolved flag
+			if (Issue != null && Issue.ResolvedAt != null && ShouldMarkFixFailed(Issue, Span))
+			{
+				Issue = await IssueCollection.UpdateIssueAsync(Issue, NewResolvedById: ObjectId.Empty);
+			}
+
+			return Issue;
+		}
+
+		/// <summary>
+		/// Determine if an issue should be marked as fix failed
+		/// </summary>
+		/// <param name="Issue"></param>
+		/// <param name="Span"></param>
+		/// <returns></returns>
+		static bool ShouldMarkFixFailed(IIssue Issue, IIssueSpan Span)
+		{
+			if (Issue.ResolvedAt != null && Span.NextSuccess == null)
+			{
+				if (Issue.FixChange != null && (Issue.Streams.FirstOrDefault(x => x.StreamId == Span.StreamId)?.ContainsFix ?? false))
+				{
+					if (Span.LastFailure.Change > Issue.FixChange.Value)
+					{
+						return true;
+					}
+				}
+				else
+				{
+					if (Span.LastFailure.StepTime > Issue.ResolvedAt.Value + TimeSpan.FromHours(24.0))
+					{
+						return true;
+					}
+				}
+			}
+			return false;
+		}
+
+		/// <summary>
+		/// Finds the streams affected by an issue
+		/// </summary>
+		/// <param name="IssueId"></param>
+		/// <param name="FixChange"></param>
+		/// <param name="PrevFixStreamIds"></param>
+		/// <returns></returns>
+		async Task<Dictionary<StreamId, bool>?> GetFixStreamsAsync(int IssueId, int FixChange, IReadOnlyDictionary<StreamId, bool>? PrevFixStreamIds = null)
+		{
+			List<IIssueSpan> Spans = await IssueCollection.FindSpansAsync(IssueId);
+			if(Spans.Count == 0)
+			{
+				return null;
+			}
+
+			Dictionary<StreamId, bool> NextFixStreamIds = new Dictionary<StreamId, bool>();
+			foreach (IIssueSpan Span in Spans)
+			{
+				if (!NextFixStreamIds.ContainsKey(Span.StreamId))
+				{
+					bool ContainsFix;
+					if (PrevFixStreamIds == null || !PrevFixStreamIds.TryGetValue(Span.StreamId, out ContainsFix))
+					{
+						List<ChangeSummary> Changes = await Perforce.GetChangesAsync(Span.StreamName, FixChange, FixChange, 1, null);
+						ContainsFix = Changes.Count > 0;
+					}
+					NextFixStreamIds[Span.StreamId] = ContainsFix;
+				}
+			}
+			return NextFixStreamIds;
+		}
+
 		/// <summary>
 		/// Find suspects that can match a given span
 		/// </summary>
@@ -928,48 +1101,23 @@ namespace HordeServer.Services.Impl
 			return ExistingIssues.FirstOrDefault(x => x.Fingerprint.IsMatch(Span.Fingerprint));
 		}
 
+		/// <summary>
+		/// 
+		/// </summary>
+		/// <param name="Issue"></param>
+		/// <returns></returns>
 		async Task<IIssue?> UpdateIssueDerivedDataAsync(IIssue Issue)
 		{
-			for (; ; )
+			IIssue? NewIssue = await IssueCollection.UpdateIssueDerivedDataAsync(Issue);
+			if (NewIssue != null)
 			{
-				List<NewIssueSuspectData> NewSuspects = new List<NewIssueSuspectData>();
-
-				// Find all the spans that are attached to the issue
-				List<IIssueSpan> Spans = await IssueCollection.FindSpansAsync(Issue.Id);
-				if (Spans.Count > 0)
-				{
-					// Get the intersection of changes that are suspects for causing this issue
-					HashSet<int> SuspectChanges = new HashSet<int>(Spans[0].Suspects.Select(x => x.OriginatingChange ?? x.Change));
-					for (int SpanIdx = 1; SpanIdx < Spans.Count; SpanIdx++)
-					{
-						SuspectChanges.IntersectWith(Spans[SpanIdx].Suspects.Select(x => x.OriginatingChange ?? x.Change));
-					}
-
-					// Update the list of suspects to match
-					NewSuspects.AddRange(Issue.Suspects.Where(x => SuspectChanges.Contains(x.Change)).Select(x => new NewIssueSuspectData(x.Author, x.Change, x.DeclinedAt)));
-					SuspectChanges.ExceptWith(NewSuspects.Select(x => x.Change));
-					NewSuspects.AddRange(Spans[0].Suspects.Where(x => SuspectChanges.Contains(x.OriginatingChange ?? x.Change)).Select(x => new NewIssueSuspectData(x.Author, x.OriginatingChange ?? x.Change, null)));
-				}
-
-				// Update the issue
-				IssueSeverity NewSeverity = Spans.Any(x => x.Severity == IssueSeverity.Error) ? IssueSeverity.Error : IssueSeverity.Warning;
-				string NewSummary = GetSummary(Issue.Fingerprint, NewSeverity);
-				bool NewResolved = Spans.All(x => x.NextSuccess != null);
-				bool NewNotifySuspects = Spans.Any(x => x.NotifySuspects);
-
-				IIssue ? NewIssue = await IssueCollection.UpdateIssueAsync(Issue, NewSeverity: NewSeverity, NewSummary: NewSummary, NewResolved: NewResolved, NewNotifySuspects: NewNotifySuspects, NewSuspects: NewSuspects);
-				if (NewIssue != null)
-				{
-					return NewIssue;
-				}
-
-				// Try to get the new issue state
-				NewIssue = await IssueCollection.GetIssueAsync(Issue.Id);
-				if (NewIssue == null)
+				string NewSummary = GetSummary(Issue.Fingerprint, Issue.Severity);
+				if (!NewSummary.Equals(Issue.Summary, StringComparison.Ordinal))
 				{
-					return null;
+					NewIssue = await IssueCollection.UpdateIssueAsync(NewIssue, NewSummary: NewSummary);
 				}
 			}
+			return NewIssue;
 		}
 
 		/// <summary>
@@ -1005,7 +1153,7 @@ namespace HordeServer.Services.Impl
 			{
 				if (Job.Change < Span.FirstFailure.Change && (Span.LastSuccess == null || Job.Change > Span.LastSuccess.Change))
 				{
-					NewIssueStepData NewLastSuccess = new NewIssueStepData(Job.Change, Job.Name, Job.Id, Batch.Id, Step.Id, Step.StartTimeUtc ?? default, Step.LogId, false);
+					NewIssueStepData NewLastSuccess = new NewIssueStepData(Job.Change, IssueSeverity.Unspecified, Job.Name, Job.Id, Batch.Id, Step.Id, Step.StartTimeUtc ?? default, Step.LogId, false);
 					List<NewIssueSpanSuspectData> NewSuspects = await FindSuspectsForSpanAsync(Stream, Span.Fingerprint, Job.Change + 1, Span.FirstFailure.Change);
 
 					if (await IssueCollection.UpdateSpanAsync(Span, NewLastSuccess: NewLastSuccess, NewSuspects: NewSuspects, NewModified: true) == null)
@@ -1017,7 +1165,7 @@ namespace HordeServer.Services.Impl
 				}
 				else if (Job.Change > Span.LastFailure.Change && (Span.NextSuccess == null || Job.Change < Span.NextSuccess.Change))
 				{
-					NewIssueStepData NewNextSucccess = new NewIssueStepData(Job.Change, Job.Name, Job.Id, Batch.Id, Step.Id, Step.StartTimeUtc ?? default, Step.LogId, false);
+					NewIssueStepData NewNextSucccess = new NewIssueStepData(Job.Change, IssueSeverity.Unspecified, Job.Name, Job.Id, Batch.Id, Step.Id, Step.StartTimeUtc ?? default, Step.LogId, false);
 					if (await IssueCollection.UpdateSpanAsync(Span, NewNextSuccess: NewNextSucccess, NewModified: true) == null)
 					{
 						return false;
diff --git a/Engine/Source/Programs/Horde/HordeServer/Services/LogFileService.cs b/Engine/Source/Programs/Horde/HordeServer/Services/LogFileService.cs
--- a/Engine/Source/Programs/Horde/HordeServer/Services/LogFileService.cs
+++ b/Engine/Source/Programs/Horde/HordeServer/Services/LogFileService.cs
@@ -635,7 +635,7 @@ namespace HordeServer.Services
 				}
 				else
 				{
-					Logger.LogDebug("Added new chunk at offset {Offset} to log {LogId}", State.Offset, LogFile.Id);
+					// Logger.LogDebug("Added new chunk at offset {Offset} to log {LogId}", State.Offset, LogFile.Id);
 					LogFile = NewLogFile;
 				}
 			}
diff --git a/Engine/Source/Programs/Horde/HordeServer/Services/PerforceService.cs b/Engine/Source/Programs/Horde/HordeServer/Services/PerforceService.cs
--- a/Engine/Source/Programs/Horde/HordeServer/Services/PerforceService.cs
+++ b/Engine/Source/Programs/Horde/HordeServer/Services/PerforceService.cs
@@ -443,17 +443,13 @@ namespace HordeServer.Services
 		/// <inheritdoc/>
 		public async override Task<PerforceUserInfo?> GetUserInfoAsync(string UserName)
 		{
-			PerforceResponseList<UserRecord> UserResponse = await Perforce.TryGetUserAsync(UserName, CancellationToken.None);
-			if (UserResponse.Succeeded)
+			PerforceResponse<UserRecord> UserResponse = await Perforce.TryGetUserAsync(UserName, CancellationToken.None);
+			if (UserResponse.Succeeded && UserResponse.Data.Access != default)
 			{
-				UserRecord? UserRecord = UserResponse.Data.FirstOrDefault();
-				if (UserRecord != null)
-				{
-					PerforceUserInfo UserInfo = new PerforceUserInfo();
-					UserInfo.Name = UserRecord.FullName;
-					UserInfo.Email = UserRecord.Email;
-					return UserInfo;
-				}
+				PerforceUserInfo UserInfo = new PerforceUserInfo();
+				UserInfo.Name = UserResponse.Data.FullName;
+				UserInfo.Email = UserResponse.Data.Email;
+				return UserInfo;
 			}
 			return null;
 		}
diff --git a/Engine/Source/Programs/Horde/HordeServer/Services/PerforceServiceApi.cs b/Engine/Source/Programs/Horde/HordeServer/Services/PerforceServiceApi.cs
new file mode 100644
--- /dev/null
+++ b/Engine/Source/Programs/Horde/HordeServer/Services/PerforceServiceApi.cs
@@ -0,0 +1,984 @@
+using HordeServer.Models;
+using System.Linq;
+using Microsoft.Extensions.Logging;
+using System;
+using System.Collections.Generic;
+using System.Net;
+using System.Text;
+using System.Threading.Tasks;
+using System.IO;
+
+using P4 = Perforce.P4;
+using Microsoft.Extensions.Options;
+using System.Runtime.InteropServices;
+using System.Globalization;
+using System.Text.RegularExpressions;
+
+namespace HordeServer.Services
+{
+	/// <summary>
+	/// P4API implementation of the Perforce service
+	/// </summary>
+	public class P4APIPerforceService : PerforceServiceBase
+	{
+
+		IOptionsMonitor<ServerSettings> Settings;
+
+		ILogger<P4APIPerforceService> Logger;
+
+		P4.Server Server;
+
+		/// <summary>
+		/// Object used for controlling access to the access user tickets
+		/// </summary>
+		object TicketLock = new object();
+
+		Dictionary<string, P4.Credential> UserTickets = new Dictionary<string, P4.Credential>();
+
+		/// <summary>
+		/// Constructor
+		/// </summary>
+		public P4APIPerforceService(IOptionsMonitor<ServerSettings> Settings, ILogger<P4APIPerforceService> Logger)
+		{
+			this.Settings = Settings;
+			this.Logger = Logger;
+
+			Server = new P4.Server(new P4.ServerAddress(Settings.CurrentValue.P4BridgeServer));
+
+			P4.LogFile.SetLoggingFunction((int LogLevel, String Source, String Message) =>
+			{
+
+				switch (LogLevel)
+				{
+					case 0:
+					case 1:
+						Serilog.Log.Error("Perforce (Error): {Message} {Source} : ", Message, Source);
+						break;
+					case 2:
+						Serilog.Log.Warning("Perforce (Warning): {Message} {Source} : ", Message, Source);
+						break;
+					case 3:
+						Serilog.Log.Information("Perforce: {Message} {Source} : ", Message, Source);
+						break;
+					default:
+						Serilog.Log.Debug("Perforce (Debug): {Message} {Source} : ", Message, Source);
+						break;
+				};
+
+			});
+		}
+
+
+		P4.Repository GetServiceUserConnection()
+		{
+
+			ServerSettings Settings = this.Settings.CurrentValue;
+
+			P4.Repository Repository = new P4.Repository(Server);
+
+			try
+			{
+				P4.Connection Connection = Repository.Connection;
+				Connection.UserName = Settings.P4BridgeServiceUsername;
+
+				P4.Options Options = new P4.Options();
+
+				Options["Ticket"] = Settings.P4BridgeServicePassword;
+
+				// connect to the server
+				if (!Connection.Connect(Options))
+				{
+					throw new Exception("Unable to get P4 server connection");
+				}
+
+			}
+			catch
+			{
+				Repository.Dispose();
+				throw;
+			}
+
+
+			return Repository;
+		}
+
+		P4.Credential GetImpersonateCredential(string ImpersonateUser)
+		{
+
+			if (!CanImpersonate)
+			{
+				throw new Exception($"Service account required to impersonate user {ImpersonateUser}");
+			}
+
+			lock (TicketLock)
+			{
+
+				ServerSettings Settings = this.Settings.CurrentValue;
+
+				// Check if we have a ticket
+				P4.Credential? Credential = null;
+
+				if (UserTickets.TryGetValue(ImpersonateUser, out Credential))
+				{
+					// if the credential expires within the next 15 minutes, refresh
+					TimeSpan Time = new TimeSpan(0, 15, 0);
+					if (Credential.Expires.Subtract(Time) <= DateTime.UtcNow)
+					{
+						UserTickets.Remove(ImpersonateUser);
+						Credential = null;
+					}
+				}
+
+				if (Credential != null)
+				{
+					return Credential;
+				}
+
+				using (P4.Repository Repository = GetServiceUserConnection())
+				{
+					Credential = Repository.Connection.Login(null, new P4.LoginCmdOptions(P4.LoginCmdFlags.AllHosts | P4.LoginCmdFlags.DisplayTicket, null), ImpersonateUser);
+				}
+
+				if (Credential == null)
+				{
+					throw new Exception($"GetImpersonateCredential - Unable to get impersonation credential for {ImpersonateUser} from {Dns.GetHostName()}");
+				}
+
+				UserTickets.Add(ImpersonateUser, Credential);
+
+				return Credential;
+			}
+
+		}
+
+		P4.Repository GetImpersonateConnection(string ImpersonateUser)
+		{
+
+			if (!CanImpersonate)
+			{
+				throw new Exception($"Service account required to impersonate user {ImpersonateUser}");
+			}
+
+			P4.Credential Credential = GetImpersonateCredential(ImpersonateUser);
+
+			// this might be able to be reused and would be threaded internally
+			P4.Repository Repository = new P4.Repository(Server);
+
+			P4.Connection Connection = Repository.Connection;
+			Connection.UserName = ImpersonateUser;
+
+			P4.Options Options = new P4.Options();
+			Options["Ticket"] = Credential.Ticket;
+
+			// connect to the server
+			if (!Connection.Connect(Options))
+			{
+				throw new Exception($"GetImpersonateConnection - Unable to get impersonated P4 connection for {ImpersonateUser} from {Dns.GetHostName()}");
+
+			}
+
+			return Repository;
+		}
+
+		P4.Repository GetConnection(string? Stream = null, string? Username = null, bool ReadOnly = true, bool CreateChange = false, int? ClientFromChange = null, bool UseClientFromChange = false, bool UsePortFromChange = false, string? ClientId = null, bool NoClient = false)
+		{
+			ServerSettings Settings = this.Settings.CurrentValue;
+
+			P4.Repository? Repository = (Username == null || Settings.P4BridgeServiceUsername == Username || !CanImpersonate) ? GetServiceUserConnection() : GetImpersonateConnection(Username);
+
+			if (!NoClient)
+			{
+				try
+				{
+					if (ClientId != null)
+					{
+						Repository.Connection.SetClient(ClientId);
+					}
+					else
+					{
+						P4.Client Client = GetOrCreateClient(Repository, Stream, Username, ReadOnly, CreateChange, ClientFromChange, UseClientFromChange, UsePortFromChange);
+						Repository.Connection.SetClient(Client.Name);
+
+						if (UseClientFromChange)
+						{
+							string? ClientHost = Repository.Connection.Client.Host;
+							if (!string.IsNullOrEmpty(ClientHost))
+							{
+								Repository.Connection.getP4Server().SetConnectionHost(ClientHost);
+							}
+						}
+					}
+				}
+				catch
+				{
+					Repository.Dispose();
+					throw;
+				}
+			}
+
+			if (Repository == null)
+			{
+				throw new Exception($"Unable to get connection for Stream:{Stream} Username:{Username} ReadOnly:{ReadOnly} CreateChange:{CreateChange} ClientFromChange:{ClientFromChange} UseClientFromChange:{UseClientFromChange} UsePortFromChange:{UsePortFromChange}");
+			}
+
+			Repository.Connection.CommandEcho += (Data) => { Serilog.Log.Information(Data);};
+
+			return Repository;
+
+		}
+
+		string GetClientName(string Stream, bool ReadOnly, bool CreateChange, string? Username = null)
+		{
+			ServerSettings Settings = this.Settings.CurrentValue;
+
+			string ClientName = $"horde-p4bridge-{Dns.GetHostName()}-{Settings.P4BridgeServiceUsername}-{Stream.Replace("/", "+", StringComparison.OrdinalIgnoreCase)}";
+
+			if (!ReadOnly)
+			{
+				ClientName += "-write";
+			}
+
+			if (CreateChange)
+			{
+				ClientName += "-create";
+			}
+
+			if (!string.IsNullOrEmpty(Username))
+			{
+				ClientName += "-" + Username;
+			}
+
+			return ClientName;
+
+		}
+
+		P4.Client GetOrCreateClient(P4.Repository Repository, string? Stream, string? Username = null, bool ReadOnly = true, bool CreateChange = false, int? ClientFromChange = null, bool UseClientFromChange = false, bool UsePortFromChange = false)
+		{
+			ServerSettings Settings = this.Settings.CurrentValue;
+
+			P4.Client? Client = null;
+			P4.Changelist? Changelist = null;
+
+			string? ClientPort = null;
+			string? ClientHost = null;
+
+			if (ClientFromChange.HasValue)
+			{
+				Changelist = Repository.GetChangelist(ClientFromChange.Value);
+
+				if (Changelist == null)
+				{
+					throw new Exception($"Unable to get changelist for client {ClientFromChange}");
+				}
+
+				Client = Repository.GetClient(Changelist.ClientId);
+
+				if (Client == null)
+				{
+					throw new Exception($"Unable to get client for id {Changelist.ClientId}");
+				}
+
+				if (UsePortFromChange)
+				{
+					if (!Client.Name.StartsWith("horde-p4bridge-", StringComparison.OrdinalIgnoreCase))
+					{
+						using (P4.P4Command Cmd = new P4.P4Command(Repository, "info", true))
+						{
+							P4.P4CommandResult Results = Cmd.Run();
+
+							if (!Results.Success)
+							{
+								throw new Exception("Unable to get server info");
+							}
+
+							P4.TaggedObject? Tag = Results.TaggedOutput.Find(T => T.ContainsKey("changeServer"));
+
+							if (Tag != null)
+							{
+								if (!Tag.TryGetValue("changeServer", out ClientPort))
+								{
+									throw new Exception("Unable to get change server from tagged server info output");
+								}
+							}
+						}
+					}
+				}
+
+				Stream = Client.Stream;
+				Username = Client.OwnerName;
+				ClientHost = Client.Host;
+
+
+			}
+
+			if (!UseClientFromChange)
+			{
+				if (Stream == null)
+				{
+					throw new Exception("Stream required for client");
+
+				}
+				string ClientName = GetClientName(Stream, ReadOnly, CreateChange, Username);
+
+				IList<P4.Client>? Clients = Repository.GetClients(new P4.ClientsCmdOptions(P4.ClientsCmdFlags.None, Username, ClientName, 1, Stream));
+
+				if (Clients != null && Clients.Count == 1)
+				{
+					Client = Clients[0];
+					if (!string.IsNullOrEmpty(Client.Host))
+					{
+						Client.Host = string.Empty;
+						Repository.UpdateClient(Client);
+					}
+				}
+				else
+				{
+
+					P4.Client NewClient = new P4.Client();
+
+					NewClient.Name = ClientName;
+					NewClient.Root = RuntimeInformation.IsOSPlatform(OSPlatform.Linux) ? $"/tmp/{ClientName}/" : $"{Path.GetTempPath()}{ClientName}\\";
+					NewClient.Stream = Stream;
+					if (ReadOnly)
+					{
+						NewClient.ClientType = P4.ClientType.@readonly;
+					}
+
+					NewClient.OwnerName = Username ?? Settings.P4BridgeServiceUsername;
+
+					Client = Repository.CreateClient(NewClient);
+
+					if (Client == null)
+					{
+						throw new Exception($"Unable to create client for ${Stream} : {Username ?? Settings.P4BridgeServiceUsername}");
+					}
+				}
+			}
+
+			if (Client == null)
+			{
+				throw new Exception($"Unable to create client for Stream:{Stream} Username:{Username} ReadOnly:{ReadOnly} CreateChange:{CreateChange} ClientFromChange:{ClientFromChange} UseClientFromChange:{UseClientFromChange} UsePortFromChange:{UsePortFromChange}");
+			}
+
+			return Client;
+
+		}
+
+		/// <inheritdoc/>
+		public override async Task<List<ChangeSummary>> GetChangesAsync(string StreamName, int? MinChange, int? MaxChange, int Results, string? ImpersonateUser)
+		{
+
+			await Task.Yield();
+
+			using (P4.Repository Repository = GetConnection(StreamName, ImpersonateUser))
+			{
+
+				P4.ChangesCmdOptions Options = new P4.ChangesCmdOptions(P4.ChangesCmdFlags.IncludeTime | P4.ChangesCmdFlags.FullDescription, null, Results, P4.ChangeListStatus.Submitted, null);
+
+				string Filter = GetFilter($"//{Repository.Connection.Client.Name}/...", MinChange, MaxChange);
+
+				P4.FileSpec FileSpec = new P4.FileSpec(new P4.DepotPath(Filter), null, null, null);
+
+				IList<P4.Changelist> Changelists = Repository.GetChangelists(Options, FileSpec);
+
+				List<ChangeSummary> Changes = new List<ChangeSummary>();
+				foreach (P4.Changelist Changelist in Changelists)
+				{
+					Changes.Add(new ChangeSummary(Changelist.Id, Changelist.OwnerName, Changelist.Description));
+				}
+
+				return Changes;
+			}
+
+		}
+
+		/// <inheritdoc/>
+		public override async Task<List<ChangeDetails>> GetChangeDetailsAsync(string StreamName, IReadOnlyList<int> ChangeNumbers, string? ImpersonateUser)
+		{
+			await Task.Yield();
+
+			using (P4.Repository Repository = GetConnection(StreamName, ImpersonateUser))
+			{
+				List<ChangeDetails> Results = new List<ChangeDetails>();
+
+				List<string> Args = new List<string> { "-s", "-S" };
+				Args.AddRange(ChangeNumbers.Select(Change => Change.ToString(CultureInfo.InvariantCulture)));
+
+				using (P4.P4Command Command = new P4.P4Command(Repository, "describe", true, Args.ToArray()))
+				{
+					P4.P4CommandResult Result = Command.Run();
+
+					if (!Result.Success)
+					{
+						throw new Exception("Unable to get changes");
+					}
+
+					if (Result.TaggedOutput == null || Result.TaggedOutput.Count <= 0)
+					{
+						return Results;
+					}
+
+					List<P4.Changelist> Changelists = new List<P4.Changelist>() { };
+
+					bool DSTMismatch = false;
+					string Offset = string.Empty;
+
+					if (Server != null && Server.Metadata != null)
+					{
+						Offset = Server.Metadata.DateTimeOffset;
+						DSTMismatch = P4.FormBase.DSTMismatch(Server.Metadata);
+					}
+
+					foreach (P4.TaggedObject TaggedObject in Result.TaggedOutput)
+					{
+						List<string> Files = new List<string>();
+
+						P4.Changelist Change = new P4.Changelist();
+						Change.FromChangeCmdTaggedOutput(TaggedObject, true, Offset, DSTMismatch);
+
+						foreach (P4.FileMetaData DescribeFile in Change.Files)
+						{
+							string? RelativePath;
+							if (TryGetStreamRelativePath(DescribeFile.DepotPath.Path, StreamName, out RelativePath))
+							{
+								Files.Add(RelativePath);
+							}
+						}
+
+						if (Change.ShelvedFiles != null && Change.ShelvedFiles.Count > 0)
+						{
+							foreach (P4.ShelvedFile ShelvedFile in Change.ShelvedFiles)
+							{
+								string? RelativePath;
+								if (TryGetStreamRelativePath(ShelvedFile.Path.ToString(), StreamName, out RelativePath))
+								{
+									Files.Add(RelativePath);
+								}
+							}
+
+						}
+
+						Results.Add(new ChangeDetails(Change.Id, Change.OwnerName, Change.Description, Files));
+					}
+
+				}
+
+				return Results;
+			}
+		}
+
+		/// <inheritdoc />
+		public override async Task<string> CreateTicket(string ImpersonateUser)
+		{
+			await Task.Yield();
+
+			P4.Credential Credential = GetImpersonateCredential(ImpersonateUser);
+
+			if (Credential == null)
+			{
+				throw new Exception($"Unable to get ticket for user {ImpersonateUser}");
+			}
+
+			return Credential.Ticket;
+
+		}
+
+		/// <inheritdoc/>
+		public override async Task<List<FileSummary>> FindFilesAsync(IEnumerable<string> Paths)
+		{
+			await Task.Yield();
+
+			List<FileSummary> Results = new List<FileSummary>();
+
+			using (P4.Repository Repository = GetConnection(NoClient: true))
+			{
+				P4.GetFileMetaDataCmdOptions Options = new P4.GetFileMetaDataCmdOptions(P4.GetFileMetadataCmdFlags.ExcludeClientData, "", "", 0, "", "", "");
+
+				IList<P4.FileMetaData>? Files = Repository.GetFileMetaData(Options, Paths.Select(Path => { P4.FileSpec FileSpec = new P4.FileSpec(new P4.DepotPath(Path)); return FileSpec; }).ToArray());
+
+				if (Files == null)
+				{
+					Files = new List<P4.FileMetaData>();
+				}
+
+				foreach (string Path in Paths)
+				{
+					P4.FileMetaData Meta = Files.FirstOrDefault(File => File.DepotPath.Path == Path);
+
+					if (Meta == null)
+					{
+						Results.Add(new FileSummary(Path, false, 0));
+					}
+					else
+					{
+						Results.Add(new FileSummary(Path, Meta.HeadAction != P4.FileAction.Delete && Meta.HeadAction != P4.FileAction.MoveDelete, Meta.HeadChange));
+					}
+
+				}
+
+				return Results;
+			}
+		}
+
+		/// <inheritdoc/>
+		public override async Task<byte[]> PrintAsync(string DepotPath)
+		{
+			await Task.Yield();
+
+			if (DepotPath.EndsWith("...", StringComparison.OrdinalIgnoreCase) || DepotPath.EndsWith("*", StringComparison.OrdinalIgnoreCase))
+			{
+				throw new Exception("PrintAsync requires exactly one file to be specified");
+			}
+
+			using (P4.Repository Repository = GetConnection(NoClient: true))
+			{
+				using (P4.P4Command Command = new P4.P4Command(Repository, "print", false, new string[] { "-q", DepotPath }))
+				{
+					P4.P4CommandResult Result = Command.Run();
+
+					return Encoding.Default.GetBytes(Result.TextOutput);
+
+				}
+
+			}
+
+		}
+
+		/// <inheritdoc/>
+		public override async Task<int> DuplicateShelvedChangeAsync(int ShelvedChange)
+		{
+			await Task.Yield();
+
+			throw new Exception("DuplicateShelvedChangeAsync is disabled due to p4 issue: https://jira.it.epicgames.com/servicedesk/customer/portal/1/ITH-144069");
+
+			/*
+			using (P4.Repository Repository = GetConnection(NoClient: true))
+			{
+				return ReshelveChange(Repository, ShelvedChange);
+			}
+
+			throw new Exception($"Failed to duplicate shelved change {ShelvedChange}");
+			*/
+
+		}
+
+		/// <inheritdoc/>
+		public override async Task DeleteShelvedChangeAsync(int ShelvedChange)
+		{
+			await Task.Yield();
+
+			using (P4.Repository Repository = GetConnection(NoClient: true))
+			{
+
+				P4.Changelist? Changelist = Repository.GetChangelist(ShelvedChange, new P4.DescribeCmdOptions(P4.DescribeChangelistCmdFlags.Omit, 0, 0));
+
+				if (Changelist == null)
+				{
+					throw new Exception($"Unable to get shelved changelist for delete: {ShelvedChange}");
+				}
+
+				string ClientId = Changelist.ClientId;
+
+				if (String.IsNullOrEmpty(ClientId))
+				{
+					throw new Exception($"Unable to get shelved changelist client id for delete: {ShelvedChange}");
+				}
+
+				using (P4.Repository ClientRepository = GetConnection(Changelist.Stream, Changelist.OwnerName, ClientId: ClientId))
+				{
+
+					if (Changelist.Shelved)
+					{
+						IList<P4.FileSpec> Files = ClientRepository.Connection.Client.ShelveFiles(new P4.ShelveFilesCmdOptions(P4.ShelveFilesCmdFlags.Delete, null, ShelvedChange));
+					}
+
+					ClientRepository.DeleteChangelist(Changelist, null);
+				}
+			}
+		}
+
+		/// <inheritdoc/>
+		public override async Task UpdateChangelistDescription(int Change, string Description)
+		{
+			await Task.Yield();
+
+			using (P4.Repository Repository = GetConnection(NoClient: true))
+			{
+				P4.Changelist Changelist = Repository.GetChangelist(Change);
+
+				Repository.Connection.Disconnect();
+
+				// the client must exist for the change list, otherwise will fail (for example, CreateNewChangeAsync deletes the client before returning)
+				using (P4.Repository UpdateRepository = GetConnection(ClientFromChange: Change, UseClientFromChange: true, Username: Changelist.OwnerName))
+				{					
+					P4.Changelist UpdatedChangelist = UpdateRepository.GetChangelist(Change);		
+					UpdatedChangelist.Description = Description;
+					UpdateRepository.UpdateChangelist(UpdatedChangelist);					
+				}
+			}
+		}
+
+		/// <inheritdoc/>
+		public async override Task<int> CreateNewChangeAsync(string StreamName, string FilePath)
+		{
+			await Task.Yield();
+
+			ServerSettings Settings = this.Settings.CurrentValue;
+
+			string Username = CanImpersonate ? "buildmachine" : Settings.P4BridgeServiceUsername!;
+
+			P4.SubmitResults? SubmitResults = null;
+
+			using (P4.Repository Repository = GetConnection(Stream: StreamName, Username: Username, ReadOnly: false, CreateChange: true))
+			{
+				P4.Client Client = Repository.Connection.Client;
+
+				string WorkspaceFilePath = $"//{Client.Name}/{FilePath.TrimStart('/')}";
+				string DiskFilePath = $"{Client.Root + FilePath.TrimStart('/')}";
+
+				P4.FileSpec WorkspaceFileSpec = new P4.FileSpec(new P4.ClientPath(WorkspaceFilePath));
+
+				IList<P4.File> Files = Repository.GetFiles(new P4.FilesCmdOptions(P4.FilesCmdFlags.None, 1), WorkspaceFileSpec);
+
+				P4.Changelist? SubmitChangelist = null;
+				P4.DepotPath? DepotPath = null;
+
+				const int MaxRetries = 10;
+				int Retry = 0;
+
+				for (; ; )
+				{
+					DepotPath = null;
+
+					if (Retry == MaxRetries)
+					{
+						break;
+					}
+
+					try
+					{
+						if (Files == null || Files.Count == 0)
+						{
+							// File does not exist, create it
+							string? DirectoryName = Path.GetDirectoryName(DiskFilePath);
+
+							if (string.IsNullOrEmpty(DirectoryName))
+							{
+								throw new Exception($"Invalid directory name for local client file, disk file path: {DiskFilePath}");
+							}
+
+							// Create the directory
+							if (!Directory.Exists(DirectoryName))
+							{
+								Directory.CreateDirectory(DirectoryName);
+
+								if (!Directory.Exists(DirectoryName))
+								{
+									throw new Exception($"Unable to create directrory: {DirectoryName}");
+								}
+							}
+
+							// Create the file
+							if (!File.Exists(DiskFilePath))
+							{
+								using (FileStream FileStream = File.OpenWrite(DiskFilePath))
+								{
+									FileStream.Close();
+								}
+
+								if (!File.Exists(DiskFilePath))
+								{
+									throw new Exception($"Unable to create local change file: {DiskFilePath}");
+								}
+
+							}
+
+							IList<P4.FileSpec> DepotFiles = Client.AddFiles(new P4.Options(), WorkspaceFileSpec);
+
+							if (DepotFiles == null || DepotFiles.Count != 1)
+							{
+								throw new Exception($"Unable to add local change file,  local: {DiskFilePath} : workspace: {WorkspaceFileSpec}");
+							}
+
+							DepotPath = DepotFiles[0].DepotPath;
+
+						}
+						else
+						{
+							IList<P4.FileSpec> SyncResults = Client.SyncFiles(new P4.SyncFilesCmdOptions(P4.SyncFilesCmdFlags.Force), WorkspaceFileSpec);
+
+							if (SyncResults == null || SyncResults.Count != 1)
+							{
+								throw new Exception($"Unable to sync file, workspace: {WorkspaceFileSpec}");
+							}
+
+							IList<P4.FileSpec> EditResults = Client.EditFiles(new P4.FileSpec[] { new P4.FileSpec(SyncResults[0].DepotPath) }, new P4.Options());
+
+							if (EditResults == null || EditResults.Count != 1)
+							{
+								throw new Exception($"Unable to edit file, workspace: {WorkspaceFileSpec}");
+							}
+
+							DepotPath = EditResults[0].DepotPath;
+
+						}
+
+						if (DepotPath == null || string.IsNullOrEmpty(DepotPath.Path))
+						{
+							throw new Exception($"Unable to get depot path for: {WorkspaceFileSpec}");
+						}
+
+						// create a new change
+						if (SubmitChangelist == null)
+						{
+							P4.Changelist Changelist = new P4.Changelist();
+							Changelist.Description = "New change for Horde job";
+							Changelist.Files.Add(new P4.FileMetaData(new P4.FileSpec(DepotPath)));
+							SubmitChangelist = Repository.CreateChangelist(Changelist);
+
+							if (SubmitChangelist == null)
+							{
+								throw new Exception($"Unable to create a changelist for: {DepotPath}");
+							}
+						}
+
+						SubmitResults = SubmitChangelist.Submit(null);
+
+						if (SubmitResults == null)
+						{
+							throw new Exception($"Unable to submit changelist for: {DepotPath}");
+						}
+
+						break;
+					}
+					catch
+					{
+						Retry++;
+
+						Client.RevertFiles(new P4.Options(), WorkspaceFileSpec);
+
+						continue;
+					}
+
+				}
+
+				string ClientName = Client.Name;
+				try
+				{
+					Repository.DeleteClient(Client, new P4.Options());
+				}
+				catch
+				{
+					Logger.LogError($"Unable to delete client {ClientName}");
+				}
+
+			}
+
+			if (SubmitResults == null)
+			{
+				throw new Exception($"Unable to submit change for {StreamName} {FilePath}");
+			}
+
+			return SubmitResults.ChangeIdAfterSubmit;
+		}
+
+		/// <inheritdoc/>
+		public override async Task<(int? Change, string Message)> SubmitShelvedChangeAsync(int Change, int OriginalChange)
+		{
+			await Task.Yield();
+
+			using (P4.Repository Repository = GetConnection(NoClient: true))
+			{
+
+				List<string> Args = new List<string> { "-S", Change.ToString(CultureInfo.InvariantCulture), OriginalChange.ToString(CultureInfo.InvariantCulture) };
+
+				using (P4.P4Command Command = new P4.P4Command(Repository, "describe", true, Args.ToArray()))
+				{
+					P4.P4CommandResult Result = Command.Run();
+
+					if (!Result.Success)
+					{
+						throw new Exception("Unable to get changes");
+					}
+
+					if (Result.TaggedOutput == null || Result.TaggedOutput.Count != 2)
+					{
+
+						throw new Exception($"Unable to get tagged output for change: {Change} and original change: {OriginalChange}");
+					}
+
+					bool DSTMismatch = false;
+					string Offset = string.Empty;
+
+					if (Server != null && Server.Metadata != null)
+					{
+						Offset = Server.Metadata.DateTimeOffset;
+						DSTMismatch = P4.FormBase.DSTMismatch(Server.Metadata);
+					}
+
+					P4.Changelist Changelist = new P4.Changelist();
+					Changelist.FromChangeCmdTaggedOutput(Result.TaggedOutput[0], true, Offset, DSTMismatch);
+
+					P4.Changelist OriginalChangelist = new P4.Changelist();
+					OriginalChangelist.FromChangeCmdTaggedOutput(Result.TaggedOutput[1], true, Offset, DSTMismatch);
+
+					if (OriginalChangelist.ShelvedFiles.Count != Changelist.ShelvedFiles.Count)
+					{
+						throw new Exception($"Mismatched number of shelved files for change: {Change} and original change: {OriginalChange}");
+					}
+
+					if (OriginalChangelist.ShelvedFiles.Count == 0)
+					{
+						throw new Exception($"No shelved file for change: {Change} and original change: {OriginalChange}");
+					}
+
+					// todo, test straight up delete
+					foreach (P4.ShelvedFile ShelvedFile in Changelist.ShelvedFiles)
+					{
+						P4.ShelvedFile? Found = OriginalChangelist.ShelvedFiles.FirstOrDefault(Original => Original.Digest == ShelvedFile.Digest && Original.Action == ShelvedFile.Action);
+
+						if (Found == null)
+						{
+							throw new Exception($"Mismatch in shelved file digest or action for {ShelvedFile.Path}");
+						}
+					}
+
+					Repository.Connection.Disconnect();
+
+					using (P4.Repository SubmitRepository = GetConnection(ReadOnly: false, ClientFromChange: Change, UseClientFromChange: true, Username: Changelist.OwnerName))
+					{
+						// we might not need a client here, possibly -e below facilitates this, check!
+						using (P4.P4Command SubmitCommand = new P4.P4Command(SubmitRepository, "submit", null, true, new string[] { "-e", Change.ToString(CultureInfo.InvariantCulture) }))
+						{
+							Result = SubmitCommand.Run();
+
+							if (!Result.Success)
+							{
+								throw new Exception($"Unable to submit {Change}");
+							}
+
+							int? SubmittedChangeId = null;
+							
+							foreach (P4.TaggedObject TaggedObject in Result.TaggedOutput)
+							{
+								string? Submitted;
+								if (TaggedObject.TryGetValue("submittedChange", out Submitted))
+								{
+									SubmittedChangeId = int.Parse(Submitted, CultureInfo.InvariantCulture);									
+								}
+							}
+
+							if (SubmittedChangeId == null)
+							{
+								throw new Exception($"Unable to get submitted changelist");
+							}
+
+							return (SubmittedChangeId, Changelist.Description);							
+
+						}
+					}
+				}
+			}
+
+			throw new Exception($"Unable to get shelve change: {Change} for original change: {OriginalChange}");
+
+		}
+
+		/// <inheritdoc/>		
+		public override async Task<PerforceUserInfo?> GetUserInfoAsync(string UserName)
+		{
+			await Task.Yield();
+
+			using (P4.Repository Repository = GetConnection(NoClient: true))
+			{
+				P4.User User = Repository.GetUser(UserName, new P4.UserCmdOptions(P4.UserCmdFlags.Output));
+
+				if (User == null)
+				{
+					return null;
+				}
+
+				return new PerforceUserInfo { Name = UserName, Email = User.EmailAddress };
+			}
+		}
+
+		void RequireEdgeServer()
+		{
+			string? Value;
+			if (Server.Metadata.RawData.TryGetValue("serverServices", out Value))
+			{
+				if (Value != "edge-server")
+				{
+					Serilog.Log.Logger.Information("Perforce: servierService is not an edge server : {ServerServices}", Value);
+					throw new Exception("Operation requires an edge server");
+				}
+			}
+			else
+			{
+				throw new Exception("unable to get server services");
+			}
+		}
+
+		bool CanImpersonate
+		{
+			get
+			{
+				ServerSettings Settings = this.Settings.CurrentValue;
+				// Todo: we need a better way of identifying we're using a service account
+				return Settings.P4BridgeServiceUsername == "SVC-P4-HordeProxy-P";
+			}
+		}
+
+		int ReshelveChange(P4.Repository Repository, int Change)
+		{
+
+			// all attempts at using this without edge server have lead to bizarre reshelving behavior, also note that -p requires an edge server
+			// https://jira.it.epicgames.com/servicedesk/customer/portal/1/ITH-144069
+			RequireEdgeServer();
+
+			using (P4.P4Command Cmd = new P4.P4Command(Repository, "reshelve", false, new string[] { "-s", Change.ToString(CultureInfo.InvariantCulture), "-f", "-p" }))
+			{
+				P4.P4CommandResult Results = Cmd.Run();
+
+				if (Results.Success)
+				{
+
+					if (Results.InfoOutput.Count != 2)
+					{
+						Serilog.Log.Logger.Information("Perforce: Unexpected info output when reshelving change");
+						throw new Exception("Unexpected info output when reshelving change");
+					}
+
+					bool Error = true;
+					int ReshelvedChange = 0;
+					string Message = Results.InfoOutput[1].Message;
+					Match ChangeMatch = Regex.Match(Message, @"Change (\d+) files shelved");
+
+					if (ChangeMatch.Success && ChangeMatch.Groups.Count == 2)
+					{
+						if (int.TryParse(ChangeMatch.Groups[1].Value, out ReshelvedChange))
+						{
+							Error = false;
+						}
+					}
+
+					if (Error)
+					{
+						Serilog.Log.Logger.Information("Perforce: Unable to parse cl for reshelf: {Message}", Message);
+						throw new Exception($"Unable to parse cl for reshelf: {Message}");
+					}
+
+					return ReshelvedChange;
+				}
+			}
+
+			Serilog.Log.Logger.Information("Perforce: General reshelving failure");
+
+			throw new Exception($"Unable to reshelve CL {Change}");
+
+		}
+
+	}
+
+}
+
diff --git a/Engine/Source/Programs/Horde/HordeServer/Services/UpgradeService.cs b/Engine/Source/Programs/Horde/HordeServer/Services/UpgradeService.cs
--- a/Engine/Source/Programs/Horde/HordeServer/Services/UpgradeService.cs
+++ b/Engine/Source/Programs/Horde/HordeServer/Services/UpgradeService.cs
@@ -81,13 +81,13 @@ namespace HordeServer.Services
 
 				if (SchemaVersion == 3)
 				{
-					IJobCollection JobCollection = ServiceProvider.GetRequiredService<IJobCollection>();
+					IJobCollection JobCollection = ServiceProvider.GetService<IJobCollection>();
 					await JobCollection.UpgradeDocumentsAsync();
 				}
 				if (SchemaVersion == 4)
 				{
 					UserCollectionV1 UserCollectionV1 = new UserCollectionV1(ServiceProvider.GetRequiredService<DatabaseService>());
-					UserCollectionV2 UserCollectionV2 = new UserCollectionV2(ServiceProvider.GetRequiredService<DatabaseService>(), ServiceProvider.GetRequiredService<ILogger<UserCollectionV2>>());
+					using UserCollectionV2 UserCollectionV2 = new UserCollectionV2(ServiceProvider.GetRequiredService<DatabaseService>(), ServiceProvider.GetRequiredService<ILogger<UserCollectionV2>>());
 					await UserCollectionV2.ResaveDocumentsAsync(UserCollectionV1);
 				}
 
diff --git a/Engine/Source/Programs/Horde/HordeServer/Startup.cs b/Engine/Source/Programs/Horde/HordeServer/Startup.cs
--- a/Engine/Source/Programs/Horde/HordeServer/Startup.cs
+++ b/Engine/Source/Programs/Horde/HordeServer/Startup.cs
@@ -291,6 +291,7 @@ namespace HordeServer
 			Services.AddGrpc(Options =>
 			{
 				Options.EnableDetailedErrors = true;
+				Options.MaxReceiveMessageSize = 20 * 1024 * 1024; // 20 MB (packaged builds of Horde agent can be large) 
 				Options.Interceptors.Add(typeof(LifetimeGrpcInterceptor));
 				Options.Interceptors.Add(typeof(GrpcExceptionInterceptor));
 			});
@@ -367,13 +368,20 @@ namespace HordeServer
 			Services.AddSingleton<ILogFileService, LogFileService>();
 			Services.AddSingleton<INotificationService, NotificationService>();			
 
-			if (String.IsNullOrEmpty(Settings.PerforceBridge))
+			
+			if (!String.IsNullOrEmpty(Settings.PerforceBridge))
 			{
-				Services.AddSingleton<IPerforceService, LocalPerforceService>();
+				// Python P4 bridge service
+				Services.AddSingleton<IPerforceService, BridgePerforceService>(SP => new BridgePerforceService(Settings.PerforceBridge!, SP.GetService<ILogger<BridgePerforceService>>()));				
+			}
+			else if (!String.IsNullOrEmpty(Settings.P4BridgeServer))
+			{
+				// Native P4 API
+				Services.AddSingleton<IPerforceService, P4APIPerforceService>();
 			}
 			else
 			{
-				Services.AddSingleton<IPerforceService, BridgePerforceService>(SP => new BridgePerforceService(Settings.PerforceBridge!, SP.GetRequiredService<ILogger<BridgePerforceService>>()));
+				Services.AddSingleton<IPerforceService, LocalPerforceService>();
 			}
 
 			Services.AddSingleton<PerforceLoadBalancer>();
@@ -585,12 +593,12 @@ namespace HordeServer
 				}
 				else
 				{
-					return new RedisLogBuilder(RedisConnectionPool, Provider.GetRequiredService<ILogger<RedisLogBuilder>>());
+					return new RedisLogBuilder(RedisConnectionPool, Provider.GetService<ILogger<RedisLogBuilder>>());
 				}
 			});
 			Services.AddSingleton<ILogStorage>(Provider =>
 			{
-				ILogStorage Storage = new PersistentLogStorage(Provider.GetRequiredService<IStorageBackend>(), Provider.GetRequiredService<ILogger<PersistentLogStorage>>());
+				ILogStorage Storage = new PersistentLogStorage(Provider.GetService<IStorageBackend>(), Provider.GetService<ILogger<PersistentLogStorage>>());
 
 //				IDatabase? RedisDb = Provider.GetService<IDatabase>();
 //				if (RedisDb != null)
diff --git a/Engine/Source/Programs/Horde/HordeServer/Tasks/Impl/ActionTaskSource.cs b/Engine/Source/Programs/Horde/HordeServer/Tasks/Impl/ActionTaskSource.cs
--- a/Engine/Source/Programs/Horde/HordeServer/Tasks/Impl/ActionTaskSource.cs
+++ b/Engine/Source/Programs/Horde/HordeServer/Tasks/Impl/ActionTaskSource.cs
@@ -23,6 +23,7 @@ using System.Diagnostics.CodeAnalysis;
 using System.Linq;
 using System.Threading;
 using System.Threading.Tasks;
+using Microsoft.Extensions.Options;
 
 namespace HordeServer.Tasks.Impl
 {
@@ -294,6 +295,7 @@ namespace HordeServer.Tasks.Impl
 		Dictionary<ObjectId, ExecuteOperation>[] ActiveOperationBuckets = new Dictionary<ObjectId, ExecuteOperation>[NumBuckets];
 		Random Random = new Random();
 		ILogger Logger;
+		RemoteExecSettings RemoteExecSettings;
 
 		/// <inheritdoc/>
 		public MessageDescriptor Descriptor => ActionTask.Descriptor;
@@ -304,11 +306,13 @@ namespace HordeServer.Tasks.Impl
 		/// <param name="ActionCacheService"></param>
 		/// <param name="LogFileService">The log file service instance</param>
 		/// <param name="Logger">The logger instance</param>
-		public ActionTaskSource(ActionCacheService ActionCacheService, ILogFileService LogFileService, ILogger<ActionTaskSource> Logger)
+		/// <param name="Settings">For getting the remote exec settings</param>
+		public ActionTaskSource(ActionCacheService ActionCacheService, ILogFileService LogFileService, ILogger<ActionTaskSource> Logger, IOptionsMonitor<ServerSettings> Settings)
 		{
 			this.ActionCacheService = ActionCacheService;
 			this.LogFileService = LogFileService;
 			this.Logger = Logger;
+			this.RemoteExecSettings = Settings.CurrentValue.RemoteExecSettings;
 
 			for (int Idx = 0; Idx < NumBuckets; Idx++)
 			{
@@ -408,9 +412,28 @@ namespace HordeServer.Tasks.Impl
 				ILogFile LogFile = await LogFileService.CreateLogFileAsync(ObjectId.Empty, Subscription.Agent.SessionId, LogType.Json);
 
 				ActionTask ActionTask = new ActionTask();
+				ActionTask.InstanceName = Request.InstanceName;
 				ActionTask.Digest = Request.ActionDigest;
 				ActionTask.LogId = LogFile.Id.ToString();
 
+				if (Request.InstanceName != null && RemoteExecSettings.Instances.TryGetValue(Request.InstanceName, out var InstanceSettings))
+				{
+					if (InstanceSettings.CasUrl != null)
+					{
+						ActionTask.CasUrl = InstanceSettings.CasUrl.ToString();
+					}
+					
+					if (InstanceSettings.ActionCacheUrl != null)
+					{
+						ActionTask.ActionCacheUrl = InstanceSettings.ActionCacheUrl.ToString();
+					}
+					
+					if (InstanceSettings.ServiceAccountToken != null)
+					{
+						ActionTask.ServiceAccountToken = InstanceSettings.ServiceAccountToken;
+					}
+				}
+
 				byte[] Payload = Any.Pack(ActionTask).ToByteArray();
 				AgentLease Lease = new AgentLease(Operation.Id, LeaseName, null, null, LogFile.Id, LeaseState.Pending, Payload, new AgentRequirements(), null);
 
diff --git a/Engine/Source/Programs/Horde/HordeServer/Tasks/Impl/JobTaskSource.cs b/Engine/Source/Programs/Horde/HordeServer/Tasks/Impl/JobTaskSource.cs
--- a/Engine/Source/Programs/Horde/HordeServer/Tasks/Impl/JobTaskSource.cs
+++ b/Engine/Source/Programs/Horde/HordeServer/Tasks/Impl/JobTaskSource.cs
@@ -892,7 +892,7 @@ namespace HordeServer.Tasks.Impl
 		/// <returns>Async task</returns>
 		public async Task UpdateUgsBadges(IJob Job, IGraph Graph, IReadOnlyList<(LabelState, LabelOutcome)> OldLabelStates, IReadOnlyList<(LabelState, LabelOutcome)> NewLabelStates)
 		{
-			if (!Job.ShowUgsBadges)
+			if (!Job.ShowUgsBadges || Job.PreflightChange != 0)
 			{
 				return;
 			}
diff --git a/Engine/Source/Programs/Horde/HordeServer/Utilities/ContentHashSerializer.cs b/Engine/Source/Programs/Horde/HordeServer/Utilities/ContentHashSerializer.cs
--- a/Engine/Source/Programs/Horde/HordeServer/Utilities/ContentHashSerializer.cs
+++ b/Engine/Source/Programs/Horde/HordeServer/Utilities/ContentHashSerializer.cs
@@ -48,7 +48,7 @@ namespace HordeServer.Utilities
 		/// <inheritdoc/>
 		public void Serialize(BsonSerializationContext Context, BsonSerializationArgs Args, ContentHash Value)
 		{
-			Context.Writer.WriteString(((ContentHash)Value).ToString());
+			Context.Writer.WriteString(Value.ToString());
 		}
 	}
 }
diff --git a/Engine/Source/Programs/Horde/HordeServer/Utilities/MongoExtensions.cs b/Engine/Source/Programs/Horde/HordeServer/Utilities/MongoExtensions.cs
--- a/Engine/Source/Programs/Horde/HordeServer/Utilities/MongoExtensions.cs
+++ b/Engine/Source/Programs/Horde/HordeServer/Utilities/MongoExtensions.cs
@@ -25,7 +25,7 @@ namespace HordeServer.Utilities
 		/// <returns>New query</returns>
 		public static IFindFluent<TDocument, TProjection> Range<TDocument, TProjection>(this IFindFluent<TDocument, TProjection> Query, int? Index, int? Count)
 		{
-			if (Index != null)
+			if (Index != null && Index.Value != 0)
 			{
 				Query = Query.Skip(Index.Value);
 			}
@@ -36,6 +36,26 @@ namespace HordeServer.Utilities
 			return Query;
 		}
 
+		/// <summary>
+		/// Filters the documents returned from a search
+		/// </summary>
+		/// <param name="Query">The query to filter</param>
+		/// <param name="Index">Index of the first document to return</param>
+		/// <param name="Count">Number of documents to return</param>
+		/// <returns>New query</returns>
+		public static IAggregateFluent<TDocument> Range<TDocument>(this IAggregateFluent<TDocument> Query, int? Index, int? Count)
+		{
+			if (Index != null && Index.Value != 0)
+			{
+				Query = Query.Skip(Index.Value);
+			}
+			if (Count != null)
+			{
+				Query = Query.Limit(Count.Value);
+			}
+			return Query;
+		}
+
 		/// <summary>
 		/// Filters the documents returned from a search
 		/// </summary>
@@ -84,6 +104,37 @@ namespace HordeServer.Utilities
 			}
 		}
 
+		/// <summary>
+		/// Attempts to insert a document into a collection, handling the error case that a document with the given key already exists
+		/// </summary>
+		/// <typeparam name="TDocument"></typeparam>
+		/// <param name="Collection">Collection to insert into</param>
+		/// <param name="NewDocuments">The document to insert</param>
+		public static async Task InsertManyIgnoreDuplicatesAsync<TDocument>(this IMongoCollection<TDocument> Collection, List<TDocument> NewDocuments)
+		{
+			try
+			{
+				if (NewDocuments.Count > 0)
+				{
+					await Collection.InsertManyAsync(NewDocuments, new InsertManyOptions { IsOrdered = false });
+				}
+			}
+			catch (MongoWriteException Ex)
+			{
+				if (Ex.WriteError.Category != ServerErrorCategory.DuplicateKey)
+				{
+					throw;
+				}
+			}
+			catch (MongoBulkWriteException Ex)
+			{
+				if(Ex.WriteErrors.Any(x => x.Category != ServerErrorCategory.DuplicateKey))
+				{
+					throw;
+				}
+			}
+		}
+
 		/// <summary>
 		/// Sets a field to a value, or unsets it if the value is null
 		/// </summary>
diff --git a/Engine/Source/Programs/Horde/HordeServer/Utilities/StringId.cs b/Engine/Source/Programs/Horde/HordeServer/Utilities/StringId.cs
--- a/Engine/Source/Programs/Horde/HordeServer/Utilities/StringId.cs
+++ b/Engine/Source/Programs/Horde/HordeServer/Utilities/StringId.cs
@@ -23,6 +23,12 @@ namespace HordeServer.Utilities
 	[TypeConverter(typeof(StringIdTypeConverter))]
 	public struct StringId<T> : IEquatable<StringId<T>>
 	{
+		/// <summary>
+		/// Empty string
+		/// </summary>
+		[SuppressMessage("Design", "CA1000:Do not declare static members on generic types", Justification = "<Pending>")]
+		public static StringId<T> Empty { get; } = new StringId<T>(String.Empty);
+
 		/// <summary>
 		/// The text representing this id
 		/// </summary>
diff --git a/Engine/Source/Programs/Horde/HordeServerTests/DebugTest.cs b/Engine/Source/Programs/Horde/HordeServerTests/DebugTest.cs
--- a/Engine/Source/Programs/Horde/HordeServerTests/DebugTest.cs
+++ b/Engine/Source/Programs/Horde/HordeServerTests/DebugTest.cs
@@ -76,7 +76,7 @@ namespace HordeServerTests
 	        TestSetup TestSetup = await GetTestSetup();
 
 	        IPool Pool = await TestSetup.PoolCollection.AddAsync(new PoolId("ue4-dev-win"), "ue4-dev-win", null, null);
-	        AwsFleetManager Manager = new AwsFleetManager(TestSetup.AgentCollection, TestSetup.ServiceProvider.GetRequiredService<ILogger<AwsFleetManager>>());
+	        AwsFleetManager Manager = new AwsFleetManager(TestSetup.AgentCollection, TestSetup.ServiceProvider.GetService<ILogger<AwsFleetManager>>());
 
 	        await Manager.ExpandPool(Pool, new List<IAgent>(), 1);
         }
@@ -88,7 +88,7 @@ namespace HordeServerTests
 	        TestSetup TestSetup = await GetTestSetupExternal();
 
 	        AutoscaleService AutoscaleService = new AutoscaleService(GetDatabaseServiceExternal(), TestSetup.AgentCollection, TestSetup.PoolCollection,
-		        TestSetup.LeaseCollection, new DefaultFleetManager(TestSetup.ServiceProvider.GetRequiredService<ILogger<DefaultFleetManager>>()),
+		        TestSetup.LeaseCollection, new DefaultFleetManager(TestSetup.ServiceProvider.GetService<ILogger<DefaultFleetManager>>()),
 				new NoOpDogStatsd(), TestSetup.ServiceProvider.GetRequiredService<ILogger<AutoscaleService>>());
 
 	        await AutoscaleService.TickSharedOnlyForTestingAsync();
diff --git a/Engine/Source/Programs/Horde/HordeServerTests/IssueServiceTests.cs b/Engine/Source/Programs/Horde/HordeServerTests/IssueServiceTests.cs
--- a/Engine/Source/Programs/Horde/HordeServerTests/IssueServiceTests.cs
+++ b/Engine/Source/Programs/Horde/HordeServerTests/IssueServiceTests.cs
@@ -106,6 +106,10 @@ namespace HordeServerTests
 		ILogFileService LogFileService => TestSetup.LogFileService;
 		IIssueService IssueService => TestSetup.IssueService;
 
+		ObjectId TimId;
+		ObjectId JerryId;
+		ObjectId BobId;
+
 		DirectoryReference WorkspaceDir;
 
 		public IssueServiceTests()
@@ -147,6 +151,10 @@ namespace HordeServerTests
 			Mock<IGraph> GraphMock = new Mock<IGraph>(MockBehavior.Strict);
 			GraphMock.SetupGet(x => x.Groups).Returns(new List<INodeGroup> { Group.Object });
 			Graph = GraphMock.Object;
+
+			TimId = TestSetup.UserCollection.FindOrAddUserByLoginAsync("Tim").Result.Id;
+			JerryId = TestSetup.UserCollection.FindOrAddUserByLoginAsync("Jerry").Result.Id;
+			BobId = TestSetup.UserCollection.FindOrAddUserByLoginAsync("Bob").Result.Id;
 		}
 
 		public static INode MockNode(string Name)
@@ -156,7 +164,7 @@ namespace HordeServerTests
 			return Node.Object;
 		}
 
-		public IJob CreateJob(StreamId StreamId, int Change, string Name, IGraph Graph)
+		public IJob CreateJob(StreamId StreamId, int Change, string Name, IGraph Graph, TimeSpan Time = default)
 		{
 			ObjectId JobId = ObjectId.GenerateNewId();
 
@@ -176,7 +184,7 @@ namespace HordeServerTests
 					Step.SetupGet(x => x.Id).Returns(StepId);
 					Step.SetupGet(x => x.NodeIdx).Returns(NodeIdx);
 					Step.SetupGet(x => x.LogId).Returns(LogFile.Id);
-					Step.SetupGet(x => x.StartTimeUtc).Returns(DateTime.UtcNow);
+					Step.SetupGet(x => x.StartTimeUtc).Returns(DateTime.UtcNow + Time);
 
 					Steps.Add(Step.Object);
 				}
@@ -211,7 +219,7 @@ namespace HordeServerTests
 
 			JobStepRefId JobStepRefId = new JobStepRefId(Job.Id, Batch.Id, Step.Id);
 			string NodeName = Graph.Groups[Batch.GroupIdx].Nodes[Step.NodeIdx].Name;
-			await TestSetup.JobStepRefCollection.InsertOrReplaceAsync(JobStepRefId, "TestJob", NodeName, Job.StreamId, Job.TemplateId, Job.Change, Step.LogId, null, null, Outcome, 0.0f, 0.0f, DateTime.MinValue, null);
+			await TestSetup.JobStepRefCollection.InsertOrReplaceAsync(JobStepRefId, "TestJob", NodeName, Job.StreamId, Job.TemplateId, Job.Change, Step.LogId, null, null, Outcome, 0.0f, 0.0f, Step.StartTimeUtc!.Value, Step.StartTimeUtc);
 		}
 
 		async Task AddEvent(IJob Job, int BatchIdx, int StepIdx, object Data, EventSeverity Severity = EventSeverity.Error)
@@ -277,12 +285,12 @@ namespace HordeServerTests
 				await AddEvent(Job, 0, 2, new { });
 				await UpdateCompleteStep(Job, 0, 2, JobStepOutcome.Failure);
 
-				List<IIssue> Issues = await IssueService.FindIssuesAsync();
+				List<IIssue> Issues = (await IssueService.FindIssuesAsync()).OrderBy(x => x.Summary).ToList();
 				Assert.AreEqual(3, Issues.Count);
 
-				Assert.AreEqual("Warnings in Update Version Files", Issues[0].Summary);
+				Assert.AreEqual("Errors in Compile ShooterGameEditor Win64", Issues[0].Summary);
 				Assert.AreEqual("Errors in Compile UnrealHeaderTool Win64", Issues[1].Summary);
-				Assert.AreEqual("Errors in Compile ShooterGameEditor Win64", Issues[2].Summary);
+				Assert.AreEqual("Warnings in Update Version Files", Issues[2].Summary);
 			}
 
 			// #3
@@ -293,12 +301,12 @@ namespace HordeServerTests
 				await AddEvent(Job, 0, 0, new { });
 				await UpdateCompleteStep(Job, 0, 0, JobStepOutcome.Failure);
 
-				List<IIssue> Issues = await IssueService.FindIssuesAsync();
+				List<IIssue> Issues = (await IssueService.FindIssuesAsync()).OrderBy(x => x.Summary).ToList();
 				Assert.AreEqual(3, Issues.Count);
 
-				Assert.AreEqual("Errors in Update Version Files", Issues[0].Summary);
+				Assert.AreEqual("Errors in Compile ShooterGameEditor Win64", Issues[0].Summary);
 				Assert.AreEqual("Errors in Compile UnrealHeaderTool Win64", Issues[1].Summary);
-				Assert.AreEqual("Errors in Compile ShooterGameEditor Win64", Issues[2].Summary);
+				Assert.AreEqual("Errors in Update Version Files", Issues[2].Summary);
 			}
 
 			// #4
@@ -438,19 +446,19 @@ namespace HordeServerTests
 				Assert.AreEqual(105, Stream.LastSuccess?.Change);
 				Assert.AreEqual(null, Stream.NextSuccess?.Change);
 
-				IReadOnlyList<IIssueSuspect> Suspects = Issue.Suspects;
+				IReadOnlyList<IIssueSuspect> Suspects = await IssueService.GetIssueSuspectsAsync(Issue);
 				Assert.AreEqual(3, Suspects.Count);
 
 				Assert.AreEqual(120, Suspects[0].Change);
-				Assert.AreEqual("Tim", Suspects[0].Author);
+				Assert.AreEqual(TimId, Suspects[0].AuthorId);
 				//				Assert.AreEqual(null, Suspects[0].OriginatingChange);
 
 				Assert.AreEqual(75 /*115*/, Suspects[1].Change);
-				Assert.AreEqual("Jerry", Suspects[1].Author);
+				Assert.AreEqual(JerryId, Suspects[1].AuthorId);
 				//				Assert.AreEqual(75, Suspects[1].OriginatingChange);
 
 				Assert.AreEqual(110, Suspects[2].Change);
-				Assert.AreEqual("Bob", Suspects[2].Author);
+				Assert.AreEqual(BobId, Suspects[2].AuthorId);
 				//			Assert.AreEqual(null, Suspects[2].OriginatingChange);
 
 				List<IIssue> OpenIssues = await IssueService.FindIssuesAsync(Resolved: false);
@@ -475,14 +483,14 @@ namespace HordeServerTests
 				Assert.AreEqual(110, Stream.LastSuccess?.Change);
 				Assert.AreEqual(null, Stream.NextSuccess?.Change);
 
-				IReadOnlyList<IIssueSuspect> Suspects = Issue.Suspects;
+				IReadOnlyList<IIssueSuspect> Suspects = (await IssueService.GetIssueSuspectsAsync(Issue)).OrderByDescending(x => x.Change).ToList();
 				Assert.AreEqual(2, Suspects.Count);
 
 				Assert.AreEqual(120, Suspects[0].Change);
-				Assert.AreEqual("Tim", Suspects[0].Author);
+				Assert.AreEqual(TimId, Suspects[0].AuthorId);
 
 				Assert.AreEqual(75, Suspects[1].Change);
-				Assert.AreEqual("Jerry", Suspects[1].Author);
+				Assert.AreEqual(JerryId, Suspects[1].AuthorId);
 
 				List<IIssue> OpenIssues = await IssueService.FindIssuesAsync(Resolved: false);
 				Assert.AreEqual(1, OpenIssues.Count);
@@ -495,7 +503,7 @@ namespace HordeServerTests
 				IJob Job = CreateJob(MainStreamId, 125, "Test Build", Graph);
 				await UpdateCompleteStep(Job, 0, 0, JobStepOutcome.Success);
 
-				List<IIssue> Issues = await IssueService.FindIssuesAsync();
+				List<IIssue> Issues = await IssueService.FindIssuesAsync(Resolved: true);
 				Assert.AreEqual(1, Issues.Count);
 
 				IIssue Issue = Issues[0];
@@ -523,7 +531,7 @@ namespace HordeServerTests
 				await ParseEventsAsync(Job, 0, 0, Lines);
 				await UpdateCompleteStep(Job, 0, 0, JobStepOutcome.Failure);
 
-				List<IIssue> Issues = await IssueService.FindIssuesAsync();
+				List<IIssue> Issues = await IssueService.FindIssuesAsync(Resolved: true);
 				Assert.AreEqual(1, Issues.Count);
 
 				IIssue Issue = Issues[0];
@@ -546,8 +554,11 @@ namespace HordeServerTests
 				await AddEvent(Job, 0, 1, new { });
 				await UpdateCompleteStep(Job, 0, 1, JobStepOutcome.Failure);
 
-				List<IIssue> Issues = await IssueService.FindIssuesAsync();
-				Assert.AreEqual(2, Issues.Count);
+				List<IIssue> ResolvedIssues = await IssueService.FindIssuesAsync(Resolved: true);
+				Assert.AreEqual(1, ResolvedIssues.Count);
+
+				List<IIssue> UnresolvedIssues = await IssueService.FindIssuesAsync(Resolved: false);
+				Assert.AreEqual(1, UnresolvedIssues.Count);
 			}
 		}
 
@@ -585,10 +596,12 @@ namespace HordeServerTests
 				List<IIssue> Issues = await IssueService.FindIssuesAsync();
 				Assert.AreEqual(1, Issues.Count);
 
-				List<string> PrimarySuspects = Issues[0].Suspects.Select(x => x.Author).OrderBy(x => x).ToList();
+				List<IIssueSuspect> Suspects = await IssueService.GetIssueSuspectsAsync(Issues[0]);
+
+				List<ObjectId> PrimarySuspects = Suspects.Select(x => x.AuthorId).ToList();
 				Assert.AreEqual(2, PrimarySuspects.Count);
-				Assert.AreEqual("Jerry", PrimarySuspects[0]); // 115
-				Assert.AreEqual("Tim", PrimarySuspects[1]); // 120
+				Assert.IsTrue(PrimarySuspects.Contains(JerryId)); // 115
+				Assert.IsTrue(PrimarySuspects.Contains(TimId)); // 120
 			}
 
 			// #3
@@ -602,11 +615,11 @@ namespace HordeServerTests
 				Assert.AreEqual(1, Issues.Count);
 
 				IIssue Issue = Issues[0];
-				Assert.AreEqual("Tim", Issue.Owner);
+				Assert.AreEqual(TimId, Issue.OwnerId);
 
 				// Also check updating an issue doesn't clear the owner
 				Assert.IsTrue(await IssueService.UpdateIssueAsync(Issue.Id));
-				Assert.AreEqual("Tim", Issue!.Owner);
+				Assert.AreEqual(TimId, Issue!.OwnerId);
 			}
 		}
 
@@ -687,9 +700,11 @@ namespace HordeServerTests
 				List<IIssue> Issues = await IssueService.FindIssuesAsync();
 				Assert.AreEqual(1, Issues.Count);
 
-				List<string> PrimarySuspects = Issues[0].Suspects.Select(x => x.Author).OrderBy(x => x).ToList();
+				List<IIssueSuspect> Suspects = await IssueService.GetIssueSuspectsAsync(Issues[0]);
+
+				List<ObjectId> PrimarySuspects = Suspects.Select(x => x.AuthorId).ToList();
 				Assert.AreEqual(1, PrimarySuspects.Count);
-				Assert.AreEqual("Jerry", PrimarySuspects[0]); // 115
+				Assert.AreEqual(JerryId, PrimarySuspects[0]); // 115
 			}
 		}
 
@@ -727,10 +742,12 @@ namespace HordeServerTests
 				List<IIssue> Issues = await IssueService.FindIssuesAsync();
 				Assert.AreEqual(1, Issues.Count);
 
-				List<string> PrimarySuspects = Issues[0].Suspects.Select(x => x.Author).OrderBy(x => x).ToList();
+				List<IIssueSuspect> Suspects = await IssueService.GetIssueSuspectsAsync(Issues[0]);
+
+				List<ObjectId> PrimarySuspects = Suspects.Select(x => x.AuthorId).ToList();
 				Assert.AreEqual(2, PrimarySuspects.Count);
-				Assert.AreEqual("Jerry", PrimarySuspects[0]); // 115
-				Assert.AreEqual("Tim", PrimarySuspects[1]); // 120
+				Assert.IsTrue(PrimarySuspects.Contains(JerryId)); // 115
+				Assert.IsTrue(PrimarySuspects.Contains(TimId)); // 120
 			}
 
 			// #3
@@ -739,14 +756,16 @@ namespace HordeServerTests
 			{
 				List<IIssue> Issues = await IssueService.FindIssuesAsync();
 				Assert.AreEqual(1, Issues.Count);
-				await IssueService.UpdateIssueAsync(Issues[0].Id, DeclinedBy: "Tim");
+				await IssueService.UpdateIssueAsync(Issues[0].Id, DeclinedById: TimId);
 
 				Issues = await IssueService.FindIssuesAsync();
 				Assert.AreEqual(1, Issues.Count);
 
-				List<string> PrimarySuspects = Issues[0].Suspects.Where(x => x.DeclinedAt == null).Select(x => x.Author).OrderBy(x => x).ToList();
+				List<IIssueSuspect> Suspects = await IssueService.GetIssueSuspectsAsync(Issues[0]);
+
+				List<ObjectId> PrimarySuspects = Suspects.Where(x => x.DeclinedAt == null).Select(x => x.AuthorId).ToList();
 				Assert.AreEqual(1, PrimarySuspects.Count);
-				Assert.AreEqual("Jerry", PrimarySuspects[0]); // 115
+				Assert.AreEqual(JerryId, PrimarySuspects[0]); // 115
 			}
 		}
 
@@ -789,9 +808,11 @@ namespace HordeServerTests
 				Assert.AreEqual(105, Stream.LastSuccess?.Change);
 				Assert.AreEqual(null, Stream.NextSuccess?.Change);
 
-				List<string> PrimarySuspects = Issue.Suspects.Select(x => x.Author).ToList();
+				List<IIssueSuspect> Suspects = await IssueService.GetIssueSuspectsAsync(Issues[0]);
+
+				List<ObjectId> PrimarySuspects = Suspects.Select(x => x.AuthorId).ToList();
 				Assert.AreEqual(1, PrimarySuspects.Count);
-				Assert.AreEqual("Jerry", PrimarySuspects[0]); // 115 = foo.cpp
+				Assert.AreEqual(JerryId, PrimarySuspects[0]); // 115 = foo.cpp
 			}
 		}
 
@@ -1028,6 +1049,153 @@ namespace HordeServerTests
 			}
 		}
 
+		[TestMethod]
+		public async Task FixFailedTest()
+		{
+			int IssueId;
+
+			// #1
+			// Scenario: Warning in first step
+			// Expected: Default issue is created
+			{
+				IJob Job = CreateJob(MainStreamId, 105, "Test Build", Graph);
+				await AddEvent(Job, 0, 0, new { level = nameof(LogLevel.Warning) }, EventSeverity.Warning);
+				await UpdateCompleteStep(Job, 0, 0, JobStepOutcome.Warnings);
+
+				List<IIssue> Issues = await IssueService.FindIssuesAsync();
+				Assert.AreEqual(1, Issues.Count);
+				Assert.AreEqual(IssueSeverity.Warning, Issues[0].Severity);
+
+				Assert.AreEqual("Warnings in Update Version Files", Issues[0].Summary);
+
+				IssueId = Issues[0].Id;
+			}
+
+			// #2
+			// Scenario: Issue is marked fixed
+			// Expected: Resolved time, owner is set
+			{
+				await IssueService.UpdateIssueAsync(IssueId, ResolvedById: BobId);
+
+				List<IIssue> OpenIssues = await IssueService.FindIssuesAsync();
+				Assert.AreEqual(0, OpenIssues.Count);
+
+				IIssue Issue = (await IssueService.GetIssueAsync(IssueId))!;
+				Assert.IsNotNull(Issue.ResolvedAt);
+				Assert.AreEqual(BobId, Issue.OwnerId);
+				Assert.AreEqual(BobId, Issue.ResolvedById);
+			}
+
+			// #3
+			// Scenario: Issue recurs an hour later
+			// Expected: Issue is still marked as resolved
+			{
+				IJob Job = CreateJob(MainStreamId, 110, "Test Build", Graph, TimeSpan.FromHours(1.0));
+				await AddEvent(Job, 0, 0, new { level = nameof(LogLevel.Warning) }, EventSeverity.Warning);
+				await UpdateCompleteStep(Job, 0, 0, JobStepOutcome.Warnings);
+
+				List<IIssue> OpenIssues = await IssueService.FindIssuesAsync();
+				Assert.AreEqual(0, OpenIssues.Count);
+
+				IIssue Issue = (await IssueService.GetIssueAsync(IssueId))!;
+				Assert.IsNotNull(Issue.ResolvedAt);
+				Assert.AreEqual(BobId, Issue.OwnerId);
+				Assert.AreEqual(BobId, Issue.ResolvedById);
+			}
+
+			// #4
+			// Scenario: Issue recurs a day later at the same change
+			// Expected: Issue is reopened
+			{
+				IJob Job = CreateJob(MainStreamId, 110, "Test Build", Graph, TimeSpan.FromHours(25.0));
+				await AddEvent(Job, 0, 0, new { level = nameof(LogLevel.Warning) }, EventSeverity.Warning);
+				await UpdateCompleteStep(Job, 0, 0, JobStepOutcome.Warnings);
+
+				List<IIssue> OpenIssues = await IssueService.FindIssuesAsync();
+				Assert.AreEqual(1, OpenIssues.Count);
+
+				IIssue Issue = OpenIssues[0];
+				Assert.AreEqual(IssueId, Issue.Id);
+				Assert.IsNull(Issue.ResolvedAt);
+				Assert.AreEqual(BobId, Issue.OwnerId);
+				Assert.IsNull(Issue.ResolvedById);
+			}
+
+			// #5
+			// Scenario: Issue is marked fixed again, at a particular changelist
+			// Expected: Resolved time, owner is set
+			{
+				await IssueService.UpdateIssueAsync(IssueId, ResolvedById: BobId, FixChange: 115);
+
+				List<IIssue> OpenIssues = await IssueService.FindIssuesAsync();
+				Assert.AreEqual(0, OpenIssues.Count);
+
+				IIssue Issue = (await IssueService.GetIssueAsync(IssueId))!;
+				Assert.IsNotNull(Issue.ResolvedAt);
+				Assert.AreEqual(BobId, Issue.OwnerId);
+				Assert.AreEqual(BobId, Issue.ResolvedById);
+			}
+
+			// #6
+			// Scenario: Issue fails again at a later changelist
+			// Expected: Issue is reopened
+			{
+				IJob Job = CreateJob(MainStreamId, 120, "Test Build", Graph, TimeSpan.FromHours(25.0));
+				await AddEvent(Job, 0, 0, new { level = nameof(LogLevel.Warning) }, EventSeverity.Warning);
+				await UpdateCompleteStep(Job, 0, 0, JobStepOutcome.Warnings);
+
+				List<IIssue> OpenIssues = await IssueService.FindIssuesAsync();
+				Assert.AreEqual(1, OpenIssues.Count);
+
+				IIssue Issue = OpenIssues[0];
+				Assert.AreEqual(IssueId, Issue.Id);
+				Assert.IsNull(Issue.ResolvedAt);
+				Assert.AreEqual(BobId, Issue.OwnerId);
+				Assert.IsNull(Issue.ResolvedById);
+			}
+
+			// #7
+			// Scenario: Issue is marked fixed again, at a particular changelist
+			// Expected: Resolved time, owner is set
+			{
+				await IssueService.UpdateIssueAsync(IssueId, ResolvedById: BobId, FixChange: 125);
+
+				List<IIssue> OpenIssues = await IssueService.FindIssuesAsync();
+				Assert.AreEqual(0, OpenIssues.Count);
+
+				IIssue Issue = (await IssueService.GetIssueAsync(IssueId))!;
+				Assert.IsNotNull(Issue.ResolvedAt);
+				Assert.AreEqual(BobId, Issue.OwnerId);
+				Assert.AreEqual(BobId, Issue.ResolvedById);
+			}
+
+			// #8
+			// Scenario: Issue succeeds at a later changelist
+			// Expected: Issue remains closed
+			{
+				IJob Job = CreateJob(MainStreamId, 125, "Test Build", Graph, TimeSpan.FromHours(25.0));
+				await UpdateCompleteStep(Job, 0, 0, JobStepOutcome.Success);
+
+				List<IIssue> OpenIssues = await IssueService.FindIssuesAsync();
+				Assert.AreEqual(0, OpenIssues.Count);
+			}
+
+			// #9
+			// Scenario: Issue fails at a later changelist
+			// Expected: New issue is opened
+			{
+				IJob Job = CreateJob(MainStreamId, 130, "Test Build", Graph, TimeSpan.FromHours(25.0));
+				await AddEvent(Job, 0, 0, new { level = nameof(LogLevel.Warning) }, EventSeverity.Warning);
+				await UpdateCompleteStep(Job, 0, 0, JobStepOutcome.Warnings);
+
+				List<IIssue> OpenIssues = await IssueService.FindIssuesAsync();
+				Assert.AreEqual(1, OpenIssues.Count);
+
+				IIssue Issue = OpenIssues[0];
+				Assert.AreNotEqual(IssueId, Issue.Id);
+			}
+		}
+
 		private async Task ParseAsync(ObjectId LogId, string[] Lines)
 		{
 			LogParserContext Context = new LogParserContext();
diff --git a/Engine/Source/Programs/Horde/HordeServerTests/LogIndexingTests.cs b/Engine/Source/Programs/Horde/HordeServerTests/LogIndexingTests.cs
--- a/Engine/Source/Programs/Horde/HordeServerTests/LogIndexingTests.cs
+++ b/Engine/Source/Programs/Horde/HordeServerTests/LogIndexingTests.cs
@@ -39,7 +39,7 @@ namespace HordeServerTests
 				.AddLogging(Builder => Builder.AddConsole().SetMinimumLevel(LogLevel.Debug))
 				.BuildServiceProvider();
 
-			ILoggerFactory LoggerFactory = ServiceProvider.GetRequiredService<ILoggerFactory>();
+			ILoggerFactory LoggerFactory = ServiceProvider.GetService<ILoggerFactory>();
 			ILogger<LogFileService> Logger = LoggerFactory.CreateLogger<LogFileService>();
 
 			// Just to satisfy the parameter, need to be fixed
diff --git a/Engine/Source/Programs/Horde/HordeServerTests/SchedulerTests.cs b/Engine/Source/Programs/Horde/HordeServerTests/SchedulerTests.cs
--- a/Engine/Source/Programs/Horde/HordeServerTests/SchedulerTests.cs
+++ b/Engine/Source/Programs/Horde/HordeServerTests/SchedulerTests.cs
@@ -383,7 +383,7 @@ namespace HordeServerTests
 		[TestMethod]
 		public async Task StreamPausing()
 		{
-			DateTime StartTime = new DateTime(2021, 1, 1, 12, 0, 0, DateTimeKind.Local); // Friday Jan 1, 2021 
+			DateTime StartTime = new DateTime(2021, 1, 1, 12, 0, 0, DateTimeKind.Utc); // Friday Jan 1, 2021 
 			TestSetup.Clock.UtcNow = StartTime;
 			
 			Schedule Schedule = new Schedule();
diff --git a/Engine/Source/Programs/Horde/HordeServerTests/ServiceAccountAuthTest.cs b/Engine/Source/Programs/Horde/HordeServerTests/ServiceAccountAuthTest.cs
--- a/Engine/Source/Programs/Horde/HordeServerTests/ServiceAccountAuthTest.cs
+++ b/Engine/Source/Programs/Horde/HordeServerTests/ServiceAccountAuthTest.cs
@@ -29,7 +29,7 @@ namespace HordeServerTests
 		{
 			DatabaseService DatabaseService = GetDatabaseService();
 			ServiceAccounts = new ServiceAccountCollection(DatabaseService);
-			ServiceAccount = ServiceAccounts.AddAsync("mytoken", new Dictionary<string, string> {{"myclaim", "myvalue"}, {"foo", "bar"}}, "mydesc").Result;
+			ServiceAccount = ServiceAccounts.AddAsync("mytoken", new List<string> {"myclaim###myvalue", "foo###bar"}, "mydesc").Result;
  		}
 
 		private async Task<ServiceAccountAuthHandler> GetAuthHandlerAsync(string? HeaderValue)
diff --git a/Engine/Source/Programs/Horde/HordeServerTests/ServiceAccountCollectionTests.cs b/Engine/Source/Programs/Horde/HordeServerTests/ServiceAccountCollectionTests.cs
--- a/Engine/Source/Programs/Horde/HordeServerTests/ServiceAccountCollectionTests.cs
+++ b/Engine/Source/Programs/Horde/HordeServerTests/ServiceAccountCollectionTests.cs
@@ -21,16 +21,16 @@ namespace HordeServerTests
 		{
 			DatabaseService DatabaseService = GetDatabaseService();
 			ServiceAccounts = new ServiceAccountCollection(DatabaseService);
-			ServiceAccount = ServiceAccounts.AddAsync(ObjectId.GenerateNewId().ToString(), new Dictionary<string, string> {{"myclaim", "myvalue"}}, "mydesc").Result;
+			ServiceAccount = ServiceAccounts.AddAsync(ObjectId.GenerateNewId().ToString(), new List<string> {"myclaim###myvalue"}, "mydesc").Result;
 		}
 
 		[TestMethod]
 		public async Task Add()
 		{
-			IServiceAccount Sa = await ServiceAccounts.AddAsync("addToken", new Dictionary<string, string> {{"myclaim", "myvalue"}}, "mydesc");
+			IServiceAccount Sa = await ServiceAccounts.AddAsync("addToken", new List<string> {"myclaim###myvalue"}, "mydesc");
 			Assert.AreEqual("addToken", Sa.SecretToken);
-			Assert.AreEqual(1, Sa.Claims.Count);
-			Assert.AreEqual("myvalue", Sa.Claims["myclaim"]);
+			Assert.AreEqual(1, Sa.GetClaims().Count);
+			Assert.AreEqual("myvalue", Sa.GetClaims()[0].Value);
 			Assert.IsTrue(Sa.Enabled);
 			Assert.AreEqual("mydesc", Sa.Description);
 		}
@@ -52,14 +52,14 @@ namespace HordeServerTests
 		[TestMethod]
 		public async Task Update()
 		{
-			Dictionary<string, string> NewClaims = new Dictionary<string, string> {{"newclaim1", "newvalue1"}, {"newclaim2", "newvalue2"}};
+			List<string> NewClaims = new List<string> {{"newclaim1###newvalue1"}, {"newclaim2###newvalue2"}};
 			await ServiceAccounts.UpdateAsync(ServiceAccount.Id, "newtoken", NewClaims, false, "newdesc");
 			IServiceAccount GetSa = (await ServiceAccounts.GetAsync(ServiceAccount.Id))!;
 			
 			Assert.AreEqual("newtoken", GetSa.SecretToken);
-			Assert.AreEqual(2, GetSa.Claims.Count);
-			Assert.AreEqual("newvalue1", GetSa.Claims["newclaim1"]);
-			Assert.AreEqual("newvalue2", GetSa.Claims["newclaim2"]);
+			Assert.AreEqual(2, GetSa.GetClaims().Count);
+			Assert.AreEqual("newvalue1", GetSa.GetClaims()[0].Value);
+			Assert.AreEqual("newvalue2", GetSa.GetClaims()[1].Value);
 			Assert.AreEqual(false, GetSa.Enabled);
 			Assert.AreEqual("newdesc", GetSa.Description);
 		}
