Add copyright

diff --git a/Engine/Source/Programs/Horde/HordeServer/Services/PerforceServiceApi.cs b/Engine/Source/Programs/Horde/HordeServer/Services/PerforceServiceApi.cs
--- a/Engine/Source/Programs/Horde/HordeServer/Services/PerforceServiceApi.cs
+++ b/Engine/Source/Programs/Horde/HordeServer/Services/PerforceServiceApi.cs
@@ -1,3 +1,5 @@
+// Copyright Epic Games, Inc. All Rights Reserved.
+
 using HordeServer.Models;
 using System.Linq;
 using Microsoft.Extensions.Logging;
