Upgrade NuGet package dependency.

diff --git a/Engine/Source/Programs/Horde/HordeAgentTests/HordeAgentTests.csproj b/Engine/Source/Programs/Horde/HordeAgentTests/HordeAgentTests.csproj
--- a/Engine/Source/Programs/Horde/HordeAgentTests/HordeAgentTests.csproj
+++ b/Engine/Source/Programs/Horde/HordeAgentTests/HordeAgentTests.csproj
@@ -1,4 +1,4 @@
-﻿<Project Sdk="Microsoft.NET.Sdk">
+<Project Sdk="Microsoft.NET.Sdk">
 
   <PropertyGroup>
     <TargetFramework>netcoreapp3.1</TargetFramework>
@@ -7,7 +7,8 @@
   </PropertyGroup>
 
   <ItemGroup>
-    <PackageReference Include="Microsoft.Extensions.Logging.Abstractions" Version="3.1.6" />
+
+    <PackageReference Include="Microsoft.Extensions.Logging.Abstractions" Version="5.0.0" />
     <PackageReference Include="Microsoft.NET.Test.Sdk" Version="16.5.0" />
     <PackageReference Include="Moq" Version="4.14.5" />
     <PackageReference Include="MSTest.TestAdapter" Version="2.1.0" />
