Fix staging of appsettings for HordeAgent.

diff --git a/Engine/Source/Programs/Horde/HordeAgent/HordeAgent.csproj b/Engine/Source/Programs/Horde/HordeAgent/HordeAgent.csproj
--- a/Engine/Source/Programs/Horde/HordeAgent/HordeAgent.csproj
+++ b/Engine/Source/Programs/Horde/HordeAgent/HordeAgent.csproj
@@ -18,6 +18,24 @@
     <None Remove="HordeAgent.xml" />
   </ItemGroup>
 
+  <ItemGroup>
+    <Content Include="appsettings.json" Condition="Exists('appsettings.json')">
+      <Generator></Generator>
+      <CopyToOutputDirectory>PreserveNewest</CopyToOutputDirectory>
+      <CopyToPublishDirectory>PreserveNewest</CopyToPublishDirectory>
+    </Content>
+    <Content Include="appsettings.Development.json" Condition="Exists('appsettings.Development.json')">
+      <Generator></Generator>
+      <CopyToOutputDirectory>PreserveNewest</CopyToOutputDirectory>
+      <CopyToPublishDirectory>PreserveNewest</CopyToPublishDirectory>
+    </Content>
+    <Content Include="appsettings.Production.json" Condition="Exists('appsettings.Production.json')">
+      <Generator></Generator>
+      <CopyToOutputDirectory>PreserveNewest</CopyToOutputDirectory>
+      <CopyToPublishDirectory>PreserveNewest</CopyToPublishDirectory>
+    </Content>
+  </ItemGroup>
+
   <ItemGroup>
     <PackageReference Include="AWSSDK.Core" Version="3.5.1.49" />
     <PackageReference Include="AWSSDK.EC2" Version="3.5.1.2" />
@@ -84,7 +102,7 @@
     <Protobuf Include="..\HordeCommon\Protos\google\rpc\code.proto" ProtoRoot="..\HordeCommon\Protos" GrpcServices="Client" />
     <Protobuf Include="..\HordeCommon\Protos\horde\action_rpc.proto" ProtoRoot="..\HordeCommon\Protos" GrpcServices="Client" />
     <Protobuf Include="..\HordeCommon\Protos\horde\horde_rpc.proto" ProtoRoot="..\HordeCommon\Protos" GrpcServices="Client" />
-    <Protobuf Include="..\HordeCommon\Protos\horde\tasks\action_task.proto" ProtoRoot="..\HordeCommon\Protos" GrpcServices="Client"/>
+    <Protobuf Include="..\HordeCommon\Protos\horde\tasks\action_task.proto" ProtoRoot="..\HordeCommon\Protos" GrpcServices="Client" />
   </ItemGroup>
 
   <ItemGroup>
