HordeAgent: Use the application directory when running the Publish command.

diff --git a/Engine/Source/Programs/Horde/HordeAgent/Commands/Utilities/PublishCommand.cs b/Engine/Source/Programs/Horde/HordeAgent/Commands/Utilities/PublishCommand.cs
--- a/Engine/Source/Programs/Horde/HordeAgent/Commands/Utilities/PublishCommand.cs
+++ b/Engine/Source/Programs/Horde/HordeAgent/Commands/Utilities/PublishCommand.cs
@@ -85,7 +85,7 @@ namespace HordeAgent.Commands.Utilities
 
 			// Read the application settings
 			IConfigurationRoot Config = new ConfigurationBuilder()
-				.SetBasePath(Directory.GetCurrentDirectory())
+				.SetBasePath(Program.AppDir)
 				.AddJsonFile("appsettings.json", true)
 				.AddJsonFile("appsettings.Production.json", true)
 				.Build();
