Fix path to P4API.NET library.

diff --git a/Engine/Source/Programs/Horde/HordeServer/HordeServer.csproj b/Engine/Source/Programs/Horde/HordeServer/HordeServer.csproj
--- a/Engine/Source/Programs/Horde/HordeServer/HordeServer.csproj
+++ b/Engine/Source/Programs/Horde/HordeServer/HordeServer.csproj
@@ -1,4 +1,4 @@
-<Project Sdk="Microsoft.NET.Sdk.Web">
+﻿<Project Sdk="Microsoft.NET.Sdk.Web">
 
   <PropertyGroup>
     <TargetFramework>netcoreapp3.1</TargetFramework>
@@ -73,7 +73,7 @@
     <ProjectReference Include="..\..\Shared\EpicGames.Core\EpicGames.Core.csproj" />
     <ProjectReference Include="..\..\Shared\EpicGames.Perforce\EpicGames.Perforce.csproj" />
     <ProjectReference Include="..\HordeCommon\HordeCommon.csproj" />
-    <ProjectReference Include="..\ThirdParty\P4Api.Net\P4Api.Net.csproj" />
+    <ProjectReference Include="..\..\..\ThirdParty\Perforce\P4Api.Net\P4Api.Net\P4Api.Net.csproj" />
   </ItemGroup>
 
   <ItemGroup>
