Fix HordeServer compile errors.

diff --git a/Engine/Source/Programs/Horde/HordeServer/Program.cs b/Engine/Source/Programs/Horde/HordeServer/Program.cs
--- a/Engine/Source/Programs/Horde/HordeServer/Program.cs
+++ b/Engine/Source/Programs/Horde/HordeServer/Program.cs
@@ -26,6 +26,8 @@ using Serilog.Events;
 using Serilog.Formatting.Json;
 using Serilog.Sinks.SystemConsole.Themes;
 
+using Log = Serilog.Log;
+
 namespace HordeServer
 {
 	static class LoggerExtensions
diff --git a/Engine/Source/Programs/Horde/HordeServer/Services/PerforceService.cs b/Engine/Source/Programs/Horde/HordeServer/Services/PerforceService.cs
--- a/Engine/Source/Programs/Horde/HordeServer/Services/PerforceService.cs
+++ b/Engine/Source/Programs/Horde/HordeServer/Services/PerforceService.cs
@@ -443,13 +443,16 @@ namespace HordeServer.Services
 		/// <inheritdoc/>
 		public async override Task<PerforceUserInfo?> GetUserInfoAsync(string UserName)
 		{
-			PerforceResponse<UserRecord> UserResponse = await Perforce.TryGetUserAsync(UserName, CancellationToken.None);
-			if (UserResponse.Succeeded && UserResponse.Data.Access != default)
+			PerforceResponseList<UserRecord> UserResponses = await Perforce.TryGetUserAsync(UserName, CancellationToken.None);
+			foreach(PerforceResponse<UserRecord> UserResponse in UserResponses)
 			{
-				PerforceUserInfo UserInfo = new PerforceUserInfo();
-				UserInfo.Name = UserResponse.Data.FullName;
-				UserInfo.Email = UserResponse.Data.Email;
-				return UserInfo;
+				if (UserResponse.Succeeded && UserResponse.Data.Access != default)
+				{
+					PerforceUserInfo UserInfo = new PerforceUserInfo();
+					UserInfo.Name = UserResponse.Data.FullName;
+					UserInfo.Email = UserResponse.Data.Email;
+					return UserInfo;
+				}
 			}
 			return null;
 		}
diff --git a/Engine/Source/Programs/Horde/HordeServer/Services/UpgradeService.cs b/Engine/Source/Programs/Horde/HordeServer/Services/UpgradeService.cs
--- a/Engine/Source/Programs/Horde/HordeServer/Services/UpgradeService.cs
+++ b/Engine/Source/Programs/Horde/HordeServer/Services/UpgradeService.cs
@@ -81,7 +81,7 @@ namespace HordeServer.Services
 
 				if (SchemaVersion == 3)
 				{
-					IJobCollection JobCollection = ServiceProvider.GetService<IJobCollection>();
+					IJobCollection JobCollection = ServiceProvider.GetRequiredService<IJobCollection>();
 					await JobCollection.UpgradeDocumentsAsync();
 				}
 				if (SchemaVersion == 4)
diff --git a/Engine/Source/Programs/Horde/HordeServer/Startup.cs b/Engine/Source/Programs/Horde/HordeServer/Startup.cs
--- a/Engine/Source/Programs/Horde/HordeServer/Startup.cs
+++ b/Engine/Source/Programs/Horde/HordeServer/Startup.cs
@@ -372,7 +372,7 @@ namespace HordeServer
 			if (!String.IsNullOrEmpty(Settings.PerforceBridge))
 			{
 				// Python P4 bridge service
-				Services.AddSingleton<IPerforceService, BridgePerforceService>(SP => new BridgePerforceService(Settings.PerforceBridge!, SP.GetService<ILogger<BridgePerforceService>>()));				
+				Services.AddSingleton<IPerforceService, BridgePerforceService>(SP => new BridgePerforceService(Settings.PerforceBridge!, SP.GetRequiredService<ILogger<BridgePerforceService>>()));				
 			}
 			else if (!String.IsNullOrEmpty(Settings.P4BridgeServer))
 			{
@@ -593,12 +593,12 @@ namespace HordeServer
 				}
 				else
 				{
-					return new RedisLogBuilder(RedisConnectionPool, Provider.GetService<ILogger<RedisLogBuilder>>());
+					return new RedisLogBuilder(RedisConnectionPool, Provider.GetRequiredService<ILogger<RedisLogBuilder>>());
 				}
 			});
 			Services.AddSingleton<ILogStorage>(Provider =>
 			{
-				ILogStorage Storage = new PersistentLogStorage(Provider.GetService<IStorageBackend>(), Provider.GetService<ILogger<PersistentLogStorage>>());
+				ILogStorage Storage = new PersistentLogStorage(Provider.GetRequiredService<IStorageBackend>(), Provider.GetRequiredService<ILogger<PersistentLogStorage>>());
 
 //				IDatabase? RedisDb = Provider.GetService<IDatabase>();
 //				if (RedisDb != null)
