Fix case of path to p4api.net project for building on Linux.

diff --git a/Engine/Source/Programs/Horde/HordeServer/HordeServer.csproj b/Engine/Source/Programs/Horde/HordeServer/HordeServer.csproj
--- a/Engine/Source/Programs/Horde/HordeServer/HordeServer.csproj
+++ b/Engine/Source/Programs/Horde/HordeServer/HordeServer.csproj
@@ -73,7 +73,7 @@
     <ProjectReference Include="..\..\Shared\EpicGames.Core\EpicGames.Core.csproj" />
     <ProjectReference Include="..\..\Shared\EpicGames.Perforce\EpicGames.Perforce.csproj" />
     <ProjectReference Include="..\HordeCommon\HordeCommon.csproj" />
-    <ProjectReference Include="..\..\..\ThirdParty\Perforce\P4Api.Net\P4Api.Net\P4Api.Net.csproj" />
+    <ProjectReference Include="..\..\..\ThirdParty\Perforce\P4Api.Net\p4api.net\p4api.net.csproj" />
   </ItemGroup>
 
   <ItemGroup>
