Horde: Merging latest agent/server code from Git.

diff --git a/Engine/Source/Programs/Horde/HordeAgent/AgentSettings.cs b/Engine/Source/Programs/Horde/HordeAgent/AgentSettings.cs
--- a/Engine/Source/Programs/Horde/HordeAgent/AgentSettings.cs
+++ b/Engine/Source/Programs/Horde/HordeAgent/AgentSettings.cs
@@ -205,7 +205,14 @@ namespace HordeAgent
 			ServerProfile? ServerProfile = ServerProfiles.FirstOrDefault(x => x.Name.Equals(Name, StringComparison.OrdinalIgnoreCase));
 			if (ServerProfile == null)
 			{
-				throw new Exception($"Unknown server profile name '{Server}'");
+				if (ServerProfiles.Count == 0)
+				{
+					throw new Exception("No server profiles are defined (missing configuration?)");
+				}
+				else
+				{
+					throw new Exception($"Unknown server profile name '{Name}' (valid profiles: {String.Join("/", ServerProfiles.Select(x => x.Name))})");
+				}
 			}
 			return ServerProfile;
 		}
diff --git a/Engine/Source/Programs/Horde/HordeAgent/Services/WorkerService.cs b/Engine/Source/Programs/Horde/HordeAgent/Services/WorkerService.cs
--- a/Engine/Source/Programs/Horde/HordeAgent/Services/WorkerService.cs
+++ b/Engine/Source/Programs/Horde/HordeAgent/Services/WorkerService.cs
@@ -1493,8 +1493,17 @@ namespace HordeAgent.Services
 			PrimaryDevice.Properties.Add($"BootTime={BootTime}");
 			PrimaryDevice.Properties.Add($"StartTime={StartTime}");
 
-			// Add disk info
-			string? DriveName = Path.GetPathRoot(WorkingDir.FullName);
+			// Add disk info based on platform
+			string? DriveName;
+			if (RuntimeInformation.IsOSPlatform(OSPlatform.Windows))
+			{
+				DriveName = Path.GetPathRoot(WorkingDir.FullName);
+			}
+			else
+			{
+				DriveName = WorkingDir.FullName;
+			}
+
 			if (DriveName != null)
 			{
 				try
diff --git a/Engine/Source/Programs/Horde/HordeServer/Authentication/AnonymousHandler.cs b/Engine/Source/Programs/Horde/HordeServer/Authentication/AnonymousHandler.cs
--- a/Engine/Source/Programs/Horde/HordeServer/Authentication/AnonymousHandler.cs
+++ b/Engine/Source/Programs/Horde/HordeServer/Authentication/AnonymousHandler.cs
@@ -1,5 +1,6 @@
 // Copyright Epic Games, Inc. All Rights Reserved.
 
+using HordeServer.Services;
 using Microsoft.AspNetCore.Authentication;
 using Microsoft.Extensions.Logging;
 using Microsoft.Extensions.Options;
@@ -32,9 +33,9 @@ namespace HordeServer.Authentication
 		{
 			Claim[] Claims =
 			{
-					new Claim(ClaimTypes.Name, AuthenticationScheme),
-					new Claim(Options.AdminClaimType, Options.AdminClaimValue)
-				};
+				new Claim(ClaimTypes.Name, AuthenticationScheme),
+				new Claim(Options.AdminClaimType, Options.AdminClaimValue)
+			};
 			ClaimsIdentity Identity = new ClaimsIdentity(Claims, Scheme.Name);
 
 			ClaimsPrincipal Principal = new ClaimsPrincipal(Identity);
diff --git a/Engine/Source/Programs/Horde/HordeServer/Collections/IJobCollection.cs b/Engine/Source/Programs/Horde/HordeServer/Collections/IJobCollection.cs
--- a/Engine/Source/Programs/Horde/HordeServer/Collections/IJobCollection.cs
+++ b/Engine/Source/Programs/Horde/HordeServer/Collections/IJobCollection.cs
@@ -162,6 +162,13 @@ namespace HordeServer.Collections
 		/// <returns>True if the groups were updated to the given list. False if another write happened first.</returns>
 		Task<bool> TryUpdateGraphAsync(IJob Job, IGraph NewGraph);
 
+		/// <summary>
+		/// Removes a job from the dispatch queue. Ignores the state of any batches still remaining to execute. Should only be used to correct for inconsistent state.
+		/// </summary>
+		/// <param name="Job">Job</param>
+		/// <returns></returns>
+		Task<bool> TryRemoveFromDispatchQueueAsync(IJob Job);
+
 		/// <summary>
 		/// Adds an issue to a job
 		/// </summary>
diff --git a/Engine/Source/Programs/Horde/HordeServer/Collections/Impl/GraphCollection.cs b/Engine/Source/Programs/Horde/HordeServer/Collections/Impl/GraphCollection.cs
--- a/Engine/Source/Programs/Horde/HordeServer/Collections/Impl/GraphCollection.cs
+++ b/Engine/Source/Programs/Horde/HordeServer/Collections/Impl/GraphCollection.cs
@@ -175,7 +175,7 @@ namespace HordeServer.Collections.Impl
 			public static GraphDocument Empty { get; } = new GraphDocument(new List<NodeGroup>(), new List<Aggregate>(), new List<Label>());
 
 			[BsonRequired, BsonId]
-			public ContentHash Id { get; private set; }
+			public ContentHash Id { get; private set; } = ContentHash.Empty;
 
 			public int Schema { get; set; }
 			public List<NodeGroup> Groups { get; private set; } = new List<NodeGroup>();
@@ -192,7 +192,6 @@ namespace HordeServer.Collections.Impl
 			[BsonConstructor]
 			private GraphDocument()
 			{
-				Id = ContentHash.Empty;
 			}
 
 			public GraphDocument(List<NodeGroup> Groups, List<Aggregate> Aggregates, List<Label> Labels)
@@ -406,10 +405,10 @@ namespace HordeServer.Collections.Impl
 		}
 
 		/// <inheritdoc/>
-		public async Task<IGraph> GetAsync(ContentHash Hash)
+		public async Task<IGraph> GetAsync(ContentHash? Hash)
 		{
 			// Special case for an empty graph request
-			if (Hash == ContentHash.Empty || Hash == GraphDocument.Empty.Id)
+			if (Hash == null || Hash == ContentHash.Empty || Hash == GraphDocument.Empty.Id)
 			{
 				return GraphDocument.Empty;
 			}
diff --git a/Engine/Source/Programs/Horde/HordeServer/Collections/Impl/JobCollection.cs b/Engine/Source/Programs/Horde/HordeServer/Collections/Impl/JobCollection.cs
--- a/Engine/Source/Programs/Horde/HordeServer/Collections/Impl/JobCollection.cs
+++ b/Engine/Source/Programs/Horde/HordeServer/Collections/Impl/JobCollection.cs
@@ -842,6 +842,7 @@ namespace HordeServer.Collections.Impl
 			if (bRefreshDependentJobSteps)
 			{
 				RefreshDependentJobSteps(JobDocument, Graph, Updates, Logger);
+				RefreshJobPriority(JobDocument, Updates);
 			}
 
 			// Update the new list of job steps
@@ -867,6 +868,22 @@ namespace HordeServer.Collections.Impl
 			return false;
 		}
 
+		/// <inheritdoc/>
+		public Task<bool> TryRemoveFromDispatchQueueAsync(IJob Job)
+		{
+			JobDocument JobDocument = (JobDocument)Job;
+
+			// Create the update 
+			UpdateDefinitionBuilder<JobDocument> UpdateBuilder = Builders<JobDocument>.Update;
+			List<UpdateDefinition<JobDocument>> Updates = new List<UpdateDefinition<JobDocument>>();
+
+			JobDocument.SchedulePriority = 0;
+			Updates.Add(UpdateBuilder.Set(x => x.SchedulePriority, JobDocument.SchedulePriority));
+
+			// Update the new list of job steps
+			return TryUpdateAsync(JobDocument, UpdateBuilder.Combine(Updates));
+		}
+
 		/// <inheritdoc/>
 		public Task<bool> TryUpdateGraphAsync(IJob Job, IGraph NewGraph)
 		{
@@ -896,7 +913,7 @@ namespace HordeServer.Collections.Impl
 		/// <inheritdoc/>
 		public async Task<List<IJob>> GetDispatchQueueAsync()
 		{
-			List<JobDocument> NewJobs = await Jobs.Find(x => x.SchedulePriority > 0).SortByDescending(x => x.SchedulePriority).ToListAsync();
+			List<JobDocument> NewJobs = await Jobs.Find(x => x.SchedulePriority > 0).SortByDescending(x => x.SchedulePriority).ThenBy(x => x.CreateTimeUtc).ToListAsync();
 			foreach (JobDocument Result in NewJobs)
 			{
 				Result.PostLoad();
@@ -1088,6 +1105,7 @@ namespace HordeServer.Collections.Impl
 			}
 
 			RefreshDependentJobSteps(JobDocument, Graph, Updates, Logger);
+			RefreshJobPriority(JobDocument, Updates);
 
 			return (Updates.Count == 0 || await TryUpdateAsync(JobDocument, UpdateBuilder.Combine(Updates)));
 		}
@@ -1121,6 +1139,7 @@ namespace HordeServer.Collections.Impl
 
 			// Update all the dependencies
 			RefreshDependentJobSteps(Job, Graph, new List<UpdateDefinition<JobDocument>>(), Logger);
+			RefreshJobPriority(Job, Updates);
 		}
 
 		/// <summary>
@@ -1503,13 +1522,21 @@ namespace HordeServer.Collections.Impl
 					}
 				}
 			}
+		}
 
+		/// <summary>
+		/// Updates the schedule priority of a job
+		/// </summary>
+		/// <param name="Job"></param>
+		/// <param name="Updates"></param>
+		static void RefreshJobPriority(JobDocument Job, List<UpdateDefinition<JobDocument>> Updates)
+		{
 			// Update the weighted priority for the job
 			int NewSchedulePriority = GetSchedulePriority(Job);
 			if (Job.SchedulePriority != NewSchedulePriority)
 			{
 				Job.SchedulePriority = NewSchedulePriority;
-				Updates.Add(UpdateBuilder.Set(x => x.SchedulePriority, NewSchedulePriority));
+				Updates.Add(Builders<JobDocument>.Update.Set(x => x.SchedulePriority, NewSchedulePriority));
 			}
 		}
 
diff --git a/Engine/Source/Programs/Horde/HordeServer/Controllers/DebugController.cs b/Engine/Source/Programs/Horde/HordeServer/Controllers/DebugController.cs
--- a/Engine/Source/Programs/Horde/HordeServer/Controllers/DebugController.cs
+++ b/Engine/Source/Programs/Horde/HordeServer/Controllers/DebugController.cs
@@ -611,7 +611,7 @@ namespace HordeServer.Controllers
 
 			return await JobTaskSource.GetStatus();
 		}
-		
+
 		/// <summary>
 		/// Returns the complete config Horde uses
 		/// </summary>
@@ -629,7 +629,7 @@ namespace HordeServer.Controllers
 			{
 				WriteIndented = true
 			};
-			
+
 			return Ok(JsonSerializer.Serialize(Settings.CurrentValue, Options));
 		}
 
@@ -803,6 +803,63 @@ namespace HordeServer.Controllers
 			return await Perforce.CreateTicket(Username);
 		}
 
+		/// <summary>
+		/// Debugs perforce service commands
+		/// </summary>
+		/// <returns>Perforce ticket</returns>
+		[HttpGet]
+		[Route("/api/v1/debug/perforce/stress")]
+		public async Task<ActionResult<string>> GetPerforceStress()
+		{
+			if (!await AclService.AuthorizeAsync(AclAction.AdminWrite, User))
+			{
+				return Forbid();
+			}
+
+			Dictionary<string, string> Results = new Dictionary<string, string>();
+			const string DevBuild = "//UE4/Dev-Build";
+			string? PerforceUser = User.GetPerforceUser();
+
+			object UpdateLock = new object();
+
+			// ParallelTest.ForAsync doesn't seem to handle async lambdas
+			IEnumerable<Task> Tasks = Enumerable.Range(0, 16).Select(async (Idx) => {
+
+				int LatestChange = await Perforce.GetLatestChangeAsync(DevBuild, PerforceUser);
+				int CodeChange = await Perforce.GetCodeChangeAsync(DevBuild, LatestChange);				
+
+				lock (UpdateLock)
+				{
+					Results.Add($"Parallel Perforce Result {Idx}", $"LatestChange: {LatestChange} - CodeChange: {CodeChange}");
+				}
+			});
+
+			await Task.WhenAll(Tasks);			
+
+			Tasks = Enumerable.Range(0, 16).Select(async (Idx) => {
+
+				int LatestChange = await Perforce.GetLatestChangeAsync(DevBuild, PerforceUser);
+				int CodeChange = await Perforce.GetCodeChangeAsync(DevBuild, LatestChange);				
+
+				lock (UpdateLock)
+				{
+					Results.Add($"Parallel Perforce Result {Idx + 16}", $"LatestChange: {LatestChange} - CodeChange: {CodeChange}");
+				}
+			});
+
+
+			await Task.WhenAll(Tasks);
+
+			StringBuilder Content = new StringBuilder();
+			Content.AppendLine("<html><body><pre>");
+			foreach (KeyValuePair<string, string> Pair in Results)
+			{
+				Content.AppendLine(HttpUtility.HtmlEncode($"{Pair.Key}={Pair.Value}"));
+			}
+			Content.Append("</pre></body></html>");
+			return new ContentResult { ContentType = "text/html", StatusCode = (int)HttpStatusCode.OK, Content = Content.ToString() };
+
+		}
 
 		/// <summary>
 		/// Debugs perforce service commands
@@ -810,7 +867,7 @@ namespace HordeServer.Controllers
 		/// <returns>Perforce ticket</returns>
 		[HttpGet]
 		[Route("/api/v1/debug/perforce/commands")]
-		public async Task<ActionResult<string>> Get()
+		public async Task<ActionResult<string>> GetPerforceCommands()
 		{
 			if (!await AclService.AuthorizeAsync(AclAction.AdminWrite, User))
 			{
@@ -832,23 +889,23 @@ namespace HordeServer.Controllers
 
 			int NewChangeId = await Perforce.CreateNewChangeAsync("//UE4/Private-Build", "Counter.txt");
 
-			List<ChangeDetails> ChangeDetails = await Perforce.GetChangeDetailsAsync("//UE4/Private-Build", new int[]{NewChangeId}, PerforceUser);
+			List<ChangeDetails> ChangeDetails = await Perforce.GetChangeDetailsAsync("//UE4/Private-Build", new int[] { NewChangeId }, PerforceUser);
 
-			Results.Add("CreateNewChangeAsync", $"{NewChangeId} : {ChangeDetails[0].Description}"); 
+			Results.Add("CreateNewChangeAsync", $"{NewChangeId} : {ChangeDetails[0].Description}");
 
-			await Perforce.UpdateChangelistDescription(TestCL, $"Updated Description: New Change CL was {NewChangeId}");			
+			await Perforce.UpdateChangelistDescription(TestCL, $"Updated Description: New Change CL was {NewChangeId}");
 
-			ChangeDetails = await Perforce.GetChangeDetailsAsync(DevBuild, new int[]{TestCL}, PerforceUser);
+			ChangeDetails = await Perforce.GetChangeDetailsAsync(DevBuild, new int[] { TestCL }, PerforceUser);
 
-			Results.Add("UpdateChangelistDescription", $"{TestCL} : {ChangeDetails[0].Description}"); 
+			Results.Add("UpdateChangelistDescription", $"{TestCL} : {ChangeDetails[0].Description}");
 
 			int LatestChange = await Perforce.GetLatestChangeAsync(DevBuild, PerforceUser);
 			Results.Add("GetLatestChangeAsync", LatestChange.ToString(CultureInfo.InvariantCulture));
 
 			int CodeChange = await Perforce.GetCodeChangeAsync(DevBuild, LatestChange);
 			Results.Add("GetCodeChangeAsync", CodeChange.ToString(CultureInfo.InvariantCulture));
-			
-			PerforceUserInfo? UserInfo =  await Perforce.GetUserInfoAsync(PerforceUser);
+
+			PerforceUserInfo? UserInfo = await Perforce.GetUserInfoAsync(PerforceUser);
 
 			if (UserInfo == null)
 			{
@@ -867,20 +924,20 @@ namespace HordeServer.Controllers
 				Results.Add($"GetChangesAsync Result {Count++}", $"{Summary.Number} : {Summary.Author} - {Summary.Description}");
 			}
 
-			ChangeDetails = await Perforce.GetChangeDetailsAsync(DevBuild, Changes.Select(Change => Change.Number).ToArray(), PerforceUser);			
+			ChangeDetails = await Perforce.GetChangeDetailsAsync(DevBuild, Changes.Select(Change => Change.Number).ToArray(), PerforceUser);
 			Count = 0;
 			foreach (ChangeDetails Details in ChangeDetails)
 			{
 				Results.Add($"GetChangeDetailsAsync Result {Count++}", $"{Details.Number} : {Details.Author} - {Details.Files.ToJson()}");
 			}
-			
+
 			String Ticket = "Failed (Expected if not running service account)";
-			try 
+			try
 			{
 				await Perforce.CreateTicket(PerforceUser);
 				Ticket = "Success";
 			}
-			catch 
+			catch
 			{
 
 			}
@@ -896,10 +953,10 @@ namespace HordeServer.Controllers
 			// print async
 
 			byte[] Bytes = await Perforce.PrintAsync($"{DevBuild}/RunUAT.bat");
-			Results.Add($"PrintAsync: {DevBuild}/RunUAT.bat ",  $"{System.Text.Encoding.Default.GetString(Bytes)}");
+			Results.Add($"PrintAsync: {DevBuild}/RunUAT.bat ", $"{System.Text.Encoding.Default.GetString(Bytes)}");
 
 			Bytes = await Perforce.PrintAsync($"{DevBuild}/RunUAT.bat@13916380");
-			Results.Add($"PrintAsync: {DevBuild}/RunUAT.bat@13916380",  $"{System.Text.Encoding.Default.GetString(Bytes)}");
+			Results.Add($"PrintAsync: {DevBuild}/RunUAT.bat@13916380", $"{System.Text.Encoding.Default.GetString(Bytes)}");
 
 
 			// DuplicateShelvedChangeAsync is disabled due to issue with p4 librarian
@@ -946,7 +1003,7 @@ namespace HordeServer.Controllers
 				Results.Add("GetChangeDetailsAsync - After Delete", $"ChangeDetails.Count: {ChangeDetails.Count}"); 
 
 			}
-			*/			
+			*/
 			StringBuilder Content = new StringBuilder();
 			Content.AppendLine("<html><body><pre>");
 			foreach (KeyValuePair<string, string> Pair in Results)
diff --git a/Engine/Source/Programs/Horde/HordeServer/HordeServer.csproj b/Engine/Source/Programs/Horde/HordeServer/HordeServer.csproj
--- a/Engine/Source/Programs/Horde/HordeServer/HordeServer.csproj
+++ b/Engine/Source/Programs/Horde/HordeServer/HordeServer.csproj
@@ -1,4 +1,4 @@
-﻿<Project Sdk="Microsoft.NET.Sdk.Web">
+<Project Sdk="Microsoft.NET.Sdk.Web">
 
   <PropertyGroup>
     <TargetFramework>netcoreapp3.1</TargetFramework>
diff --git a/Engine/Source/Programs/Horde/HordeServer/Models/Agent.cs b/Engine/Source/Programs/Horde/HordeServer/Models/Agent.cs
--- a/Engine/Source/Programs/Horde/HordeServer/Models/Agent.cs
+++ b/Engine/Source/Programs/Horde/HordeServer/Models/Agent.cs
@@ -619,7 +619,7 @@ namespace HordeServer.Models
 					}
 				}
 
-				if(Requirements.Devices != null)
+				if (Requirements.Devices != null)
 				{
 					if (!TryCreateDeviceLeases(Capabilities.Devices, Requirements.Devices, 0, 0, Leases.SelectMany(x => x.Devices), out Result))
 					{
@@ -651,7 +651,7 @@ namespace HordeServer.Models
 			if (RequiredDeviceIdx == RequiredDevices.Count)
 			{
 				LeasedDevices = new List<AgentLeaseDevice>(RequiredDevices.Count);
-				foreach(DeviceRequirements Device in RequiredDevices)
+				foreach (DeviceRequirements Device in RequiredDevices)
 				{
 					LeasedDevices.Add(new AgentLeaseDevice(-1, Device.Handle, Device.Resources));
 				}
@@ -786,13 +786,19 @@ namespace HordeServer.Models
 				}
 
 				IPerforceServer? Server = await LoadBalancer.SelectServerAsync(ClusterInfo, Agent);
-				if(Server == null)
+				if (Server == null)
 				{
 					return false;
 				}
 				ServerAndPort = Server.ServerAndPort;
 			}
 
+			// HACK: Trying to debug issues with inconsistent workspace state. Possibly related to edge servers?
+			if (Agent.Id == new AgentId("10-99-114-249") || Agent.Id == new AgentId("10-99-115-36"))
+			{
+				ServerAndPort = "perforce:1666";
+			}
+
 			// Find the matching credentials for the desired user
 			PerforceCredentials? Credentials = null;
 			if (ClusterInfo != null)
diff --git a/Engine/Source/Programs/Horde/HordeServer/Program.cs b/Engine/Source/Programs/Horde/HordeServer/Program.cs
--- a/Engine/Source/Programs/Horde/HordeServer/Program.cs
+++ b/Engine/Source/Programs/Horde/HordeServer/Program.cs
@@ -26,8 +26,6 @@ using Serilog.Events;
 using Serilog.Formatting.Json;
 using Serilog.Sinks.SystemConsole.Themes;
 
-using Log = Serilog.Log;
-
 namespace HordeServer
 {
 	static class LoggerExtensions
@@ -71,7 +69,7 @@ namespace HordeServer
 			ServerSettings HordeSettings = new ServerSettings();
 			Config.GetSection("Horde").Bind(HordeSettings);
 
-			Log.Logger = new LoggerConfiguration()
+			Serilog.Log.Logger = new LoggerConfiguration()
 				.MinimumLevel.Debug()
 //				.MinimumLevel.Override("HordeServer.Services.DatabaseService", LogEventLevel.Verbose) // For MongoDB query tracing
 				.MinimumLevel.Override("Microsoft", LogEventLevel.Warning)
@@ -79,6 +77,7 @@ namespace HordeServer
 				.MinimumLevel.Override("HordeServer.Authentication.OktaHandler", LogEventLevel.Warning)
 				.MinimumLevel.Override("HordeServer.Authentication.OktaViaJwtHandler", LogEventLevel.Warning)
 				.MinimumLevel.Override("HordeServer.Authentication.HordeJwtBearerHandler", LogEventLevel.Warning)
+				.MinimumLevel.Override("HordeServer.Authentication.ServiceAccountAuthHandler", LogEventLevel.Warning)
 				.MinimumLevel.Override("System.Net.Http.HttpClient", LogEventLevel.Warning)
 				.MinimumLevel.Override("Grpc", LogEventLevel.Warning)
 				.Enrich.FromLogContext()
@@ -98,7 +97,7 @@ namespace HordeServer
 			catch (Exception Ex)
 #pragma warning restore CA1031 // Do not catch general exception types
 			{
-				Log.Logger.Error(Ex, "Unhandled exception");
+				Serilog.Log.Logger.Error(Ex, "Unhandled exception");
 			}
 		}
 
diff --git a/Engine/Source/Programs/Horde/HordeServer/ServerSettings.cs b/Engine/Source/Programs/Horde/HordeServer/ServerSettings.cs
--- a/Engine/Source/Programs/Horde/HordeServer/ServerSettings.cs
+++ b/Engine/Source/Programs/Horde/HordeServer/ServerSettings.cs
@@ -1,6 +1,7 @@
 // Copyright Epic Games, Inc. All Rights Reserved.
 
 using HordeServer.Services;
+using HordeServer.Utilities;
 using System;
 using System.Collections.Generic;
 using System.ComponentModel.DataAnnotations;
@@ -111,12 +112,12 @@ namespace HordeServer
 		/// <summary>
 		/// The claim type for administrators
 		/// </summary>
-		public string? AdminClaimType { get; set; }
+		public string AdminClaimType { get; set; } = HordeClaimTypes.InternalRole;
 
 		/// <summary>
 		/// Value of the claim type for administrators
 		/// </summary>
-		public string? AdminClaimValue { get; set; }
+		public string AdminClaimValue { get; set; } = "admin";
 
 		/// <summary>
 		/// Optional certificate to trust in order to access the database (eg. AWS public cert for TLS)
diff --git a/Engine/Source/Programs/Horde/HordeServer/Services/Impl/IssueService.cs b/Engine/Source/Programs/Horde/HordeServer/Services/Impl/IssueService.cs
--- a/Engine/Source/Programs/Horde/HordeServer/Services/Impl/IssueService.cs
+++ b/Engine/Source/Programs/Horde/HordeServer/Services/Impl/IssueService.cs
@@ -959,7 +959,14 @@ namespace HordeServer.Services.Impl
 				Dictionary<StreamId, bool>? FixStreamIds = await GetFixStreamsAsync(Issue.Id, Issue.FixChange.Value, Issue.GetFixStreamIds());
 				if (FixStreamIds != null)
 				{
-					Issue = await IssueCollection.UpdateIssueAsync(Issue, NewFixStreamIds: FixStreamIds);
+					if(FixStreamIds.Count > Issue.Streams.Count)
+					{
+						Issue = await UpdateIssueDerivedDataAsync(Issue);
+					}
+					if (Issue != null)
+					{
+						Issue = await IssueCollection.UpdateIssueAsync(Issue, NewFixStreamIds: FixStreamIds);
+					}
 				}
 			}
 
@@ -1023,12 +1030,22 @@ namespace HordeServer.Services.Impl
 					bool ContainsFix;
 					if (PrevFixStreamIds == null || !PrevFixStreamIds.TryGetValue(Span.StreamId, out ContainsFix))
 					{
+						Logger.LogInformation("Querying fix changelist {FixChange} for issue {IssueId} in {StreamId} (prev={HavePrev})", FixChange, IssueId, Span.StreamId, PrevFixStreamIds != null);
 						List<ChangeSummary> Changes = await Perforce.GetChangesAsync(Span.StreamName, FixChange, FixChange, 1, null);
 						ContainsFix = Changes.Count > 0;
 					}
 					NextFixStreamIds[Span.StreamId] = ContainsFix;
 				}
 			}
+
+			if (PrevFixStreamIds != null)
+			{
+				if(PrevFixStreamIds.Count != NextFixStreamIds.Count || PrevFixStreamIds.Any(x => !NextFixStreamIds.TryGetValue(x.Key, out bool Value) || x.Value != Value))
+				{
+					return null;
+				}
+			}
+
 			return NextFixStreamIds;
 		}
 
diff --git a/Engine/Source/Programs/Horde/HordeServer/Services/PerforceServiceApi.cs b/Engine/Source/Programs/Horde/HordeServer/Services/PerforceServiceApi.cs
--- a/Engine/Source/Programs/Horde/HordeServer/Services/PerforceServiceApi.cs
+++ b/Engine/Source/Programs/Horde/HordeServer/Services/PerforceServiceApi.cs
@@ -63,7 +63,7 @@ namespace HordeServer.Services
 						Serilog.Log.Information("Perforce: {Message} {Source} : ", Message, Source);
 						break;
 					default:
-						Serilog.Log.Debug("Perforce (Debug): {Message} {Source} : ", Message, Source);
+						// Serilog.Log.Debug("Perforce (Debug): {Message} {Source} : ", Message, Source);
 						break;
 				};
 
diff --git a/Engine/Source/Programs/Horde/HordeServer/Tasks/Impl/JobTaskSource.cs b/Engine/Source/Programs/Horde/HordeServer/Tasks/Impl/JobTaskSource.cs
--- a/Engine/Source/Programs/Horde/HordeServer/Tasks/Impl/JobTaskSource.cs
+++ b/Engine/Source/Programs/Horde/HordeServer/Tasks/Impl/JobTaskSource.cs
@@ -471,6 +471,19 @@ namespace HordeServer.Tasks.Impl
 			{
 				IJob? NewJob = NewJobs[Idx];
 
+				if (NewJob.GraphHash == null)
+				{
+					Logger.LogError("Job {JobId} has a null graph hash and can't be started.", NewJob.Id);
+					await Jobs.TryRemoveFromDispatchQueueAsync(NewJob);
+					continue;
+				}
+				if (NewJob.AbortedByUser != null)
+				{
+					Logger.LogError("Job {JobId} was aborted but not removed from dispatch queue", NewJob.Id);
+					await Jobs.TryRemoveFromDispatchQueueAsync(NewJob);
+					continue;
+				}
+
 				// Get the graph for this job
 				IGraph Graph = await Graphs.GetAsync(NewJob.GraphHash);
 
@@ -483,6 +496,7 @@ namespace HordeServer.Tasks.Impl
 				}
 
 				// Update all the batches
+				bool IsRunning = false;
 				for (int BatchIdx = 0; NewJob != null && BatchIdx < NewJob.Batches.Count; BatchIdx++)
 				{
 					// Skip any batches which aren't ready to execute
@@ -520,6 +534,15 @@ namespace HordeServer.Tasks.Impl
 							NewBatchIdToQueueItem[(NewJob.Id, Batch.Id)] = NewQueueItem;
 						}
 					}
+					IsRunning |= (Batch.State == JobStepBatchState.Ready || Batch.State == JobStepBatchState.Running);
+				}
+
+				// Add a warning if a job looks to be idle
+				if (!IsRunning && NewJob != null)
+				{
+					Logger.LogError("Job {JobId} is in dispatch queue but not currently executing", NewJob.Id);
+					await Jobs.TryRemoveFromDispatchQueueAsync(NewJob);
+					continue;
 				}
 			}
 
diff --git a/Engine/Source/Programs/Horde/HordeServer/Utilities/HordeClaimTypes.cs b/Engine/Source/Programs/Horde/HordeServer/Utilities/HordeClaimTypes.cs
--- a/Engine/Source/Programs/Horde/HordeServer/Utilities/HordeClaimTypes.cs
+++ b/Engine/Source/Programs/Horde/HordeServer/Utilities/HordeClaimTypes.cs
@@ -25,6 +25,11 @@ namespace HordeServer.Utilities
 		/// </summary>
 		public const string Role = Prefix + "role";
 
+		/// <summary>
+		/// Claim for a well known internal role (eg. admin)
+		/// </summary>
+		public const string InternalRole = Prefix + "internal-role";
+
 		/// <summary>
 		/// Claim for a particular session
 		/// </summary>
diff --git a/Engine/Source/Programs/Horde/HordeServerTests/DebugTest.cs b/Engine/Source/Programs/Horde/HordeServerTests/DebugTest.cs
--- a/Engine/Source/Programs/Horde/HordeServerTests/DebugTest.cs
+++ b/Engine/Source/Programs/Horde/HordeServerTests/DebugTest.cs
@@ -76,7 +76,7 @@ namespace HordeServerTests
 	        TestSetup TestSetup = await GetTestSetup();
 
 	        IPool Pool = await TestSetup.PoolCollection.AddAsync(new PoolId("ue4-dev-win"), "ue4-dev-win", null, null);
-	        AwsFleetManager Manager = new AwsFleetManager(TestSetup.AgentCollection, TestSetup.ServiceProvider.GetService<ILogger<AwsFleetManager>>());
+	        AwsFleetManager Manager = new AwsFleetManager(TestSetup.AgentCollection, TestSetup.ServiceProvider.GetRequiredService<ILogger<AwsFleetManager>>());
 
 	        await Manager.ExpandPool(Pool, new List<IAgent>(), 1);
         }
@@ -88,7 +88,7 @@ namespace HordeServerTests
 	        TestSetup TestSetup = await GetTestSetupExternal();
 
 	        AutoscaleService AutoscaleService = new AutoscaleService(GetDatabaseServiceExternal(), TestSetup.AgentCollection, TestSetup.PoolCollection,
-		        TestSetup.LeaseCollection, new DefaultFleetManager(TestSetup.ServiceProvider.GetService<ILogger<DefaultFleetManager>>()),
+		        TestSetup.LeaseCollection, new DefaultFleetManager(TestSetup.ServiceProvider.GetRequiredService<ILogger<DefaultFleetManager>>()),
 				new NoOpDogStatsd(), TestSetup.ServiceProvider.GetRequiredService<ILogger<AutoscaleService>>());
 
 	        await AutoscaleService.TickSharedOnlyForTestingAsync();
diff --git a/Engine/Source/Programs/Horde/HordeServerTests/LogIndexingTests.cs b/Engine/Source/Programs/Horde/HordeServerTests/LogIndexingTests.cs
--- a/Engine/Source/Programs/Horde/HordeServerTests/LogIndexingTests.cs
+++ b/Engine/Source/Programs/Horde/HordeServerTests/LogIndexingTests.cs
@@ -39,7 +39,7 @@ namespace HordeServerTests
 				.AddLogging(Builder => Builder.AddConsole().SetMinimumLevel(LogLevel.Debug))
 				.BuildServiceProvider();
 
-			ILoggerFactory LoggerFactory = ServiceProvider.GetService<ILoggerFactory>();
+			ILoggerFactory LoggerFactory = ServiceProvider.GetRequiredService<ILoggerFactory>();
 			ILogger<LogFileService> Logger = LoggerFactory.CreateLogger<LogFileService>();
 
 			// Just to satisfy the parameter, need to be fixed
