Horde: Fix warnings running tests.

diff --git a/Engine/Source/Programs/Horde/HordeServer/Collections/Impl/IssueCollection.cs b/Engine/Source/Programs/Horde/HordeServer/Collections/Impl/IssueCollection.cs
--- a/Engine/Source/Programs/Horde/HordeServer/Collections/Impl/IssueCollection.cs
+++ b/Engine/Source/Programs/Horde/HordeServer/Collections/Impl/IssueCollection.cs
@@ -749,14 +749,14 @@ namespace HordeServer.Collections.Impl
 			if (Count == null)
 			{
 				List<ProjectedIssueId> IssueIds = await Cursor.ToListAsync();
-				return await NextStageFunc(IssueIds.Where(x => x._id != null).Select(x => x._id.Value), Index, null);
+				return await NextStageFunc(IssueIds.Where(x => x._id != null).Select(x => x._id!.Value), Index, null);
 			}
 			else
 			{
 				List<Issue> Results = new List<Issue>();
 				while (await Cursor.MoveNextAsync() && Results.Count < Count.Value)
 				{
-					List<Issue> NextResults = await NextStageFunc(Cursor.Current.Where(x => x._id != null).Select(x => x._id.Value), 0, Count.Value - Results.Count);
+					List<Issue> NextResults = await NextStageFunc(Cursor.Current.Where(x => x._id != null).Select(x => x._id!.Value), 0, Count.Value - Results.Count);
 					int RemoveCount = Math.Min(Index, NextResults.Count);
 					NextResults.RemoveRange(0, RemoveCount);
 					Index -= RemoveCount;
