Horde: Add a message to the account page showing that it's built from Perforce.

diff --git a/Engine/Source/Programs/Horde/HordeServer/Controllers/AccountController.cs b/Engine/Source/Programs/Horde/HordeServer/Controllers/AccountController.cs
--- a/Engine/Source/Programs/Horde/HordeServer/Controllers/AccountController.cs
+++ b/Engine/Source/Programs/Horde/HordeServer/Controllers/AccountController.cs
@@ -82,6 +82,8 @@ namespace HordeServer.Controllers
 				}
 				Content.Append("</table>");
 				Content.Append("</p>");
+
+				Content.Append($"<p>Built from Perforce</p>");
 			}
 			else
 			{
