Horde dashboard updates from git repo

diff --git a/Engine/Source/Programs/Horde/HordeDashboard/src/Plugins/UnrealAutomatedTests/models/UnrealAutomatedTests.ts b/Engine/Source/Programs/Horde/HordeDashboard/src/Plugins/UnrealAutomatedTests/models/UnrealAutomatedTests.ts
--- a/Engine/Source/Programs/Horde/HordeDashboard/src/Plugins/UnrealAutomatedTests/models/UnrealAutomatedTests.ts
+++ b/Engine/Source/Programs/Horde/HordeDashboard/src/Plugins/UnrealAutomatedTests/models/UnrealAutomatedTests.ts
@@ -3,6 +3,9 @@
 export enum TestState {
     Failed = "Fail",
     Success = "Success",
+    SuccessWithWarnings = "SuccessWithWarnings",
+    NotRun = "NotRun",
+    InProcess = "InProcess",
 }
 
 export enum EventType {
@@ -40,6 +43,8 @@ export type TestResult = {
     FullTestPath: string;
     State: string;
     ArtifactName: string;
+    Errors: number;
+    Warnings: number;
 }
 
 export type TestPassSummary = {
@@ -47,6 +52,7 @@ export type TestPassSummary = {
     ReportURL: string;
     FailedCount: number;
     NotRunCount: number;
+    InProcessCount: number;
     ReportCreatedOn: string;
     SucceededCount: number;
     SucceededWithWarningsCount: number;
diff --git a/Engine/Source/Programs/Horde/HordeDashboard/src/backend/Api.ts b/Engine/Source/Programs/Horde/HordeDashboard/src/backend/Api.ts
--- a/Engine/Source/Programs/Horde/HordeDashboard/src/backend/Api.ts
+++ b/Engine/Source/Programs/Horde/HordeDashboard/src/backend/Api.ts
@@ -281,6 +281,7 @@ export type IssueQuery = {
 	maxChange?: number;
 	index?: number;
 	count?: number;
+	user?: string;
 	resolved?: boolean;
 }
 
@@ -2354,6 +2355,9 @@ export type GetIssueStepResponse = {
 	/**The changelist number */
 	change: number;
 
+	/** Severity of the issue in this step */
+	severity: IssueSeverity;
+
 	/** Name of the job containing this step */
 	jobName: string;
 
diff --git a/Engine/Source/Programs/Horde/HordeDashboard/src/backend/Backend.ts b/Engine/Source/Programs/Horde/HordeDashboard/src/backend/Backend.ts
--- a/Engine/Source/Programs/Horde/HordeDashboard/src/backend/Backend.ts
+++ b/Engine/Source/Programs/Horde/HordeDashboard/src/backend/Backend.ts
@@ -804,6 +804,18 @@ export class Backend {
 
     }
 
+    getIssue(issueId: number): Promise<IssueData> {
+
+        return new Promise<IssueData>((resolve, reject) => {
+            this.backend.get(`/api/v1/issues/${issueId}`).then((value) => {
+                resolve(value.data as IssueData);
+            }).catch((reason) => {
+                reject(reason);
+            });
+        });
+
+    }
+
 
     getIssues(queryIn?: IssueQuery): Promise<IssueData[]> {
 
@@ -845,7 +857,7 @@ export class Backend {
 
     }
 
-    getIssueEvents(issueId: number, jobId?: string, logIds?: string[]): Promise<GetLogEventResponse[]> {
+    getIssueEvents(issueId: number, jobId?: string, logIds?: string[], count?:number): Promise<GetLogEventResponse[]> {
 
         if (logIds && logIds.length) {
             jobId = undefined;
@@ -853,7 +865,8 @@ export class Backend {
 
         const params = {
             jobId: jobId,
-            logIds: logIds
+            logIds: logIds,
+            count: count
         };
 
         return new Promise<GetLogEventResponse[]>((resolve, reject) => {
diff --git a/Engine/Source/Programs/Horde/HordeDashboard/src/backend/JobDetails.ts b/Engine/Source/Programs/Horde/HordeDashboard/src/backend/JobDetails.ts
--- a/Engine/Source/Programs/Horde/HordeDashboard/src/backend/JobDetails.ts
+++ b/Engine/Source/Programs/Horde/HordeDashboard/src/backend/JobDetails.ts
@@ -581,7 +581,7 @@ export class JobDetails {
 
     private getIssues() {
 
-        return new Promise<void>((resolve, reject) => {
+        return new Promise<void>(async (resolve, reject) => {
 
             let stepId = this.stepId;
             let labelIdx = this.labelIdx;
@@ -591,7 +591,14 @@ export class JobDetails {
             }
 
             this.issues = [];
-            backend.getIssues({ jobId: this.id, stepId: stepId, label: this.labelIdx, count: 50 }).then(async (issues) => {
+
+            const issueRequests = [ backend.getIssues({ jobId: this.id, stepId: stepId, label: this.labelIdx, count: 50, resolved: false }),
+            backend.getIssues({ jobId: this.id, stepId: stepId, label: this.labelIdx, count: 50, resolved: true })];
+            
+            Promise.all(issueRequests).then(async (issueResults) => {
+
+                let issues:IssueData[] = [];
+                issueResults.forEach( r => issues.push(...r));
 
                 let nstepId = this.stepId;
                 let nlabelIdx = this.labelIdx;
diff --git a/Engine/Source/Programs/Horde/HordeDashboard/src/base/utilities/filter.ts b/Engine/Source/Programs/Horde/HordeDashboard/src/base/utilities/filter.ts
--- a/Engine/Source/Programs/Horde/HordeDashboard/src/base/utilities/filter.ts
+++ b/Engine/Source/Programs/Horde/HordeDashboard/src/base/utilities/filter.ts
@@ -1,4 +1,6 @@
+
 import { GetJobResponse } from "../../backend/Api";
+import graphCache from "../../backend/GraphCache";
 
 
 export type JobFilterSimple = {
@@ -24,6 +26,34 @@ const getJobKeywords = (job: GetJobResponse): string[] => {
         keywords.push("Scheduler");
     }
 
+    if (job.graphHash) {
+        const graph = graphCache.cache.get(job.graphHash);
+
+        if (graph) {
+
+            const nodes: string[] = [];
+            graph.groups?.forEach(g => g.nodes.forEach(n => nodes.push(n.name)));
+
+            job.batches?.forEach(batch => {
+
+                const group = graph.groups?.[batch.groupIdx];
+
+                if (group) {
+
+                    batch.steps.forEach(step => {
+                        const name = group.nodes[step.nodeIdx]?.name;
+                        if (name && keywords.indexOf(name) === -1) {                            
+                            keywords.push(name);
+                        }
+
+                    })
+                }
+
+            });
+
+        } 
+    }
+
 
     keywords.push(...job.arguments);
 
@@ -31,7 +61,7 @@ const getJobKeywords = (job: GetJobResponse): string[] => {
 
 }
 
-export const filterJob = (job: GetJobResponse, keywordIn?: string, additionalKeywords?:string[]): boolean => {
+export const filterJob = (job: GetJobResponse, keywordIn?: string, additionalKeywords?: string[]): boolean => {
 
     const keyword = keywordIn?.toLowerCase();
 
