Fix Horde Dashboard plugin case

diff --git a/Engine/Source/Programs/Horde/HordeDashboard/src/backend/Config.ts b/Engine/Source/Programs/Horde/HordeDashboard/src/backend/Config.ts
--- a/Engine/Source/Programs/Horde/HordeDashboard/src/backend/Config.ts
+++ b/Engine/Source/Programs/Horde/HordeDashboard/src/backend/Config.ts
@@ -10,6 +10,7 @@ export type SiteConfig = {
 
 const dev = {
 	"environment": "dev",
+	"datadogClient" : "",
 	"plugins": [
 		"UnrealAutomatedTests",
 		"SimpleTestReport"
@@ -18,6 +19,7 @@ const dev = {
 
 const prod = {
 	"environment": "production",
+	"datadogClient" : "",
 	"plugins": [
 		"UnrealAutomatedTests",
 		"SimpleTestReport"
