Horde: Add an endpoint for setting/getting logos associated with a project.

diff --git a/Engine/Source/Programs/Horde/HordeServer/Collections/IProjectCollection.cs b/Engine/Source/Programs/Horde/HordeServer/Collections/IProjectCollection.cs
--- a/Engine/Source/Programs/Horde/HordeServer/Collections/IProjectCollection.cs
+++ b/Engine/Source/Programs/Horde/HordeServer/Collections/IProjectCollection.cs
@@ -60,6 +60,22 @@ namespace HordeServer.Collections
 		/// <returns>The project document</returns>
 		Task<IProject?> GetAsync(ProjectId ProjectId);
 
+		/// <summary>
+		/// Gets the logo for a project
+		/// </summary>
+		/// <param name="ProjectId">The project id</param>
+		/// <returns>The project logo document</returns>
+		Task<IProjectLogo?> GetLogoAsync(ProjectId ProjectId);
+
+		/// <summary>
+		/// Sets the logo for a project
+		/// </summary>
+		/// <param name="ProjectId">The project id</param>
+		/// <param name="MimeType"></param>
+		/// <param name="Data"></param>
+		/// <returns></returns>
+		Task SetLogoAsync(ProjectId ProjectId, string MimeType, byte[] Data);
+
 		/// <summary>
 		/// Gets a project's permissions info by ID
 		/// </summary>
diff --git a/Engine/Source/Programs/Horde/HordeServer/Collections/Impl/ProjectCollection.cs b/Engine/Source/Programs/Horde/HordeServer/Collections/Impl/ProjectCollection.cs
--- a/Engine/Source/Programs/Horde/HordeServer/Collections/Impl/ProjectCollection.cs
+++ b/Engine/Source/Programs/Horde/HordeServer/Collections/Impl/ProjectCollection.cs
@@ -56,6 +56,18 @@ namespace HordeServer.Collections.Impl
 			}
 		}
 
+		/// <summary>
+		/// Logo for a project
+		/// </summary>
+		class ProjectLogoDocument : IProjectLogo
+		{
+			public ProjectId Id { get; set; }
+
+			public string MimeType { get; set; } = String.Empty;
+
+			public byte[] Data { get; set; } = Array.Empty<byte>();
+		}
+
 		/// <summary>
 		/// Projection of a project definition to just include permissions info
 		/// </summary>
@@ -78,6 +90,11 @@ namespace HordeServer.Collections.Impl
 		/// </summary>
 		IMongoCollection<ProjectDocument> Projects;
 
+		/// <summary>
+		/// Collection of project logo documents
+		/// </summary>
+		IMongoCollection<ProjectLogoDocument> ProjectLogos;
+
 		/// <summary>
 		/// Constructor
 		/// </summary>
@@ -85,6 +102,7 @@ namespace HordeServer.Collections.Impl
 		public ProjectCollection(DatabaseService DatabaseService)
 		{
 			Projects = DatabaseService.GetCollection<ProjectDocument>("Projects");
+			ProjectLogos = DatabaseService.GetCollection<ProjectLogoDocument>("ProjectLogos");
 		}
 
 		/// <inheritdoc/>
@@ -179,6 +197,22 @@ namespace HordeServer.Collections.Impl
 			return await Projects.Find<ProjectDocument>(x => x.Id == ProjectId).FirstOrDefaultAsync();
 		}
 
+		/// <inheritdoc/>
+		public async Task<IProjectLogo?> GetLogoAsync(ProjectId ProjectId)
+		{
+			return await ProjectLogos.Find(x => x.Id == ProjectId).FirstOrDefaultAsync();
+		}
+
+		/// <inheritdoc/>
+		public async Task SetLogoAsync(ProjectId ProjectId, string MimeType, byte[] Data)
+		{
+			ProjectLogoDocument Logo = new ProjectLogoDocument();
+			Logo.Id = ProjectId;
+			Logo.MimeType = MimeType;
+			Logo.Data = Data;
+			await ProjectLogos.UpdateOneAsync(x => x.Id == ProjectId, Builders<ProjectLogoDocument>.Update.Set(x => x.MimeType, MimeType).Set(x => x.Data, Data), new UpdateOptions { IsUpsert = true });
+		}
+
 		/// <inheritdoc/>
 		public async Task<IProjectPermissions?> GetPermissionsAsync(ProjectId ProjectId)
 		{
diff --git a/Engine/Source/Programs/Horde/HordeServer/Controllers/ProjectsController.cs b/Engine/Source/Programs/Horde/HordeServer/Controllers/ProjectsController.cs
--- a/Engine/Source/Programs/Horde/HordeServer/Controllers/ProjectsController.cs
+++ b/Engine/Source/Programs/Horde/HordeServer/Controllers/ProjectsController.cs
@@ -5,11 +5,13 @@ using HordeServer.Models;
 using HordeServer.Services;
 using HordeServer.Utilities;
 using Microsoft.AspNetCore.Authorization;
+using Microsoft.AspNetCore.Http;
 using Microsoft.AspNetCore.Mvc;
 using MongoDB.Bson;
 using System;
 using System.Collections.Generic;
 using System.ComponentModel.DataAnnotations;
+using System.IO;
 using System.Linq;
 using System.Threading.Tasks;
 
@@ -246,5 +248,68 @@ namespace HordeServer.Controllers
 			await ProjectService.DeleteProjectAsync(ProjectIdValue, StreamService, JobService);
 			return new OkResult();
 		}
+
+		/// <summary>
+		/// Retrieve information about a specific project
+		/// </summary>
+		/// <param name="ProjectId">Id of the project to get information about</param>
+		/// <returns>Information about the requested project</returns>
+		[HttpGet]
+		[Route("/api/v1/projects/{ProjectId}/logo")]
+		[ProducesResponseType(typeof(List<GetProjectResponse>), 200)]
+		public async Task<ActionResult<object>> GetProjectLogoAsync(string ProjectId)
+		{
+			ProjectId ProjectIdValue = new ProjectId(ProjectId);
+
+			IProject? Project = await ProjectService.GetProjectAsync(ProjectIdValue);
+			if (Project == null)
+			{
+				return NotFound();
+			}
+			if (!await ProjectService.AuthorizeAsync(Project, AclAction.ViewProject, User, null))
+			{
+				return Forbid();
+			}
+
+			IProjectLogo? ProjectLogo = await ProjectService.Collection.GetLogoAsync(ProjectIdValue);
+			if (ProjectLogo == null)
+			{
+				return NotFound();
+			}
+
+			return new FileContentResult(ProjectLogo.Data, ProjectLogo.MimeType);
+		}
+
+		/// <summary>
+		/// Retrieve information about a specific project
+		/// </summary>
+		/// <param name="ProjectId">Id of the project to get information about</param>
+		/// <param name="File">The file content</param>
+		/// <returns>Information about the requested project</returns>
+		[HttpPut]
+		[Route("/api/v1/projects/{ProjectId}/logo")]
+		[ProducesResponseType(typeof(List<GetProjectResponse>), 200)]
+		public async Task<ActionResult> SetProjectLogoAsync(string ProjectId, IFormFile File)
+		{
+			ProjectId ProjectIdValue = new ProjectId(ProjectId);
+
+			IProject? Project = await ProjectService.GetProjectAsync(ProjectIdValue);
+			if (Project == null)
+			{
+				return NotFound();
+			}
+			if (!await ProjectService.AuthorizeAsync(Project, AclAction.UpdateProject, User, null))
+			{
+				return Forbid();
+			}
+
+			using (MemoryStream Stream = new MemoryStream())
+			{
+				await File.CopyToAsync(Stream);
+				await ProjectService.Collection.SetLogoAsync(ProjectIdValue, File.ContentType, Stream.ToArray());
+			}
+
+			return Ok();
+		}
 	}
 }
diff --git a/Engine/Source/Programs/Horde/HordeServer/Models/Project.cs b/Engine/Source/Programs/Horde/HordeServer/Models/Project.cs
--- a/Engine/Source/Programs/Horde/HordeServer/Models/Project.cs
+++ b/Engine/Source/Programs/Horde/HordeServer/Models/Project.cs
@@ -124,6 +124,27 @@ namespace HordeServer.Models
 		/// </summary>
 		public IReadOnlyDictionary<string, string> Properties { get; }
 	}
+	
+	/// <summary>
+	/// Interface for a document containing a project logo
+	/// </summary>
+	public interface IProjectLogo
+	{
+		/// <summary>
+		/// The project id
+		/// </summary>
+		public ProjectId Id { get; }
+
+		/// <summary>
+		/// Mime type for the image
+		/// </summary>
+		public string MimeType { get; }
+
+		/// <summary>
+		/// Image data
+		/// </summary>
+		public byte[] Data { get; }
+	}
 
 	/// <summary>
 	/// Extension methods for projects
diff --git a/Engine/Source/Programs/Horde/HordeServer/Services/ProjectService.cs b/Engine/Source/Programs/Horde/HordeServer/Services/ProjectService.cs
--- a/Engine/Source/Programs/Horde/HordeServer/Services/ProjectService.cs
+++ b/Engine/Source/Programs/Horde/HordeServer/Services/ProjectService.cs
@@ -43,6 +43,11 @@ namespace HordeServer.Services
 		/// </summary>
 		IProjectCollection Projects;
 
+		/// <summary>
+		/// Accessor for the collection of project documents
+		/// </summary>
+		public IProjectCollection Collection => Projects;
+
 		/// <summary>
 		/// Constructor
 		/// </summary>
