Horde: Missing file.

diff --git a/Engine/Source/Programs/Horde/HordeServer/Models/Config.cs b/Engine/Source/Programs/Horde/HordeServer/Models/Config.cs
new file mode 100644
--- /dev/null
+++ b/Engine/Source/Programs/Horde/HordeServer/Models/Config.cs
@@ -0,0 +1,203 @@
+// Copyright Epic Games, Inc. All Rights Reserved.
+
+using HordeServer.Api;
+using System;
+using System.Collections.Generic;
+using System.ComponentModel.DataAnnotations;
+using System.Linq;
+using System.Threading.Tasks;
+
+using ProjectId = HordeServer.Utilities.StringId<HordeServer.Models.IProject>;
+using StreamId = HordeServer.Utilities.StringId<HordeServer.Models.IStream>;
+
+namespace HordeServer.Models
+{
+	/// <summary>
+	/// Global configuration
+	/// </summary>
+	public class GlobalConfig
+	{
+		/// <summary>
+		/// List of projects
+		/// </summary>
+		public List<ProjectConfigRef> Projects { get; set; } = new List<ProjectConfigRef>();
+
+		/// <summary>
+		/// Manually added status messages
+		/// </summary>
+		public List<Notice> Notices { get; set; } = new List<Notice>();
+
+		/// <summary>
+		/// List of scheduled downtime
+		/// </summary>
+		public List<ScheduledDowntime> Downtime { get; set; } = new List<ScheduledDowntime>();
+
+		/// <summary>
+		/// List of Perforce clusters
+		/// </summary>
+		public List<PerforceCluster> PerforceClusters { get; set; } = new List<PerforceCluster>();
+
+		/// <summary>
+		/// Access control list
+		/// </summary>
+		public UpdateAclRequest? Acl { get; set; }
+	}
+
+	/// <summary>
+	/// References a project configuration
+	/// </summary>
+	public class ProjectConfigRef
+	{
+		/// <summary>
+		/// Unique id for the project
+		/// </summary>
+		[Required]
+		public ProjectId Id { get; set; } = ProjectId.Empty;
+
+		/// <summary>
+		/// Config path for the project
+		/// </summary>
+		[Required]
+		public string Path { get; set; } = null!;
+	}
+
+	/// <summary>
+	/// Stores configuration for a project
+	/// </summary>
+	public class ProjectConfig
+	{
+		/// <summary>
+		/// Name for the new project
+		/// </summary>
+		[Required]
+		public string Name { get; set; } = null!;
+
+		/// <summary>
+		/// Path to the project logo
+		/// </summary>
+		public string? Logo { get; set; }
+
+		/// <summary>
+		/// Categories to include in this project
+		/// </summary>
+		public List<CreateProjectCategoryRequest> Categories { get; set; } = new List<CreateProjectCategoryRequest>();
+
+		/// <summary>
+		/// List of streams
+		/// </summary>
+		public List<StreamConfigRef> Streams { get; set; } = new List<StreamConfigRef>();
+
+		/// <summary>
+		/// Acl entries
+		/// </summary>
+		public UpdateAclRequest? Acl { get; set; }
+
+		/// <summary>
+		/// Default constructor for serialization
+		/// </summary>
+		protected ProjectConfig()
+		{
+		}
+	}
+
+	/// <summary>
+	/// Reference to configuration for a stream
+	/// </summary>
+	public class StreamConfigRef
+	{
+		/// <summary>
+		/// Unique id for the stream
+		/// </summary>
+		[Required]
+		public StreamId Id { get; set; } = StreamId.Empty;
+
+		/// <summary>
+		/// Path to the configuration file
+		/// </summary>
+		[Required]
+		public string Path { get; set; } = null!;
+	}
+
+	/// <summary>
+	/// Config for a stream
+	/// </summary>
+	public class StreamConfig
+	{
+		/// <summary>
+		/// Name of the stream
+		/// </summary>
+		[Required]
+		public string Name { get; set; } = String.Empty;
+
+		/// <summary>
+		/// Order for this stream
+		/// </summary>
+		public int? Order { get; set; }
+
+		/// <summary>
+		/// Notification channel for all jobs in this stream
+		/// </summary>
+		public string? NotificationChannel { get; set; }
+
+		/// <summary>
+		/// Notification channel filter for this template. Can be Success|Failure|Warnings
+		/// </summary>
+		public string? NotificationChannelFilter { get; set; }
+
+		/// <summary>
+		/// Channel to post issue triage notifications
+		/// </summary>
+		public string? TriageChannel { get; set; }
+
+		/// <summary>
+		/// Default template for running preflights
+		/// </summary>
+		public string? DefaultPreflightTemplate { get; set; }
+
+		/// <summary>
+		/// Default template for running preflights
+		/// </summary>
+		public DefaultPreflightRequest? DefaultPreflight { get; set; }
+
+		/// <summary>
+		/// List of tabs to show for the new stream
+		/// </summary>
+		public List<CreateStreamTabRequest> Tabs { get; set; } = new List<CreateStreamTabRequest>();
+
+		/// <summary>
+		/// Map of agent name to type
+		/// </summary>
+		public Dictionary<string, CreateAgentTypeRequest> AgentTypes { get; set; } = new Dictionary<string, CreateAgentTypeRequest>();
+
+		/// <summary>
+		/// Map of workspace name to type
+		/// </summary>
+		public Dictionary<string, CreateWorkspaceTypeRequest> WorkspaceTypes { get; set; } = new Dictionary<string, CreateWorkspaceTypeRequest>();
+
+		/// <summary>
+		/// List of templates to create
+		/// </summary>
+		public List<CreateTemplateRefRequest> Templates { get; set; } = new List<CreateTemplateRefRequest>();
+
+		/// <summary>
+		/// Properties for the new stream
+		/// </summary>
+		public Dictionary<string, string>? Properties { get; set; }
+
+		/// <summary>
+		/// Custom permissions for this object
+		/// </summary>
+		public UpdateAclRequest? Acl { get; set; }
+
+		/// <summary>
+		/// Pause stream builds until specified date
+		/// </summary>
+		public DateTime? PausedUntil { get; set; }
+
+		/// <summary>
+		/// Reason for pausing builds of the stream
+		/// </summary>
+		public string? PauseComment { get; set; }
+	}
+}
+
