Horde: Another missing file.

diff --git a/Engine/Source/Programs/Horde/HordeServer/Services/ConfigService.cs b/Engine/Source/Programs/Horde/HordeServer/Services/ConfigService.cs
new file mode 100644
--- /dev/null
+++ b/Engine/Source/Programs/Horde/HordeServer/Services/ConfigService.cs
@@ -0,0 +1,354 @@
+// Copyright Epic Games, Inc. All Rights Reserved.
+
+using EpicGames.Core;
+using HordeServer.Api;
+using HordeServer.Controllers;
+using HordeServer.Models;
+using HordeServer.Notifications;
+using HordeServer.Utilities;
+using Microsoft.AspNetCore.StaticFiles;
+using Microsoft.Extensions.Logging;
+using Microsoft.Extensions.Options;
+using MongoDB.Bson;
+using System;
+using System.Collections.Generic;
+using System.ComponentModel.DataAnnotations;
+using System.IO;
+using System.Linq;
+using System.Text;
+using System.Text.Json;
+using System.Threading;
+using System.Threading.Tasks;
+
+using ProjectId = HordeServer.Utilities.StringId<HordeServer.Models.IProject>;
+using StreamId = HordeServer.Utilities.StringId<HordeServer.Models.IStream>;
+
+namespace HordeServer.Services
+{
+	/// <summary>
+	/// Polls Perforce for stream config changes
+	/// </summary>
+	public class ConfigService : ElectedBackgroundService
+	{
+		/// <summary>
+		/// Config file version number
+		/// </summary>
+		const int Version = 2;
+
+		/// <summary>
+		/// Database service
+		/// </summary>
+		DatabaseService DatabaseService;
+
+		/// <summary>
+		/// Project service
+		/// </summary>
+		ProjectService ProjectService;
+
+		/// <summary>
+		/// Stream service
+		/// </summary>
+		StreamService StreamService;
+
+		/// <summary>
+		/// Instance of the perforce service
+		/// </summary>
+		IPerforceService PerforceService;
+
+		/// <summary>
+		/// Instance of the notification service
+		/// </summary>
+		INotificationService NotificationService;
+
+		/// <summary>
+		/// The server settings
+		/// </summary>
+		ServerSettings Settings;
+
+		/// <summary>
+		/// Logging device
+		/// </summary>
+		ILogger Logger;
+
+		/// <summary>
+		/// Constructor
+		/// </summary>
+		/// <param name="DatabaseService"></param>
+		/// <param name="PerforceService"></param>
+		/// <param name="ProjectService"></param>
+		/// <param name="StreamService"></param>
+		/// <param name="NotificationService"></param>
+		/// <param name="Settings"></param>
+		/// <param name="Logger"></param>
+		public ConfigService(DatabaseService DatabaseService, IPerforceService PerforceService, ProjectService ProjectService, StreamService StreamService, INotificationService NotificationService, IOptionsMonitor<ServerSettings> Settings, ILogger<ConfigService> Logger)
+			: base(DatabaseService, new ObjectId("5ff60549e7632b15e64ac2f7"), Logger)
+		{
+			this.DatabaseService = DatabaseService;
+			this.PerforceService = PerforceService;
+			this.ProjectService = ProjectService;
+			this.StreamService = StreamService;
+			this.NotificationService = NotificationService;
+			this.Settings = Settings.CurrentValue;
+			this.Logger = Logger;
+		}
+
+		GlobalConfig? CachedGlobalConfig;
+		Dictionary<ProjectId, ProjectConfig> CachedProjectConfigs = new Dictionary<ProjectId, ProjectConfig>();
+		Dictionary<ProjectId, string?> CachedLogoRevisions = new Dictionary<ProjectId, string?>();
+
+		async Task UpdateConfigAsync(string ConfigPath)
+		{
+			Logger.LogInformation("Updating configuration from {ConfigPath}", ConfigPath);
+
+			// Update the globals singleton
+			GlobalConfig GlobalConfig;
+			for (; ; )
+			{
+				Dictionary<string, string> GlobalRevisions = await FindRevisionsAsync(new[] { ConfigPath });
+				if (GlobalRevisions.Count == 0)
+				{
+					throw new Exception($"Invalid config path: {ConfigPath}");
+				}
+
+				string Revision = GlobalRevisions.First().Value;
+
+				Globals Globals = await DatabaseService.GetGlobalsAsync();
+				if (Globals.ConfigRevision != Revision || CachedGlobalConfig == null)
+				{
+					Logger.LogInformation("Caching global config from {Revision}", Revision);
+					CachedGlobalConfig = await ReadDataAsync<GlobalConfig>(ConfigPath);
+				}
+				GlobalConfig = CachedGlobalConfig;
+
+				if (Globals.ConfigRevision == Revision)
+				{
+					break;
+				}
+
+				Globals.ConfigRevision = Revision;
+				Globals.Notices = CachedGlobalConfig.Notices;
+				Globals.PerforceClusters = CachedGlobalConfig.PerforceClusters;
+				Globals.ScheduledDowntime = CachedGlobalConfig.Downtime;
+
+				if (await DatabaseService.TryUpdateSingletonAsync(Globals))
+				{
+					break;
+				}
+			}
+
+			// Projects to remove
+			List<IProject> Projects = await ProjectService.GetProjectsAsync();
+
+			// Get the path to all the project configs
+			List<(ProjectConfigRef ProjectRef, string Path)> ProjectConfigs = GlobalConfig.Projects.Select(x => (x, CombinePaths(ConfigPath, x.Path))).ToList();
+
+			Dictionary<ProjectId, ProjectConfig> PrevCachedProjectConfigs = CachedProjectConfigs;
+			CachedProjectConfigs = new Dictionary<ProjectId, ProjectConfig>();
+
+			List<(ProjectId ProjectId, string Path)> ProjectLogos = new List<(ProjectId ProjectId, string Path)>();
+
+			List<(ProjectId ProjectId, StreamConfigRef StreamRef, string Path)> StreamConfigs = new List<(ProjectId, StreamConfigRef, string)>();
+
+			// Update any existing projects
+			Dictionary<string, string> ProjectRevisions = await FindRevisionsAsync(ProjectConfigs.Select(x => x.Path));
+			for (int Idx = 0; Idx < ProjectConfigs.Count; Idx++)
+			{
+				(ProjectConfigRef ProjectRef, string ProjectPath) = ProjectConfigs[Idx];
+				if (ProjectRevisions.TryGetValue(ProjectPath, out string? Revision))
+				{
+					IProject? Project = Projects.FirstOrDefault(x => x.Id == ProjectRef.Id);
+					bool Update = (Project == null || Project.Revision != Revision);
+
+					ProjectConfig? ProjectConfig;
+					if (Update || !PrevCachedProjectConfigs.TryGetValue(ProjectRef.Id, out ProjectConfig))
+					{
+						Logger.LogInformation("Caching configuration for project {ProjectId} ({Revision})", ProjectRef.Id, Revision);
+						ProjectConfig = await ReadDataAsync<ProjectConfig>(ProjectPath);
+					}
+					if (Update)
+					{
+						Logger.LogInformation("Updating configuration for project {ProjectId} ({Revision})", ProjectRef.Id, Revision);
+						await ProjectService.Collection.AddOrUpdateAsync(ProjectRef.Id, Revision, Idx, ProjectConfig);
+					}
+
+					if (ProjectConfig.Logo != null)
+					{
+						ProjectLogos.Add((ProjectRef.Id, CombinePaths(ProjectPath, ProjectConfig.Logo)));
+					}
+
+					CachedProjectConfigs[ProjectRef.Id] = ProjectConfig;
+					StreamConfigs.AddRange(ProjectConfig.Streams.Select(x => (ProjectRef.Id, x, CombinePaths(ProjectPath, x.Path))));
+				}
+			}
+
+			// Get the logo revisions
+			Dictionary<string, string> LogoRevisions = await FindRevisionsAsync(ProjectLogos.Select(x => x.Path));
+			for (int Idx = 0; Idx < ProjectLogos.Count; Idx++)
+			{
+				(ProjectId ProjectId, string Path) = ProjectLogos[Idx];
+				if (LogoRevisions.TryGetValue(Path, out string? Revision))
+				{
+					string? CurrentRevision;
+					if (!CachedLogoRevisions.TryGetValue(ProjectId, out CurrentRevision))
+					{
+						CurrentRevision = (await ProjectService.Collection.GetLogoAsync(ProjectId))?.Revision;
+						CachedLogoRevisions[ProjectId] = CurrentRevision;
+					}
+					if (Revision != CurrentRevision)
+					{
+						Logger.LogInformation("Updating logo for project {ProjectId} ({Revision})", ProjectId, Revision);
+						await ProjectService.Collection.SetLogoAsync(ProjectId, Revision, GetMimeTypeFromPath(Path), await ReadDataAsync(Path));
+						CachedLogoRevisions[ProjectId] = Revision;
+					}
+				}
+			}
+
+			// Get the current streams
+			List<IStream> Streams = await StreamService.GetStreamsAsync();
+
+			// Get the revisions for all the stream documents
+			Dictionary<string, string> StreamRevisions = await FindRevisionsAsync(StreamConfigs.Select(x => x.Path));
+			for (int Idx = 0; Idx < StreamConfigs.Count; Idx++)
+			{
+				(ProjectId ProjectId, StreamConfigRef StreamRef, string Path) = StreamConfigs[Idx];
+				if (StreamRevisions.TryGetValue(Path, out string? Revision))
+				{
+					IStream? Stream = Streams.FirstOrDefault(x => x.Id == StreamRef.Id);
+					if (Stream == null || Stream.Revision != Revision)
+					{
+						Logger.LogInformation("Updating configuration for stream {StreamRef} ({Revision})", StreamRef.Id, Revision);
+
+						StreamConfig StreamConfig = await ReadDataAsync<StreamConfig>(Path);
+						for (; ; )
+						{
+							Stream = await StreamService.StreamCollection.TryCreateOrReplaceAsync(StreamRef.Id, Stream, Revision, ProjectId, StreamConfig);
+
+							if (Stream != null)
+							{
+								break;
+							}
+
+							Stream = await StreamService.GetStreamAsync(StreamRef.Id);
+						}
+					}
+				}
+			}
+
+			// Remove any projects which are no longer used
+			HashSet<ProjectId> RemoveProjectIds = new HashSet<ProjectId>(Projects.Select(x => x.Id));
+			RemoveProjectIds.ExceptWith(ProjectConfigs.Select(y => y.ProjectRef.Id));
+
+			foreach (ProjectId RemoveProjectId in RemoveProjectIds)
+			{
+				await ProjectService.DeleteProjectAsync(RemoveProjectId);
+			}
+
+			// Remove any streams that are no longer used
+			HashSet<StreamId> RemoveStreamIds = new HashSet<StreamId>(Streams.Select(x => x.Id));
+			RemoveStreamIds.ExceptWith(StreamConfigs.Select(x => x.StreamRef.Id));
+
+			foreach (StreamId RemoveStreamId in RemoveStreamIds)
+			{
+				await StreamService.DeleteStreamAsync(RemoveStreamId);
+			}
+		}
+
+		static FileExtensionContentTypeProvider ContentTypeProvider = new FileExtensionContentTypeProvider();
+
+		static string GetMimeTypeFromPath(string Path)
+		{
+			string ContentType;
+			if (!ContentTypeProvider.TryGetContentType(Path, out ContentType))
+			{
+				ContentType = "application/octet-stream";
+			}
+			return ContentType;
+		}
+
+		static string CombinePaths(string BasePath, string RelativePath)
+		{
+			// Handle Perforce paths
+			if (RelativePath.StartsWith("//", StringComparison.Ordinal))
+			{
+				return RelativePath;
+			}
+			if (BasePath.StartsWith("//", StringComparison.Ordinal))
+			{
+				return BasePath.Substring(0, BasePath.LastIndexOf('/') + 1) + RelativePath;
+			}
+
+			// Handle filesystem paths
+			return Path.GetFullPath(Path.Combine(Path.GetDirectoryName(BasePath)!, RelativePath));
+		}
+
+		async Task<Dictionary<string, string>> FindRevisionsAsync(IEnumerable<string> Paths)
+		{
+			Dictionary<string, string> Revisions = new Dictionary<string, string>();
+
+			// Find all the Perforce uris
+			List<string> PerforcePaths = new List<string>();
+			foreach (string Path in Paths)
+			{
+				if (Path.StartsWith("//", StringComparison.Ordinal))
+				{
+					PerforcePaths.Add(Path);
+				}
+				else
+				{
+					Revisions[Path] = $"ver={Version},md5={ContentHash.MD5(Path)}";
+				}
+			}
+
+			// Query all the Perforce revisions
+			if (PerforcePaths.Count > 0)
+			{
+				List<FileSummary> Files = await PerforceService.FindFilesAsync(PerforcePaths);
+				foreach (FileSummary File in Files)
+				{
+					if (File.Error == null)
+					{
+						Revisions[File.DepotPath] = $"ver={Version},p4={File.DepotPath}@{File.Change}";
+					}
+					else
+					{
+						NotificationService.NotifyUpdateStreamFailure(File);
+					}
+				}
+			}
+
+			return Revisions;
+		}
+
+		async Task<T> ReadDataAsync<T>(string ConfigPath)
+		{
+			byte[] Data = await ReadDataAsync(ConfigPath);
+
+			JsonSerializerOptions Options = new JsonSerializerOptions();
+			Startup.ConfigureJsonSerializer(Options);
+
+			return JsonSerializer.Deserialize<T>(Data, Options);
+		}
+
+		Task<byte[]> ReadDataAsync(string ConfigPath)
+		{
+			if (ConfigPath.StartsWith("//", StringComparison.Ordinal))
+			{
+				return PerforceService.PrintAsync(ConfigPath);
+			}
+			else
+			{
+				return File.ReadAllBytesAsync(ConfigPath);
+			}
+		}
+
+		/// <inheritdoc/>
+		protected override async Task<DateTime> TickLeaderAsync(CancellationToken StoppingToken)
+		{
+			if (Settings.ConfigPath != null)
+			{
+				await UpdateConfigAsync(Settings.ConfigPath);
+			}
+			return DateTime.UtcNow + TimeSpan.FromMinutes(1.0);
+		}
+	}
+}
