Horde: Fix unit tests.

diff --git a/Engine/Source/Programs/Horde/HordeServerTests/Fixture.cs b/Engine/Source/Programs/Horde/HordeServerTests/Fixture.cs
--- a/Engine/Source/Programs/Horde/HordeServerTests/Fixture.cs
+++ b/Engine/Source/Programs/Horde/HordeServerTests/Fixture.cs
@@ -27,8 +27,6 @@ namespace HordeServerTests
 		public IStream? Stream { get; private set; }
 		public TemplateRefId TemplateRefId1 { get; private set; }
 		public TemplateRefId TemplateRefId2 { get; private set; }
-		public TemplateRef TemplateRef1 { get; private set; } = null!;
-		public TemplateRef TemplateRef2 { get; private set; } = null!;
 		public Artifact Job1Artifact { get; private set; } = null!;
 		public string Job1ArtifactData { get; private set; } = null!;
 		public IAgent Agent1 { get; private set; } = null!;
@@ -65,29 +63,20 @@ namespace HordeServerTests
 
 			TemplateRefId1 = new TemplateRefId("template1");
 			TemplateRefId2 = new TemplateRefId("template2");
-			TemplateRef1 = new TemplateRef(Template, false, false, null, null, null);
-			TemplateRef2 = new TemplateRef(Template, false, false, null, null, null);
 
-			Dictionary<TemplateRefId, TemplateRef> TemplateRefs = new Dictionary<TemplateRefId, TemplateRef>
-			{
-				{TemplateRefId1, TemplateRef1},
-				{TemplateRefId2, TemplateRef2}
-			};
-			
-			List<StreamTab> Tabs = new List<StreamTab>();
-			Tabs.Add(new JobsTab("foo", true, new List<TemplateRefId> { TemplateRefId1, TemplateRefId2 }, new List<string>(), new List<JobsTabColumn>()));
+			List<CreateTemplateRefRequest> Templates = new List<CreateTemplateRefRequest>();
+			Templates.Add(new CreateTemplateRefRequest { Id = TemplateRefId1.ToString(), Name = "Test Template" });
+			Templates.Add(new CreateTemplateRefRequest { Id = TemplateRefId2.ToString(), Name = "Test Template" });
 
-			Stream = await StreamService.TryCreateStreamAsync(
-				Id: new StreamId("ue5-main"),
-				Name: "//UE5/Main",
-				ProjectId: new ProjectId("does-not-exist"),
-				Order: null,
-				Tabs: Tabs,
-				AgentTypes: null,
-				WorkspaceTypes: null,
-				TemplateRefs: TemplateRefs,
-				Properties: null,
-				Acl: null
+			List<CreateStreamTabRequest> Tabs = new List<CreateStreamTabRequest>();
+			Tabs.Add(new CreateJobsTabRequest { Title = "foo", Templates = new List<string> { TemplateRefId1.ToString(), TemplateRefId2.ToString() } });
+
+			Stream = await StreamService.StreamCollection.TryCreateOrReplaceAsync(
+				new StreamId("ue5-main"),
+				null,
+				"",
+				new ProjectId("does-not-exist"),
+				new StreamConfig { Name = "//UE5/Main", Tabs = Tabs, Templates = Templates }
 			);
 			
 			Job1 = await JobService.CreateJobAsync(
diff --git a/Engine/Source/Programs/Horde/HordeServerTests/IssueServiceTests.cs b/Engine/Source/Programs/Horde/HordeServerTests/IssueServiceTests.cs
--- a/Engine/Source/Programs/Horde/HordeServerTests/IssueServiceTests.cs
+++ b/Engine/Source/Programs/Horde/HordeServerTests/IssueServiceTests.cs
@@ -127,9 +127,9 @@ namespace HordeServerTests
 
 			IProject Project = TestSetup.ProjectCollection.TryAddAsync(new ProjectId("ue4"), "UE4").Result!;
 
-			IStream MainStream = TestSetup.StreamCollection.TryCreateAsync(MainStreamId, MainStreamName, Project.Id).Result!;
-			IStream ReleaseStream = TestSetup.StreamCollection.TryCreateAsync(ReleaseStreamId, ReleaseStreamName, Project.Id).Result!;
-			IStream DevStream = TestSetup.StreamCollection.TryCreateAsync(DevStreamId, DevStreamName, Project.Id).Result!;
+			IStream MainStream = TestSetup.StreamCollection.TryCreateOrReplaceAsync(MainStreamId, null, "", Project.Id, new StreamConfig { Name = MainStreamName }).Result!;
+			IStream ReleaseStream = TestSetup.StreamCollection.TryCreateOrReplaceAsync(ReleaseStreamId, null, "", Project.Id, new StreamConfig { Name = ReleaseStreamName }).Result!;
+			IStream DevStream = TestSetup.StreamCollection.TryCreateOrReplaceAsync(DevStreamId, null, "", Project.Id, new StreamConfig { Name = DevStreamName }).Result!;
 
 			Perforce = TestSetup.PerforceService;
 			Perforce.AddChange(MainStreamName, 100, "Bill", "Description", new string[] { "a/b.cpp" });
diff --git a/Engine/Source/Programs/Horde/HordeServerTests/JobServiceTests.cs b/Engine/Source/Programs/Horde/HordeServerTests/JobServiceTests.cs
--- a/Engine/Source/Programs/Horde/HordeServerTests/JobServiceTests.cs
+++ b/Engine/Source/Programs/Horde/HordeServerTests/JobServiceTests.cs
@@ -23,8 +23,9 @@ namespace HordeServerTests
 		public async Task TestChainedJobs()
 		{
 			TestSetup TestSetup = await GetTestSetup();
-			
-			IProject? Project = await TestSetup.ProjectService.TryCreateProjectAsync(new ProjectId("ue5"), "UE5", null, null, null);
+
+			ProjectId ProjectId = new ProjectId("ue5");
+			IProject? Project = await TestSetup.ProjectService.TryCreateProjectAsync(ProjectId, "UE5", null, null, null);
 			Assert.IsNotNull(Project);
 
 			ITemplate Template = await TestSetup.TemplateService.CreateTemplateAsync("Test template", null, false, null, null, new List<TemplateCounter>(), new List<string>(), new List<Parameter>());
@@ -32,26 +33,17 @@ namespace HordeServerTests
 
 			TemplateRefId TemplateRefId1 = new TemplateRefId("template1");
 			TemplateRefId TemplateRefId2 = new TemplateRefId("template2");
-			TemplateRef TemplateRef1 = new TemplateRef(Template, Triggers: new List<ChainedJobTemplate> { new ChainedJobTemplate("Setup Build", TemplateRefId2) });
-			TemplateRef TemplateRef2 = new TemplateRef(Template);
 
-			Dictionary<TemplateRefId, TemplateRef> TemplateRefs = new Dictionary<TemplateRefId, TemplateRef>
-			{
-				{ TemplateRefId1, TemplateRef1 },
-				{ TemplateRefId2, TemplateRef2 }
-			};
-
-			List<StreamTab> Tabs = new List<StreamTab>();
-			Tabs.Add(new JobsTab("foo", true, new List<TemplateRefId> { TemplateRefId1, TemplateRefId2 }, new List<string>(), new List<JobsTabColumn>()));
+			StreamConfig StreamConfig = new StreamConfig();
+			StreamConfig.Templates.Add(new CreateTemplateRefRequest { Id = TemplateRefId1.ToString(), Name = "Test Template", ChainedJobs = new List<CreateChainedJobTemplateRequest> { new CreateChainedJobTemplateRequest { TemplateId = TemplateRefId2.ToString(), Trigger = "Setup Build" } } });
+			StreamConfig.Templates.Add(new CreateTemplateRefRequest { Id = TemplateRefId2.ToString(), Name = "Test Template" });
+			StreamConfig.Tabs.Add(new CreateJobsTabRequest { Title = "foo", Templates = new List<string> { TemplateRefId1.ToString(), TemplateRefId2.ToString() } });
 
 			StreamId StreamId = new StreamId("ue5-main");
-			await TestSetup.StreamService.TryCreateStreamAsync(new StreamId("ue5-main"), "//UE5/Main", Project!.Id);
-
 			IStream? Stream = await TestSetup.StreamService.GetStreamAsync(StreamId);
-			Assert.IsNotNull(Stream);
-			await TestSetup.StreamService.UpdateStreamAsync(Stream, NewTabs: Tabs, NewTemplateRefs: TemplateRefs);
+			Stream = await TestSetup.StreamService.StreamCollection.TryCreateOrReplaceAsync(new StreamId("ue5-main"), Stream, String.Empty, ProjectId, StreamConfig);
 
-			IJob Job = await TestSetup.JobService.CreateJobAsync(null, Stream!.Id, TemplateRefId1, Template.Id, Graph, "Hello", 1234, 1233, 999, null, null, "joe", null, null, TemplateRef1.ChainedJobs, true, true, null, null, null, Template.Counters, new List<string>());
+			IJob Job = await TestSetup.JobService.CreateJobAsync(null, Stream!.Id, TemplateRefId1, Template.Id, Graph, "Hello", 1234, 1233, 999, null, null, "joe", null, null, Stream.Templates[TemplateRefId1].ChainedJobs, true, true, null, null, null, Template.Counters, new List<string>());
 			Assert.AreEqual(1, Job.ChainedJobs.Count);
 
 			Assert.IsTrue(await TestSetup.JobService.UpdateBatchAsync(Job, Job.Batches[0].Id, ObjectId.GenerateNewId(), JobStepBatchState.Running));
@@ -159,7 +151,7 @@ namespace HordeServerTests
 			Assert.IsNotNull(Project);
 
 			StreamId StreamId = new StreamId("ue5-main");
-			await TestSetup.StreamService.TryCreateStreamAsync(new StreamId("ue5-main"), "//UE5/Main", Project!.Id);
+			await TestSetup.StreamCollection.TryCreateOrReplaceAsync(new StreamId("ue5-main"), null, "", Project!.Id, new StreamConfig { Name = "//UE5/Main" });
 
 			ITemplate Template = await TestSetup.TemplateService.CreateTemplateAsync("Test template", null, false, null, null, new List<TemplateCounter>(), new List<string>(), new List<Parameter>());
 			IGraph Graph = await TestSetup.GraphCollection.AddAsync(Template);
diff --git a/Engine/Source/Programs/Horde/HordeServerTests/SchedulerTests.cs b/Engine/Source/Programs/Horde/HordeServerTests/SchedulerTests.cs
--- a/Engine/Source/Programs/Horde/HordeServerTests/SchedulerTests.cs
+++ b/Engine/Source/Programs/Horde/HordeServerTests/SchedulerTests.cs
@@ -46,7 +46,7 @@ namespace HordeServerTests
 			Stream = TestSetup.StreamCollection.GetAsync(StreamId).Result!;
 			if (Stream == null)
 			{
-				Stream = TestSetup.StreamService.TryCreateStreamAsync(new StreamId("ue5-main"), "//UE5/Main", Project!.Id).Result!;
+				Stream = TestSetup.StreamCollection.TryCreateOrReplaceAsync(new StreamId("ue5-main"), null, "", Project!.Id, new StreamConfig { Name = "//UE5/Main" }).Result!;
 			}
 
 			InitialJobIds = new HashSet<ObjectId>(TestSetup.JobCollection.FindAsync().Result.Select(x => x.Id));
