binary print support

diff --git a/Engine/Source/Programs/Horde/HordeServer/Services/PerforceService.cs b/Engine/Source/Programs/Horde/HordeServer/Services/PerforceService.cs
--- a/Engine/Source/Programs/Horde/HordeServer/Services/PerforceService.cs
+++ b/Engine/Source/Programs/Horde/HordeServer/Services/PerforceService.cs
@@ -746,6 +746,7 @@ namespace HordeServer.Services
 		{
 			public string Data { get; set; } = String.Empty;
 			public List<string>? Errors { get; set; } = null;
+			public bool Binary { get; set; } = false;
 		}
 #pragma warning restore CA1812
 
@@ -1013,6 +1014,10 @@ namespace HordeServer.Services
 				throw new Exception($"{Uri} resulted in no data being returned. Was {Path} deleted?");
 			}
 
+			if(Result.Binary)
+			{
+				return Convert.FromBase64String(Result.Data);
+			}
 			return Encoding.Default.GetBytes(Result.Data);
 		}
 
