Horde: Restore endpoint to allow retrieving project logos.

diff --git a/Engine/Source/Programs/Horde/HordeServer/Controllers/ProjectsController.cs b/Engine/Source/Programs/Horde/HordeServer/Controllers/ProjectsController.cs
--- a/Engine/Source/Programs/Horde/HordeServer/Controllers/ProjectsController.cs
+++ b/Engine/Source/Programs/Horde/HordeServer/Controllers/ProjectsController.cs
@@ -172,6 +172,36 @@ namespace HordeServer.Controllers
 			return Project.ToResponse(bIncludeStreams, bIncludeCategories, VisibleStreams, bIncludeAcl).ApplyFilter(Filter);
 		}
 
+		/// <summary>
+		/// Retrieve information about a specific project
+		/// </summary>
+		/// <param name="ProjectId">Id of the project to get information about</param>
+		/// <returns>Information about the requested project</returns>
+		[HttpGet]
+		[Route("/api/v1/projects/{ProjectId}/logo")]
+		public async Task<ActionResult<object>> GetProjectLogoAsync(string ProjectId)
+		{
+			ProjectId ProjectIdValue = new ProjectId(ProjectId);
+
+			IProject? Project = await ProjectService.GetProjectAsync(ProjectIdValue);
+			if (Project == null)
+			{
+				return NotFound();
+			}
+			if (!await ProjectService.AuthorizeAsync(Project, AclAction.ViewProject, User, null))
+			{
+				return Forbid();
+			}
+
+			IProjectLogo? ProjectLogo = await ProjectService.Collection.GetLogoAsync(ProjectIdValue);
+			if (ProjectLogo == null)
+			{
+				return NotFound();
+			}
+
+			return new FileContentResult(ProjectLogo.Data, ProjectLogo.MimeType);
+		}
+
 		/// <summary>
 		/// Update a project's properties.
 		/// </summary>
