Horde: Force update of stored config data to fix missing logos.

diff --git a/Engine/Source/Programs/Horde/HordeServer/Services/ConfigService.cs b/Engine/Source/Programs/Horde/HordeServer/Services/ConfigService.cs
--- a/Engine/Source/Programs/Horde/HordeServer/Services/ConfigService.cs
+++ b/Engine/Source/Programs/Horde/HordeServer/Services/ConfigService.cs
@@ -33,7 +33,7 @@ namespace HordeServer.Services
 		/// <summary>
 		/// Config file version number
 		/// </summary>
-		const int Version = 2;
+		const int Version = 3;
 
 		/// <summary>
 		/// Database service
