Horde: adding P4 api native binaries

diff --git a/Engine/Source/Programs/Horde/HordeServer/HordeServer.csproj b/Engine/Source/Programs/Horde/HordeServer/HordeServer.csproj
--- a/Engine/Source/Programs/Horde/HordeServer/HordeServer.csproj
+++ b/Engine/Source/Programs/Horde/HordeServer/HordeServer.csproj
@@ -77,11 +77,11 @@
   </ItemGroup>
 
   <ItemGroup>
-    <Content Include="..\ThirdParty\P4Api.Net\Native\Windows\Release\p4bridge.*" Condition=" '$(OS)' == 'Windows_NT' ">
-      <CopyToOutputDirectory>PreserveNewest</CopyToOutputDirectory>
+    <Content Include="..\..\..\ThirdParty\Perforce\P4Api.Net\native\windows\release\p4bridge.*" Condition=" '$(OS)' == 'Windows_NT' ">
+      <CopyToOutputDirectory>Always</CopyToOutputDirectory>
     </Content>
-    <Content Include="..\ThirdParty\P4Api.Net\Native\Linux\Release\p4bridge.*" Condition="$([MSBuild]::IsOSPlatform('Linux'))" >
-      <CopyToOutputDirectory>PreserveNewest</CopyToOutputDirectory>
+    <Content Include="..\..\..\ThirdParty\Perforce\P4Api.Net\native\linux\release\p4bridge.*" Condition="$([MSBuild]::IsOSPlatform('Linux'))" >
+      <CopyToOutputDirectory>Always</CopyToOutputDirectory>
     </Content>	
   </ItemGroup>
 
