Increase gRPC send/receive sizes for Horde Agent

Blobs fetched from CAS can be quite large.

diff --git a/Engine/Source/Programs/Horde/HordeAgent/Services/GrpcService.cs b/Engine/Source/Programs/Horde/HordeAgent/Services/GrpcService.cs
--- a/Engine/Source/Programs/Horde/HordeAgent/Services/GrpcService.cs
+++ b/Engine/Source/Programs/Horde/HordeAgent/Services/GrpcService.cs
@@ -89,7 +89,15 @@ namespace HordeAgent.Services
 			HttpClient.Timeout = TimeSpan.FromSeconds(210); // Need to make sure this doesn't cancel any long running gRPC streaming calls (eg. session update)
 
 			Logger.LogInformation("Connecting to rpc server {BaseUrl}", ServerProfile.Url);
-			return GrpcChannel.ForAddress(Address, new GrpcChannelOptions { HttpClient = HttpClient, DisposeHttpClient = true });
+			return GrpcChannel.ForAddress(Address, new GrpcChannelOptions
+			{
+				// Required payloads coming from CAS service can be large
+				MaxReceiveMessageSize = 1024 * 1024 * 1024, // 1 GB
+				MaxSendMessageSize = 1024 * 1024 * 1024, // 1 GB
+				
+				HttpClient = HttpClient,
+				DisposeHttpClient = true
+			});
 		}
 	}
 }
diff --git a/Engine/Source/Programs/Horde/HordeAgent/Services/WorkerService.cs b/Engine/Source/Programs/Horde/HordeAgent/Services/WorkerService.cs
--- a/Engine/Source/Programs/Horde/HordeAgent/Services/WorkerService.cs
+++ b/Engine/Source/Programs/Horde/HordeAgent/Services/WorkerService.cs
@@ -777,7 +777,12 @@ namespace HordeAgent.Services
 					if (!string.IsNullOrEmpty(Token))
 						return GrpcService.CreateGrpcChannel(Url, new AuthenticationHeaderValue("ServiceAccount", Token));
 					
-					return GrpcChannel.ForAddress(Url);
+					return GrpcChannel.ForAddress(Url, new GrpcChannelOptions
+					{
+						// Required payloads coming from CAS service can be large
+						MaxReceiveMessageSize = 1024 * 1024 * 1024, // 1 GB
+						MaxSendMessageSize = 1024 * 1024 * 1024 // 1 GB
+					});
 				}
 
 				GrpcChannel CasChannel = GetChannel(ActionTask.CasUrl, ActionTask.ServiceAccountToken);// == null ? Client.Channel : GrpcChannel.ForAddress(ActionTask.CasUrl);
