Optionally download output files from a remote execution job

Also add the ability to skip cache lookup (not currently implemented in Jupiter).

diff --git a/Engine/Source/Programs/Horde/HordeAgent/Commands/Execution/ExecuteCommand.cs b/Engine/Source/Programs/Horde/HordeAgent/Commands/Execution/ExecuteCommand.cs
--- a/Engine/Source/Programs/Horde/HordeAgent/Commands/Execution/ExecuteCommand.cs
+++ b/Engine/Source/Programs/Horde/HordeAgent/Commands/Execution/ExecuteCommand.cs
@@ -119,6 +119,7 @@ namespace HordeAgent.Commands
 
 		class JsonCommand
 		{
+			public List<string> OutputPaths { get; set; } = new List<string>();
 			public List<string> Arguments { get; set; } = new List<string>();
 			public string WorkingDirectory { get; set; } = String.Empty;
 
@@ -127,6 +128,7 @@ namespace HordeAgent.Commands
 				RpcCommand Command = new RpcCommand();
 				Command.Arguments.Add(Arguments);
 				Command.WorkingDirectory = WorkingDirectory;
+				Command.OutputPaths.Add(OutputPaths);
 				return UploadList.Add(Command);
 			}
 		}
@@ -136,12 +138,13 @@ namespace HordeAgent.Commands
 			public JsonCommand Command { get; set; } = new JsonCommand();
 			public JsonDirectory Workspace { get; set; } = new JsonDirectory();
 
-			public async Task<Digest> BuildAsync(UploadList UploadList, ByteString? Salt)
+			public async Task<Digest> BuildAsync(UploadList UploadList, ByteString? Salt, bool DoNotCache)
 			{
 				Action NewAction = new Action();
 				NewAction.CommandDigest = Command.Build(UploadList);
 				NewAction.InputRootDigest = await Workspace.BuildAsync(UploadList);
 				NewAction.Salt = Salt;
+				NewAction.DoNotCache = DoNotCache;
 				return UploadList.Add(NewAction);
 			}
 		}
@@ -186,6 +189,18 @@ namespace HordeAgent.Commands
 		/// </summary>
 		[CommandLine("-InstanceName=")]
 		public string? InstanceName = null;
+				
+		/// <summary>
+		/// Skip checking if a result is already available in the action cache
+		/// </summary>
+		[CommandLine("-SkipCacheLookup")]
+		public bool SkipCacheLookup = false;
+		
+		/// <summary>
+		/// Directory to download the output files to. If not set, no results will be downloaded.
+		/// </summary>
+		[CommandLine("-OutputDir")]
+		public string? OutputDir = null;
 
 		/// <inheritdoc/>
 		public override async Task<int> ExecuteAsync(ILogger Logger)
@@ -228,7 +243,7 @@ namespace HordeAgent.Commands
 				}
 
 				UploadList UploadList = new UploadList();
-				Digest ActionDigest = await Action.BuildAsync(UploadList, SaltBytes);
+				Digest ActionDigest = await Action.BuildAsync(UploadList, SaltBytes, SkipCacheLookup);
 
 				using (GrpcChannel Channel = GrpcService.CreateGrpcChannel())
 				{
@@ -279,6 +294,7 @@ namespace HordeAgent.Commands
 
 					ExecuteRequest ExecuteRequest = new ExecuteRequest();
 					ExecuteRequest.ActionDigest = ActionDigest;
+					ExecuteRequest.SkipCacheLookup = SkipCacheLookup;
 					if (InstanceName != null)
 					{
 						ExecuteRequest.InstanceName = InstanceName;
@@ -316,6 +332,18 @@ namespace HordeAgent.Commands
 								}
 
 								Logger.LogInformation("Cached: {CachedResult}", ExecuteResponse.CachedResult);
+								
+								Logger.LogInformation("Output directories:");
+								foreach (OutputDirectory Dir in ActionResult.OutputDirectories)
+								{
+									Logger.LogInformation("Path: {Path} Digest: {Digest}", Dir.Path, Dir.TreeDigest);
+								}
+								
+								Logger.LogInformation("Output files:");
+								foreach (OutputFile File in ActionResult.OutputFiles)
+								{
+									Logger.LogInformation("Path: {Path} Digest: {Digest}", File.Path, File.Digest);
+								}
 
 								// Print the result
 								BatchReadBlobsRequest BatchReadRequest = new BatchReadBlobsRequest();
@@ -328,6 +356,10 @@ namespace HordeAgent.Commands
 								{
 									BatchReadRequest.Digests.Add(ActionResult.StderrDigest);
 								}
+								if (OutputDir != null && ActionResult.OutputFiles != null && ActionResult.OutputFiles.Count > 0)
+								{
+									BatchReadRequest.Digests.Add(ActionResult.OutputFiles.Select(x => x.Digest));
+								}
 
 								BatchReadBlobsResponse BatchReadResponse = await StorageClient.BatchReadBlobsAsync(BatchReadRequest);
 								if (ActionResult.StdoutDigest != null)
@@ -352,6 +384,21 @@ namespace HordeAgent.Commands
 										}
 									}
 								}
+								if (OutputDir != null && ActionResult.OutputFiles != null)
+								{
+									foreach (OutputFile OutputFile in ActionResult.OutputFiles)
+									{
+										BatchReadBlobsResponse.Types.Response? FileResponse = BatchReadResponse.Responses.FirstOrDefault(x => x.Digest.Hash == OutputFile.Digest.Hash);
+										if (FileResponse == null)
+										{
+											continue;
+										}
+										string OutputPath = Path.Join(OutputDir, OutputFile.Path);
+										System.IO.Directory.CreateDirectory(Path.GetDirectoryName(OutputPath));
+										await File.WriteAllBytesAsync(OutputPath, FileResponse.Data.ToByteArray());
+										Logger.LogInformation("Wrote {OutputFilePath}", OutputFile.Path);
+									}
+								}
 
 								Logger.LogInformation("exit: {ExitCode}", ActionResult.ExitCode);
 								break;
diff --git a/Engine/Source/Programs/Horde/HordeAgent/Execution/ActionExecutor.cs b/Engine/Source/Programs/Horde/HordeAgent/Execution/ActionExecutor.cs
--- a/Engine/Source/Programs/Horde/HordeAgent/Execution/ActionExecutor.cs
+++ b/Engine/Source/Programs/Horde/HordeAgent/Execution/ActionExecutor.cs
@@ -105,9 +105,9 @@ namespace HordeAgent
 				DirectoryReference.CreateDirectory(OutputFileLocation.Directory);
 			}
 
-			if (Command.OutputPaths.Count > 0 || Command.OutputDirectories.Count > 0)
+			if (Command.OutputFiles.Count > 0 || Command.OutputDirectories.Count > 0)
 			{
-				throw new NotImplementedException();
+				throw new NotImplementedException("OutputFiles/OutputDirectories are deprecated. Use OutputPaths instead.");
 			}
 
 			ActionResult Result;
@@ -156,16 +156,16 @@ namespace HordeAgent
 
 					Logger.LogInformation("exit: {ExitCode}", Process.ExitCode);
 
-					foreach (string OutputFile in Command.OutputFiles)
+					foreach (string OutputPath in Command.OutputPaths)
 					{
-						FileReference FileRef = FileReference.Combine(SandboxDir, OutputFile);
+						FileReference FileRef = FileReference.Combine(SandboxDir, OutputPath);
 						if (FileReference.Exists(FileRef))
 						{
 							byte[] Bytes = await FileReference.ReadAllBytesAsync(FileRef);
 							Digest Digest = await Storage.PutBulkDataAsync(InstanceName, Bytes);
 
 							OutputFile OutputFileInfo = new OutputFile();
-							OutputFileInfo.Path = OutputFile;
+							OutputFileInfo.Path = OutputPath;
 							OutputFileInfo.Digest = Digest;
 							Result.OutputFiles.Add(OutputFileInfo);
 
