Horde: Log when removing projects and streams, and soft-delete rather than hard-deleting streams.

diff --git a/Engine/Source/Programs/Horde/HordeServer/Collections/Impl/StreamCollection.cs b/Engine/Source/Programs/Horde/HordeServer/Collections/Impl/StreamCollection.cs
--- a/Engine/Source/Programs/Horde/HordeServer/Collections/Impl/StreamCollection.cs
+++ b/Engine/Source/Programs/Horde/HordeServer/Collections/Impl/StreamCollection.cs
@@ -59,6 +59,7 @@ namespace HordeServer.Collections.Impl
 			public string? PauseComment { get; set; }
 			public Acl? Acl { get; set; }
 			public int UpdateIndex { get; set; }
+			public bool Deleted { get; set; }
 
 			IReadOnlyList<StreamTab> IStream.Tabs => Tabs;
 			IReadOnlyDictionary<string, AgentType> IStream.AgentTypes => AgentTypes;
@@ -356,14 +357,14 @@ namespace HordeServer.Collections.Impl
 		/// <inheritdoc/>
 		public async Task<List<IStream>> FindAllAsync()
 		{
-			List<StreamDocument> Results = await Streams.Find(FilterDefinition<StreamDocument>.Empty).ToListAsync();
+			List<StreamDocument> Results = await Streams.Find(Builders<StreamDocument>.Filter.Ne(x => x.Deleted, true)).ToListAsync();
 			return Results.ConvertAll<IStream>(x => x);
 		}
 
 		/// <inheritdoc/>
 		public async Task<List<IStream>> FindForProjectsAsync(ProjectId[] ProjectIds)
 		{
-			FilterDefinition<StreamDocument> Filter = Builders<StreamDocument>.Filter.In(x => x.ProjectId, ProjectIds);
+			FilterDefinition<StreamDocument> Filter = Builders<StreamDocument>.Filter.In(x => x.ProjectId, ProjectIds) & Builders<StreamDocument>.Filter.Ne(x => x.Deleted, true);
 			List<StreamDocument> Results = await Streams.Find(Filter).ToListAsync();
 			return Results.ConvertAll<IStream>(x => x);
 		}
@@ -555,7 +556,7 @@ namespace HordeServer.Collections.Impl
 		/// <inheritdoc/>
 		public async Task DeleteAsync(StreamId StreamId)
 		{
-			await Streams.DeleteOneAsync<StreamDocument>(x => x.Id == StreamId);
+			await Streams.UpdateOneAsync<StreamDocument>(x => x.Id == StreamId, Builders<StreamDocument>.Update.Set(x => x.Deleted, true).Inc(x => x.UpdateIndex, 1));
 		}
 	}
 }
diff --git a/Engine/Source/Programs/Horde/HordeServer/Services/ConfigService.cs b/Engine/Source/Programs/Horde/HordeServer/Services/ConfigService.cs
--- a/Engine/Source/Programs/Horde/HordeServer/Services/ConfigService.cs
+++ b/Engine/Source/Programs/Horde/HordeServer/Services/ConfigService.cs
@@ -240,6 +240,7 @@ namespace HordeServer.Services
 
 			foreach (ProjectId RemoveProjectId in RemoveProjectIds)
 			{
+				Logger.LogInformation("Removing project {ProjectId}", RemoveProjectId);
 				await ProjectService.DeleteProjectAsync(RemoveProjectId);
 			}
 
@@ -249,6 +250,7 @@ namespace HordeServer.Services
 
 			foreach (StreamId RemoveStreamId in RemoveStreamIds)
 			{
+				Logger.LogInformation("Removing stream {StreamId}", RemoveStreamId);
 				await StreamService.DeleteStreamAsync(RemoveStreamId);
 			}
 		}
