Horde: Add static files serving for server, so we can host on one machine.

diff --git a/Engine/Source/Programs/Horde/HordeServer/Startup.cs b/Engine/Source/Programs/Horde/HordeServer/Startup.cs
--- a/Engine/Source/Programs/Horde/HordeServer/Startup.cs
+++ b/Engine/Source/Programs/Horde/HordeServer/Startup.cs
@@ -694,7 +694,8 @@ namespace HordeServer
 			{
 				App.UseMiddleware<RequestTrackerMiddleware>();	
 			}
-			
+
+			App.UseStaticFiles();
 			App.UseRouting();
 
 			App.UseAuthentication();
