Horde: Support for different outcomes when picking base CL for a preflight. The default preflight returned for a stream will now contain a "Change" property allowing the user to set a target and list of valid outcomes (success, warning, etc...). These can be forwarded to the create job request.

diff --git a/Engine/Source/Programs/Horde/HordeServer/Api/Jobs.cs b/Engine/Source/Programs/Horde/HordeServer/Api/Jobs.cs
--- a/Engine/Source/Programs/Horde/HordeServer/Api/Jobs.cs
+++ b/Engine/Source/Programs/Horde/HordeServer/Api/Jobs.cs
@@ -35,22 +35,6 @@ namespace HordeServer.Api
 		Complete,
 	}
 
-	/// <summary>
-	/// Query selecting the base changelist to use
-	/// </summary>
-	public class ChangeQueryRequest
-	{
-		/// <summary>
-		/// The template id to query
-		/// </summary>
-		public string TemplateId { get; set; } = String.Empty;
-
-		/// <summary>
-		/// The target to query
-		/// </summary>
-		public string? Target { get; set; }
-	}
-
 	/// <summary>
 	/// Parameters required to create a job
 	/// </summary>
diff --git a/Engine/Source/Programs/Horde/HordeServer/Api/Streams.cs b/Engine/Source/Programs/Horde/HordeServer/Api/Streams.cs
--- a/Engine/Source/Programs/Horde/HordeServer/Api/Streams.cs
+++ b/Engine/Source/Programs/Horde/HordeServer/Api/Streams.cs
@@ -14,6 +14,43 @@ using TemplateRefId = HordeServer.Utilities.StringId<HordeServer.Models.Template
 
 namespace HordeServer.Api
 {
+	/// <summary>
+	/// Query selecting the base changelist to use
+	/// </summary>
+	public class ChangeQueryRequest
+	{
+		/// <summary>
+		/// The template id to query
+		/// </summary>
+		public string? TemplateId { get; set; }
+
+		/// <summary>
+		/// The target to query
+		/// </summary>
+		public string? Target { get; set; }
+
+		/// <summary>
+		/// Whether to match a job that produced warnings
+		/// </summary>
+		public List<JobStepOutcome>? Outcomes { get; set; }
+
+		/// <summary>
+		/// Convert to a model object
+		/// </summary>
+		/// <returns></returns>
+		public ChangeQuery ToModel()
+		{
+			ChangeQuery Query = new ChangeQuery();
+			if (TemplateId != null)
+			{
+				Query.TemplateRefId = new StringId<TemplateRef>(TemplateId);
+			}
+			Query.Target = Target;
+			Query.Outcomes = Outcomes;
+			return Query;
+		}
+	}
+
 	/// <summary>
 	/// Specifies defaults for running a preflight
 	/// </summary>
@@ -27,15 +64,28 @@ namespace HordeServer.Api
 		/// <summary>
 		/// The last successful job type to use for the base changelist
 		/// </summary>
+		[Obsolete("Use Change.TemplateId instead")]
 		public string? ChangeTemplateId { get; set; }
 
+		/// <summary>
+		/// Query for the change to use
+		/// </summary>
+		public ChangeQueryRequest? Change { get; set; }
+
 		/// <summary>
 		/// Convert to a model object
 		/// </summary>
 		/// <returns></returns>
 		public DefaultPreflight ToModel()
 		{
-			return new DefaultPreflight((TemplateId != null) ? (TemplateRefId?)new TemplateRefId(TemplateId) : null, (ChangeTemplateId != null) ? (TemplateRefId?)new TemplateRefId(ChangeTemplateId) : null);
+#pragma warning disable CS0618 // Type or member is obsolete
+			if (ChangeTemplateId != null)
+			{
+				Change ??= new ChangeQueryRequest();
+				Change.TemplateId = ChangeTemplateId;
+			}
+			return new DefaultPreflight((TemplateId != null) ? (TemplateRefId?)new TemplateRefId(TemplateId) : null, Change?.ToModel());
+#pragma warning restore CS0618 // Type or member is obsolete
 		}
 	}
 
diff --git a/Engine/Source/Programs/Horde/HordeServer/Controllers/JobsController.cs b/Engine/Source/Programs/Horde/HordeServer/Controllers/JobsController.cs
--- a/Engine/Source/Programs/Horde/HordeServer/Controllers/JobsController.cs
+++ b/Engine/Source/Programs/Horde/HordeServer/Controllers/JobsController.cs
@@ -189,7 +189,7 @@ namespace HordeServer.Controllers
 			}
 			else if (Create.ChangeQuery != null)
 			{
-				Change = await ExecuteChangeQueryAsync(Stream, Create.ChangeQuery);
+				Change = await ExecuteChangeQueryAsync(Stream, new TemplateRefId(Create.ChangeQuery.TemplateId ?? Create.TemplateId), Create.ChangeQuery.Target, Create.ChangeQuery.Outcomes ?? new List<JobStepOutcome> { JobStepOutcome.Success });
 			}
 			else if (Create.PreflightChange == null && Template.SubmitNewChange != null)
 			{
@@ -216,19 +216,21 @@ namespace HordeServer.Controllers
 		/// Evaluate a change query to determine which CL to run a job at
 		/// </summary>
 		/// <param name="Stream"></param>
-		/// <param name="Request"></param>
+		/// <param name="TemplateId"></param>
+		/// <param name="Target"></param>
+		/// <param name="Outcomes"></param>
 		/// <returns></returns>
-		async Task<int> ExecuteChangeQueryAsync(IStream Stream, ChangeQueryRequest Request)
+		async Task<int> ExecuteChangeQueryAsync(IStream Stream, TemplateRefId TemplateId, string? Target, List<JobStepOutcome> Outcomes)
 		{
-			IList<IJob> Jobs = await JobService.FindJobsAsync(StreamId: Stream.Id, Templates: new[] { new TemplateRefId(Request.TemplateId) }, Target: Request.Target, State: new[] { JobStepState.Completed }, Outcome: new[] { JobStepOutcome.Success }, Count: 1);
+			IList<IJob> Jobs = await JobService.FindJobsAsync(StreamId: Stream.Id, Templates: new[] { TemplateId }, Target: Target, State: new[] { JobStepState.Completed }, Outcome: Outcomes.ToArray(), Count: 1);
 			if (Jobs.Count == 0)
 			{
-				Logger.LogInformation("Unable to find successful build of {TemplateRefId} target {Target}. Using latest change instead", Request.TemplateId, Request.Target);
+				Logger.LogInformation("Unable to find successful build of {TemplateRefId} target {Target}. Using latest change instead", TemplateId, Target);
 				return await Perforce.GetLatestChangeAsync(Stream.Name, null);
 			}
 			else
 			{
-				Logger.LogInformation("Last successful build of {TemplateRefId} target {Target} was job {JobId} at change {Change}", Request.TemplateId, Request.Target, Jobs[0].Id, Jobs[0].Change);
+				Logger.LogInformation("Last successful build of {TemplateRefId} target {Target} was job {JobId} at change {Change}", TemplateId, Target, Jobs[0].Id, Jobs[0].Change);
 				return Jobs[0].Change;
 			}
 		}
diff --git a/Engine/Source/Programs/Horde/HordeServer/Models/Stream.cs b/Engine/Source/Programs/Horde/HordeServer/Models/Stream.cs
--- a/Engine/Source/Programs/Horde/HordeServer/Models/Stream.cs
+++ b/Engine/Source/Programs/Horde/HordeServer/Models/Stream.cs
@@ -377,6 +377,36 @@ namespace HordeServer.Models
 		}
 	}
 
+	/// <summary>
+	/// Query used to identify a base changelist for a preflight
+	/// </summary>
+	public class ChangeQuery
+	{
+		/// <summary>
+		/// Template to search for
+		/// </summary>
+		public TemplateRefId? TemplateRefId { get; set; }
+
+		/// <summary>
+		/// The target to look at the status for
+		/// </summary>
+		public string? Target { get; set; }
+
+		/// <summary>
+		/// Whether to match a job that contains warnings
+		/// </summary>
+		public List<JobStepOutcome>? Outcomes { get; set; }
+
+		/// <summary>
+		/// Convert to a request object
+		/// </summary>
+		/// <returns></returns>
+		public ChangeQueryRequest ToRequest()
+		{
+			return new ChangeQueryRequest { TemplateId = TemplateRefId?.ToString(), Target = Target, Outcomes = Outcomes };
+		}
+	}
+
 	/// <summary>
 	/// Definition of a query to execute to find the changelist to run a build at
 	/// </summary>
@@ -387,20 +417,26 @@ namespace HordeServer.Models
 		/// </summary>
 		public TemplateRefId? TemplateRefId { get; set; }
 
+		/// <summary>
+		/// Query specifying a changelist to use
+		/// </summary>
+		public ChangeQuery? Change { get; set; }
+
 		/// <summary>
 		/// The job type to query for the change to use
 		/// </summary>
+		[Obsolete("Use Change.TemplateRefId instead")]
 		public TemplateRefId? ChangeTemplateRefId { get; set; }
 
 		/// <summary>
 		/// Constructor
 		/// </summary>
 		/// <param name="TemplateRefId"></param>
-		/// <param name="ChangeTemplateRefId">The job type to query for the change to use</param>
-		public DefaultPreflight(TemplateRefId? TemplateRefId, TemplateRefId? ChangeTemplateRefId)
+		/// <param name="Change">The job type to query for the change to use</param>
+		public DefaultPreflight(TemplateRefId? TemplateRefId, ChangeQuery? Change)
 		{
 			this.TemplateRefId = TemplateRefId;
-			this.ChangeTemplateRefId = ChangeTemplateRefId;
+			this.Change = Change;
 		}
 
 		/// <summary>
@@ -409,7 +445,9 @@ namespace HordeServer.Models
 		/// <returns></returns>
 		public DefaultPreflightRequest ToRequest()
 		{
-			return new DefaultPreflightRequest { TemplateId = TemplateRefId?.ToString(), ChangeTemplateId = ChangeTemplateRefId?.ToString() };
+#pragma warning disable CS0618 // Type or member is obsolete
+			return new DefaultPreflightRequest { TemplateId = TemplateRefId?.ToString(), Change = Change?.ToRequest(), ChangeTemplateId = Change?.TemplateRefId?.ToString() };
+#pragma warning restore CS0618 // Type or member is obsolete
 		}
 	}
 
@@ -610,10 +648,6 @@ namespace HordeServer.Models
 				{
 					throw new InvalidStreamException($"Default preflight template was listed as '{Stream.DefaultPreflight.TemplateRefId.Value}', but no template was found by that name");
 				}
-				if (Stream.DefaultPreflight.ChangeTemplateRefId != null && !Stream.Templates.ContainsKey(Stream.DefaultPreflight.ChangeTemplateRefId.Value))
-				{
-					throw new InvalidStreamException($"Default preflight template for CL was listed as '{Stream.DefaultPreflight.ChangeTemplateRefId.Value}', but no template was found by that name");
-				}
 			}
 
 			// Check that all the templates are referenced by a tab
