Horde: Adding preflight base CL outcome support to frontend

diff --git a/Engine/Source/Programs/Horde/HordeDashboard/src/backend/Api.ts b/Engine/Source/Programs/Horde/HordeDashboard/src/backend/Api.ts
--- a/Engine/Source/Programs/Horde/HordeDashboard/src/backend/Api.ts
+++ b/Engine/Source/Programs/Horde/HordeDashboard/src/backend/Api.ts
@@ -691,8 +691,6 @@ export type CreateLabelRequest = {
 	includedNodes: string[];
 }
 
-
-
 /** Query selecting the base changelist to use */
 export type ChangeQueryRequest = {
 
@@ -701,6 +699,9 @@ export type ChangeQueryRequest = {
 
 	/** The target to query	*/
 	target?: string;
+
+	/** Whether to match a job that produced warnings */
+	outcomes?: JobStepOutcome[];
 }
 
 
@@ -1635,8 +1636,9 @@ export type DefaultPreflightRequest = {
 	/** The template id to query */
 	templateId?: string;
 
-	/**  The last successful job type to use for the base changelist */
-	changeTemplateId?: string;
+	/** Query for the change to use */
+	change?: ChangeQueryRequest;
+
 }
 
 /**Response describing a stream */
