Horde: Fix missing changeTemplateId response in stream response.

diff --git a/Engine/Source/Programs/Horde/HordeServer/Models/Stream.cs b/Engine/Source/Programs/Horde/HordeServer/Models/Stream.cs
--- a/Engine/Source/Programs/Horde/HordeServer/Models/Stream.cs
+++ b/Engine/Source/Programs/Horde/HordeServer/Models/Stream.cs
@@ -446,7 +446,17 @@ namespace HordeServer.Models
 		public DefaultPreflightRequest ToRequest()
 		{
 #pragma warning disable CS0618 // Type or member is obsolete
-			return new DefaultPreflightRequest { TemplateId = TemplateRefId?.ToString(), Change = Change?.ToRequest(), ChangeTemplateId = Change?.TemplateRefId?.ToString() };
+			ChangeQueryRequest? ChangeRequest = null;
+			if (Change != null)
+			{
+				ChangeRequest = Change.ToRequest();
+			}
+			else if (ChangeTemplateRefId != null)
+			{
+				ChangeRequest = new ChangeQueryRequest { TemplateId = ChangeTemplateRefId?.ToString() };
+			}
+
+			return new DefaultPreflightRequest { TemplateId = TemplateRefId?.ToString(), Change = ChangeRequest, ChangeTemplateId = Change?.TemplateRefId?.ToString() };
 #pragma warning restore CS0618 // Type or member is obsolete
 		}
 	}
