Horde: Add support for reading ignore rules from the workspace. Looks in Engine/Build/Horde/IgnorePatterns.txt and any platform/restricted equivalents.

diff --git a/Engine/Source/Programs/Horde/HordeAgent/Execution/BuildGraphExecutor.cs b/Engine/Source/Programs/Horde/HordeAgent/Execution/BuildGraphExecutor.cs
--- a/Engine/Source/Programs/Horde/HordeAgent/Execution/BuildGraphExecutor.cs
+++ b/Engine/Source/Programs/Horde/HordeAgent/Execution/BuildGraphExecutor.cs
@@ -627,6 +627,49 @@ namespace HordeAgent.Execution
 			return Result;
 		}
 
+		static void AddRestrictedDirs(List<DirectoryReference> Directories, string SubFolder)
+		{
+			int NumDirs = Directories.Count;
+			for (int Idx = 0; Idx < NumDirs; Idx++)
+			{
+				DirectoryReference SubDir = DirectoryReference.Combine(Directories[Idx], SubFolder);
+				if(DirectoryReference.Exists(SubDir))
+				{
+					Directories.AddRange(DirectoryReference.EnumerateDirectories(SubDir));
+				}
+			}
+		}
+
+		static async Task<List<string>> ReadIgnorePatternsAsync(DirectoryReference WorkspaceDir)
+		{
+			List<DirectoryReference> BaseDirs = new List<DirectoryReference>();
+			BaseDirs.Add(DirectoryReference.Combine(WorkspaceDir, "Engine"));
+			AddRestrictedDirs(BaseDirs, "Restricted");
+			AddRestrictedDirs(BaseDirs, "Platforms");
+
+			List<string> IgnorePatternLines = new List<string>(Properties.Resources.IgnorePatterns.Split('\n'));
+			foreach (DirectoryReference BaseDir in BaseDirs)
+			{
+				FileReference IgnorePatternFile = FileReference.Combine(BaseDir, "Build", "Horde", "IgnorePatterns.txt");
+				if (FileReference.Exists(IgnorePatternFile))
+				{
+					IgnorePatternLines.AddRange(await FileReference.ReadAllLinesAsync(IgnorePatternFile));
+				}
+			}
+
+			HashSet<string> IgnorePatterns = new HashSet<string>(StringComparer.Ordinal);
+			foreach (string Line in IgnorePatternLines)
+			{
+				string TrimLine = Line.Trim();
+				if (TrimLine.Length > 0 && TrimLine[0] != '#')
+				{
+					IgnorePatterns.Add(TrimLine);
+				}
+			}
+
+			return IgnorePatterns.ToList();
+		}
+
 		async Task<int> ExecuteCommandAsync(BeginStepResponse Step, DirectoryReference WorkspaceDir, string FileName, string Arguments, IReadOnlyDictionary<string, string> EnvVars, IReadOnlyDictionary<string, string> Credentials, ILogger Logger, CancellationToken CancellationToken)
 		{
 			Dictionary<string, string> NewEnvironment = new Dictionary<string, string>(EnvVars);
@@ -688,17 +731,8 @@ namespace HordeAgent.Execution
 					Context.WorkspaceDir = WorkspaceDir;
 					Context.PerforceStream = Stream.Name;
 					Context.PerforceChange = Job.Change;
-					
-					List<string> IgnorePatterns = new List<string>();
-					foreach (string Line in Properties.Resources.IgnorePatterns.Split('\n'))
-					{
-						string TrimLine = Line.Trim();
-						if (TrimLine.Length > 0 && TrimLine[0] != '#')
-						{
-							IgnorePatterns.Add(TrimLine);
-						}
-					}
-					
+
+					List<string> IgnorePatterns = await ReadIgnorePatternsAsync(WorkspaceDir);
 					using (LogParser Filter = new LogParser(Logger, Context, IgnorePatterns))
 					{
 						await Process.CopyToAsync((Buffer, Offset, Length) => Filter.WriteData(Buffer.AsSpan(Offset, Length), false), 4096, CancellationToken);
