Horde: Do not [REDACTED_HANDLE] users when positing to the triage channel.

diff --git a/Engine/Source/Programs/Horde/HordeServer/Notifications/Impl/SlackNotificationSink.cs b/Engine/Source/Programs/Horde/HordeServer/Notifications/Impl/SlackNotificationSink.cs
--- a/Engine/Source/Programs/Horde/HordeServer/Notifications/Impl/SlackNotificationSink.cs
+++ b/Engine/Source/Programs/Horde/HordeServer/Notifications/Impl/SlackNotificationSink.cs
@@ -665,11 +665,11 @@ namespace HordeServer.Notifications.Impl
 				{
 					if (Suspect.DeclinedAt == null)
 					{
-						DeclinedLines.Add($":heavy_minus_sign: Ignored by {await FormatMentionAsync(Suspect.AuthorId)}");
+						DeclinedLines.Add($":heavy_minus_sign: Ignored by {await FormatNameAsync(Suspect.AuthorId)} (CL {Suspect.Change})");
 					}
 					else
 					{
-						DeclinedLines.Add($":downvote: Declined by {await FormatMentionAsync(Suspect.AuthorId)} at {FormatSlackTime(Suspect.DeclinedAt.Value)}");
+						DeclinedLines.Add($":downvote: Declined by {await FormatNameAsync(Suspect.AuthorId)} at {FormatSlackTime(Suspect.DeclinedAt.Value)} (CL {Suspect.Change})");
 					}
 				}
 				Attachment.Blocks.Add(new SectionBlock(String.Join("\n", DeclinedLines)));
@@ -683,6 +683,16 @@ namespace HordeServer.Notifications.Impl
 			return $"issue_{Issue.Id}";
 		}
 
+		async Task<string> FormatNameAsync(ObjectId UserId)
+		{
+			IUser? User = await UserCollection.GetUserAsync(UserId);
+			if (User == null)
+			{
+				return $"User {UserId}";
+			}
+			return User.Name;
+		}
+
 		async Task<string> FormatMentionAsync(ObjectId UserId)
 		{
 			IUser? User = await UserCollection.GetUserAsync(UserId);
