Horde: Add caching for Slack user ids to prevent rate limiting kicking in.

diff --git a/Engine/Source/Programs/Horde/HordeServer/Notifications/Impl/SlackNotificationSink.cs b/Engine/Source/Programs/Horde/HordeServer/Notifications/Impl/SlackNotificationSink.cs
--- a/Engine/Source/Programs/Horde/HordeServer/Notifications/Impl/SlackNotificationSink.cs
+++ b/Engine/Source/Programs/Horde/HordeServer/Notifications/Impl/SlackNotificationSink.cs
@@ -9,6 +9,7 @@ using HordeServer.Services;
 using HordeServer.Utilities;
 using HordeServer.Utilities.BlockKit;
 using HordeServer.Utilities.Slack.BlockKit;
+using Microsoft.Extensions.Caching.Memory;
 using Microsoft.Extensions.Hosting;
 using Microsoft.Extensions.Logging;
 using Microsoft.Extensions.Options;
@@ -117,18 +118,27 @@ namespace HordeServer.Notifications.Impl
 			public string Digest { get; set; } = String.Empty;
 		}
 
+		class UserDocument
+		{
+			[BsonId]
+			public string Email { get; set; } = String.Empty;
+			public string? SlackUserId { get; set; }
+			public DateTime Time { get; set; }
+		}
+
 		IIssueService IssueService;
 		IUserCollection UserCollection;
 		ILogFileService LogFileService;
 		StreamService StreamService;
 		ServerSettings Settings;
 		IMongoCollection<MessageStateDocument> MessageStates;
+		IMongoCollection<UserDocument> Users;
 		ILogger Logger;
 
 		/// <summary>
 		/// Map of email address to Slack user ID.
 		/// </summary>
-		private readonly ConcurrentDictionary<string, string> EmailAddressToSlackResponse = new ConcurrentDictionary<string, string>();
+		private MemoryCache UserCache = new MemoryCache(new MemoryCacheOptions());
 
 		/// <summary>
 		/// Constructor
@@ -148,9 +158,20 @@ namespace HordeServer.Notifications.Impl
 			this.StreamService = StreamService;
 			this.Settings = Settings.Value;
 			this.MessageStates = DatabaseService.Database.GetCollection<MessageStateDocument>("Slack");
+			this.Users = DatabaseService.Database.GetCollection<UserDocument>("Slack.Users");
 			this.Logger = Logger;
 		}
 
+		/// <inheritdoc/>
+		public override void Dispose()
+		{
+			base.Dispose();
+
+			UserCache.Dispose();
+
+			GC.SuppressFinalize(this);
+		}
+
 		#region Message state 
 
 		async Task<MessageStateDocument?> GetMessageStateAsync(string Recipient, string EventId)
@@ -877,6 +898,25 @@ namespace HordeServer.Notifications.Impl
 			}
 		}
 
+		static bool ShouldUpdateUser(UserDocument? Document)
+		{
+			if(Document == null)
+			{
+				return true;
+			}
+
+			TimeSpan ExpiryTime;
+			if (Document.SlackUserId == null)
+			{
+				ExpiryTime = TimeSpan.FromMinutes(10.0);
+			}
+			else
+			{
+				ExpiryTime = TimeSpan.FromDays(1.0);
+			}
+			return Document.Time + ExpiryTime < DateTime.UtcNow;
+		}
+
 		private async Task<string?> GetSlackUserId(IUser User)
 		{
 			string? Email = User.Email;
@@ -886,28 +926,54 @@ namespace HordeServer.Notifications.Impl
 				return null;
 			}
 
-			string? SlackResponse;
-			if (!EmailAddressToSlackResponse.TryGetValue(Email, out SlackResponse))
+			UserDocument? UserDocument;
+			if (!UserCache.TryGetValue(Email, out UserDocument))
 			{
-				using HttpClient Client = new HttpClient();
+				UserDocument = await Users.Find(x => x.Email == Email).FirstOrDefaultAsync();
+				if (UserDocument == null || ShouldUpdateUser(UserDocument))
+				{
+					string? UserId = await GetSlackUserIdByEmail(Email);
+					if (UserDocument == null || UserId != null)
+					{
+						UserDocument = new UserDocument { Email = Email, SlackUserId = UserId, Time = DateTime.UtcNow };
+						await Users.ReplaceOneAsync(x => x.Email == Email, UserDocument, new ReplaceOptions { IsUpsert = true });
+					}
+				}
+				using (ICacheEntry Entry = UserCache.CreateEntry(Email))
+				{
+					Entry.SlidingExpiration = TimeSpan.FromMinutes(10.0);
+					Entry.Value = UserDocument;
+				}
+			}
 
-				using HttpRequestMessage GetUserIdRequest = new HttpRequestMessage(HttpMethod.Post, $"https://slack.com/api/users.lookupByEmail?email={Email}");
-				GetUserIdRequest.Headers.Add("Authorization", $"Bearer {Settings.SlackToken ?? ""}");
+			return UserDocument?.SlackUserId;
+		}
 
-				HttpResponseMessage Response = await Client.SendAsync(GetUserIdRequest);
-				SlackResponse = await Response.Content.ReadAsStringAsync();
+		private async Task<string?> GetSlackUserIdByEmail(string Email)
+		{
+			using HttpClient Client = new HttpClient();
 
-				EmailAddressToSlackResponse.TryAdd(Email, SlackResponse);
-			}
+			using HttpRequestMessage GetUserIdRequest = new HttpRequestMessage(HttpMethod.Post, $"https://slack.com/api/users.lookupByEmail?email={Email}");
+			GetUserIdRequest.Headers.Add("Authorization", $"Bearer {Settings.SlackToken ?? ""}");
+
+			HttpResponseMessage Response = await Client.SendAsync(GetUserIdRequest);
+			string? SlackResponse = await Response.Content.ReadAsStringAsync();
 
 			JObject ResponseObject = JObject.Parse(SlackResponse);
 			if (!ResponseObject["ok"].ToObject<bool>())
 			{
-				Logger.LogWarning("Unable to find Slack user id for {UserId} ({Name}): {Response}", User.Id, User.Name, SlackResponse);
+				Logger.LogWarning("Unable to find Slack user id for {Email}: {Response}", Email, SlackResponse);
+				return null;
+			}
+
+			string? UserId = ResponseObject["user"]["id"].ToString();
+			if (String.IsNullOrEmpty(UserId))
+			{
+				Logger.LogWarning("Unexpected response from Slack user query on {Email}: {Response}", Email, SlackResponse);
 				return null;
 			}
 
-			return ResponseObject["user"]["id"].ToString();
+			return UserId;
 		}
 
 		class SlackResponse
