Horde: Only create build health issues for scheduled builds by default.

diff --git a/Engine/Source/Programs/Horde/HordeServer/Api/Jobs.cs b/Engine/Source/Programs/Horde/HordeServer/Api/Jobs.cs
--- a/Engine/Source/Programs/Horde/HordeServer/Api/Jobs.cs
+++ b/Engine/Source/Programs/Horde/HordeServer/Api/Jobs.cs
@@ -82,6 +82,11 @@ namespace HordeServer.Api
 		/// </summary>
 		public bool? AutoSubmit { get; set; }
 
+		/// <summary>
+		/// Whether to update issues based on the outcome of this job
+		/// </summary>
+		public bool? UpdateIssues { get; set; }
+
 		/// <summary>
 		/// Nodes for the new job
 		/// </summary>
diff --git a/Engine/Source/Programs/Horde/HordeServer/Collections/IJobCollection.cs b/Engine/Source/Programs/Horde/HordeServer/Collections/IJobCollection.cs
--- a/Engine/Source/Programs/Horde/HordeServer/Collections/IJobCollection.cs
+++ b/Engine/Source/Programs/Horde/HordeServer/Collections/IJobCollection.cs
@@ -47,6 +47,7 @@ namespace HordeServer.Collections
 		/// <param name="StartedByUserName">User that started the job</param>
 		/// <param name="Priority">Priority of the job</param>
 		/// <param name="AutoSubmit">Whether to automatically submit the preflighted change on completion</param>
+		/// <param name="UpdateIssues">Whether to update issues based on the outcome of this job</param>
 		/// <param name="JobTriggers">List of downstream job triggers</param>
 		/// <param name="ShowUgsBadges">Whether to show badges in UGS for this job</param>
 		/// <param name="ShowUgsAlerts">Whether to show alerts in UGS for this job</param>
@@ -55,7 +56,7 @@ namespace HordeServer.Collections
 		/// <param name="HelixSwarmCallbackUrl">Callback URL to a Helix Server review, if any</param>
 		/// <param name="Arguments">Arguments for the job</param>
 		/// <returns>The new job document</returns>
-		Task<IJob> AddAsync(ObjectId JobId, StreamId StreamId, TemplateRefId TemplateRefId, ContentHash TemplateHash, IGraph Graph, string Name, int Change, int CodeChange, int? PreflightChange, int? ClonedPreflightChange, ObjectId? StartedByUserId, string? StartedByUserName, Priority? Priority, bool? AutoSubmit, List<ChainedJobTemplate>? JobTriggers, bool ShowUgsBadges, bool ShowUgsAlerts, string? NotificationChannel, string? NotificationChannelFilter, string? HelixSwarmCallbackUrl, List<string>? Arguments);
+		Task<IJob> AddAsync(ObjectId JobId, StreamId StreamId, TemplateRefId TemplateRefId, ContentHash TemplateHash, IGraph Graph, string Name, int Change, int CodeChange, int? PreflightChange, int? ClonedPreflightChange, ObjectId? StartedByUserId, string? StartedByUserName, Priority? Priority, bool? AutoSubmit, bool? UpdateIssues, List<ChainedJobTemplate>? JobTriggers, bool ShowUgsBadges, bool ShowUgsAlerts, string? NotificationChannel, string? NotificationChannelFilter, string? HelixSwarmCallbackUrl, List<string>? Arguments);
 
 		/// <summary>
 		/// Gets a job with the given unique id
diff --git a/Engine/Source/Programs/Horde/HordeServer/Collections/Impl/JobCollection.cs b/Engine/Source/Programs/Horde/HordeServer/Collections/Impl/JobCollection.cs
--- a/Engine/Source/Programs/Horde/HordeServer/Collections/Impl/JobCollection.cs
+++ b/Engine/Source/Programs/Horde/HordeServer/Collections/Impl/JobCollection.cs
@@ -217,6 +217,8 @@ namespace HordeServer.Collections.Impl
 			[BsonIgnoreIfNull]
 			public string? AutoSubmitMessage { get; set; }
 
+			public bool UpdateIssues { get; set; }
+
 			[BsonIgnoreIfNull]
 			public DateTimeOffset? CreateTime { get; set; }
 
@@ -268,7 +270,7 @@ namespace HordeServer.Collections.Impl
 				GraphHash = null!;
 			}
 
-			public JobDocument(ObjectId Id, StreamId StreamId, TemplateRefId TemplateId, ContentHash TemplateHash, ContentHash GraphHash, string Name, int Change, int CodeChange, int PreflightChange, int ClonedPreflightChange, ObjectId? StartedByUserId, string? StartedByUserName, Priority? Priority, bool? AutoSubmit, DateTime CreateTimeUtc, List<ChainedJobDocument> ChainedJobs, bool ShowUgsBadges, bool ShowUgsAlerts, string? NotificationChannel, string? NotificationChannelFilter, string? HelixSwarmCallbackUrl, List<string>? Arguments)
+			public JobDocument(ObjectId Id, StreamId StreamId, TemplateRefId TemplateId, ContentHash TemplateHash, ContentHash GraphHash, string Name, int Change, int CodeChange, int PreflightChange, int ClonedPreflightChange, ObjectId? StartedByUserId, string? StartedByUserName, Priority? Priority, bool? AutoSubmit, bool? UpdateIssues, DateTime CreateTimeUtc, List<ChainedJobDocument> ChainedJobs, bool ShowUgsBadges, bool ShowUgsAlerts, string? NotificationChannel, string? NotificationChannelFilter, string? HelixSwarmCallbackUrl, List<string>? Arguments)
 			{
 				this.Id = Id;
 				this.StreamId = StreamId;
@@ -284,6 +286,7 @@ namespace HordeServer.Collections.Impl
 				this.StartedByUser = StartedByUserName;
 				this.Priority = Priority ?? HordeCommon.Priority.Normal;
 				this.AutoSubmit = AutoSubmit ?? false;
+				this.UpdateIssues = UpdateIssues ?? (StartedByUserId == null && PreflightChange == 0);
 				this.CreateTimeUtc = CreateTimeUtc;
 				this.ChainedJobs = ChainedJobs;
 				this.ShowUgsBadges = ShowUgsBadges;
@@ -354,7 +357,7 @@ namespace HordeServer.Collections.Impl
 
 		/// <inheritdoc/>
 		[SuppressMessage("Compiler", "CA1054:URI parameters should not be strings")]
-		public async Task<IJob> AddAsync(ObjectId JobId, StreamId StreamId, TemplateRefId TemplateRefId, ContentHash TemplateHash, IGraph Graph, string Name, int Change, int CodeChange, int? PreflightChange, int? ClonedPreflightChange, ObjectId? StartedByUserId, string? StartedByUserName, Priority? Priority, bool? AutoSubmit, List<ChainedJobTemplate>? ChainedJobs, bool ShowUgsBadges, bool ShowUgsAlerts, string? NotificationChannel, string? NotificationChannelFilter, string? HelixSwarmCallbackUrl, List<string>? Arguments)
+		public async Task<IJob> AddAsync(ObjectId JobId, StreamId StreamId, TemplateRefId TemplateRefId, ContentHash TemplateHash, IGraph Graph, string Name, int Change, int CodeChange, int? PreflightChange, int? ClonedPreflightChange, ObjectId? StartedByUserId, string? StartedByUserName, Priority? Priority, bool? AutoSubmit, bool? UpdateIssues, List<ChainedJobTemplate>? ChainedJobs, bool ShowUgsBadges, bool ShowUgsAlerts, string? NotificationChannel, string? NotificationChannelFilter, string? HelixSwarmCallbackUrl, List<string>? Arguments)
 		{
 			List<ChainedJobDocument> JobTriggers = new List<ChainedJobDocument>();
 			if (ChainedJobs == null)
@@ -366,7 +369,7 @@ namespace HordeServer.Collections.Impl
 				JobTriggers = ChainedJobs.ConvertAll(x => new ChainedJobDocument(x));
 			}
 
-			JobDocument NewJob = new JobDocument(JobId, StreamId, TemplateRefId, TemplateHash, Graph.Id, Name, Change, CodeChange, PreflightChange ?? 0, ClonedPreflightChange ?? 0, StartedByUserId, StartedByUserName, Priority, AutoSubmit, DateTime.UtcNow, JobTriggers, ShowUgsBadges, ShowUgsAlerts, NotificationChannel, NotificationChannelFilter, HelixSwarmCallbackUrl, Arguments);
+			JobDocument NewJob = new JobDocument(JobId, StreamId, TemplateRefId, TemplateHash, Graph.Id, Name, Change, CodeChange, PreflightChange ?? 0, ClonedPreflightChange ?? 0, StartedByUserId, StartedByUserName, Priority, AutoSubmit, UpdateIssues, DateTime.UtcNow, JobTriggers, ShowUgsBadges, ShowUgsAlerts, NotificationChannel, NotificationChannelFilter, HelixSwarmCallbackUrl, Arguments);
 			CreateBatches(NewJob, Graph, Logger);
 
 			await Jobs.InsertOneAsync(NewJob);
diff --git a/Engine/Source/Programs/Horde/HordeServer/Controllers/JobsController.cs b/Engine/Source/Programs/Horde/HordeServer/Controllers/JobsController.cs
--- a/Engine/Source/Programs/Horde/HordeServer/Controllers/JobsController.cs
+++ b/Engine/Source/Programs/Horde/HordeServer/Controllers/JobsController.cs
@@ -207,7 +207,7 @@ namespace HordeServer.Controllers
 			List<string> Arguments = Create.Arguments ?? Template.GetDefaultArguments();
 
 			// Create the job
-			IJob Job = await JobService.CreateJobAsync(null, Stream.Id, TemplateRefId, Template.Id, Graph, Name, Change, CodeChange, Create.PreflightChange, null, User.GetUserId(), User.GetUserName(), Priority, Create.AutoSubmit, TemplateRef.ChainedJobs, TemplateRef.ShowUgsBadges, TemplateRef.ShowUgsAlerts, TemplateRef.NotificationChannel, TemplateRef.NotificationChannelFilter, null, Template.Counters, Arguments);
+			IJob Job = await JobService.CreateJobAsync(null, Stream.Id, TemplateRefId, Template.Id, Graph, Name, Change, CodeChange, Create.PreflightChange, null, User.GetUserId(), User.GetUserName(), Priority, Create.AutoSubmit, Create.UpdateIssues, TemplateRef.ChainedJobs, TemplateRef.ShowUgsBadges, TemplateRef.ShowUgsAlerts, TemplateRef.NotificationChannel, TemplateRef.NotificationChannelFilter, null, Template.Counters, Arguments);
 			await UpdateNotificationsAsync(Job.Id.ToString(), new UpdateNotificationsRequest { Slack = true });
 			return new CreateJobResponse(Job.Id.ToString());
 		}
diff --git a/Engine/Source/Programs/Horde/HordeServer/Models/Job.cs b/Engine/Source/Programs/Horde/HordeServer/Models/Job.cs
--- a/Engine/Source/Programs/Horde/HordeServer/Models/Job.cs
+++ b/Engine/Source/Programs/Horde/HordeServer/Models/Job.cs
@@ -586,6 +586,11 @@ namespace HordeServer.Models
 		/// </summary>
 		public string? AutoSubmitMessage { get; }
 
+		/// <summary>
+		/// Whether to update issues based on the outcome of this job
+		/// </summary>
+		public bool UpdateIssues { get; }
+
 		/// <summary>
 		/// Time that the job was created (in UTC)
 		/// </summary>
diff --git a/Engine/Source/Programs/Horde/HordeServer/Services/JobService.cs b/Engine/Source/Programs/Horde/HordeServer/Services/JobService.cs
--- a/Engine/Source/Programs/Horde/HordeServer/Services/JobService.cs
+++ b/Engine/Source/Programs/Horde/HordeServer/Services/JobService.cs
@@ -218,6 +218,7 @@ namespace HordeServer.Services
 		/// <param name="StartedByUserName">User that started the job</param>
 		/// <param name="Priority">Priority of the job</param>
 		/// <param name="AutoSubmit">Whether to automatically submit the preflighted change on completion</param>
+		/// <param name="UpdateIssues">Whether to update issues when this job completes</param>
 		/// <param name="JobTriggers">List of downstream job triggers</param>
 		/// <param name="ShowUgsBadges">Whether to show badges in UGS for this job</param>
 		/// <param name="ShowUgsAlerts">Whether to show alerts in UGS for this job</param>
@@ -227,7 +228,7 @@ namespace HordeServer.Services
 		/// <param name="Counters">Counters for the job</param>
 		/// <param name="Arguments">Arguments for the job</param>
 		/// <returns>Unique id representing the job</returns>
-		public async Task<IJob> CreateJobAsync(ObjectId? JobId, StreamId StreamId, TemplateRefId TemplateRefId, ContentHash TemplateHash, IGraph Graph, string Name, int Change, int CodeChange, int? PreflightChange, int? ClonedPreflightChange, ObjectId? StartedByUserId, string? StartedByUserName, Priority? Priority, bool? AutoSubmit, List<ChainedJobTemplate>? JobTriggers, bool ShowUgsBadges, bool ShowUgsAlerts, string? NotificationChannel, string? NotificationChannelFilter, string? HelixSwarmCallbackUrl, IReadOnlyList<ITemplateCounter> Counters, IReadOnlyList<string> Arguments)
+		public async Task<IJob> CreateJobAsync(ObjectId? JobId, StreamId StreamId, TemplateRefId TemplateRefId, ContentHash TemplateHash, IGraph Graph, string Name, int Change, int CodeChange, int? PreflightChange, int? ClonedPreflightChange, ObjectId? StartedByUserId, string? StartedByUserName, Priority? Priority, bool? AutoSubmit, bool? UpdateIssues, List<ChainedJobTemplate>? JobTriggers, bool ShowUgsBadges, bool ShowUgsAlerts, string? NotificationChannel, string? NotificationChannelFilter, string? HelixSwarmCallbackUrl, IReadOnlyList<ITemplateCounter> Counters, IReadOnlyList<string> Arguments)
 		{
 			ObjectId JobIdValue = JobId ?? ObjectId.GenerateNewId();
 			using IDisposable Scope = Logger.BeginScope("CreateJobAsync({JobId})", JobIdValue);
@@ -268,7 +269,7 @@ namespace HordeServer.Services
 
 			Name = StringUtils.ExpandProperties(Name, Properties);
 
-			IJob NewJob = await Jobs.AddAsync(JobIdValue, StreamId, TemplateRefId, TemplateHash, Graph, Name, Change, CodeChange, PreflightChange, ClonedPreflightChange, StartedByUserId, StartedByUserName, Priority, AutoSubmit, JobTriggers, ShowUgsBadges, ShowUgsAlerts, NotificationChannel, NotificationChannelFilter, HelixSwarmCallbackUrl, ExpandedArguments);
+			IJob NewJob = await Jobs.AddAsync(JobIdValue, StreamId, TemplateRefId, TemplateHash, Graph, Name, Change, CodeChange, PreflightChange, ClonedPreflightChange, StartedByUserId, StartedByUserName, Priority, AutoSubmit, UpdateIssues, JobTriggers, ShowUgsBadges, ShowUgsAlerts, NotificationChannel, NotificationChannelFilter, HelixSwarmCallbackUrl, ExpandedArguments);
 			JobTaskSource.UpdateQueuedJob(NewJob, Graph);
 
 			await JobTaskSource.UpdateUgsBadges(NewJob, Graph, new List<(LabelState, LabelOutcome)>());
@@ -912,7 +913,7 @@ namespace HordeServer.Services
 					}
 
 					// Update any issues that depend on this step
-					if (NewState == JobStepState.Completed && Job.PreflightChange == 0)
+					if (NewState == JobStepState.Completed && Job.UpdateIssues)
 					{
 						if (IssueService != null)
 						{
@@ -1091,7 +1092,7 @@ namespace HordeServer.Services
 					IGraph TriggerGraph = await Graphs.AddAsync(Template);
 					Logger.LogInformation("Creating downstream job {ChainedJobId} from job {JobId}", ChainedJobId, Job.Id);
 
-					await CreateJobAsync(ChainedJobId, Job.StreamId, JobTrigger.TemplateRefId, TemplateRef.Hash, TriggerGraph, TemplateRef.Name, Job.Change, Job.CodeChange, Job.PreflightChange, Job.ClonedPreflightChange, Job.StartedByUserId, Job.StartedByUser, Template.Priority, null, TemplateRef.ChainedJobs, false, false, TemplateRef.NotificationChannel, TemplateRef.NotificationChannelFilter, null, Template.Counters, Template.Arguments);
+					await CreateJobAsync(ChainedJobId, Job.StreamId, JobTrigger.TemplateRefId, TemplateRef.Hash, TriggerGraph, TemplateRef.Name, Job.Change, Job.CodeChange, Job.PreflightChange, Job.ClonedPreflightChange, Job.StartedByUserId, Job.StartedByUser, Template.Priority, null, Job.UpdateIssues, TemplateRef.ChainedJobs, false, false, TemplateRef.NotificationChannel, TemplateRef.NotificationChannelFilter, null, Template.Counters, Template.Arguments);
 					break;
 				}
 
diff --git a/Engine/Source/Programs/Horde/HordeServer/Services/ScheduleService.cs b/Engine/Source/Programs/Horde/HordeServer/Services/ScheduleService.cs
--- a/Engine/Source/Programs/Horde/HordeServer/Services/ScheduleService.cs
+++ b/Engine/Source/Programs/Horde/HordeServer/Services/ScheduleService.cs
@@ -355,7 +355,7 @@ namespace HordeServer.Services
 			foreach ((int Change, int CodeChange) in TriggerChanges.OrderBy(x => x.Change))
 			{
 				List<string> DefaultArguments = Template.GetDefaultArguments();
-				IJob NewJob = await JobService.CreateJobAsync(null, Stream.Id, TemplateRefId, Template.Id, Graph, Template.Name, Change, CodeChange, null, null, null, null, Template.Priority, null, TemplateRef.ChainedJobs, TemplateRef.ShowUgsBadges, TemplateRef.ShowUgsAlerts, TemplateRef.NotificationChannel, TemplateRef.NotificationChannelFilter, null, Template.Counters, DefaultArguments);
+				IJob NewJob = await JobService.CreateJobAsync(null, Stream.Id, TemplateRefId, Template.Id, Graph, Template.Name, Change, CodeChange, null, null, null, null, Template.Priority, null, null, TemplateRef.ChainedJobs, TemplateRef.ShowUgsBadges, TemplateRef.ShowUgsAlerts, TemplateRef.NotificationChannel, TemplateRef.NotificationChannelFilter, null, Template.Counters, DefaultArguments);
 				Logger.LogInformation("Started new job for {StreamName} template {TemplateName} at CL {Change} (Code CL {CodeChange}): {JobId}", Stream.Id, TemplateRef.Name, Change, CodeChange, NewJob.Id);
 				CreatedJobs.Add(NewJob);
 			}
diff --git a/Engine/Source/Programs/Horde/HordeServerTests/Fixture.cs b/Engine/Source/Programs/Horde/HordeServerTests/Fixture.cs
--- a/Engine/Source/Programs/Horde/HordeServerTests/Fixture.cs
+++ b/Engine/Source/Programs/Horde/HordeServerTests/Fixture.cs
@@ -95,12 +95,13 @@ namespace HordeServerTests
 				Priority: Priority.Normal,
 				null,
 				null,
+				null,
 				false,
 				false,
 				null,
 				null,
 				null,
-				Template.Counters,
+				Counters: Template.Counters,
 				Arguments: new List<string>()
 			);
 			Job1 = (await JobService.GetJobAsync(Job1.Id))!;
@@ -121,6 +122,7 @@ namespace HordeServerTests
 				Priority: Priority.Normal,
 				null,
 				null,
+				null,
 				false,
 				false,
 				null,
diff --git a/Engine/Source/Programs/Horde/HordeServerTests/JobCollectionTests.cs b/Engine/Source/Programs/Horde/HordeServerTests/JobCollectionTests.cs
--- a/Engine/Source/Programs/Horde/HordeServerTests/JobCollectionTests.cs
+++ b/Engine/Source/Programs/Horde/HordeServerTests/JobCollectionTests.cs
@@ -85,7 +85,7 @@ namespace HordeServerTests
 			Arguments.Add("-Target=Publish Client");
 			Arguments.Add("-Target=Post-Publish Client");
 
-			IJob Job = await JobCollection.AddAsync(ObjectId.GenerateNewId(), new StreamId("ue4-main"), new TemplateRefId("test-build"), ContentHash.SHA1("hello"), BaseGraph, "Test job", 123, 123, null, null, null, "Ben", null, null, null, false, false, null, null, null, Arguments);
+			IJob Job = await JobCollection.AddAsync(ObjectId.GenerateNewId(), new StreamId("ue4-main"), new TemplateRefId("test-build"), ContentHash.SHA1("hello"), BaseGraph, "Test job", 123, 123, null, null, null, "Ben", null, null, null, null, false, false, null, null, null, Arguments);
 
 			await StartBatch(Job, BaseGraph, 0);
 			await RunStep(Job, BaseGraph, 0, 0, JobStepOutcome.Success); // Setup Build
diff --git a/Engine/Source/Programs/Horde/HordeServerTests/JobServiceTests.cs b/Engine/Source/Programs/Horde/HordeServerTests/JobServiceTests.cs
--- a/Engine/Source/Programs/Horde/HordeServerTests/JobServiceTests.cs
+++ b/Engine/Source/Programs/Horde/HordeServerTests/JobServiceTests.cs
@@ -43,7 +43,7 @@ namespace HordeServerTests
 			IStream? Stream = await TestSetup.StreamService.GetStreamAsync(StreamId);
 			Stream = await TestSetup.StreamService.StreamCollection.TryCreateOrReplaceAsync(new StreamId("ue5-main"), Stream, String.Empty, ProjectId, StreamConfig);
 
-			IJob Job = await TestSetup.JobService.CreateJobAsync(null, Stream!.Id, TemplateRefId1, Template.Id, Graph, "Hello", 1234, 1233, 999, null, null, "joe", null, null, Stream.Templates[TemplateRefId1].ChainedJobs, true, true, null, null, null, Template.Counters, new List<string>());
+			IJob Job = await TestSetup.JobService.CreateJobAsync(null, Stream!.Id, TemplateRefId1, Template.Id, Graph, "Hello", 1234, 1233, 999, null, null, "joe", null, null, null, Stream.Templates[TemplateRefId1].ChainedJobs, true, true, null, null, null, Template.Counters, new List<string>());
 			Assert.AreEqual(1, Job.ChainedJobs.Count);
 
 			Assert.IsTrue(await TestSetup.JobService.UpdateBatchAsync(Job, Job.Batches[0].Id, ObjectId.GenerateNewId(), JobStepBatchState.Running));
@@ -108,6 +108,7 @@ namespace HordeServerTests
 				Priority: Priority.Normal,
 				null,
 				null,
+				null,
 				false,
 				false,
 				null,
@@ -166,7 +167,7 @@ namespace HordeServerTests
 
 			Graph = await TestSetup.GraphCollection.AppendAsync(Graph, new List<CreateGroupRequest> { GroupA, GroupB });
 
-			IJob Job = await TestSetup.JobService.CreateJobAsync(null, StreamId, new TemplateRefId("temp"), Template.Id, Graph, "Hello", 1234, 1233, 999, null, null, "joe", null, null, null, true, true, null, null, null, Template.Counters, new List<string> { "-Target=Pak" });
+			IJob Job = await TestSetup.JobService.CreateJobAsync(null, StreamId, new TemplateRefId("temp"), Template.Id, Graph, "Hello", 1234, 1233, 999, null, null, "joe", null, null, null, null, true, true, null, null, null, Template.Counters, new List<string> { "-Target=Pak" });
 
 			Assert.IsTrue(await TestSetup.JobService.UpdateBatchAsync(Job, Job.Batches[0].Id, ObjectId.GenerateNewId(), JobStepBatchState.Running));
 			Assert.IsNotNull(await TestSetup.JobService.UpdateStepAsync(Job, Job.Batches[0].Id, Job.Batches[0].Steps[0].Id, JobStepState.Running));
diff --git a/Engine/Source/Programs/Horde/HordeServerTests/SchedulerTests.cs b/Engine/Source/Programs/Horde/HordeServerTests/SchedulerTests.cs
--- a/Engine/Source/Programs/Horde/HordeServerTests/SchedulerTests.cs
+++ b/Engine/Source/Programs/Horde/HordeServerTests/SchedulerTests.cs
@@ -355,7 +355,7 @@ namespace HordeServerTests
 			Assert.AreEqual(0, Jobs2.Count);
 
 			// Create a job and fail it
-			IJob Job1 = await TestSetup.JobService.CreateJobAsync(null, StreamId, NewTemplateRefId1, Template.Id, GraphA, "Hello", 1234, 1233, 999, null, null, "joe", null, null, null, true, true, null, null, null, Template.Counters, new List<string> { "-Target=TriggerNext" });
+			IJob Job1 = await TestSetup.JobService.CreateJobAsync(null, StreamId, NewTemplateRefId1, Template.Id, GraphA, "Hello", 1234, 1233, 999, null, null, "joe", null, null, null, null, true, true, null, null, null, Template.Counters, new List<string> { "-Target=TriggerNext" });
 			Assert.IsTrue(await TestSetup.JobService.UpdateBatchAsync(Job1, Job1.Batches[0].Id, ObjectId.GenerateNewId(), JobStepBatchState.Running));
 			Assert.IsNotNull(await TestSetup.JobService.UpdateStepAsync(Job1, Job1.Batches[0].Id, Job1.Batches[0].Steps[0].Id, JobStepState.Completed, JobStepOutcome.Failure));
 			Assert.IsTrue(await TestSetup.JobService.UpdateBatchAsync(Job1, Job1.Batches[0].Id, ObjectId.GenerateNewId(), JobStepBatchState.Complete));
@@ -368,7 +368,7 @@ namespace HordeServerTests
 			Assert.AreEqual(0, Jobs3.Count);
 
 			// Create a job and make it succeed
-			IJob Job2 = await TestSetup.JobService.CreateJobAsync(null, StreamId, NewTemplateRefId1, Template.Id, GraphA, "Hello", 1234, 1233, 999, null, null, "joe", null, null, null, true, true, null, null, null, Template.Counters, new List<string> { "-Target=TriggerNext" });
+			IJob Job2 = await TestSetup.JobService.CreateJobAsync(null, StreamId, NewTemplateRefId1, Template.Id, GraphA, "Hello", 1234, 1233, 999, null, null, "joe", null, null, null, null, true, true, null, null, null, Template.Counters, new List<string> { "-Target=TriggerNext" });
 			Assert.IsTrue(await TestSetup.JobService.UpdateBatchAsync(Job2, Job2.Batches[0].Id, ObjectId.GenerateNewId(), JobStepBatchState.Running));
 			Assert.IsNotNull(await TestSetup.JobService.UpdateStepAsync(Job2, Job2.Batches[0].Id, Job2.Batches[0].Steps[0].Id, JobStepState.Completed, JobStepOutcome.Success));
 
