Add support for sending environment variables in REAPI test client

diff --git a/Engine/Source/Programs/Horde/HordeAgent/Commands/Execution/ExecuteCommand.cs b/Engine/Source/Programs/Horde/HordeAgent/Commands/Execution/ExecuteCommand.cs
--- a/Engine/Source/Programs/Horde/HordeAgent/Commands/Execution/ExecuteCommand.cs
+++ b/Engine/Source/Programs/Horde/HordeAgent/Commands/Execution/ExecuteCommand.cs
@@ -121,6 +121,7 @@ namespace HordeAgent.Commands
 		{
 			public List<string> OutputPaths { get; set; } = new List<string>();
 			public List<string> Arguments { get; set; } = new List<string>();
+			public Dictionary<string, string> EnvVars { get; set; } = new Dictionary<string, string>();
 			public string WorkingDirectory { get; set; } = String.Empty;
 
 			public Digest Build(UploadList UploadList)
@@ -129,6 +130,12 @@ namespace HordeAgent.Commands
 				Command.Arguments.Add(Arguments);
 				Command.WorkingDirectory = WorkingDirectory;
 				Command.OutputPaths.Add(OutputPaths);
+				
+				foreach (KeyValuePair<string, string> Pair in EnvVars)
+				{
+					Command.EnvironmentVariables.Add(new RpcCommand.Types.EnvironmentVariable { Name = Pair.Key, Value = Pair.Value });
+				}
+				
 				return UploadList.Add(Command);
 			}
 		}
