Horde: New Unreal logo branding

diff --git a/Engine/Source/Programs/Horde/HordeDashboard/public/index.html b/Engine/Source/Programs/Horde/HordeDashboard/public/index.html
--- a/Engine/Source/Programs/Horde/HordeDashboard/public/index.html
+++ b/Engine/Source/Programs/Horde/HordeDashboard/public/index.html
@@ -2,7 +2,7 @@
 <html lang="en">
   <head>
     <meta charset="utf-8" />
-    <link rel="icon" href="%PUBLIC_URL%/favicon.ico" />
+    <link rel="icon" href="%PUBLIC_URL%/favicon.svg" />
     <meta name="viewport" content="width=device-width, initial-scale=1" />
     <meta name="theme-color" content="#000000" />
     <meta
diff --git a/Engine/Source/Programs/Horde/HordeDashboard/public/manifest.json b/Engine/Source/Programs/Horde/HordeDashboard/public/manifest.json
--- a/Engine/Source/Programs/Horde/HordeDashboard/public/manifest.json
+++ b/Engine/Source/Programs/Horde/HordeDashboard/public/manifest.json
@@ -3,9 +3,8 @@
   "name": "Horde Dashboard",
   "icons": [
     {
-      "src": "favicon.ico",
-      "sizes": "64x64 32x32 24x24 16x16",
-      "type": "image/x-icon"
+      "src": "favicon.svg",
+      "type": "image/svg+xml"
     },
     {
       "src": "logo192.png",
