Detect and expand directories in list of output paths to upload from REAPI tasks

Earlier, only specific files could be uploaded as specified in the command.
The REAPI spec allows specifiying directories here as well.

diff --git a/Engine/Source/Programs/Horde/HordeAgent/Execution/ActionExecutor.cs b/Engine/Source/Programs/Horde/HordeAgent/Execution/ActionExecutor.cs
--- a/Engine/Source/Programs/Horde/HordeAgent/Execution/ActionExecutor.cs
+++ b/Engine/Source/Programs/Horde/HordeAgent/Execution/ActionExecutor.cs
@@ -156,21 +156,17 @@ namespace HordeAgent
 
 					Logger.LogInformation("exit: {ExitCode}", Process.ExitCode);
 
-					foreach (string OutputPath in Command.OutputPaths)
+					foreach (var (FileRef, OutputPath) in ResolveOutputPaths(SandboxDir, Command.OutputPaths))
 					{
-						FileReference FileRef = FileReference.Combine(SandboxDir, OutputPath);
-						if (FileReference.Exists(FileRef))
-						{
-							byte[] Bytes = await FileReference.ReadAllBytesAsync(FileRef);
-							Digest Digest = await Storage.PutBulkDataAsync(InstanceName, Bytes);
-
-							OutputFile OutputFileInfo = new OutputFile();
-							OutputFileInfo.Path = OutputPath;
-							OutputFileInfo.Digest = Digest;
-							Result.OutputFiles.Add(OutputFileInfo);
-
-							Logger.LogInformation("Uploaded {File} (digest: {Digest}, size: {Size})", FileRef, Digest.Hash, Digest.SizeBytes);
-						}
+						byte[] Bytes = await FileReference.ReadAllBytesAsync(FileRef);
+						Digest Digest = await Storage.PutBulkDataAsync(InstanceName, Bytes);
+				
+						OutputFile OutputFileInfo = new OutputFile();
+						OutputFileInfo.Path = OutputPath;
+						OutputFileInfo.Digest = Digest;
+						Result.OutputFiles.Add(OutputFileInfo);
+				
+						Logger.LogInformation("Uploaded {File} (digest: {Digest}, size: {Size})", FileRef, Digest.Hash, Digest.SizeBytes);
 					}
 
 					PostActionResultRequest PostResultRequest = new PostActionResultRequest();
@@ -220,5 +216,42 @@ namespace HordeAgent
 				await SetupSandboxAsync(InputSubDirectory, OutputSubDirectory);
 			}
 		}
+
+		/// <summary>
+		/// Resolves a list of output paths into file references
+		///
+		/// The REAPI spec allows directories to be specified as an output path which require all sub dirs and files
+		/// to be resolved.
+		/// </summary>
+		/// <param name="SandboxDir">Base directory where execution is taking place</param>
+		/// <param name="OutputPaths">List of output paths relative to SandboxDir</param>
+		/// <returns>List of resolved paths (incl expanded dirs)</returns>
+		internal static List<(FileReference FileRef, string RelativePath)> ResolveOutputPaths(DirectoryReference SandboxDir, IEnumerable<string> OutputPaths)
+		{
+			var Files = new List<(FileReference FileRef, string RelativePath)>();
+
+			foreach (string OutputPath in OutputPaths)
+			{
+				DirectoryReference DirRef = DirectoryReference.Combine(SandboxDir, OutputPath);
+				if (DirectoryReference.Exists(DirRef))
+				{
+					IEnumerable<FileReference> ListedFiles = DirectoryReference.EnumerateFiles(DirRef, "*", SearchOption.AllDirectories);
+					foreach (FileReference ListedFileRef in ListedFiles)
+					{
+						Files.Add((ListedFileRef, ListedFileRef.MakeRelativeTo(SandboxDir).Replace("\\", "/")));
+					}
+				}
+				else
+				{
+					FileReference FileRef = FileReference.Combine(SandboxDir, OutputPath);
+					if (FileReference.Exists(FileRef))
+					{
+						Files.Add((FileRef, OutputPath));
+					}
+				}
+			}
+
+			return Files;
+		}
 	}
 }
diff --git a/Engine/Source/Programs/Horde/HordeAgentTests/ActionExecutorTest.cs b/Engine/Source/Programs/Horde/HordeAgentTests/ActionExecutorTest.cs
new file mode 100644
--- /dev/null
+++ b/Engine/Source/Programs/Horde/HordeAgentTests/ActionExecutorTest.cs
@@ -0,0 +1,76 @@
+// Copyright Epic Games, Inc. All Rights Reserved.
+
+using System.IO;
+using EpicGames.Core;
+using HordeAgent;
+using Microsoft.VisualStudio.TestTools.UnitTesting;
+
+namespace HordeAgentTests
+{
+	[TestClass]
+	public class ActionExecutorTest
+	{
+		private readonly string TempDir;
+
+		public ActionExecutorTest()
+		{
+			TempDir = GetTemporaryDirectory();
+		}
+
+		[TestMethod]
+		public void TestPaths()
+		{
+			string SubDirPath = Path.Join(TempDir, "1");
+			string SubSubDirPath = Path.Join(SubDirPath, "2");
+			string SubSubSubDirPath = Path.Join(SubSubDirPath, "3");
+			Directory.CreateDirectory(SubSubSubDirPath);
+			File.WriteAllText(Path.Join(TempDir, "foo"), "a");
+			File.WriteAllText(Path.Join(SubDirPath, "bar"), "b");
+			File.WriteAllText(Path.Join(SubSubDirPath, "baz"), "b");
+			File.WriteAllText(Path.Join(SubSubDirPath, "qux"), "c");
+			File.WriteAllText(Path.Join(SubSubSubDirPath, "waldo"), "d");
+
+			string[] OutputPaths =
+			{
+				"foo",
+				"1/bar",
+				"1/2", // directory
+				"doesNotExist",
+				//"foo", // duplicated (to be fixed)
+			};
+
+			DirectoryReference SandboxDir = new DirectoryReference(TempDir);
+			var Resolved = ActionExecutor.ResolveOutputPaths(SandboxDir, OutputPaths);
+			Assert.AreEqual(Path.Join(TempDir, "foo"), Resolved[0].FileRef.FullName);
+			Assert.AreEqual("foo", Resolved[0].RelativePath);
+			
+			Assert.AreEqual(Path.Join(TempDir, "1", "bar"), Resolved[1].FileRef.FullName);
+			Assert.AreEqual("1/bar", Resolved[1].RelativePath);
+			Assert.AreEqual(Path.Join(TempDir, "1", "2", "baz"), Resolved[2].FileRef.FullName);
+			Assert.AreEqual("1/2/baz", Resolved[2].RelativePath);
+			Assert.AreEqual(Path.Join(TempDir, "1", "2", "qux"), Resolved[3].FileRef.FullName);
+			Assert.AreEqual("1/2/qux", Resolved[3].RelativePath);
+			Assert.AreEqual(Path.Join(TempDir, "1", "2", "3", "waldo"), Resolved[4].FileRef.FullName);
+			Assert.AreEqual("1/2/3/waldo", Resolved[4].RelativePath);
+		}
+
+		[TestCleanup]
+		public void Cleanup()
+		{
+			Directory.Delete(TempDir, true);
+		}
+		
+		private static string GetTemporaryDirectory()
+		{
+			string TempDir = Path.Join(Path.GetTempPath(), "horde-" + Path.GetRandomFileName());
+			if (Directory.Exists(TempDir))
+			{
+				Directory.Delete(TempDir, true);
+			}
+
+			Directory.CreateDirectory(TempDir);
+
+			return TempDir;
+		}
+	}
+}
\ No newline at end of file
