Horde: Adding perforce service router

diff --git a/Engine/Source/Programs/Horde/HordeServer/ServerSettings.cs b/Engine/Source/Programs/Horde/HordeServer/ServerSettings.cs
--- a/Engine/Source/Programs/Horde/HordeServer/ServerSettings.cs
+++ b/Engine/Source/Programs/Horde/HordeServer/ServerSettings.cs
@@ -371,6 +371,11 @@ namespace HordeServer
 		/// </summary>
 		public string? P4BridgeServicePassword { get; set; }
 
+		/// <summary>
+		/// Whether to use the (temporary) P4 router while transitioning to p4 api 
+		/// </summary>
+		public bool P4UseRouter { get; set; } = false;
+
 		/// <summary>
 		/// Set the minimum size of the global thread pool
 		/// This value has been found in need of tweaking to avoid timeouts with the Redis client during bursts
diff --git a/Engine/Source/Programs/Horde/HordeServer/Services/PerforceServiceRouter.cs b/Engine/Source/Programs/Horde/HordeServer/Services/PerforceServiceRouter.cs
new file mode 100644
--- /dev/null
+++ b/Engine/Source/Programs/Horde/HordeServer/Services/PerforceServiceRouter.cs
@@ -0,0 +1,120 @@
+// Copyright Epic Games, Inc. All Rights Reserved.
+
+using HordeServer.Models;
+using Microsoft.Extensions.Logging;
+using System.Collections.Generic;
+using System.Threading.Tasks;
+using Microsoft.Extensions.Options;
+using System;
+
+namespace HordeServer.Services
+{
+	/// <summary>
+	/// A temporary router for easing in the new P4 API based perforce service
+	/// which will replace the python based one
+	/// </summary>
+	public sealed class PerforceServiceRouter : PerforceServiceBase, IDisposable
+	{
+
+		IOptionsMonitor<ServerSettings> Settings;
+
+		ILogger<PerforceServiceRouter> Logger;
+
+		BridgePerforceService PythonBridgeService;
+
+		P4APIPerforceService P4APIService;
+
+		/// <summary>
+		/// Constructor
+		/// </summary>
+		public PerforceServiceRouter(IOptionsMonitor<ServerSettings> Settings, ILogger<PerforceServiceRouter> Logger, ILogger<P4APIPerforceService> P4APILogger, ILogger<BridgePerforceService> BridgePerforceLogger)
+		{
+			this.Settings = Settings;
+			this.Logger = Logger;
+
+			this.P4APIService = new P4APIPerforceService(Settings, P4APILogger);
+			this.PythonBridgeService = new BridgePerforceService(Settings.CurrentValue.PerforceBridge!, BridgePerforceLogger);
+		}
+
+		/// <inheritdoc/>
+		override public Task<PerforceUserInfo?> GetUserInfoAsync(string UserName)
+		{
+			return P4APIService.GetUserInfoAsync(UserName);
+		}
+
+		/// <inheritdoc/>
+		override public Task<int> CreateNewChangeAsync(string StreamName, string Path)
+		{
+			return PythonBridgeService.CreateNewChangeAsync(StreamName, Path);
+		}
+
+		/// <inheritdoc/>
+		override public Task<List<ChangeSummary>> GetChangesAsync(string StreamName, int? MinChange, int? MaxChange, int Results, string? ImpersonateUser)
+		{
+			return P4APIService.GetChangesAsync(StreamName, MinChange, MaxChange, Results, ImpersonateUser);
+		}
+
+		/// <inheritdoc/>
+		override public Task<List<ChangeDetails>> GetChangeDetailsAsync(string StreamName, IReadOnlyList<int> ChangeNumbers, string? ImpersonateUser)
+		{
+			return P4APIService.GetChangeDetailsAsync(StreamName, ChangeNumbers, ImpersonateUser);
+		}
+
+		/// <inheritdoc />
+		override public Task<string> CreateTicket(string ImpersonateUser)
+		{
+			return P4APIService.CreateTicket(ImpersonateUser);
+		}
+
+		/// <inheritdoc/>
+		override public Task<int> GetCodeChangeAsync(string StreamName, int Change)
+		{
+			return P4APIService.GetCodeChangeAsync(StreamName, Change);
+		}
+
+		/// <inheritdoc/>
+		override public Task<List<FileSummary>> FindFilesAsync(IEnumerable<string> Paths)
+		{
+			return P4APIService.FindFilesAsync(Paths);
+		}
+
+		/// <inheritdoc/>
+		override public Task<byte[]> PrintAsync(string Path)
+		{
+			return PythonBridgeService.PrintAsync(Path);
+		}
+
+		/// <inheritdoc/>
+		override public Task<int> DuplicateShelvedChangeAsync(int ShelvedChange)
+		{
+			throw new Exception("DuplicateShelvedChangeAsync is disabled due to p4 issue: https://jira.it.epicgames.com/servicedesk/customer/portal/1/ITH-144069");	
+		}
+
+		/// <inheritdoc/>
+		override public Task<(int? Change, string Message)> SubmitShelvedChangeAsync(int ShelvedChange, int OriginalChange)
+		{
+			return PythonBridgeService.SubmitShelvedChangeAsync(ShelvedChange, OriginalChange);
+		}
+
+		/// <inheritdoc/>
+		override public Task DeleteShelvedChangeAsync(int ShelvedChange)
+		{
+			return PythonBridgeService.DeleteShelvedChangeAsync(ShelvedChange);
+		}
+
+		/// <inheritdoc/>
+		override public Task UpdateChangelistDescription(int Change, string Description)
+		{
+			return P4APIService.UpdateChangelistDescription(Change, Description );
+		}
+
+		/// <inheritdoc/>
+		public void Dispose()
+		{
+			PythonBridgeService.Dispose();
+		}
+
+	}
+
+}
+
diff --git a/Engine/Source/Programs/Horde/HordeServer/Startup.cs b/Engine/Source/Programs/Horde/HordeServer/Startup.cs
--- a/Engine/Source/Programs/Horde/HordeServer/Startup.cs
+++ b/Engine/Source/Programs/Horde/HordeServer/Startup.cs
@@ -366,10 +366,16 @@ namespace HordeServer
 			Services.AddSingleton<JobService>();
 			Services.AddSingleton<LifetimeService>();
 			Services.AddSingleton<ILogFileService, LogFileService>();
-			Services.AddSingleton<INotificationService, NotificationService>();			
+			Services.AddSingleton<INotificationService, NotificationService>();
 
-			
-			if (!String.IsNullOrEmpty(Settings.PerforceBridge))
+			if (Settings.P4UseRouter && !String.IsNullOrEmpty(Settings.PerforceBridge) && !String.IsNullOrEmpty(Settings.P4BridgeServer))
+			{
+				// using p4 service router, which will be removed along with the python p4 service
+				// once the p4 api service is hardened
+				Services.AddSingleton<IPerforceService, PerforceServiceRouter>();
+
+			}
+			else if (!String.IsNullOrEmpty(Settings.PerforceBridge))
 			{
 				// Python P4 bridge service
 				Services.AddSingleton<IPerforceService, BridgePerforceService>(SP => new BridgePerforceService(Settings.PerforceBridge!, SP.GetRequiredService<ILogger<BridgePerforceService>>()));				
