Horde: Fixing issue with pinned jobs and removed streamsm also putting in support for deprecated preflight template ref

diff --git a/Engine/Source/Programs/Horde/HordeDashboard/src/backend/Api.ts b/Engine/Source/Programs/Horde/HordeDashboard/src/backend/Api.ts
--- a/Engine/Source/Programs/Horde/HordeDashboard/src/backend/Api.ts
+++ b/Engine/Source/Programs/Horde/HordeDashboard/src/backend/Api.ts
@@ -1639,6 +1639,8 @@ export type DefaultPreflightRequest = {
 	/** Query for the change to use */
 	change?: ChangeQueryRequest;
 
+	/**  The last successful job type to use for the base changelist (DEPRECATED!)*/
+	changeTemplateId?: string;
 }
 
 /**Response describing a stream */
