Horde: List the problematic file when unable to parse project/stream configs, and allow other project/streams to continue to be updated.

diff --git a/Engine/Source/Programs/Horde/HordeServer/Services/ConfigService.cs b/Engine/Source/Programs/Horde/HordeServer/Services/ConfigService.cs
--- a/Engine/Source/Programs/Horde/HordeServer/Services/ConfigService.cs
+++ b/Engine/Source/Programs/Horde/HordeServer/Services/ConfigService.cs
@@ -117,6 +117,11 @@ namespace HordeServer.Services
 				{
 					Logger.LogInformation("Caching global config from {Revision}", Revision);
 					CachedGlobalConfig = await ReadDataAsync<GlobalConfig>(ConfigPath);
+					if (CachedGlobalConfig == null)
+					{
+						Logger.LogWarning("Unable to update any projects or streams due to invalid global config.");
+						return;
+					}
 				}
 				GlobalConfig = CachedGlobalConfig;
 
@@ -149,36 +154,49 @@ namespace HordeServer.Services
 
 			List<(ProjectId ProjectId, StreamConfigRef StreamRef, string Path)> StreamConfigs = new List<(ProjectId, StreamConfigRef, string)>();
 
+			// List of project ids that were not able to be updated. We will avoid removing any existing project or stream definitions for these.
+			HashSet<ProjectId> SkipProjectIds = new HashSet<ProjectId>();
+
 			// Update any existing projects
 			Dictionary<string, string> ProjectRevisions = await FindRevisionsAsync(ProjectConfigs.Select(x => x.Path));
 			for (int Idx = 0; Idx < ProjectConfigs.Count; Idx++)
 			{
+				// Make sure we were able to fetch metadata for 
 				(ProjectConfigRef ProjectRef, string ProjectPath) = ProjectConfigs[Idx];
-				if (ProjectRevisions.TryGetValue(ProjectPath, out string? Revision))
+				if (!ProjectRevisions.TryGetValue(ProjectPath, out string? Revision))
 				{
-					IProject? Project = Projects.FirstOrDefault(x => x.Id == ProjectRef.Id);
-					bool Update = (Project == null || Project.Revision != Revision);
+					Logger.LogWarning("Unable to update project {ProjectId} due to missing revision information", ProjectRef.Id);
+					SkipProjectIds.Add(ProjectRef.Id);
+					continue;
+				}
+
+				IProject? Project = Projects.FirstOrDefault(x => x.Id == ProjectRef.Id);
+				bool Update = (Project == null || Project.Revision != Revision);
 
-					ProjectConfig? ProjectConfig;
-					if (Update || !PrevCachedProjectConfigs.TryGetValue(ProjectRef.Id, out ProjectConfig))
+				ProjectConfig? ProjectConfig;
+				if (Update || !PrevCachedProjectConfigs.TryGetValue(ProjectRef.Id, out ProjectConfig))
+				{
+					Logger.LogInformation("Caching configuration for project {ProjectId} ({Revision})", ProjectRef.Id, Revision);
+					ProjectConfig = await ReadDataAsync<ProjectConfig>(ProjectPath);
+					if (ProjectConfig == null)
 					{
-						Logger.LogInformation("Caching configuration for project {ProjectId} ({Revision})", ProjectRef.Id, Revision);
-						ProjectConfig = await ReadDataAsync<ProjectConfig>(ProjectPath);
+						SkipProjectIds.Add(ProjectRef.Id);
+						continue;
 					}
 					if (Update)
 					{
 						Logger.LogInformation("Updating configuration for project {ProjectId} ({Revision})", ProjectRef.Id, Revision);
 						await ProjectService.Collection.AddOrUpdateAsync(ProjectRef.Id, Revision, Idx, ProjectConfig);
 					}
+				}
 
-					if (ProjectConfig.Logo != null)
-					{
-						ProjectLogos.Add((ProjectRef.Id, CombinePaths(ProjectPath, ProjectConfig.Logo)));
-					}
-
-					CachedProjectConfigs[ProjectRef.Id] = ProjectConfig;
-					StreamConfigs.AddRange(ProjectConfig.Streams.Select(x => (ProjectRef.Id, x, CombinePaths(ProjectPath, x.Path))));
+				if (ProjectConfig.Logo != null)
+				{
+					ProjectLogos.Add((ProjectRef.Id, CombinePaths(ProjectPath, ProjectConfig.Logo)));
 				}
+
+				CachedProjectConfigs[ProjectRef.Id] = ProjectConfig;
+				StreamConfigs.AddRange(ProjectConfig.Streams.Select(x => (ProjectRef.Id, x, CombinePaths(ProjectPath, x.Path))));
 			}
 
 			// Get the logo revisions
@@ -218,17 +236,20 @@ namespace HordeServer.Services
 					{
 						Logger.LogInformation("Updating configuration for stream {StreamRef} ({Revision})", StreamRef.Id, Revision);
 
-						StreamConfig StreamConfig = await ReadDataAsync<StreamConfig>(Path);
-						for (; ; )
+						StreamConfig? StreamConfig = await ReadDataAsync<StreamConfig>(Path);
+						if (StreamConfig != null)
 						{
-							Stream = await StreamService.StreamCollection.TryCreateOrReplaceAsync(StreamRef.Id, Stream, Revision, ProjectId, StreamConfig);
-
-							if (Stream != null)
+							for (; ; )
 							{
-								break;
-							}
+								Stream = await StreamService.StreamCollection.TryCreateOrReplaceAsync(StreamRef.Id, Stream, Revision, ProjectId, StreamConfig);
+
+								if (Stream != null)
+								{
+									break;
+								}
 
-							Stream = await StreamService.GetStreamAsync(StreamRef.Id);
+								Stream = await StreamService.GetStreamAsync(StreamRef.Id);
+							}
 						}
 					}
 				}
@@ -245,7 +266,7 @@ namespace HordeServer.Services
 			}
 
 			// Remove any streams that are no longer used
-			HashSet<StreamId> RemoveStreamIds = new HashSet<StreamId>(Streams.Select(x => x.Id));
+			HashSet<StreamId> RemoveStreamIds = new HashSet<StreamId>(Streams.Where(x => !SkipProjectIds.Contains(x.ProjectId)).Select(x => x.Id));
 			RemoveStreamIds.ExceptWith(StreamConfigs.Select(x => x.StreamRef.Id));
 
 			foreach (StreamId RemoveStreamId in RemoveStreamIds)
@@ -321,14 +342,22 @@ namespace HordeServer.Services
 			return Revisions;
 		}
 
-		async Task<T> ReadDataAsync<T>(string ConfigPath)
+		async Task<T?> ReadDataAsync<T>(string ConfigPath) where T : class
 		{
-			byte[] Data = await ReadDataAsync(ConfigPath);
+			try
+			{
+				byte[] Data = await ReadDataAsync(ConfigPath);
 
-			JsonSerializerOptions Options = new JsonSerializerOptions();
-			Startup.ConfigureJsonSerializer(Options);
+				JsonSerializerOptions Options = new JsonSerializerOptions();
+				Startup.ConfigureJsonSerializer(Options);
 
-			return JsonSerializer.Deserialize<T>(Data, Options);
+				return JsonSerializer.Deserialize<T>(Data, Options);
+			}
+			catch (Exception Ex)
+			{
+				Logger.LogError(Ex, "Unable to read data from {ConfigPath}: {Message}", ConfigPath, Ex.Message);
+				return null;
+			}
 		}
 
 		Task<byte[]> ReadDataAsync(string ConfigPath)
