Horde: Remove endpoints to create/delete projects through REST API. This is now driven by config files.

diff --git a/Engine/Source/Programs/Horde/HordeServer/Api/Projects.cs b/Engine/Source/Programs/Horde/HordeServer/Api/Projects.cs
--- a/Engine/Source/Programs/Horde/HordeServer/Api/Projects.cs
+++ b/Engine/Source/Programs/Horde/HordeServer/Api/Projects.cs
@@ -241,11 +241,6 @@ namespace HordeServer.Api
 		/// </summary>
 		public List<GetProjectCategoryResponse>? Categories { get; set; }
 
-		/// <summary>
-		/// Properties for this project.
-		/// </summary>
-		public Dictionary<string, string> Properties { get; set; }
-
 		/// <summary>
 		/// Custom permissions for this object
 		/// </summary>
@@ -258,7 +253,6 @@ namespace HordeServer.Api
 		{
 			Id = null!;
 			Name = null!;
-			Properties = null!;
 			Streams = null!;
 			Categories = null!;
 		}
@@ -271,16 +265,14 @@ namespace HordeServer.Api
 		/// <param name="Order">Order to show this project on the dashboard</param>
 		/// <param name="Streams">List of streams to display</param>
 		/// <param name="Categories">List of stream categories to display</param>
-		/// <param name="Properties">Properties for this project</param>
 		/// <param name="Acl">Custom permissions for this object</param>
-		public GetProjectResponse(string Id, string Name, int Order, List<GetProjectStreamResponse>? Streams, List<GetProjectCategoryResponse>? Categories, Dictionary<string, string> Properties, GetAclResponse? Acl)
+		public GetProjectResponse(string Id, string Name, int Order, List<GetProjectStreamResponse>? Streams, List<GetProjectCategoryResponse>? Categories, GetAclResponse? Acl)
 		{
 			this.Id = Id;
 			this.Name = Name;
 			this.Order = Order;
 			this.Streams = Streams;
 			this.Categories = Categories;
-			this.Properties = Properties;
 			this.Acl = Acl;
 		}
 	}
diff --git a/Engine/Source/Programs/Horde/HordeServer/Collections/IProjectCollection.cs b/Engine/Source/Programs/Horde/HordeServer/Collections/IProjectCollection.cs
--- a/Engine/Source/Programs/Horde/HordeServer/Collections/IProjectCollection.cs
+++ b/Engine/Source/Programs/Horde/HordeServer/Collections/IProjectCollection.cs
@@ -34,29 +34,6 @@ namespace HordeServer.Collections
 		/// <returns>New project instance</returns>
 		Task<IProject?> AddOrUpdateAsync(ProjectId Id, string Revision, int Order, ProjectConfig Config);
 
-		/// <summary>
-		/// Attempts to create a new project
-		/// </summary>
-		/// <param name="Id">Identifier for the new project</param>
-		/// <param name="Name">Name of the new project</param>
-		/// <param name="Order">Order to display this project on the dashboard</param>
-		/// <param name="Categories">Categories for the new project</param>
-		/// <param name="Properties">Properties for the new project</param>
-		/// <returns>The new project document</returns>
-		Task<IProject?> TryAddAsync(ProjectId Id, string Name, int? Order = null, List<StreamCategory>? Categories = null, Dictionary<string, string>? Properties = null);
-
-		/// <summary>
-		/// Updates an existing project
-		/// </summary>
-		/// <param name="ProjectId">Unique id of the project</param>
-		/// <param name="NewName">The new name for the project</param>
-		/// <param name="NewOrder">The new order for this project on the dashboard</param>
-		/// <param name="NewCategories">List of categories for the project</param>
-		/// <param name="NewProperties">Properties on the project to update. Any properties with a value of null will be removed.</param>
-		/// <param name="NewAcl">New acl for the project</param>
-		/// <returns>Async task object</returns>
-		Task UpdateAsync(ProjectId ProjectId, string? NewName, int? NewOrder, List<StreamCategory>? NewCategories, Dictionary<string, string>? NewProperties, Acl? NewAcl);
-
 		/// <summary>
 		/// Gets all the available projects
 		/// </summary>
diff --git a/Engine/Source/Programs/Horde/HordeServer/Collections/Impl/ProjectCollection.cs b/Engine/Source/Programs/Horde/HordeServer/Collections/Impl/ProjectCollection.cs
--- a/Engine/Source/Programs/Horde/HordeServer/Collections/Impl/ProjectCollection.cs
+++ b/Engine/Source/Programs/Horde/HordeServer/Collections/Impl/ProjectCollection.cs
@@ -42,13 +42,11 @@ namespace HordeServer.Collections.Impl
 			public int Order { get; set; } = DefaultOrder;
 			public Acl? Acl { get; set; }
 			public List<StreamCategory> Categories { get; set; } = new List<StreamCategory>();
-			public Dictionary<string, string> Properties { get; set; } = new Dictionary<string, string>(StringComparer.Ordinal);
 
 			[BsonIgnoreIfDefault, BsonDefaultValue(false)]
 			public bool Deleted { get; set; }
 
 			IReadOnlyList<StreamCategory> IProject.Categories => Categories;
-			IReadOnlyDictionary<string, string> IProject.Properties => Properties;
 
 			[BsonConstructor]
 			private ProjectDocument()
@@ -126,85 +124,6 @@ namespace HordeServer.Collections.Impl
 			return NewProject;
 		}
 
-		/// <inheritdoc/>
-		public async Task<IProject?> TryAddAsync(ProjectId Id, string Name, int? Order, List<StreamCategory>? Categories, Dictionary<string, string>? Properties)
-		{
-			ProjectDocument NewProject = new ProjectDocument(Id, Name);
-			if (Order != null)
-			{
-				NewProject.Order = Order.Value;
-			}
-			if (Categories != null)
-			{
-				NewProject.Categories = Categories;
-			}
-			if (Properties != null)
-			{
-				NewProject.Properties = new Dictionary<string, string>(Properties, NewProject.Properties.Comparer);
-			}
-
-			// Attempt to insert the stream, but fail gracefully if it already exists
-			try
-			{
-				await Projects.InsertOneAsync(NewProject);
-				return NewProject;
-			}
-			catch (MongoWriteException Ex)
-			{
-				if (Ex.WriteError.Category == ServerErrorCategory.DuplicateKey)
-				{
-					return null;
-				}
-				else
-				{
-					throw;
-				}
-			}
-		}
-
-		/// <inheritdoc/>
-		public async Task UpdateAsync(ProjectId ProjectId, string? NewName, int? NewOrder, List<StreamCategory>? NewCategories, Dictionary<string, string>? NewProperties, Acl? NewAcl)
-		{
-			UpdateDefinitionBuilder<ProjectDocument> UpdateBuilder = Builders<ProjectDocument>.Update;
-
-			List<UpdateDefinition<ProjectDocument>> Updates = new List<UpdateDefinition<ProjectDocument>>();
-			if (NewName != null)
-			{
-				Updates.Add(UpdateBuilder.Set(x => x.Name, NewName));
-			}
-			if (NewOrder != null)
-			{
-				Updates.Add(UpdateBuilder.Set(x => x.Order, NewOrder.Value));
-			}
-			if (NewCategories != null)
-			{
-				Updates.Add(UpdateBuilder.Set(x => x.Categories, NewCategories));
-			}
-			if (NewProperties != null)
-			{
-				foreach (KeyValuePair<string, string> Pair in NewProperties)
-				{
-					if (Pair.Value == null)
-					{
-						Updates.Add(UpdateBuilder.Unset(x => x.Properties[Pair.Key]));
-					}
-					else
-					{
-						Updates.Add(UpdateBuilder.Set(x => x.Properties[Pair.Key], Pair.Value));
-					}
-				}
-			}
-			if (NewAcl != null)
-			{
-				Updates.Add(Acl.CreateUpdate<ProjectDocument>(x => x.Acl!, NewAcl));
-			}
-
-			if (Updates.Count > 0)
-			{
-				await Projects.FindOneAndUpdateAsync<ProjectDocument>(x => x.Id == ProjectId, UpdateBuilder.Combine(Updates));
-			}
-		}
-
 		/// <inheritdoc/>
 		public async Task<List<IProject>> FindAllAsync()
 		{
diff --git a/Engine/Source/Programs/Horde/HordeServer/Controllers/ProjectsController.cs b/Engine/Source/Programs/Horde/HordeServer/Controllers/ProjectsController.cs
--- a/Engine/Source/Programs/Horde/HordeServer/Controllers/ProjectsController.cs
+++ b/Engine/Source/Programs/Horde/HordeServer/Controllers/ProjectsController.cs
@@ -67,31 +67,6 @@ namespace HordeServer.Controllers
 			this.ScheduleService = ScheduleService;
 		}
 
-		/// <summary>
-		/// Creates a new project
-		/// </summary>
-		/// <param name="Create">Parameters for the new project.</param>
-		/// <returns>Http result code</returns>
-		[HttpPost]
-		[Route("/api/v1/projects")]
-		public async Task<ActionResult<CreateProjectResponse>> CreateProjectAsync([FromBody] CreateProjectRequest Create)
-		{
-			if(!await AclService.AuthorizeAsync(AclAction.CreateProject, User, null))
-			{
-				return Forbid();
-			}
-
-			IProject? NewProject = await ProjectService.TryCreateProjectAsync(ProjectId.Sanitize(Create.Name), Create.Name, Create.Order, Create.Categories?.ConvertAll(x => new StreamCategory(x)), Create.Properties);
-			if (NewProject == null)
-			{
-				return BadRequest("Project already exists");
-			}
-			else
-			{
-				return new CreateProjectResponse(NewProject.Id.ToString());
-			}
-		}
-
 		/// <summary>
 		/// Query all the projects
 		/// </summary>
@@ -201,80 +176,5 @@ namespace HordeServer.Controllers
 
 			return new FileContentResult(ProjectLogo.Data, ProjectLogo.MimeType);
 		}
-
-		/// <summary>
-		/// Update a project's properties.
-		/// </summary>
-		/// <param name="ProjectId">Id of the project to update</param>
-		/// <param name="Update">Items on the project to update</param>
-		/// <returns>Http result code</returns>
-		[HttpPut]
-		[Route("/api/v1/projects/{ProjectId}")]
-		public async Task<ActionResult> UpdateProjectAsync(string ProjectId, [FromBody] UpdateProjectRequest Update)
-		{
-			ProjectId ProjectIdValue = new ProjectId(ProjectId);
-			for (; ; )
-			{
-				IProject? Project = await ProjectService.GetProjectAsync(ProjectIdValue);
-				if (Project == null)
-				{
-					if (!await AclService.AuthorizeAsync(AclAction.CreateProject, User, null))
-					{
-						return Forbid();
-					}
-					if (Update.Name == null)
-					{
-						return BadRequest("Project does not exist; 'Name' must be specified on update request to create it");
-					}
-
-					IProject? NewProject = await ProjectService.TryCreateProjectAsync(ProjectIdValue, Update.Name, Update.Order, Update.Categories?.ConvertAll(x => new StreamCategory(x)), Update.Properties);
-					if (NewProject != null)
-					{
-						return Ok();
-					}
-				}
-				else
-				{
-					ProjectPermissionsCache PermissionsCache = new ProjectPermissionsCache();
-					if (!await ProjectService.AuthorizeAsync(Project, AclAction.UpdateProject, User, PermissionsCache))
-					{
-						return Forbid();
-					}
-					if (Update.Acl != null && !await ProjectService.AuthorizeAsync(Project, AclAction.ChangePermissions, User, PermissionsCache))
-					{
-						return Forbid();
-					}
-
-					List<StreamCategory>? NewCategories = Update.Categories?.ConvertAll(x => new StreamCategory(x));
-					await ProjectService.UpdateProjectAsync(ProjectIdValue, Update.Name, Update.Order, NewCategories, Update.Properties, Acl.Merge(Project.Acl, Update.Acl));
-					return new OkResult();
-				}
-			}
-		}
-
-		/// <summary>
-		/// Delete a project
-		/// </summary>
-		/// <param name="ProjectId">Id of the project to update</param>
-		/// <returns>Http result code</returns>
-		[HttpDelete]
-		[Route("/api/v1/projects/{ProjectId}")]
-		public async Task<ActionResult> DeleteProjectAsync(string ProjectId)
-		{
-			ProjectId ProjectIdValue = new ProjectId(ProjectId);
-
-			IProject? Project = await ProjectService.GetProjectAsync(ProjectIdValue);
-			if (Project == null)
-			{
-				return NotFound();
-			}
-			if (!await ProjectService.AuthorizeAsync(Project, AclAction.DeleteProject, User, null))
-			{
-				return Forbid();
-			}
-
-			await ProjectService.DeleteProjectAsync(ProjectIdValue);
-			return new OkResult();
-		}
 	}
 }
diff --git a/Engine/Source/Programs/Horde/HordeServer/Models/Project.cs b/Engine/Source/Programs/Horde/HordeServer/Models/Project.cs
--- a/Engine/Source/Programs/Horde/HordeServer/Models/Project.cs
+++ b/Engine/Source/Programs/Horde/HordeServer/Models/Project.cs
@@ -123,11 +123,6 @@ namespace HordeServer.Models
 		/// List of stream categories for this project. Controls how streams are displayed on the dashboard.
 		/// </summary>
 		public IReadOnlyList<StreamCategory> Categories { get; }
-
-		/// <summary>
-		/// Optional user-defined properties for this project
-		/// </summary>
-		public IReadOnlyDictionary<string, string> Properties { get; }
 	}
 	
 	/// <summary>
@@ -208,7 +203,7 @@ namespace HordeServer.Models
 			}
 
 			GetAclResponse? AclResponse = (bIncludeAcl && Project.Acl != null) ? new GetAclResponse(Project.Acl) : null;
-			return new GetProjectResponse(Project.Id.ToString(), Project.Name, Project.Order, StreamResponses, CategoryResponses, new Dictionary<string, string>(Project.Properties), AclResponse);
+			return new GetProjectResponse(Project.Id.ToString(), Project.Name, Project.Order, StreamResponses, CategoryResponses, AclResponse);
 		}
 
 		/// <summary>
diff --git a/Engine/Source/Programs/Horde/HordeServer/Services/ProjectService.cs b/Engine/Source/Programs/Horde/HordeServer/Services/ProjectService.cs
--- a/Engine/Source/Programs/Horde/HordeServer/Services/ProjectService.cs
+++ b/Engine/Source/Programs/Horde/HordeServer/Services/ProjectService.cs
@@ -59,35 +59,6 @@ namespace HordeServer.Services
 			this.Projects = Projects;
 		}
 
-		/// <summary>
-		/// Attempts to create a new project
-		/// </summary>
-		/// <param name="Id">Identifier for the new project</param>
-		/// <param name="Name">Name of the new project</param>
-		/// <param name="Order">Order to display this project on the dashboard</param>
-		/// <param name="Categories">Categories for the new project</param>
-		/// <param name="Properties">Properties for the new project</param>
-		/// <returns>The new project document</returns>
-		public Task<IProject?> TryCreateProjectAsync(ProjectId Id, string Name, int? Order, List<StreamCategory>? Categories, Dictionary<string, string>? Properties)
-		{
-			return Projects.TryAddAsync(Id, Name, Order, Categories, Properties);
-		}
-
-		/// <summary>
-		/// Updates an existing project
-		/// </summary>
-		/// <param name="ProjectId">Unique id of the project</param>
-		/// <param name="NewName">The new name for the project</param>
-		/// <param name="NewOrder">The new order for this project on the dashboard</param>
-		/// <param name="NewCategories">List of categories for the project</param>
-		/// <param name="NewProperties">Properties on the project to update. Any properties with a value of null will be removed.</param>
-		/// <param name="NewAcl">New acl for the project</param>
-		/// <returns>Async task object</returns>
-		public Task UpdateProjectAsync(ProjectId ProjectId, string? NewName, int? NewOrder, List<StreamCategory>? NewCategories, Dictionary<string, string>? NewProperties, Acl? NewAcl)
-		{
-			return Projects.UpdateAsync(ProjectId, NewName, NewOrder, NewCategories, NewProperties, NewAcl);
-		}
-
 		/// <summary>
 		/// Gets all the available projects
 		/// </summary>
