Horde: Remove some unused dependencies from the projects controller.

diff --git a/Engine/Source/Programs/Horde/HordeServer/Controllers/ProjectsController.cs b/Engine/Source/Programs/Horde/HordeServer/Controllers/ProjectsController.cs
--- a/Engine/Source/Programs/Horde/HordeServer/Controllers/ProjectsController.cs
+++ b/Engine/Source/Programs/Horde/HordeServer/Controllers/ProjectsController.cs
@@ -25,11 +25,6 @@ namespace HordeServer.Controllers
 	[Route("[controller]")]
 	public class ProjectsController : ControllerBase
 	{
-		/// <summary>
-		/// Singleton instance of the ACL service
-		/// </summary>
-		private readonly AclService AclService;
-
 		/// <summary>
 		/// Singleton instance of the project service
 		/// </summary>
@@ -40,31 +35,15 @@ namespace HordeServer.Controllers
 		/// </summary>
 		private readonly StreamService StreamService;
 
-		/// <summary>
-		/// Singleton instance of the job service
-		/// </summary>
-		private readonly JobService JobService;
-
-		/// <summary>
-		/// Singleton instance of the schedule service
-		/// </summary>
-		private readonly ScheduleService ScheduleService;
-
 		/// <summary>
 		/// Constructor
 		/// </summary>
-		/// <param name="AclService">The ACL service</param>
 		/// <param name="ProjectService">The project service</param>
 		/// <param name="StreamService">The stream service</param>
-		/// <param name="JobService">The job service</param>
-		/// <param name="ScheduleService">The schedule service</param>
-		public ProjectsController(AclService AclService, ProjectService ProjectService, StreamService StreamService, JobService JobService, ScheduleService ScheduleService)
+		public ProjectsController(ProjectService ProjectService, StreamService StreamService)
 		{
-			this.AclService = AclService;
 			this.ProjectService = ProjectService;
 			this.StreamService = StreamService;
-			this.JobService = JobService;
-			this.ScheduleService = ScheduleService;
 		}
 
 		/// <summary>
