Horde: Remove properties from stream objects (unused)

diff --git a/Engine/Source/Programs/Horde/HordeServer/Api/Streams.cs b/Engine/Source/Programs/Horde/HordeServer/Api/Streams.cs
--- a/Engine/Source/Programs/Horde/HordeServer/Api/Streams.cs
+++ b/Engine/Source/Programs/Horde/HordeServer/Api/Streams.cs
@@ -634,11 +634,6 @@ namespace HordeServer.Api
 		/// </summary>
 		public List<GetTemplateRefResponse> Templates { get; set; }
 
-		/// <summary>
-		/// Properties for this stream
-		/// </summary>
-		public Dictionary<string, string> Properties { get; set; }
-
 		/// <summary>
 		/// Custom permissions for this object
 		/// </summary>
@@ -665,7 +660,6 @@ namespace HordeServer.Api
 			Tabs = null!;
 			AgentTypes = null!;
 			Templates = null!;
-			Properties = null!;
 		}
 
 		/// <summary>
@@ -684,11 +678,10 @@ namespace HordeServer.Api
 		/// <param name="AgentTypes">Map of agent type name to description</param>
 		/// <param name="WorkspaceTypes">Map of workspace name to description</param>
 		/// <param name="Templates">Templates for this stream</param>
-		/// <param name="Properties">Properties for the stream</param>
 		/// <param name="Acl">Permissions for this object</param>
 		/// <param name="PausedUntil">Stream paused for new builds until this date</param>
 		/// <param name="PauseComment">Reason for stream being paused</param>
-		public GetStreamResponse(string Id, string ProjectId, string Name, string? Revision, int Order, string? NotificationChannel, string? NotificationChannelFilter, string? TriageChannel, DefaultPreflightRequest? DefaultPreflight, List<GetStreamTabResponse> Tabs, Dictionary<string, GetAgentTypeResponse> AgentTypes, Dictionary<string, GetWorkspaceTypeResponse>? WorkspaceTypes, List<GetTemplateRefResponse> Templates, Dictionary<string, string> Properties, GetAclResponse? Acl, DateTime? PausedUntil, string? PauseComment)
+		public GetStreamResponse(string Id, string ProjectId, string Name, string? Revision, int Order, string? NotificationChannel, string? NotificationChannelFilter, string? TriageChannel, DefaultPreflightRequest? DefaultPreflight, List<GetStreamTabResponse> Tabs, Dictionary<string, GetAgentTypeResponse> AgentTypes, Dictionary<string, GetWorkspaceTypeResponse>? WorkspaceTypes, List<GetTemplateRefResponse> Templates, GetAclResponse? Acl, DateTime? PausedUntil, string? PauseComment)
 		{
 			this.Id = Id;
 			this.ProjectId = ProjectId;
@@ -704,7 +697,6 @@ namespace HordeServer.Api
 			this.AgentTypes = AgentTypes;
 			this.WorkspaceTypes = WorkspaceTypes;
 			this.Templates = Templates;
-			this.Properties = Properties;
 			this.Acl = Acl;
 			this.PausedUntil = PausedUntil;
 			this.PauseComment = PauseComment;
diff --git a/Engine/Source/Programs/Horde/HordeServer/Models/Config.cs b/Engine/Source/Programs/Horde/HordeServer/Models/Config.cs
--- a/Engine/Source/Programs/Horde/HordeServer/Models/Config.cs
+++ b/Engine/Source/Programs/Horde/HordeServer/Models/Config.cs
@@ -179,11 +179,6 @@ namespace HordeServer.Models
 		/// </summary>
 		public List<CreateTemplateRefRequest> Templates { get; set; } = new List<CreateTemplateRefRequest>();
 
-		/// <summary>
-		/// Properties for the new stream
-		/// </summary>
-		public Dictionary<string, string>? Properties { get; set; }
-
 		/// <summary>
 		/// Custom permissions for this object
 		/// </summary>
diff --git a/Engine/Source/Programs/Horde/HordeServer/Models/Stream.cs b/Engine/Source/Programs/Horde/HordeServer/Models/Stream.cs
--- a/Engine/Source/Programs/Horde/HordeServer/Models/Stream.cs
+++ b/Engine/Source/Programs/Horde/HordeServer/Models/Stream.cs
@@ -547,11 +547,6 @@ namespace HordeServer.Models
 		/// </summary>
 		public IReadOnlyDictionary<TemplateRefId, TemplateRef> Templates { get; }
 
-		/// <summary>
-		/// Optional user-defined properties for this stream
-		/// </summary>
-		public IReadOnlyDictionary<string, string> Properties { get; }
-
 		/// <summary>
 		/// Last time that we queried for commits
 		/// </summary>
@@ -698,7 +693,7 @@ namespace HordeServer.Models
 			Dictionary<string, GetAgentTypeResponse> ApiAgentTypes = Stream.AgentTypes.ToDictionary(x => x.Key, x => x.Value.ToApiResponse());
 			Dictionary<string, GetWorkspaceTypeResponse> ApiWorkspaceTypes = Stream.WorkspaceTypes.ToDictionary(x => x.Key, x => x.Value.ToApiResponse());
 			GetAclResponse? ApiAcl = (bIncludeAcl && Stream.Acl != null)? new GetAclResponse(Stream.Acl) : null;
-			return new Api.GetStreamResponse(Stream.Id.ToString(), Stream.ProjectId.ToString(), Stream.Name, Stream.Revision, Stream.Order, Stream.NotificationChannel, Stream.NotificationChannelFilter, Stream.TriageChannel, Stream.DefaultPreflight?.ToRequest(), ApiTabs, ApiAgentTypes, ApiWorkspaceTypes, ApiTemplateRefs, new Dictionary<string, string>(Stream.Properties), ApiAcl, Stream.PausedUntil, Stream.PauseComment);
+			return new Api.GetStreamResponse(Stream.Id.ToString(), Stream.ProjectId.ToString(), Stream.Name, Stream.Revision, Stream.Order, Stream.NotificationChannel, Stream.NotificationChannelFilter, Stream.TriageChannel, Stream.DefaultPreflight?.ToRequest(), ApiTabs, ApiAgentTypes, ApiWorkspaceTypes, ApiTemplateRefs, ApiAcl, Stream.PausedUntil, Stream.PauseComment);
 		}
 
 		/// <summary>
@@ -711,7 +706,6 @@ namespace HordeServer.Models
 			HordeCommon.Rpc.GetStreamResponse Response = new HordeCommon.Rpc.GetStreamResponse();
 			Response.Name = Stream.Name;
 			Response.AgentTypes.Add(Stream.AgentTypes.ToDictionary(x => x.Key, x => x.Value.ToRpcResponse()));
-			Response.Properties.Add(Stream.Properties);
 			Response.LastCommitTime = Stream.LastCommitTime.HasValue? Timestamp.FromDateTime(Stream.LastCommitTime.Value) : new Timestamp();
 			return Response;
 		}
