Horde: Store config paths on project and stream objects.

diff --git a/Engine/Source/Programs/Horde/HordeServer/Api/Streams.cs b/Engine/Source/Programs/Horde/HordeServer/Api/Streams.cs
--- a/Engine/Source/Programs/Horde/HordeServer/Api/Streams.cs
+++ b/Engine/Source/Programs/Horde/HordeServer/Api/Streams.cs
@@ -582,7 +582,12 @@ namespace HordeServer.Api
 		/// <summary>
 		/// The config file path on the server
 		/// </summary>
-		public string? Revision { get; set; }
+		public string ConfigPath { get; set; } = String.Empty;
+
+		/// <summary>
+		/// Revision of the config file 
+		/// </summary>
+		public string ConfigRevision { get; set; } = String.Empty;
 
 		/// <summary>
 		/// Order to display in the list
@@ -668,7 +673,8 @@ namespace HordeServer.Api
 		/// <param name="Id">Unique id of the stream</param>
 		/// <param name="ProjectId">Unique id of the project containing the stream</param>
 		/// <param name="Name">Name of the stream</param>
-		/// <param name="Revision">The config file path on the server</param>
+		/// <param name="ConfigPath">Path to the config file for this stream</param>
+		/// <param name="ConfigRevision">The config file path on the server</param>
 		/// <param name="Order">Order to display this stream</param>
 		/// <param name="NotificationChannel"></param>
 		/// <param name="NotificationChannelFilter"></param>
@@ -681,12 +687,13 @@ namespace HordeServer.Api
 		/// <param name="Acl">Permissions for this object</param>
 		/// <param name="PausedUntil">Stream paused for new builds until this date</param>
 		/// <param name="PauseComment">Reason for stream being paused</param>
-		public GetStreamResponse(string Id, string ProjectId, string Name, string? Revision, int Order, string? NotificationChannel, string? NotificationChannelFilter, string? TriageChannel, DefaultPreflightRequest? DefaultPreflight, List<GetStreamTabResponse> Tabs, Dictionary<string, GetAgentTypeResponse> AgentTypes, Dictionary<string, GetWorkspaceTypeResponse>? WorkspaceTypes, List<GetTemplateRefResponse> Templates, GetAclResponse? Acl, DateTime? PausedUntil, string? PauseComment)
+		public GetStreamResponse(string Id, string ProjectId, string Name, string ConfigPath, string ConfigRevision, int Order, string? NotificationChannel, string? NotificationChannelFilter, string? TriageChannel, DefaultPreflightRequest? DefaultPreflight, List<GetStreamTabResponse> Tabs, Dictionary<string, GetAgentTypeResponse> AgentTypes, Dictionary<string, GetWorkspaceTypeResponse>? WorkspaceTypes, List<GetTemplateRefResponse> Templates, GetAclResponse? Acl, DateTime? PausedUntil, string? PauseComment)
 		{
 			this.Id = Id;
 			this.ProjectId = ProjectId;
 			this.Name = Name;
-			this.Revision = Revision;
+			this.ConfigPath = ConfigPath;
+			this.ConfigRevision = ConfigRevision;
 			this.Order = Order;
 			this.NotificationChannel = NotificationChannel;
 			this.NotificationChannelFilter = NotificationChannelFilter;
diff --git a/Engine/Source/Programs/Horde/HordeServer/Collections/IProjectCollection.cs b/Engine/Source/Programs/Horde/HordeServer/Collections/IProjectCollection.cs
--- a/Engine/Source/Programs/Horde/HordeServer/Collections/IProjectCollection.cs
+++ b/Engine/Source/Programs/Horde/HordeServer/Collections/IProjectCollection.cs
@@ -28,11 +28,12 @@ namespace HordeServer.Collections
 		/// Updates the project configuration
 		/// </summary>
 		/// <param name="Id">The project id</param>
+		/// <param name="ConfigPath">Path to the config file used to configure this project</param>
 		/// <param name="Revision">The config file revision</param>
 		/// <param name="Order">Order of the project</param>
 		/// <param name="Config">The configuration</param>
 		/// <returns>New project instance</returns>
-		Task<IProject?> AddOrUpdateAsync(ProjectId Id, string Revision, int Order, ProjectConfig Config);
+		Task<IProject?> AddOrUpdateAsync(ProjectId Id, string ConfigPath, string Revision, int Order, ProjectConfig Config);
 
 		/// <summary>
 		/// Gets all the available projects
@@ -58,11 +59,12 @@ namespace HordeServer.Collections
 		/// Sets the logo for a project
 		/// </summary>
 		/// <param name="ProjectId">The project id</param>
+		/// <param name="Path">Path to the source file</param>
 		/// <param name="Revision">Revision of the file</param>
 		/// <param name="MimeType"></param>
 		/// <param name="Data"></param>
 		/// <returns></returns>
-		Task SetLogoAsync(ProjectId ProjectId, string Revision, string MimeType, byte[] Data);
+		Task SetLogoAsync(ProjectId ProjectId, string Path, string Revision, string MimeType, byte[] Data);
 
 		/// <summary>
 		/// Gets a project's permissions info by ID
diff --git a/Engine/Source/Programs/Horde/HordeServer/Collections/IStreamCollection.cs b/Engine/Source/Programs/Horde/HordeServer/Collections/IStreamCollection.cs
--- a/Engine/Source/Programs/Horde/HordeServer/Collections/IStreamCollection.cs
+++ b/Engine/Source/Programs/Horde/HordeServer/Collections/IStreamCollection.cs
@@ -29,11 +29,12 @@ namespace HordeServer.Collections
 		/// </summary>
 		/// <param name="Id">Unique id for the new stream</param>
 		/// <param name="Stream">The current stream value. If not-null, this will attempt to replace the existing instance.</param>
+		/// <param name="ConfigPath">Path to the config file</param>
 		/// <param name="Revision">The config file revision</param>
 		/// <param name="ProjectId">The project id</param>
 		/// <param name="Config">The stream configuration</param>
 		/// <returns></returns>
-		Task<IStream?> TryCreateOrReplaceAsync(StreamId Id, IStream? Stream, string Revision, ProjectId ProjectId, StreamConfig Config);
+		Task<IStream?> TryCreateOrReplaceAsync(StreamId Id, IStream? Stream, string ConfigPath, string Revision, ProjectId ProjectId, StreamConfig Config);
 
 		/// <summary>
 		/// Gets a stream by ID
diff --git a/Engine/Source/Programs/Horde/HordeServer/Collections/Impl/ProjectCollection.cs b/Engine/Source/Programs/Horde/HordeServer/Collections/Impl/ProjectCollection.cs
--- a/Engine/Source/Programs/Horde/HordeServer/Collections/Impl/ProjectCollection.cs
+++ b/Engine/Source/Programs/Horde/HordeServer/Collections/Impl/ProjectCollection.cs
@@ -37,7 +37,8 @@ namespace HordeServer.Collections.Impl
 			[BsonRequired]
 			public string Name { get; set; }
 
-			public string? Revision { get; set; }
+			public string? ConfigPath { get; set; }
+			public string? ConfigRevision { get; set; }
 
 			public int Order { get; set; } = DefaultOrder;
 			public Acl? Acl { get; set; }
@@ -68,6 +69,7 @@ namespace HordeServer.Collections.Impl
 		{
 			public ProjectId Id { get; set; }
 
+			public string Path { get; set; } = String.Empty;
 			public string Revision { get; set; } = String.Empty;
 			public string MimeType { get; set; } = String.Empty;
 
@@ -112,10 +114,11 @@ namespace HordeServer.Collections.Impl
 		}
 
 		/// <inheritdoc/>
-		public async Task<IProject?> AddOrUpdateAsync(ProjectId Id, string Revision, int Order, ProjectConfig Config)
+		public async Task<IProject?> AddOrUpdateAsync(ProjectId Id, string ConfigPath, string Revision, int Order, ProjectConfig Config)
 		{
 			ProjectDocument NewProject = new ProjectDocument(Id, Config.Name);
-			NewProject.Revision = Revision;
+			NewProject.ConfigPath = ConfigPath;
+			NewProject.ConfigRevision = Revision;
 			NewProject.Order = Order;
 			NewProject.Categories = Config.Categories.ConvertAll(x => new StreamCategory(x));
 			NewProject.Acl = Acl.Merge(new Acl(), Config.Acl);
@@ -144,11 +147,12 @@ namespace HordeServer.Collections.Impl
 		}
 
 		/// <inheritdoc/>
-		public async Task SetLogoAsync(ProjectId ProjectId, string Revision, string MimeType, byte[] Data)
+		public async Task SetLogoAsync(ProjectId ProjectId, string LogoPath, string LogoRevision, string MimeType, byte[] Data)
 		{
 			ProjectLogoDocument Logo = new ProjectLogoDocument();
 			Logo.Id = ProjectId;
-			Logo.Revision = Revision;
+			Logo.Path = LogoPath;
+			Logo.Revision = LogoRevision;
 			Logo.MimeType = MimeType;
 			Logo.Data = Data;
 			await ProjectLogos.ReplaceOneAsync(x => x.Id == ProjectId, Logo, new ReplaceOptions { IsUpsert = true });
diff --git a/Engine/Source/Programs/Horde/HordeServer/Collections/Impl/StreamCollection.cs b/Engine/Source/Programs/Horde/HordeServer/Collections/Impl/StreamCollection.cs
--- a/Engine/Source/Programs/Horde/HordeServer/Collections/Impl/StreamCollection.cs
+++ b/Engine/Source/Programs/Horde/HordeServer/Collections/Impl/StreamCollection.cs
@@ -41,8 +41,8 @@ namespace HordeServer.Collections.Impl
 			[BsonRequired]
 			public string Name { get; set; }
 
-			[BsonIgnoreIfNull]
-			public string? Revision { get; set; }
+			public string ConfigPath { get; set; } = String.Empty;
+			public string ConfigRevision { get; set; } = String.Empty;
 
 			public int Order { get; set; } = DefaultOrder;
 			public string? NotificationChannel { get; set; }
@@ -53,7 +53,6 @@ namespace HordeServer.Collections.Impl
 			public Dictionary<string, AgentType> AgentTypes { get; set; } = new Dictionary<string, AgentType>(StringComparer.Ordinal);
 			public Dictionary<string, WorkspaceType> WorkspaceTypes { get; set; } = new Dictionary<string, WorkspaceType>(StringComparer.Ordinal);
 			public Dictionary<TemplateRefId, TemplateRef> Templates { get; set; } = new Dictionary<TemplateRefId, TemplateRef>();
-			public Dictionary<string, string> Properties { get; set; } = new Dictionary<string, string>(StringComparer.Ordinal);
 			public DateTime? LastCommitTime { get; set; }
 			public DateTime? PausedUntil { get; set; }
 			public string? PauseComment { get; set; }
@@ -65,7 +64,6 @@ namespace HordeServer.Collections.Impl
 			IReadOnlyDictionary<string, AgentType> IStream.AgentTypes => AgentTypes;
 			IReadOnlyDictionary<string, WorkspaceType> IStream.WorkspaceTypes => WorkspaceTypes;
 			IReadOnlyDictionary<TemplateRefId, TemplateRef> IStream.Templates => Templates;
-			IReadOnlyDictionary<string, string> IStream.Properties => Properties;
 
 			[BsonConstructor]
 			private StreamDocument()
@@ -115,7 +113,7 @@ namespace HordeServer.Collections.Impl
 		}
 
 		/// <inheritdoc/>
-		public async Task<IStream?> TryCreateOrReplaceAsync(StreamId Id, IStream? Stream, string Revision, ProjectId ProjectId, StreamConfig Config)
+		public async Task<IStream?> TryCreateOrReplaceAsync(StreamId Id, IStream? Stream, string ConfigPath, string Revision, ProjectId ProjectId, StreamConfig Config)
 		{
 			List<StreamTab> Tabs = Config.Tabs.ConvertAll(x => StreamTab.FromRequest(x));
 			Dictionary<TemplateRefId, TemplateRef> TemplateRefs = await CreateTemplateRefsAsync(Config.Templates, null, TemplateCollection);
@@ -141,11 +139,11 @@ namespace HordeServer.Collections.Impl
 			Acl? Acl = Acl.Merge(new Acl(), Config.Acl);
 			if (Stream == null)
 			{
-				return await TryCreateAsync(Id, Config.Name, ProjectId, Revision, Config.Order, Config.NotificationChannel, Config.NotificationChannelFilter, Config.TriageChannel, DefaultPreflight, Tabs, AgentTypes, WorkspaceTypes, TemplateRefs, null, Acl);
+				return await TryCreateAsync(Id, ProjectId, ConfigPath, Revision, Config, DefaultPreflight, Tabs, AgentTypes, WorkspaceTypes, TemplateRefs, Acl);
 			}
 			else
 			{
-				return await TryReplaceAsync(Stream, Config.Name, ProjectId, Revision, Config.Order, Config.NotificationChannel, Config.NotificationChannelFilter, Config.TriageChannel, DefaultPreflight, Tabs, AgentTypes, WorkspaceTypes, TemplateRefs, null, Acl);
+				return await TryReplaceAsync(Stream, ProjectId, ConfigPath, Revision, Config, DefaultPreflight, Tabs, AgentTypes, WorkspaceTypes, TemplateRefs, Acl);
 			}
 		}
 
@@ -203,29 +201,16 @@ namespace HordeServer.Collections.Impl
 		}
 
 		/// <inheritdoc/>
-		public async Task<IStream?> TryCreateAsync(StreamId Id, string Name, ProjectId ProjectId, string? Revision, int? Order, string? NotificationChannel, string? NotificationChannelFilter, string? TriageChannel, DefaultPreflight? DefaultPreflight, List<StreamTab>? Tabs, Dictionary<string, AgentType>? AgentTypes, Dictionary<string, WorkspaceType>? WorkspaceTypes, Dictionary<TemplateRefId, TemplateRef>? TemplateRefs, Dictionary<string, string>? Properties, Acl? Acl)
+		public async Task<IStream?> TryCreateAsync(StreamId Id, ProjectId ProjectId, string ConfigPath, string ConfigRevision, StreamConfig Config, DefaultPreflight? DefaultPreflight, List<StreamTab>? Tabs, Dictionary<string, AgentType>? AgentTypes, Dictionary<string, WorkspaceType>? WorkspaceTypes, Dictionary<TemplateRefId, TemplateRef>? TemplateRefs, Acl? Acl)
 		{
-			StreamDocument NewStream = new StreamDocument(Id, Name, ProjectId);
-			if (Revision != null)
-			{
-				NewStream.Revision = Revision;
-			}
-			if (Order != null)
-			{
-				NewStream.Order = Order.Value;
-			}
-			if(NotificationChannel != null)
-			{
-				NewStream.NotificationChannel = NotificationChannel;
-			}
-			if (NotificationChannelFilter != null)
-			{
-				NewStream.NotificationChannelFilter = NotificationChannelFilter;
-			}
-			if (TriageChannel != null)
-			{
-				NewStream.TriageChannel = TriageChannel;
-			}
+			StreamDocument NewStream = new StreamDocument(Id, Config.Name, ProjectId);
+			NewStream.ConfigPath = ConfigPath;
+			NewStream.ConfigRevision = ConfigRevision;
+			NewStream.Order = Config.Order ?? StreamDocument.DefaultOrder;
+			NewStream.NotificationChannel = Config.NotificationChannel;
+			NewStream.NotificationChannelFilter = Config.NotificationChannelFilter;
+			NewStream.TriageChannel = Config.TriageChannel;
+
 			if (DefaultPreflight != null)
 			{
 				NewStream.DefaultPreflight = DefaultPreflight;
@@ -246,10 +231,6 @@ namespace HordeServer.Collections.Impl
 			{
 				NewStream.Templates = TemplateRefs;
 			}
-			if (Properties != null)
-			{
-				NewStream.Properties = new Dictionary<string, string>(Properties, NewStream.Properties.Comparer);
-			}
 			if (Acl != null)
 			{
 				NewStream.Acl = Acl;
@@ -275,7 +256,7 @@ namespace HordeServer.Collections.Impl
 		}
 
 		/// <inheritdoc/>
-		public async Task<IStream?> TryReplaceAsync(IStream StreamInterface, string Name, ProjectId ProjectId, string? Revision, int? Order, string? NotificationChannel, string? NotificationChannelFilter, string? TriageChannel, DefaultPreflight? DefaultPreflight, List<StreamTab>? Tabs, Dictionary<string, AgentType>? AgentTypes, Dictionary<string, WorkspaceType>? WorkspaceTypes, Dictionary<TemplateRefId, TemplateRef>? TemplateRefs, Dictionary<string, string>? Properties, Acl? Acl, DateTime? PausedUntil = null, string? PauseComment = null)
+		public async Task<IStream?> TryReplaceAsync(IStream StreamInterface, ProjectId ProjectId, string ConfigPath, string ConfigRevision, StreamConfig Config, DefaultPreflight? DefaultPreflight, List<StreamTab>? Tabs, Dictionary<string, AgentType>? AgentTypes, Dictionary<string, WorkspaceType>? WorkspaceTypes, Dictionary<TemplateRefId, TemplateRef>? TemplateRefs, Acl? Acl)
 		{
 			StreamDocument Stream = (StreamDocument)StreamInterface;
 
@@ -283,25 +264,28 @@ namespace HordeServer.Collections.Impl
 
 			List<UpdateDefinition<StreamDocument>> Updates = new List<UpdateDefinition<StreamDocument>>();
 
-			Stream.Name = Name;
+			Stream.Name = Config.Name;
 			Updates.Add(UpdateBuilder.Set(x => x.Name, Stream.Name));
 
 			Stream.ProjectId = ProjectId;
 			Updates.Add(UpdateBuilder.Set(x => x.ProjectId, Stream.ProjectId));
 
-			Stream.Revision = Revision;
-			Updates.Add(UpdateBuilder.Set(x => x.Revision, Stream.Revision));
+			Stream.ConfigPath = ConfigPath;
+			Updates.Add(UpdateBuilder.Set(x => x.ConfigPath, Stream.ConfigPath));
+
+			Stream.ConfigRevision = ConfigRevision;
+			Updates.Add(UpdateBuilder.Set(x => x.ConfigRevision, Stream.ConfigRevision));
 
-			Stream.Order = Order ?? StreamDocument.DefaultOrder;
+			Stream.Order = Config.Order ?? StreamDocument.DefaultOrder;
 			Updates.Add(UpdateBuilder.Set(x => x.Order, Stream.Order));
 
-			Stream.NotificationChannel = NotificationChannel;
+			Stream.NotificationChannel = Config.NotificationChannel;
 			Updates.Add(UpdateBuilder.Set(x => x.NotificationChannel, Stream.NotificationChannel));
 
-			Stream.NotificationChannelFilter = NotificationChannelFilter;
+			Stream.NotificationChannelFilter = Config.NotificationChannelFilter;
 			Updates.Add(UpdateBuilder.Set(x => x.NotificationChannelFilter, Stream.NotificationChannelFilter));
 
-			Stream.TriageChannel = TriageChannel;
+			Stream.TriageChannel = Config.TriageChannel;
 			Updates.Add(UpdateBuilder.Set(x => x.TriageChannel, Stream.TriageChannel));
 
 			Stream.DefaultPreflight = DefaultPreflight;
@@ -319,18 +303,9 @@ namespace HordeServer.Collections.Impl
 			Stream.Templates = TemplateRefs ?? new Dictionary<TemplateRefId, TemplateRef>();
 			Updates.Add(UpdateBuilder.Set(x => x.Templates, Stream.Templates));
 
-			Stream.Properties = Properties ?? new Dictionary<string, string>();
-			Updates.Add(UpdateBuilder.Set(x => x.Properties, Stream.Properties));
-
 			Stream.Acl = Acl;
 			Updates.Add(UpdateBuilder.SetOrUnsetNullRef(x => x.Acl, Acl));
 			
-			Stream.PausedUntil = PausedUntil;
-			Updates.Add(UpdateBuilder.Set(x => x.PausedUntil, Stream.PausedUntil));
-			
-			Stream.PauseComment = PauseComment;
-			Updates.Add(UpdateBuilder.Set(x => x.PauseComment, Stream.PauseComment));
-
 			Stream.Validate();
 			if(await TryUpdateStreamAsync(Stream, UpdateBuilder.Combine(Updates)))
 			{
@@ -444,22 +419,6 @@ namespace HordeServer.Collections.Impl
 				Stream.Templates = NewTemplateRefs;
 				Updates.Add(UpdateBuilder.Set(x => x.Templates, Stream.Templates));
 			}
-			if (NewProperties != null)
-			{
-				foreach (KeyValuePair<string, string?> Pair in NewProperties)
-				{
-					if (Pair.Value == null)
-					{
-						Stream.Properties.Remove(Pair.Key);
-						Updates.Add(UpdateBuilder.Unset(x => x.Properties[Pair.Key]));
-					}
-					else
-					{
-						Stream.Properties[Pair.Key] = Pair.Value;
-						Updates.Add(UpdateBuilder.Set(x => x.Properties[Pair.Key], Pair.Value));
-					}
-				}
-			}
 			if (NewAcl != null)
 			{
 				Stream.Acl = NewAcl;
diff --git a/Engine/Source/Programs/Horde/HordeServer/Models/Project.cs b/Engine/Source/Programs/Horde/HordeServer/Models/Project.cs
--- a/Engine/Source/Programs/Horde/HordeServer/Models/Project.cs
+++ b/Engine/Source/Programs/Horde/HordeServer/Models/Project.cs
@@ -104,10 +104,15 @@ namespace HordeServer.Models
 		/// </summary>
 		public string Name { get; }
 
+		/// <summary>
+		/// Path to the config file used to configure this project
+		/// </summary>
+		public string? ConfigPath { get; }
+
 		/// <summary>
 		/// Revision of the config file used to configure this project
 		/// </summary>
-		public string? Revision { get; set; }
+		public string? ConfigRevision { get; }
 
 		/// <summary>
 		/// Order to display on the dashboard
@@ -135,6 +140,11 @@ namespace HordeServer.Models
 		/// </summary>
 		public ProjectId Id { get; }
 
+		/// <summary>
+		/// Path to the logo
+		/// </summary>
+		public string Path { get; }
+
 		/// <summary>
 		/// Revision of the file
 		/// </summary>
diff --git a/Engine/Source/Programs/Horde/HordeServer/Models/Stream.cs b/Engine/Source/Programs/Horde/HordeServer/Models/Stream.cs
--- a/Engine/Source/Programs/Horde/HordeServer/Models/Stream.cs
+++ b/Engine/Source/Programs/Horde/HordeServer/Models/Stream.cs
@@ -497,10 +497,15 @@ namespace HordeServer.Models
 		/// </summary>
 		public string Name { get; }
 
+		/// <summary>
+		/// Path to the configuration file for this stream
+		/// </summary>
+		public string ConfigPath { get; }
+
 		/// <summary>
 		/// The revision of config file used for this stream
 		/// </summary>
-		public string? Revision { get; }
+		public string ConfigRevision { get; }
 
 		/// <summary>
 		/// Order to display on the dashboard's drop-down list
@@ -693,7 +698,7 @@ namespace HordeServer.Models
 			Dictionary<string, GetAgentTypeResponse> ApiAgentTypes = Stream.AgentTypes.ToDictionary(x => x.Key, x => x.Value.ToApiResponse());
 			Dictionary<string, GetWorkspaceTypeResponse> ApiWorkspaceTypes = Stream.WorkspaceTypes.ToDictionary(x => x.Key, x => x.Value.ToApiResponse());
 			GetAclResponse? ApiAcl = (bIncludeAcl && Stream.Acl != null)? new GetAclResponse(Stream.Acl) : null;
-			return new Api.GetStreamResponse(Stream.Id.ToString(), Stream.ProjectId.ToString(), Stream.Name, Stream.Revision, Stream.Order, Stream.NotificationChannel, Stream.NotificationChannelFilter, Stream.TriageChannel, Stream.DefaultPreflight?.ToRequest(), ApiTabs, ApiAgentTypes, ApiWorkspaceTypes, ApiTemplateRefs, ApiAcl, Stream.PausedUntil, Stream.PauseComment);
+			return new Api.GetStreamResponse(Stream.Id.ToString(), Stream.ProjectId.ToString(), Stream.Name, Stream.ConfigPath, Stream.ConfigRevision, Stream.Order, Stream.NotificationChannel, Stream.NotificationChannelFilter, Stream.TriageChannel, Stream.DefaultPreflight?.ToRequest(), ApiTabs, ApiAgentTypes, ApiWorkspaceTypes, ApiTemplateRefs, ApiAcl, Stream.PausedUntil, Stream.PauseComment);
 		}
 
 		/// <summary>
diff --git a/Engine/Source/Programs/Horde/HordeServer/Services/ConfigService.cs b/Engine/Source/Programs/Horde/HordeServer/Services/ConfigService.cs
--- a/Engine/Source/Programs/Horde/HordeServer/Services/ConfigService.cs
+++ b/Engine/Source/Programs/Horde/HordeServer/Services/ConfigService.cs
@@ -171,7 +171,7 @@ namespace HordeServer.Services
 				}
 
 				IProject? Project = Projects.FirstOrDefault(x => x.Id == ProjectRef.Id);
-				bool Update = (Project == null || Project.Revision != Revision);
+				bool Update = (Project == null || Project.ConfigPath != ConfigPath || Project.ConfigRevision != Revision);
 
 				ProjectConfig? ProjectConfig;
 				if (Update || !PrevCachedProjectConfigs.TryGetValue(ProjectRef.Id, out ProjectConfig))
@@ -186,7 +186,7 @@ namespace HordeServer.Services
 					if (Update)
 					{
 						Logger.LogInformation("Updating configuration for project {ProjectId} ({Revision})", ProjectRef.Id, Revision);
-						await ProjectService.Collection.AddOrUpdateAsync(ProjectRef.Id, Revision, Idx, ProjectConfig);
+						await ProjectService.Collection.AddOrUpdateAsync(ProjectRef.Id, ProjectPath, Revision, Idx, ProjectConfig);
 					}
 				}
 
@@ -215,7 +215,7 @@ namespace HordeServer.Services
 					if (Revision != CurrentRevision)
 					{
 						Logger.LogInformation("Updating logo for project {ProjectId} ({Revision})", ProjectId, Revision);
-						await ProjectService.Collection.SetLogoAsync(ProjectId, Revision, GetMimeTypeFromPath(Path), await ReadDataAsync(Path));
+						await ProjectService.Collection.SetLogoAsync(ProjectId, Path, Revision, GetMimeTypeFromPath(Path), await ReadDataAsync(Path));
 						CachedLogoRevisions[ProjectId] = Revision;
 					}
 				}
@@ -232,7 +232,7 @@ namespace HordeServer.Services
 				if (StreamRevisions.TryGetValue(Path, out string? Revision))
 				{
 					IStream? Stream = Streams.FirstOrDefault(x => x.Id == StreamRef.Id);
-					if (Stream == null || Stream.Revision != Revision)
+					if (Stream == null || Stream.ConfigPath != Path || Stream.ConfigRevision != Revision)
 					{
 						Logger.LogInformation("Updating configuration for stream {StreamRef} ({Revision})", StreamRef.Id, Revision);
 
@@ -241,7 +241,7 @@ namespace HordeServer.Services
 						{
 							for (; ; )
 							{
-								Stream = await StreamService.StreamCollection.TryCreateOrReplaceAsync(StreamRef.Id, Stream, Revision, ProjectId, StreamConfig);
+								Stream = await StreamService.StreamCollection.TryCreateOrReplaceAsync(StreamRef.Id, Stream, Path, Revision, ProjectId, StreamConfig);
 
 								if (Stream != null)
 								{
