Horde: Fix spurious errors about jobs not being valid, when they've only just been aborted.

diff --git a/Engine/Source/Programs/Horde/HordeServer/Tasks/Impl/JobTaskSource.cs b/Engine/Source/Programs/Horde/HordeServer/Tasks/Impl/JobTaskSource.cs
--- a/Engine/Source/Programs/Horde/HordeServer/Tasks/Impl/JobTaskSource.cs
+++ b/Engine/Source/Programs/Horde/HordeServer/Tasks/Impl/JobTaskSource.cs
@@ -499,8 +499,10 @@ namespace HordeServer.Tasks.Impl
 				bool IsRunning = false;
 				for (int BatchIdx = 0; NewJob != null && BatchIdx < NewJob.Batches.Count; BatchIdx++)
 				{
-					// Skip any batches which aren't ready to execute
 					IJobStepBatch Batch = NewJob.Batches[BatchIdx];
+					IsRunning |= (Batch.State == JobStepBatchState.Ready || Batch.State == JobStepBatchState.Running);
+
+					// Skip any batches which aren't ready to execute
 					if (Batch.State == JobStepBatchState.Ready)
 					{
 						// Validate the agent type and workspace settings
@@ -534,7 +536,6 @@ namespace HordeServer.Tasks.Impl
 							NewBatchIdToQueueItem[(NewJob.Id, Batch.Id)] = NewQueueItem;
 						}
 					}
-					IsRunning |= (Batch.State == JobStepBatchState.Ready || Batch.State == JobStepBatchState.Running);
 				}
 
 				// Add a warning if a job looks to be idle
@@ -578,6 +579,8 @@ namespace HordeServer.Tasks.Impl
 
 		private async Task<IJob?> SkipBatchAsync(IJob Job, int BatchIdx, IGraph Graph, JobStepBatchError Reason)
 		{
+			Logger.LogInformation("Skipping batch {BatchIdx} for job {JobId}", BatchIdx, Job.Id);
+
 			IReadOnlyList<(LabelState, LabelOutcome)> OldLabelStates = Job.GetLabelStates(Graph);
 			IJob? NewJob = await Jobs.SkipBatchAsync(Job, BatchIdx, Graph, Reason);
 			if(NewJob != null)
