Horde: Reduce some log spam.

diff --git a/Engine/Source/Programs/Horde/HordeServer/Tasks/Impl/JobTaskSource.cs b/Engine/Source/Programs/Horde/HordeServer/Tasks/Impl/JobTaskSource.cs
--- a/Engine/Source/Programs/Horde/HordeServer/Tasks/Impl/JobTaskSource.cs
+++ b/Engine/Source/Programs/Horde/HordeServer/Tasks/Impl/JobTaskSource.cs
@@ -823,7 +823,7 @@ namespace HordeServer.Tasks.Impl
 			else
 			{
 				// Unable to assign job
-				Logger.LogDebug("Refreshing queue entries for job {JobId}", Job.Id);
+				Logger.LogTrace("Refreshing queue entries for job {JobId}", Job.Id);
 
 				// Get the new copy of the job
 				IJob? NewJob = await Jobs.GetAsync(Job.Id);
