Horde: Fix tests.

diff --git a/Engine/Source/Programs/Horde/HordeServer/Models/Config.cs b/Engine/Source/Programs/Horde/HordeServer/Models/Config.cs
--- a/Engine/Source/Programs/Horde/HordeServer/Models/Config.cs
+++ b/Engine/Source/Programs/Horde/HordeServer/Models/Config.cs
@@ -91,13 +91,6 @@ namespace HordeServer.Models
 		/// Acl entries
 		/// </summary>
 		public UpdateAclRequest? Acl { get; set; }
-
-		/// <summary>
-		/// Default constructor for serialization
-		/// </summary>
-		protected ProjectConfig()
-		{
-		}
 	}
 
 	/// <summary>
diff --git a/Engine/Source/Programs/Horde/HordeServerTests/Fixture.cs b/Engine/Source/Programs/Horde/HordeServerTests/Fixture.cs
--- a/Engine/Source/Programs/Horde/HordeServerTests/Fixture.cs
+++ b/Engine/Source/Programs/Horde/HordeServerTests/Fixture.cs
@@ -75,6 +75,7 @@ namespace HordeServerTests
 				new StreamId("ue5-main"),
 				null,
 				"",
+				"",
 				new ProjectId("does-not-exist"),
 				new StreamConfig { Name = "//UE5/Main", Tabs = Tabs, Templates = Templates }
 			);
diff --git a/Engine/Source/Programs/Horde/HordeServerTests/IssueServiceTests.cs b/Engine/Source/Programs/Horde/HordeServerTests/IssueServiceTests.cs
--- a/Engine/Source/Programs/Horde/HordeServerTests/IssueServiceTests.cs
+++ b/Engine/Source/Programs/Horde/HordeServerTests/IssueServiceTests.cs
@@ -125,11 +125,11 @@ namespace HordeServerTests
 				WorkspaceDir = new DirectoryReference("/Horde");
 			}
 
-			IProject Project = TestSetup.ProjectCollection.TryAddAsync(new ProjectId("ue4"), "UE4").Result!;
+			IProject Project = TestSetup.ProjectCollection.AddOrUpdateAsync(new ProjectId("ue4"), "", "", 0, new ProjectConfig { Name = "UE4" }).Result!;
 
-			IStream MainStream = TestSetup.StreamCollection.TryCreateOrReplaceAsync(MainStreamId, null, "", Project.Id, new StreamConfig { Name = MainStreamName }).Result!;
-			IStream ReleaseStream = TestSetup.StreamCollection.TryCreateOrReplaceAsync(ReleaseStreamId, null, "", Project.Id, new StreamConfig { Name = ReleaseStreamName }).Result!;
-			IStream DevStream = TestSetup.StreamCollection.TryCreateOrReplaceAsync(DevStreamId, null, "", Project.Id, new StreamConfig { Name = DevStreamName }).Result!;
+			IStream MainStream = TestSetup.StreamCollection.TryCreateOrReplaceAsync(MainStreamId, null, "", "", Project.Id, new StreamConfig { Name = MainStreamName }).Result!;
+			IStream ReleaseStream = TestSetup.StreamCollection.TryCreateOrReplaceAsync(ReleaseStreamId, null, "", "", Project.Id, new StreamConfig { Name = ReleaseStreamName }).Result!;
+			IStream DevStream = TestSetup.StreamCollection.TryCreateOrReplaceAsync(DevStreamId, null, "", "", Project.Id, new StreamConfig { Name = DevStreamName }).Result!;
 
 			Perforce = TestSetup.PerforceService;
 			Perforce.AddChange(MainStreamName, 100, "Bill", "Description", new string[] { "a/b.cpp" });
diff --git a/Engine/Source/Programs/Horde/HordeServerTests/JobServiceTests.cs b/Engine/Source/Programs/Horde/HordeServerTests/JobServiceTests.cs
--- a/Engine/Source/Programs/Horde/HordeServerTests/JobServiceTests.cs
+++ b/Engine/Source/Programs/Horde/HordeServerTests/JobServiceTests.cs
@@ -25,7 +25,7 @@ namespace HordeServerTests
 			TestSetup TestSetup = await GetTestSetup();
 
 			ProjectId ProjectId = new ProjectId("ue5");
-			IProject? Project = await TestSetup.ProjectService.TryCreateProjectAsync(ProjectId, "UE5", null, null, null);
+			IProject? Project = await TestSetup.ProjectService.Collection.AddOrUpdateAsync(ProjectId, "", "", 0, new ProjectConfig { Name = "UE5" });
 			Assert.IsNotNull(Project);
 
 			ITemplate Template = await TestSetup.TemplateService.CreateTemplateAsync("Test template", null, false, null, null, new List<TemplateCounter>(), new List<string>(), new List<Parameter>());
@@ -41,7 +41,7 @@ namespace HordeServerTests
 
 			StreamId StreamId = new StreamId("ue5-main");
 			IStream? Stream = await TestSetup.StreamService.GetStreamAsync(StreamId);
-			Stream = await TestSetup.StreamService.StreamCollection.TryCreateOrReplaceAsync(new StreamId("ue5-main"), Stream, String.Empty, ProjectId, StreamConfig);
+			Stream = await TestSetup.StreamService.StreamCollection.TryCreateOrReplaceAsync(new StreamId("ue5-main"), Stream, String.Empty, String.Empty, ProjectId, StreamConfig);
 
 			IJob Job = await TestSetup.JobService.CreateJobAsync(null, Stream!.Id, TemplateRefId1, Template.Id, Graph, "Hello", 1234, 1233, 999, null, null, "joe", null, null, null, Stream.Templates[TemplateRefId1].ChainedJobs, true, true, null, null, null, Template.Counters, new List<string>());
 			Assert.AreEqual(1, Job.ChainedJobs.Count);
@@ -147,12 +147,12 @@ namespace HordeServerTests
 		public async Task TestRunEarly()
 		{
 			TestSetup TestSetup = await GetTestSetup();
-			
-			IProject? Project = await TestSetup.ProjectService.TryCreateProjectAsync(new ProjectId("ue5"), "UE5", null, null, null);
+
+			IProject? Project = await TestSetup.ProjectService.Collection.AddOrUpdateAsync(new ProjectId("ue5"), "", "", 0, new ProjectConfig { Name = "UE5" });
 			Assert.IsNotNull(Project);
 
 			StreamId StreamId = new StreamId("ue5-main");
-			await TestSetup.StreamCollection.TryCreateOrReplaceAsync(new StreamId("ue5-main"), null, "", Project!.Id, new StreamConfig { Name = "//UE5/Main" });
+			await TestSetup.StreamCollection.TryCreateOrReplaceAsync(new StreamId("ue5-main"), null, "", "", Project!.Id, new StreamConfig { Name = "//UE5/Main" });
 
 			ITemplate Template = await TestSetup.TemplateService.CreateTemplateAsync("Test template", null, false, null, null, new List<TemplateCounter>(), new List<string>(), new List<Parameter>());
 			IGraph Graph = await TestSetup.GraphCollection.AddAsync(Template);
diff --git a/Engine/Source/Programs/Horde/HordeServerTests/SchedulerTests.cs b/Engine/Source/Programs/Horde/HordeServerTests/SchedulerTests.cs
--- a/Engine/Source/Programs/Horde/HordeServerTests/SchedulerTests.cs
+++ b/Engine/Source/Programs/Horde/HordeServerTests/SchedulerTests.cs
@@ -38,7 +38,7 @@ namespace HordeServerTests
 		{
 			TestSetup = GetTestSetup().Result;
 
-			IProject ? Project = TestSetup.ProjectService.TryCreateProjectAsync(new ProjectId("ue5"), "UE5", null, null, null).Result;
+			IProject ? Project = TestSetup.ProjectService.Collection.AddOrUpdateAsync(new ProjectId("ue5"), "", "", 0, new ProjectConfig { Name = "UE5" }).Result;
 			Assert.IsNotNull(Project);
 
 			Template = TestSetup.TemplateService.CreateTemplateAsync("Test template", null, false, null, null, new List<TemplateCounter>(), new List<string>(), new List<Parameter>()).Result;
@@ -46,7 +46,7 @@ namespace HordeServerTests
 			Stream = TestSetup.StreamCollection.GetAsync(StreamId).Result!;
 			if (Stream == null)
 			{
-				Stream = TestSetup.StreamCollection.TryCreateOrReplaceAsync(new StreamId("ue5-main"), null, "", Project!.Id, new StreamConfig { Name = "//UE5/Main" }).Result!;
+				Stream = TestSetup.StreamCollection.TryCreateOrReplaceAsync(new StreamId("ue5-main"), null, "", "", Project!.Id, new StreamConfig { Name = "//UE5/Main" }).Result!;
 			}
 
 			InitialJobIds = new HashSet<ObjectId>(TestSetup.JobCollection.FindAsync().Result.Select(x => x.Id));
