Horde: Add a systemic log event for XGE not running.

diff --git a/Engine/Source/Programs/Horde/HordeAgent/Parser/Matchers/XoreaxEventMatcher.cs b/Engine/Source/Programs/Horde/HordeAgent/Parser/Matchers/XoreaxEventMatcher.cs
--- a/Engine/Source/Programs/Horde/HordeAgent/Parser/Matchers/XoreaxEventMatcher.cs
+++ b/Engine/Source/Programs/Horde/HordeAgent/Parser/Matchers/XoreaxEventMatcher.cs
@@ -17,7 +17,12 @@ namespace HordeAgent.Parser.Matchers
 	{
 		public LogEvent? Match(ILogCursor Cursor, ILogContext Context)
 		{
-			// Return it as information, but add a property indicating it's something that should be 
+			if (Cursor.IsMatch(@"\(BuildService.exe\) is not running"))
+			{
+				LogEventBuilder Builder = new LogEventBuilder(Cursor);
+				return Builder.ToLogEvent(LogEventPriority.High, LogLevel.Information, KnownLogEvents.Systemic_Xge_ServiceNotRunning);
+			}
+
 			if (Cursor.IsMatch(@"^\s*--------------------Build System Warning[- ]"))
 			{
 				LogEventBuilder Builder = new LogEventBuilder(Cursor);
diff --git a/Engine/Source/Programs/Horde/HordeAgentTests/LogParserTests.cs b/Engine/Source/Programs/Horde/HordeAgentTests/LogParserTests.cs
--- a/Engine/Source/Programs/Horde/HordeAgentTests/LogParserTests.cs
+++ b/Engine/Source/Programs/Horde/HordeAgentTests/LogParserTests.cs
@@ -968,15 +968,19 @@ namespace HordeAgentTests
 				@"Remote tasks distribution:",
 				@"    Tasks execution is impeded due to low agent responsiveness",
 				@"-------------------------------------------------------------------------------",
+				@"    LogXGEController: Warning: XGE's background service (BuildService.exe) is not running - service is likely disabled on this machine.",
 			};
 
 			List<CapturedEvent> Events = Parse(Lines);
-			Assert.AreEqual(2, Events.Count);
+			Assert.AreEqual(3, Events.Count);
 			Assert.AreEqual(LogLevel.Information, Events[0].Level);
 			Assert.AreEqual(KnownLogEvents.Systemic_Xge_Standalone, Events[0].Id);
 
 			Assert.AreEqual(LogLevel.Information, Events[1].Level);
 			Assert.AreEqual(KnownLogEvents.Systemic_Xge, Events[1].Id);
+
+			Assert.AreEqual(LogLevel.Information, Events[2].Level);
+			Assert.AreEqual(KnownLogEvents.Systemic_Xge_ServiceNotRunning, Events[2].Id);
 		}
 
 		string GetSubProperty(CapturedEvent Event, string SpanName, string Name)
diff --git a/Engine/Source/Programs/Horde/HordeCommon/KnownLogEvents.cs b/Engine/Source/Programs/Horde/HordeCommon/KnownLogEvents.cs
--- a/Engine/Source/Programs/Horde/HordeCommon/KnownLogEvents.cs
+++ b/Engine/Source/Programs/Horde/HordeCommon/KnownLogEvents.cs
@@ -134,6 +134,11 @@ namespace HordeCommon
 		/// </summary>
 		public static EventId Systemic_Xge_Standalone { get; } = new EventId(711);
 
+		/// <summary>
+		/// BuildService.exe is not running
+		/// </summary>
+		public static EventId Systemic_Xge_ServiceNotRunning { get; } = new EventId(712);
+
 		/// <summary>
 		/// DDC is slow
 		/// </summary>
