Horde: Fix agents executing work for a job before they have synced the required workspaces.

diff --git a/Engine/Source/Programs/Horde/HordeServer/Models/Agent.cs b/Engine/Source/Programs/Horde/HordeServer/Models/Agent.cs
--- a/Engine/Source/Programs/Horde/HordeServer/Models/Agent.cs
+++ b/Engine/Source/Programs/Horde/HordeServer/Models/Agent.cs
@@ -754,6 +754,53 @@ namespace HordeServer.Models
 			return true;
 		}
 
+		/// <summary>
+		/// Get the AutoSDK workspace required for an agent
+		/// </summary>
+		/// <param name="Agent"></param>
+		/// <returns></returns>
+		public static AgentWorkspace? GetAutoSdkWorkspace(this IAgent Agent)
+		{
+			if (Agent.Capabilities.Devices.Count > 0)
+			{
+				DeviceCapabilities PrimaryDevice = Agent.Capabilities.Devices[0];
+				if (PrimaryDevice.Properties != null)
+				{
+					if (PrimaryDevice.Properties.Contains("OSFamily=Windows"))
+					{
+						return new AgentWorkspace(null, null, AgentWorkspace.AutoSdkIdentifier, "//UE4/Private-AutoSDK-Windows", null, true);
+					}
+					else if (PrimaryDevice.Properties.Contains("OSVersion=MacOS"))
+					{
+						return new AgentWorkspace(null, null, AgentWorkspace.AutoSdkIdentifier, "//UE4/Private-AutoSDK-Mac", null, true);
+					}
+				}
+			}
+			return null;
+		}
+
+		/// <summary>
+		/// Check that an agent has the required workspaces to execute a task
+		/// </summary>
+		/// <param name="Agent"></param>
+		/// <param name="Workspace"></param>
+		/// <returns></returns>
+		public static bool HasRequiredWorkspaces(this IAgent Agent, AgentWorkspace Workspace)
+		{
+			if (!Agent.Workspaces.Contains(Workspace))
+			{
+				return false;
+			}
+
+			AgentWorkspace? AutoSdkWorkspace = Agent.GetAutoSdkWorkspace();
+			if (AutoSdkWorkspace != null && !Agent.Workspaces.Contains(AutoSdkWorkspace))
+			{
+				return false;
+			}
+
+			return true;
+		}
+
 		/// <summary>
 		/// Converts this workspace to an RPC message
 		/// </summary>
@@ -793,12 +840,6 @@ namespace HordeServer.Models
 				ServerAndPort = Server.ServerAndPort;
 			}
 
-			// HACK: Trying to debug issues with inconsistent workspace state. Possibly related to edge servers?
-			if (Agent.Id == new AgentId("10-99-114-249") || Agent.Id == new AgentId("10-99-115-36"))
-			{
-				ServerAndPort = "perforce:1666";
-			}
-
 			// Find the matching credentials for the desired user
 			PerforceCredentials? Credentials = null;
 			if (ClusterInfo != null)
diff --git a/Engine/Source/Programs/Horde/HordeServer/Services/PoolService.cs b/Engine/Source/Programs/Horde/HordeServer/Services/PoolService.cs
--- a/Engine/Source/Programs/Horde/HordeServer/Services/PoolService.cs
+++ b/Engine/Source/Programs/Horde/HordeServer/Services/PoolService.cs
@@ -127,20 +127,11 @@ namespace HordeServer.Services
 		public async Task<HashSet<AgentWorkspace>> GetWorkspacesAsync(IAgent Agent, DateTime ValidAtTime)
 		{
 			HashSet<AgentWorkspace> Workspaces = new HashSet<AgentWorkspace>();
-			if (Agent.Capabilities.Devices.Count > 0)
+
+			AgentWorkspace? AutoSdkWorkspace = Agent.GetAutoSdkWorkspace();
+			if(AutoSdkWorkspace != null)
 			{
-				DeviceCapabilities PrimaryDevice = Agent.Capabilities.Devices[0];
-				if (PrimaryDevice.Properties != null)
-				{
-					if (PrimaryDevice.Properties.Contains("OSFamily=Windows"))
-					{
-						Workspaces.Add(new AgentWorkspace(null, null, AgentWorkspace.AutoSdkIdentifier, "//UE4/Private-AutoSDK-Windows", null, true));
-					}
-					else if (PrimaryDevice.Properties.Contains("OSVersion=MacOS"))
-					{
-						Workspaces.Add(new AgentWorkspace(null, null, AgentWorkspace.AutoSdkIdentifier, "//UE4/Private-AutoSDK-Mac", null, true));
-					}
-				}
+				Workspaces.Add(AutoSdkWorkspace);
 			}
 
 			Dictionary<PoolId, IPool> PoolMapping = await GetPoolLookupAsync(ValidAtTime);
diff --git a/Engine/Source/Programs/Horde/HordeServer/Tasks/Impl/JobTaskSource.cs b/Engine/Source/Programs/Horde/HordeServer/Tasks/Impl/JobTaskSource.cs
--- a/Engine/Source/Programs/Horde/HordeServer/Tasks/Impl/JobTaskSource.cs
+++ b/Engine/Source/Programs/Horde/HordeServer/Tasks/Impl/JobTaskSource.cs
@@ -643,6 +643,12 @@ namespace HordeServer.Tasks.Impl
 			}
 		}
 
+		/// <summary>
+		/// Attempts to assign a queue item to an agent waiting for work
+		/// </summary>
+		/// <param name="Item"></param>
+		/// <param name="Waiter"></param>
+		/// <returns></returns>
 		bool TryAssignItemToWaiter(QueueItem Item, QueueWaiter Waiter)
 		{
 			if (Item.AssignTask == null && Item.Batch.SessionId == null)
@@ -877,7 +883,7 @@ namespace HordeServer.Tasks.Impl
 
 			List<HordeCommon.Rpc.Messages.AgentWorkspace> Workspaces = new List<HordeCommon.Rpc.Messages.AgentWorkspace>();
 
-			AgentWorkspace? AutoSdkWorkspace = Agent.Workspaces.FirstOrDefault(x => x.Identifier == AgentWorkspace.AutoSdkIdentifier);
+			AgentWorkspace? AutoSdkWorkspace = Agent.GetAutoSdkWorkspace();
 			if (AutoSdkWorkspace != null)
 			{
 				if (!await Agent.TryAddWorkspaceMessage(AutoSdkWorkspace, Globals, PerforceLoadBalancer, Workspaces))
