Horde: Scrub the conform list of invalid entries every 30 minutes. Also move configuration into the global config file.

diff --git a/Engine/Source/Programs/Horde/HordeServer/Controllers/AdminController.cs b/Engine/Source/Programs/Horde/HordeServer/Controllers/AdminController.cs
--- a/Engine/Source/Programs/Horde/HordeServer/Controllers/AdminController.cs
+++ b/Engine/Source/Programs/Horde/HordeServer/Controllers/AdminController.cs
@@ -78,11 +78,6 @@ namespace HordeServer.Controllers
 		/// </summary>
 		UpgradeService UpgradeService;
 
-		/// <summary>
-		/// List of conforms running
-		/// </summary>
-		SingletonDocument<ConformList> ConformList;
-
 		/// <summary>
 		/// Settings for the server
 		/// </summary>
@@ -100,7 +95,6 @@ namespace HordeServer.Controllers
 			this.DatabaseService = DatabaseService;
 			this.AclService = AclService;
 			this.UpgradeService = UpgradeService;
-			this.ConformList = new SingletonDocument<ConformList>(DatabaseService);
 			this.Settings = Settings;
 		}
 
@@ -137,52 +131,6 @@ namespace HordeServer.Controllers
 			}
 		}
 
-
-		/// <summary>
-		/// Gets the scheduled downtime 
-		/// </summary>
-		/// <returns>Async task</returns>
-		[HttpGet]
-		[Route("/api/v1/admin/downtime")]
-		public async Task<ActionResult<List<ScheduledDowntime>>> GetScheduledDowntimeAsync()
-		{
-			if (!await AclService.AuthorizeAsync(AclAction.AdminWrite, User))
-			{
-				return Forbid();
-			}
-
-			Globals Globals = await DatabaseService.GetGlobalsAsync();
-			return Globals.ScheduledDowntime;
-		}
-
-		/// <summary>
-		/// Sets the scheduled downtime
-		/// </summary>
-		/// <param name="Schedules">List of downtime</param>
-		/// <returns>Async task</returns>
-		[HttpPut]
-		[Route("/api/v1/admin/downtime")]
-		public async Task<ActionResult> SetScheduledDowntimeAsync(List<ScheduledDowntime> Schedules)
-		{
-			if (!await AclService.AuthorizeAsync(AclAction.AdminWrite, User))
-			{
-				return Forbid();
-			}
-
-			for (; ; )
-			{
-				Globals Globals = await DatabaseService.GetGlobalsAsync();
-				Globals.ScheduledDowntime = Schedules;
-
-				if (await DatabaseService.TryUpdateSingletonAsync(Globals))
-				{
-					break;
-				}
-			}
-
-			return Ok();
-		}
-
 		/// <summary>
 		/// Upgrade the database to the latest schema
 		/// </summary>
@@ -200,88 +148,6 @@ namespace HordeServer.Controllers
 			return Ok();
 		}
 
-		/// <summary>
-		/// Update the conform limit
-		/// </summary>
-		[HttpPut]
-		[Route("/api/v1/admin/conform")]
-		public async Task<ActionResult> SetConformSettingsAsync([FromBody] ConformSettings NewSettings)
-		{
-			if (!await AclService.AuthorizeAsync(AclAction.AdminWrite, User))
-			{
-				return Forbid();
-			}
-
-			for (; ; )
-			{
-				ConformList CurrentValue = await ConformList.GetAsync();
-				CurrentValue.MaxCount = NewSettings.MaxCount;
-
-				if (await ConformList.TryUpdateAsync(CurrentValue))
-				{
-					return Ok();
-				}
-			}
-		}
-
-		/// <summary>
-		/// Gets the current perforce settinsg
-		/// </summary>
-		/// <returns>List of Perforce clusters</returns>
-		[HttpGet]
-		[Route("/api/v1/admin/conform")]
-		public async Task<ActionResult<ConformSettings>> GetConformSettingsAsync()
-		{
-			if (!await AclService.AuthorizeAsync(AclAction.AdminRead, User))
-			{
-				return Forbid();
-			}
-
-			ConformList CurrentValue = await ConformList.GetAsync();
-			return new ConformSettings { MaxCount = CurrentValue.MaxCount };
-		}
-
-		/// <summary>
-		/// Update Perforce server settings
-		/// </summary>
-		[HttpPut]
-		[Route("/api/v1/admin/perforce")]
-		public async Task<ActionResult> UpdatePerforceSettingsAsync([FromBody] List<PerforceCluster> Clusters)
-		{
-			if (!await AclService.AuthorizeAsync(AclAction.AdminWrite, User))
-			{
-				return Forbid();
-			}
-
-			for (; ; )
-			{
-				Globals Globals = await DatabaseService.GetGlobalsAsync();
-				Globals.PerforceClusters = Clusters;
-
-				if (await DatabaseService.TryUpdateSingletonAsync(Globals))
-				{
-					return Ok("Settings have been updated");
-				}
-			}
-		}
-
-		/// <summary>
-		/// Gets the current perforce settinsg
-		/// </summary>
-		/// <returns>List of Perforce clusters</returns>
-		[HttpGet]
-		[Route("/api/v1/admin/perforce")]
-		public async Task<ActionResult<List<PerforceCluster>>> GetPerforceSettingsAsync()
-		{
-			if (!await AclService.AuthorizeAsync(AclAction.AdminRead, User))
-			{
-				return Forbid();
-			}
-
-			Globals Globals = await DatabaseService.GetGlobalsAsync();
-			return Globals.PerforceClusters;
-		}
-
 		/// <summary>
 		/// Issues a token for the given roles. Issues a token for the current user if not specified.
 		/// </summary>
diff --git a/Engine/Source/Programs/Horde/HordeServer/Models/Config.cs b/Engine/Source/Programs/Horde/HordeServer/Models/Config.cs
--- a/Engine/Source/Programs/Horde/HordeServer/Models/Config.cs
+++ b/Engine/Source/Programs/Horde/HordeServer/Models/Config.cs
@@ -37,6 +37,11 @@ namespace HordeServer.Models
 		/// </summary>
 		public List<PerforceCluster> PerforceClusters { get; set; } = new List<PerforceCluster>();
 
+		/// <summary>
+		/// Maximum number of conforms to run at once
+		/// </summary>
+		public int MaxConformCount { get; set; }
+
 		/// <summary>
 		/// Access control list
 		/// </summary>
diff --git a/Engine/Source/Programs/Horde/HordeServer/Models/ConformList.cs b/Engine/Source/Programs/Horde/HordeServer/Models/ConformList.cs
--- a/Engine/Source/Programs/Horde/HordeServer/Models/ConformList.cs
+++ b/Engine/Source/Programs/Horde/HordeServer/Models/ConformList.cs
@@ -2,6 +2,7 @@
 
 using HordeServer.Utilities;
 using MongoDB.Bson;
+using MongoDB.Bson.Serialization.Attributes;
 using System;
 using System.Collections.Generic;
 using System.Linq;
@@ -10,19 +11,38 @@ using System.Threading.Tasks;
 namespace HordeServer.Models
 {
 	/// <summary>
-	/// List of machines that are currently conforming
+	/// Entry to the conform list
 	/// </summary>
-	[SingletonDocument("5ff62d06a3bdb49f82da3cac")]
-	public class ConformList : SingletonBase
+	public class ConformListEntry
 	{
 		/// <summary>
-		/// Maximum allowed at once
+		/// The agent id
+		/// </summary>
+		[BsonElement("a")]
+		public AgentId AgentId { get; set; }
+
+		/// <summary>
+		/// The lease id
+		/// </summary>
+		[BsonElement("l")]
+		public ObjectId LeaseId { get; set; }
+
+		/// <summary>
+		/// Last timestamp that the lease was checked for validity 
 		/// </summary>
-		public int MaxCount { get; set; }
+		[BsonElement("t"), BsonIgnoreIfNull]
+		public DateTime? LastCheckTimeUtc { get; set; }
+	}
 
+	/// <summary>
+	/// List of machines that are currently conforming
+	/// </summary>
+	[SingletonDocument("60afc737f0d2a70754229300")]
+	public class ConformList : SingletonBase
+	{
 		/// <summary>
-		/// List of current conforms that are running
+		/// List of entries
 		/// </summary>
-		public List<ObjectId> LeaseIds { get; set; } = new List<ObjectId>();
+		public List<ConformListEntry> Entries { get; set; } = new List<ConformListEntry>();
 	}
 }
diff --git a/Engine/Source/Programs/Horde/HordeServer/Models/Globals.cs b/Engine/Source/Programs/Horde/HordeServer/Models/Globals.cs
--- a/Engine/Source/Programs/Horde/HordeServer/Models/Globals.cs
+++ b/Engine/Source/Programs/Horde/HordeServer/Models/Globals.cs
@@ -283,6 +283,11 @@ namespace HordeServer.Models
 		/// </summary>
 		public List<PerforceCluster> PerforceClusters { get; set; } = new List<PerforceCluster>();
 
+		/// <summary>
+		/// Maximum number of simultaneous conforms
+		/// </summary>
+		public int MaxConformCount { get; set; }
+
 		/// <summary>
 		/// Randomly generated JWT signing key
 		/// </summary>
diff --git a/Engine/Source/Programs/Horde/HordeServer/Services/ConfigService.cs b/Engine/Source/Programs/Horde/HordeServer/Services/ConfigService.cs
--- a/Engine/Source/Programs/Horde/HordeServer/Services/ConfigService.cs
+++ b/Engine/Source/Programs/Horde/HordeServer/Services/ConfigService.cs
@@ -134,6 +134,7 @@ namespace HordeServer.Services
 				Globals.Notices = CachedGlobalConfig.Notices;
 				Globals.PerforceClusters = CachedGlobalConfig.PerforceClusters;
 				Globals.ScheduledDowntime = CachedGlobalConfig.Downtime;
+				Globals.MaxConformCount = CachedGlobalConfig.MaxConformCount;
 
 				if (await DatabaseService.TryUpdateSingletonAsync(Globals))
 				{
diff --git a/Engine/Source/Programs/Horde/HordeServer/Startup.cs b/Engine/Source/Programs/Horde/HordeServer/Startup.cs
--- a/Engine/Source/Programs/Horde/HordeServer/Startup.cs
+++ b/Engine/Source/Programs/Horde/HordeServer/Startup.cs
@@ -526,6 +526,8 @@ namespace HordeServer
 			Services.AddSingleton<ITaskSource, JobTaskSource>(Provider => Provider.GetRequiredService<JobTaskSource>());
 			Services.AddSingleton<ITaskSource, ActionTaskSource>(Provider => Provider.GetRequiredService<ActionTaskSource>());
 
+			Services.AddHostedService(Provider => Provider.GetRequiredService<ConformTaskSource>());
+
 			// Allow longer to shutdown so we can debug missing cancellation tokens
 			Services.Configure<HostOptions>(Options =>
 			{
diff --git a/Engine/Source/Programs/Horde/HordeServer/Tasks/Impl/ConformTaskSource.cs b/Engine/Source/Programs/Horde/HordeServer/Tasks/Impl/ConformTaskSource.cs
--- a/Engine/Source/Programs/Horde/HordeServer/Tasks/Impl/ConformTaskSource.cs
+++ b/Engine/Source/Programs/Horde/HordeServer/Tasks/Impl/ConformTaskSource.cs
@@ -6,14 +6,17 @@ using Google.Protobuf.WellKnownTypes;
 using HordeCommon;
 using HordeCommon.Rpc.Tasks;
 using HordeServer.Api;
+using HordeServer.Collections;
 using HordeServer.Models;
 using HordeServer.Services;
 using HordeServer.Utilities;
+using Microsoft.Extensions.Hosting;
 using Microsoft.Extensions.Logging;
 using MongoDB.Bson;
 using System;
 using System.Collections.Generic;
 using System.Linq;
+using System.Threading;
 using System.Threading.Tasks;
 
 namespace HordeServer.Tasks.Impl
@@ -21,26 +24,30 @@ namespace HordeServer.Tasks.Impl
 	/// <summary>
 	/// Generates tasks telling agents to sync their workspaces
 	/// </summary>
-	public class ConformTaskSource : TaskSourceBase<ConformTask>
+	public sealed class ConformTaskSource : TaskSourceBase<ConformTask>, IHostedService, IDisposable
 	{
 		DatabaseService DatabaseService;
+		IAgentCollection AgentCollection;
 		PoolService PoolService;
 		SingletonDocument<ConformList> ConformList;
 		PerforceLoadBalancer PerforceLoadBalancer;
 		ILogFileService LogService;
 		ILogger Logger;
+		ElectedTick? TickConformList;
 
 		/// <summary>
 		/// Constructor
 		/// </summary>
 		/// <param name="DatabaseService"></param>
+		/// <param name="AgentCollection"></param>
 		/// <param name="PoolService"></param>
 		/// <param name="LogService"></param>
 		/// <param name="PerforceLoadBalancer"></param>
 		/// <param name="Logger"></param>
-		public ConformTaskSource(DatabaseService DatabaseService, PoolService PoolService, ILogFileService LogService, PerforceLoadBalancer PerforceLoadBalancer, ILogger<ConformTaskSource> Logger)
+		public ConformTaskSource(DatabaseService DatabaseService, IAgentCollection AgentCollection, PoolService PoolService, ILogFileService LogService, PerforceLoadBalancer PerforceLoadBalancer, ILogger<ConformTaskSource> Logger)
 		{
 			this.DatabaseService = DatabaseService;
+			this.AgentCollection = AgentCollection;
 			this.PoolService = PoolService;
 			this.ConformList = new SingletonDocument<ConformList>(DatabaseService);
 			this.PerforceLoadBalancer = PerforceLoadBalancer;
@@ -48,6 +55,96 @@ namespace HordeServer.Tasks.Impl
 			this.Logger = Logger;
 		}
 
+		/// <inheritdoc/>
+		public void Dispose()
+		{
+			TickConformList?.Dispose();
+		}
+
+		/// <inheritdoc/>
+		public Task StartAsync(CancellationToken CancellationToken)
+		{
+			if (TickConformList == null)
+			{
+				TickConformList = new ElectedTick(DatabaseService, new ObjectId("60afc5cf555a9a76aff0a50c"), CleanConformListAsync, TimeSpan.FromMinutes(0.0), Logger);
+			}
+			return Task.CompletedTask;
+		}
+
+		/// <inheritdoc/>
+		public async Task StopAsync(CancellationToken CancellationToken)
+		{
+			if (TickConformList != null)
+			{
+				await TickConformList.DisposeAsync();
+				TickConformList = null;
+			}
+		}
+
+		/// <summary>
+		/// Clean up the conform list of any outdated entries
+		/// </summary>
+		/// <param name="CancellationToken">Cancellation token</param>
+		/// <returns>Async task</returns>
+		async Task CleanConformListAsync(CancellationToken CancellationToken)
+		{
+			DateTime UtcNow = DateTime.UtcNow;
+			DateTime LastCheckTimeUtc = UtcNow - TimeSpan.FromMinutes(30.0);
+
+			// Get the current state of the conform list
+			ConformList List = await ConformList.GetAsync();
+
+			// Update any leases that are older than LastCheckTimeUtc
+			Dictionary<ObjectId, bool> RemoveLeases = new Dictionary<ObjectId, bool>();
+			foreach (ConformListEntry Entry in List.Entries)
+			{
+				if (Entry.LastCheckTimeUtc < LastCheckTimeUtc)
+				{
+					IAgent? Agent = await AgentCollection.GetAsync(Entry.AgentId);
+
+					bool Remove = false;
+					if (Agent == null || !Agent.Leases.Any(x => x.Id == Entry.LeaseId))
+					{
+						Logger.LogWarning("Removing invalid lease from conform list: {AgentId} lease {LeaseId}", Entry.AgentId, Entry.LeaseId);
+						Remove = true;
+					}
+
+					RemoveLeases[Entry.LeaseId] = Remove;
+				}
+			}
+
+			// If there's anything to change, update the list
+			if (RemoveLeases.Count > 0)
+			{
+				await ConformList.UpdateAsync(List => UpdateConformList(List, UtcNow, RemoveLeases));
+			}
+		}
+
+		/// <summary>
+		/// Remove items from the conform list, and update timestamps for them
+		/// </summary>
+		/// <param name="List">The list to update</param>
+		/// <param name="UtcNow">Current time</param>
+		/// <param name="RemoveLeases">List of leases to update. Entries with values set to true will be removed, entries with values set to false will have their timestamp updated.</param>
+		static void UpdateConformList(ConformList List, DateTime UtcNow, Dictionary<ObjectId, bool> RemoveLeases)
+		{
+			for (int Idx = 0; Idx < List.Entries.Count; Idx++)
+			{
+				bool Remove;
+				if (RemoveLeases.TryGetValue(List.Entries[Idx].LeaseId, out Remove))
+				{
+					if (Remove)
+					{
+						List.Entries.RemoveAt(Idx--);
+					}
+					else if (List.Entries[Idx].LastCheckTimeUtc < UtcNow)
+					{
+						List.Entries[Idx].LastCheckTimeUtc = UtcNow;
+					}
+				}
+			}
+		}
+
 		/// <inheritdoc/>
 		public override async Task<ITaskListener?> SubscribeAsync(IAgent Agent)
 		{
@@ -60,11 +157,12 @@ namespace HordeServer.Tasks.Impl
 			if (Agent.Leases.Count == 0)
 			{
 				ObjectId LeaseId = ObjectId.GenerateNewId();
-				if (await AllocateConformLeaseAsync(LeaseId))
+				if (await AllocateConformLeaseAsync(Agent.Id, LeaseId))
 				{
 					ConformTask Task = new ConformTask();
 					if (!await GetWorkspacesAsync(Agent, Task.Workspaces))
 					{
+						await ReleaseConformLeaseAsync(LeaseId);
 						return null;
 					}
 
@@ -112,19 +210,26 @@ namespace HordeServer.Tasks.Impl
 		/// <summary>
 		/// Atempt to allocate a conform resource for the given lease
 		/// </summary>
+		/// <param name="AgentId">The agent id</param>
 		/// <param name="LeaseId">The lease id</param>
 		/// <returns>True if the resource was allocated, false otherwise</returns>
-		private async Task<bool> AllocateConformLeaseAsync(ObjectId LeaseId)
+		private async Task<bool> AllocateConformLeaseAsync(AgentId AgentId, ObjectId LeaseId)
 		{
+			Globals Globals = await DatabaseService.GetGlobalsAsync();
 			for (; ; )
 			{
 				ConformList CurrentValue = await ConformList.GetAsync();
-				if (CurrentValue.MaxCount != 0 && CurrentValue.LeaseIds.Count >= CurrentValue.MaxCount)
+				if (Globals.MaxConformCount != 0 && CurrentValue.Entries.Count >= Globals.MaxConformCount)
 				{
 					return false;
 				}
 
-				CurrentValue.LeaseIds.Add(LeaseId);
+				ConformListEntry Entry = new ConformListEntry();
+				Entry.AgentId = AgentId;
+				Entry.LeaseId = LeaseId;
+				Entry.LastCheckTimeUtc = DateTime.UtcNow;
+				CurrentValue.Entries.Add(Entry);
+
 				if (await ConformList.TryUpdateAsync(CurrentValue))
 				{
 					Logger.LogInformation("Added conform lease {LeaseId}", LeaseId);
@@ -143,7 +248,7 @@ namespace HordeServer.Tasks.Impl
 			for (; ; )
 			{
 				ConformList CurrentValue = await ConformList.GetAsync();
-				if (!CurrentValue.LeaseIds.Remove(LeaseId))
+				if (CurrentValue.Entries.RemoveAll(x => x.LeaseId == LeaseId) == 0)
 				{
 					Logger.LogInformation("Conform lease {LeaseId} is not in singelton", LeaseId);
 					break;
diff --git a/Engine/Source/Programs/Horde/HordeServer/Tasks/Impl/JobTaskSource.cs b/Engine/Source/Programs/Horde/HordeServer/Tasks/Impl/JobTaskSource.cs
--- a/Engine/Source/Programs/Horde/HordeServer/Tasks/Impl/JobTaskSource.cs
+++ b/Engine/Source/Programs/Horde/HordeServer/Tasks/Impl/JobTaskSource.cs
@@ -280,11 +280,6 @@ namespace HordeServer.Tasks.Impl
 		/// </summary>
 		IOptionsMonitor<ServerSettings> Settings;
 
-		/// <summary>
-		/// List of active conform tasks
-		/// </summary>
-		SingletonDocument<ConformList> ConformList;
-
 		/// <summary>
 		/// Log writer
 		/// </summary>
@@ -369,7 +364,6 @@ namespace HordeServer.Tasks.Impl
 			this.PerforceLoadBalancer = PerforceLoadBalancer;
 			this.ApplicationLifetime = ApplicationLifetime;
 			this.StoppingTask = new CancellationTask(ApplicationLifetime.ApplicationStopping);
-			this.ConformList = new SingletonDocument<ConformList>(DatabaseService);
 			this.Settings = Settings;
 			this.Logger = Logger;
 		}
diff --git a/Engine/Source/Programs/Horde/HordeServer/Utilities/ElectedTick.cs b/Engine/Source/Programs/Horde/HordeServer/Utilities/ElectedTick.cs
new file mode 100644
--- /dev/null
+++ b/Engine/Source/Programs/Horde/HordeServer/Utilities/ElectedTick.cs
@@ -0,0 +1,159 @@
+// Copyright Epic Games, Inc. All Rights Reserved.
+
+using HordeServer.Models;
+using HordeServer.Services;
+using Microsoft.Extensions.Logging;
+using MongoDB.Bson;
+using System;
+using System.Collections.Generic;
+using System.Linq;
+using System.Threading;
+using System.Threading.Tasks;
+
+namespace HordeServer.Utilities
+{
+	/// <summary>
+	/// Service which is ticked on one machine in the cluster at a specified interval
+	/// </summary>
+	public sealed class ElectedTick : IAsyncDisposable, IDisposable
+	{
+		class State : SingletonBase
+		{
+			public string? Host { get; set; }
+			public DateTime NextUpdateTime { get; set; }
+		}
+
+		SingletonDocument<State> StateAccessor;
+		TimeSpan PollInterval;
+		TimeSpan ElectionTimeout;
+		Task BackgroundTask;
+		CancellationTokenSource CancellationTokenSource;
+		Func<CancellationToken, Task<DateTime>> TickAsync;
+		ILogger Logger;
+
+		/// <summary>
+		/// Constructor
+		/// </summary>
+		/// <param name="DatabaseService">The database service instance</param>
+		/// <param name="ServiceId">Object id for the singleton tracking the update state</param>
+		/// <param name="TickAsync">The tick function to be called</param>
+		/// <param name="Interval">Interval for calling the tick function</param>
+		/// <param name="Logger">The log service</param>
+		public ElectedTick(DatabaseService DatabaseService, ObjectId ServiceId, Func<CancellationToken, Task> TickAsync, TimeSpan Interval, ILogger Logger)
+			: this(DatabaseService, ServiceId, async (Token) => { await TickAsync(Token); return DateTime.UtcNow + Interval; }, Logger)
+		{
+		}
+
+		/// <summary>
+		/// Constructor
+		/// </summary>
+		/// <param name="DatabaseService">The database service instance</param>
+		/// <param name="ServiceId">Object id for the singleton tracking the update state</param>
+		/// <param name="TickAsync">The tick function to be called</param>
+		/// <param name="Logger">The log service</param>
+		public ElectedTick(DatabaseService DatabaseService, ObjectId ServiceId, Func<CancellationToken, Task<DateTime>> TickAsync, ILogger Logger)
+			: this(DatabaseService, ServiceId, TimeSpan.FromMinutes(1.0), TimeSpan.FromMinutes(5.0), TickAsync, Logger)
+		{
+		}
+
+		/// <summary>
+		/// Constructor
+		/// </summary>
+		/// <param name="DatabaseService">The database service instance</param>
+		/// <param name="ServiceId">Object id for the singleton tracking the update state</param>
+		/// <param name="PollInterval">Frequency to poll state</param>
+		/// <param name="ElectionTimeout"></param>
+		/// <param name="TickAsync">The tick function to be called</param>
+		/// <param name="Logger">The log service</param>
+		public ElectedTick(DatabaseService DatabaseService, ObjectId ServiceId, TimeSpan PollInterval, TimeSpan ElectionTimeout, Func<CancellationToken, Task<DateTime>> TickAsync, ILogger Logger)
+		{
+			this.StateAccessor = new SingletonDocument<State>(DatabaseService, ServiceId);
+			this.PollInterval = PollInterval;
+			this.ElectionTimeout = ElectionTimeout;
+			this.TickAsync = TickAsync;
+			this.Logger = Logger;
+
+			CancellationTokenSource = new CancellationTokenSource();
+			BackgroundTask = Task.Run(() => RunAsync(CancellationTokenSource.Token));
+		}
+
+		/// <inheritdoc/>
+		public void Dispose()
+		{
+			DisposeAsync().AsTask().Wait();
+		}
+
+		/// <inheritdoc/>
+		public async ValueTask DisposeAsync()
+		{
+			CancellationTokenSource.Cancel();
+			try
+			{
+				await BackgroundTask;
+			}
+			catch(OperationCanceledException)
+			{
+			}
+			CancellationTokenSource.Dispose();
+		}
+
+		/// <summary>
+		/// Tick the local service, and poll for whether to tick the shared service.
+		/// </summary>
+		/// <param name="StoppingToken">Stopping token</param>
+		/// <returns>Async task</returns>
+		async Task RunAsync(CancellationToken StoppingToken)
+		{
+			for (; ; )
+			{
+				DateTime UtcNow = DateTime.UtcNow;
+				State State = await StateAccessor.GetAsync();
+
+				// Check if it's time to update yet
+				TimeSpan Delay = State.NextUpdateTime - UtcNow;
+				if (Delay > TimeSpan.Zero)
+				{
+					TimeSpan RetryDelay = Delay + TimeSpan.FromSeconds(0.5);
+					if (PollInterval < RetryDelay)
+					{
+						RetryDelay = PollInterval;
+					}
+					await Task.Delay(RetryDelay);
+					continue;
+				}
+
+				// Try to claim the election
+				State.Host = System.Net.Dns.GetHostName();
+				State.NextUpdateTime = UtcNow + ElectionTimeout;
+				if (!await StateAccessor.TryUpdateAsync(State))
+				{
+					continue;
+				}
+
+				// Start the tick function, and keep prolonging the next update time until it finishes
+				Task<DateTime> TickTask = Task.Run(() => TickAsync(StoppingToken));
+				while (await Task.WhenAny(TickTask, Task.Delay(ElectionTimeout / 2)) != TickTask)
+				{
+					State.NextUpdateTime = DateTime.UtcNow + ElectionTimeout;
+					await StateAccessor.TryUpdateAsync(State);
+				}
+
+				// Wait for the task to complete
+				DateTime NextUpdateTime;
+				try
+				{
+					NextUpdateTime = await TickTask;
+				}
+				catch (Exception Ex)
+				{
+					Logger.LogError(Ex, "Exception while executing tick function");
+					continue;
+				}
+
+				// Update the next update time
+				State.NextUpdateTime = NextUpdateTime;
+				await StateAccessor.TryUpdateAsync(State);
+			}
+		}
+	}
+}
