Horde: Fix node dependencies within the same group not scheduled to be run again when a batch fails.

diff --git a/Engine/Source/Programs/Horde/HordeServer/Collections/Impl/JobCollection.cs b/Engine/Source/Programs/Horde/HordeServer/Collections/Impl/JobCollection.cs
--- a/Engine/Source/Programs/Horde/HordeServer/Collections/Impl/JobCollection.cs
+++ b/Engine/Source/Programs/Horde/HordeServer/Collections/Impl/JobCollection.cs
@@ -1299,7 +1299,7 @@ namespace HordeServer.Collections.Impl
 						{
 							if (Dependency.GroupIdx == GroupIdx)
 							{
-								NewNodesToExecute.Add(Node);
+								NewNodesToExecute.Add(Group.Nodes[Dependency.NodeIdx]);
 							}
 						}
 					}
diff --git a/Engine/Source/Programs/Horde/HordeServerTests/JobCollectionTests.cs b/Engine/Source/Programs/Horde/HordeServerTests/JobCollectionTests.cs
--- a/Engine/Source/Programs/Horde/HordeServerTests/JobCollectionTests.cs
+++ b/Engine/Source/Programs/Horde/HordeServerTests/JobCollectionTests.cs
@@ -133,5 +133,72 @@ namespace HordeServerTests
 			
 			// Manually verify the log output
 		}
+
+		[TestMethod]
+		public Task LostLeaseTestWithDependency()
+		{
+			return LostLeaseTestInternal(true);
+		}
+
+		[TestMethod]
+		public Task LostLeaseTestWithoutDependency()
+		{
+			return LostLeaseTestInternal(false);
+		}
+
+		public async Task LostLeaseTestInternal(bool HasDependency)
+		{
+			Mock<ITemplate> TemplateMock = new Mock<ITemplate>(MockBehavior.Strict);
+			TemplateMock.SetupGet(x => x.InitialAgentType).Returns((string?)null);
+
+			IGraph BaseGraph = await GraphCollection.AddAsync(TemplateMock.Object);
+
+			List<string> Arguments = new List<string>();
+			Arguments.Add("-Target=Step 1");
+			Arguments.Add("-Target=Step 3");
+
+			IJob Job = await JobCollection.AddAsync(ObjectId.GenerateNewId(), new StreamId("ue4-main"), new TemplateRefId("test-build"), ContentHash.SHA1("hello"), BaseGraph, "Test job", 123, 123, null, null, null, "Ben", null, null, null, null, false, false, null, null, null, Arguments);
+
+			await StartBatch(Job, BaseGraph, 0);
+			await RunStep(Job, BaseGraph, 0, 0, JobStepOutcome.Success); // Setup Build
+
+			List<CreateGroupRequest> NewGroups = new List<CreateGroupRequest>();
+
+			CreateGroupRequest InitialGroup = AddGroup(NewGroups);
+			AddNode(InitialGroup, "Step 1", null);
+			AddNode(InitialGroup, "Step 2", HasDependency? new[] { "Step 1" } : null);
+			AddNode(InitialGroup, "Step 3", new[] { "Step 2" });
+
+			IGraph Graph = await GraphCollection.AppendAsync(BaseGraph, NewGroups, null, null);
+			Assert.IsTrue(await JobCollection.TryUpdateGraphAsync(Job, Graph));
+
+			await StartBatch(Job, Graph, 1);
+			await RunStep(Job, Graph, 1, 0, JobStepOutcome.Success); // Step 1
+			await RunStep(Job, Graph, 1, 1, JobStepOutcome.Success); // Step 2
+
+			// Force an error executing the batch
+			Assert.IsTrue(await JobCollection.TryUpdateBatchAsync(Job, Graph, Job.Batches[1].Id, null, JobStepBatchState.Complete, JobStepBatchError.Incomplete));
+
+			// Check that it restarted all three nodes
+			IJob NewJob = (await JobCollection.GetAsync(Job.Id))!;
+			Assert.AreEqual(3, NewJob.Batches.Count);
+			Assert.AreEqual(1, NewJob.Batches[2].GroupIdx);
+
+			if (HasDependency)
+			{
+				Assert.AreEqual(3, NewJob.Batches[2].Steps.Count);
+
+				Assert.AreEqual(0, NewJob.Batches[2].Steps[0].NodeIdx);
+				Assert.AreEqual(1, NewJob.Batches[2].Steps[1].NodeIdx);
+				Assert.AreEqual(2, NewJob.Batches[2].Steps[2].NodeIdx);
+			}
+			else
+			{
+				Assert.AreEqual(2, NewJob.Batches[2].Steps.Count);
+
+				Assert.AreEqual(1, NewJob.Batches[2].Steps[0].NodeIdx);
+				Assert.AreEqual(2, NewJob.Batches[2].Steps[1].NodeIdx);
+			}
+		}
 	}
 }
