Horde: Do not output errors when jobs are cancelled or fail due to a configuration issue. These are not systemic Horde issues.

diff --git a/Engine/Source/Programs/Horde/HordeServer/Tasks/Impl/JobTaskSource.cs b/Engine/Source/Programs/Horde/HordeServer/Tasks/Impl/JobTaskSource.cs
--- a/Engine/Source/Programs/Horde/HordeServer/Tasks/Impl/JobTaskSource.cs
+++ b/Engine/Source/Programs/Horde/HordeServer/Tasks/Impl/JobTaskSource.cs
@@ -1038,7 +1038,7 @@ namespace HordeServer.Tasks.Impl
 				{
 					if (Batch.Error != JobStepBatchError.None)
 					{
-						Logger.LogError("Failed job {JobId}, batch {BatchId} with error {Error}", Job.Id, Batch.Id, Batch.Error);
+						Logger.LogInformation("Failed job {JobId}, batch {BatchId} with error {Error}", Job.Id, Batch.Id, Batch.Error);
 					}
 					if (RunningStepIdx != -1)
 					{
