Horde: Truncate Slack messages at 2048 characters to fix errors.

diff --git a/Engine/Source/Programs/Horde/HordeServer/Utilities/BlockKit/TextObject.cs b/Engine/Source/Programs/Horde/HordeServer/Utilities/BlockKit/TextObject.cs
--- a/Engine/Source/Programs/Horde/HordeServer/Utilities/BlockKit/TextObject.cs
+++ b/Engine/Source/Programs/Horde/HordeServer/Utilities/BlockKit/TextObject.cs
@@ -47,7 +47,16 @@ namespace HordeServer.Utilities.Slack.BlockKit
 		{
 			Writer.WriteStartObject("text");
 			Writer.WriteString("type", IsMarkdown ? "mrkdwn" : "plain_text");
-			Writer.WriteString("text", Text);
+
+			const int MaxLength = 2048;
+
+			string TruncatedText = Text;
+			if (TruncatedText.Length > MaxLength)
+			{
+				TruncatedText = TruncatedText.Substring(0, MaxLength) + "...\n(message truncated)";
+			}
+
+			Writer.WriteString("text", TruncatedText);
 			if (Emoji)
 			{
 				Writer.WriteBoolean("emoji", Emoji);
