Horde: Fix confusing Slack notification when a job step is retried.

diff --git a/Engine/Source/Programs/Horde/HordeServer/Notifications/Impl/SlackNotificationSink.cs b/Engine/Source/Programs/Horde/HordeServer/Notifications/Impl/SlackNotificationSink.cs
--- a/Engine/Source/Programs/Horde/HordeServer/Notifications/Impl/SlackNotificationSink.cs
+++ b/Engine/Source/Programs/Horde/HordeServer/Notifications/Impl/SlackNotificationSink.cs
@@ -260,20 +260,19 @@ namespace HordeServer.Notifications.Impl
 			{
 				List<string> FailedStepStrings = new List<string>();
 				List<string> WarningStepStrings = new List<string>();
-				foreach (IJobStepBatch Batch in Job.Batches)
+
+				IReadOnlyDictionary<NodeRef, IJobStep> NodeToStep = Job.GetStepForNodeMap();
+				foreach ((NodeRef NodeRef, IJobStep Step) in NodeToStep)
 				{
-					foreach (IJobStep Step in Batch.Steps)
+					INode StepNode = Graph.GetNode(NodeRef);
+					string StepName = $"<{JobLink}?step={Step.Id}|{StepNode.Name}>";
+					if (Step.Outcome == JobStepOutcome.Failure)
 					{
-						INode StepNode = Graph.GetNode(new NodeRef(Batch.GroupIdx, Step.NodeIdx));
-						string StepName = $"<{JobLink}?step={Step.Id}|{StepNode.Name}>";
-						if (Step.Outcome == JobStepOutcome.Failure)
-						{
-							FailedStepStrings.Add(StepName);
-						}
-						else if (Step.Outcome == JobStepOutcome.Warnings)
-						{
-							WarningStepStrings.Add(StepName);
-						}
+						FailedStepStrings.Add(StepName);
+					}
+					else if (Step.Outcome == JobStepOutcome.Warnings)
+					{
+						WarningStepStrings.Add(StepName);
 					}
 				}
 
