Horde: Update the state for schedules immediately after creating a job, to ensure the server isn't shutdown while working on the requests.

diff --git a/Engine/Source/Programs/Horde/HordeServer/Collections/Impl/StreamCollection.cs b/Engine/Source/Programs/Horde/HordeServer/Collections/Impl/StreamCollection.cs
--- a/Engine/Source/Programs/Horde/HordeServer/Collections/Impl/StreamCollection.cs
+++ b/Engine/Source/Programs/Horde/HordeServer/Collections/Impl/StreamCollection.cs
@@ -480,7 +480,7 @@ namespace HordeServer.Collections.Impl
 				Updates.Add(Builders<StreamDocument>.Update.Set(LastTriggerTimeField, LastTriggerTime.Value));
 				Schedule.LastTriggerTime = LastTriggerTime.Value;
 			}
-			if (LastTriggerChange.HasValue && LastTriggerChange.Value != Schedule.LastTriggerChange)
+			if (LastTriggerChange.HasValue && LastTriggerChange.Value > Schedule.LastTriggerChange)
 			{
 				FieldDefinition<StreamDocument, int> LastTriggerChangeField = $"{nameof(Stream.Templates)}.{TemplateRefId}.{nameof(Schedule)}.{nameof(Schedule.LastTriggerChange)}";
 				Updates.Add(Builders<StreamDocument>.Update.Set(LastTriggerChangeField, LastTriggerChange.Value));
diff --git a/Engine/Source/Programs/Horde/HordeServer/Services/ScheduleService.cs b/Engine/Source/Programs/Horde/HordeServer/Services/ScheduleService.cs
--- a/Engine/Source/Programs/Horde/HordeServer/Services/ScheduleService.cs
+++ b/Engine/Source/Programs/Horde/HordeServer/Services/ScheduleService.cs
@@ -206,28 +206,21 @@ namespace HordeServer.Services
 					RemoveJobIds.Add(ActiveJobId);
 				}
 			}
+			if (RemoveJobIds.Count > 0)
+			{
+				await StreamService.UpdateScheduleTriggerAsync(Stream, TemplateRefId, Now, null, new List<ObjectId>(), RemoveJobIds);
+			}
 
 			// Trigger this schedule
-			List<IJob> CreatedJobs = new List<IJob>();
 			try
 			{
-				await TriggerAsync(Stream, TemplateRefId, TemplateRef, Schedule, Schedule.ActiveJobs.Count - RemoveJobIds.Count, CreatedJobs);
+				await TriggerAsync(Stream, TemplateRefId, TemplateRef, Schedule, Schedule.ActiveJobs.Count - RemoveJobIds.Count, Now);
 			}
 			catch (Exception Ex)
 			{
 				Logger.LogError(Ex, "Failed to start schedule {StreamId}/{TemplateRefId}", Stream.Id, TemplateRefId);
 			}
 
-			// Update the stream
-			if (CreatedJobs.Count > 0)
-			{
-				await StreamService.UpdateScheduleTriggerAsync(Stream, TemplateRefId, Now, CreatedJobs.Max(x => x.Change), CreatedJobs.ConvertAll(x => x.Id), RemoveJobIds);
-			}
-			else if (RemoveJobIds.Count > 0)
-			{
-				await StreamService.UpdateScheduleTriggerAsync(Stream, TemplateRefId, null, null, new List<ObjectId>(), RemoveJobIds);
-			}
-
 			// Return the next update time
 			return Schedule.GetNextTriggerTime(Now, TimeZone);
 		}
@@ -240,9 +233,9 @@ namespace HordeServer.Services
 		/// <param name="TemplateRef"></param>
 		/// <param name="Schedule"></param>
 		/// <param name="NumActiveJobs"></param>
-		/// <param name="CreatedJobs"></param>
+		/// <param name="Now"></param>
 		/// <returns>Async task</returns>
-		private async Task TriggerAsync(IStream Stream, TemplateRefId TemplateRefId, TemplateRef TemplateRef, Schedule Schedule, int NumActiveJobs, List<IJob> CreatedJobs)
+		private async Task TriggerAsync(IStream Stream, TemplateRefId TemplateRefId, TemplateRef TemplateRef, Schedule Schedule, int NumActiveJobs, DateTimeOffset Now)
 		{
 			// Check we're not already at the maximum number of allowed jobs
 			if (Schedule.MaxActive != 0 && NumActiveJobs >= Schedule.MaxActive)
@@ -357,7 +350,7 @@ namespace HordeServer.Services
 				List<string> DefaultArguments = Template.GetDefaultArguments();
 				IJob NewJob = await JobService.CreateJobAsync(null, Stream.Id, TemplateRefId, Template.Id, Graph, Template.Name, Change, CodeChange, null, null, null, null, Template.Priority, null, null, TemplateRef.ChainedJobs, TemplateRef.ShowUgsBadges, TemplateRef.ShowUgsAlerts, TemplateRef.NotificationChannel, TemplateRef.NotificationChannelFilter, null, Template.Counters, DefaultArguments);
 				Logger.LogInformation("Started new job for {StreamName} template {TemplateName} at CL {Change} (Code CL {CodeChange}): {JobId}", Stream.Id, TemplateRef.Name, Change, CodeChange, NewJob.Id);
-				CreatedJobs.Add(NewJob);
+				await StreamService.UpdateScheduleTriggerAsync(Stream, TemplateRefId, Now, Change, new List<ObjectId> { NewJob.Id }, new List<ObjectId>());
 			}
 		}
 
