Horde: Fix ambiguous references to Digest class.

diff --git a/Engine/Source/Programs/Horde/HordeAgent/Commands/Execution/ExecuteCommand.cs b/Engine/Source/Programs/Horde/HordeAgent/Commands/Execution/ExecuteCommand.cs
--- a/Engine/Source/Programs/Horde/HordeAgent/Commands/Execution/ExecuteCommand.cs
+++ b/Engine/Source/Programs/Horde/HordeAgent/Commands/Execution/ExecuteCommand.cs
@@ -27,6 +27,8 @@ using Microsoft.Extensions.Hosting;
 using Microsoft.Extensions.DependencyInjection;
 using Microsoft.Extensions.Configuration;
 
+using Digest = Build.Bazel.Remote.Execution.V2.Digest;
+
 namespace HordeAgent.Commands
 {
 	/// <summary>
diff --git a/Engine/Source/Programs/Horde/HordeAgent/Execution/ActionExecutor.cs b/Engine/Source/Programs/Horde/HordeAgent/Execution/ActionExecutor.cs
--- a/Engine/Source/Programs/Horde/HordeAgent/Execution/ActionExecutor.cs
+++ b/Engine/Source/Programs/Horde/HordeAgent/Execution/ActionExecutor.cs
@@ -14,6 +14,7 @@ using System.Threading;
 using System.Threading.Tasks;
 using System.Diagnostics;
 using Action = Build.Bazel.Remote.Execution.V2.Action;
+using Digest = Build.Bazel.Remote.Execution.V2.Digest;
 using Directory = Build.Bazel.Remote.Execution.V2.Directory;
 using RpcCommand = Build.Bazel.Remote.Execution.V2.Command;
 using Status = Google.Rpc.Status;
diff --git a/Engine/Source/Programs/Horde/HordeAgent/Utility/StorageExtensions.cs b/Engine/Source/Programs/Horde/HordeAgent/Utility/StorageExtensions.cs
--- a/Engine/Source/Programs/Horde/HordeAgent/Utility/StorageExtensions.cs
+++ b/Engine/Source/Programs/Horde/HordeAgent/Utility/StorageExtensions.cs
@@ -9,6 +9,8 @@ using Blake3;
 using Build.Bazel.Remote.Execution.V2;
 using EpicGames.Core;
 using Google.Protobuf;
+
+using Digest = Build.Bazel.Remote.Execution.V2.Digest;
 using ContentAddressableStorageClient = Build.Bazel.Remote.Execution.V2.ContentAddressableStorage.ContentAddressableStorageClient;
 
 namespace HordeAgent.Utility
