Horde: Add support for step reports
#rnx

diff --git a/Engine/Source/Programs/Horde/HordeDashboard/src/backend/Api.ts b/Engine/Source/Programs/Horde/HordeDashboard/src/backend/Api.ts
--- a/Engine/Source/Programs/Horde/HordeDashboard/src/backend/Api.ts
+++ b/Engine/Source/Programs/Horde/HordeDashboard/src/backend/Api.ts
@@ -781,17 +781,9 @@ export enum ReportPlacement {
 	Summary = "Summary"
 }
 
-export enum ReportScope {
-	Job = "Job",
-	Step = "Step"
-}
-
 /**Information about a report associated with a job*/
 export type GetReportResponse = {
 
-	// Report scope
-	scope: ReportScope;
-
 	// Report placement
 	placement: ReportPlacement;
 
@@ -941,6 +933,9 @@ export type GetStepResponse = {
 	/**Time at which the run finished (UTC) */
 	finishTime: Date | string;
 
+	/** List of reports */
+	reports?: GetReportResponse[];
+
 	/**User-defined properties for this jobstep. */
 	properties: { [key: string]: string };
 
diff --git a/Engine/Source/Programs/Horde/HordeDashboard/src/backend/JobDetails.ts b/Engine/Source/Programs/Horde/HordeDashboard/src/backend/JobDetails.ts
--- a/Engine/Source/Programs/Horde/HordeDashboard/src/backend/JobDetails.ts
+++ b/Engine/Source/Programs/Horde/HordeDashboard/src/backend/JobDetails.ts
@@ -5,7 +5,7 @@ import { getTheme, mergeStyles, mergeStyleSets } from 'office-ui-fabric-react/li
 import backend from '.';
 import { getBatchInitElapsed, getNiceTime, getStepElapsed, getStepETA, getStepFinishTime, getStepTimingDelta } from '../base/utilities/timeUtils';
 import { getBatchText } from '../components/JobDetailCommon';
-import { AgentData, ArtifactData, BatchData, EventData, GetGroupResponse, GetJobStepRefResponse, GetJobTimingResponse, GetLabelResponse, GetLabelStateResponse, GetLabelTimingInfoResponse, GetTemplateResponse, GroupData, IssueData, JobData, JobStepBatchState, JobStepOutcome, JobStepState, LabelState, NodeData, ReportPlacement, ReportScope, StepData, StreamData, TestData } from './Api';
+import { AgentData, ArtifactData, BatchData, EventData, GetGroupResponse, GetJobStepRefResponse, GetJobTimingResponse, GetLabelResponse, GetLabelStateResponse, GetLabelTimingInfoResponse, GetTemplateResponse, GroupData, IssueData, JobData, JobStepBatchState, JobStepOutcome, JobStepState, LabelState, NodeData, ReportPlacement, StepData, StreamData, TestData } from './Api';
 import { projectStore } from './ProjectStore';
 
 const theme = getTheme();
@@ -633,43 +633,54 @@ export class JobDetails {
         })
     }    
 
-    getReportData(scope: ReportScope, placement: ReportPlacement): string | undefined {
-        
-        // placement and scope not currently returned for reports
+    getReportData(placement: ReportPlacement, stepId?: string): string | undefined {
+                
+        if (stepId) {
 
-        if (scope === ReportScope.Job && placement === ReportPlacement.Summary) {
-            if (this.jobdata?.reports?.length) {
-                return this.reportData.get(this.jobdata.reports[0].artifactId);
+            const step = this.stepById(stepId);
+
+            if (!step) {
+                return undefined;
             }
-        }
 
-        return undefined;
+            const report = step.reports?.find(r => r.placement === placement);
+
+            if (!report) {
+                return undefined;
+            }
+
+            return this.reportData.get(report.artifactId);
+
+        } else {
 
-        /*
-        const report = this.jobdata?.reports?.find(r => r.placement === placement && r.scope === scope);
+            if (this.jobdata?.reports?.length) {
+                return this.reportData.get(this.jobdata.reports[0].artifactId);
+            }
 
-        if (!report) {
             return undefined;
+
         }
-        return this.reportData.get(report.artifactId);
-        */
     }
 
     private async queryReports():Promise<void> {
 
-        if (!this.jobdata?.reports?.length) {
-            return;
-        }
+        const artifacts: string[] = [];
+
+        this.jobdata?.reports?.forEach(r => artifacts.push(r.artifactId));
+        this.jobdata?.batches?.forEach(b => {
+            b.steps.forEach(s => s.reports?.forEach( r => artifacts.push(r.artifactId)));
+        });
+        
 
-        for (let i = 0; i < this.jobdata.reports.length; i++) {
+        for (let i = 0; i < artifacts.length; i++) {
 
-            const report = this.jobdata.reports[i];
+            const artifactId = artifacts[i];
 
-            if (!this.reportData.has(report.artifactId)) {
+            if (!this.reportData.has(artifactId)) {
 
-                const r = await backend.getArtifactDataById(report.artifactId) as unknown as string;
+                const r = await backend.getArtifactDataById(artifactId) as unknown as string;
                 
-                this.reportData.set(report.artifactId, r);                                
+                this.reportData.set(artifactId, r);                                
             }
         }
     }
