Horde: Additional logging for BeginStep() RPC responses.

diff --git a/Engine/Source/Programs/Horde/HordeServer/Services/RpcService.cs b/Engine/Source/Programs/Horde/HordeServer/Services/RpcService.cs
--- a/Engine/Source/Programs/Horde/HordeServer/Services/RpcService.cs
+++ b/Engine/Source/Programs/Horde/HordeServer/Services/RpcService.cs
@@ -562,6 +562,7 @@ namespace HordeServer.Services
 					// If there aren't any more steps, send a complete message
 					if (StepIdx == Batch.Steps.Count)
 					{
+						Logger.LogDebug("All steps in batch complete");
 						BeginStepResponse Response = new BeginStepResponse();
 						Response.State = BeginStepResponse.Types.Result.Complete;
 						return Response;
@@ -575,6 +576,7 @@ namespace HordeServer.Services
 					}
 					if (Step.State == JobStepState.Waiting)
 					{
+						Logger.LogDebug("Waiting for job {JobId}, batch {BatchId}, step {StepId}")
 						BeginStepResponse Response = new BeginStepResponse();
 						Response.State = BeginStepResponse.Types.Result.Waiting;
 						return Response;
diff --git a/Engine/Source/Programs/Horde/HordeServer/Tasks/Impl/JobTaskSource.cs b/Engine/Source/Programs/Horde/HordeServer/Tasks/Impl/JobTaskSource.cs
--- a/Engine/Source/Programs/Horde/HordeServer/Tasks/Impl/JobTaskSource.cs
+++ b/Engine/Source/Programs/Horde/HordeServer/Tasks/Impl/JobTaskSource.cs
@@ -573,7 +573,7 @@ namespace HordeServer.Tasks.Impl
 
 		private async Task<IJob?> SkipBatchAsync(IJob Job, int BatchIdx, IGraph Graph, JobStepBatchError Reason)
 		{
-			Logger.LogInformation("Skipping batch {BatchIdx} for job {JobId}", BatchIdx, Job.Id);
+			Logger.LogInformation("Skipping batch {BatchIdx} for job {JobId} (reason: {Reason})", BatchIdx, Job.Id, Reason);
 
 			IReadOnlyList<(LabelState, LabelOutcome)> OldLabelStates = Job.GetLabelStates(Graph);
 			IJob? NewJob = await Jobs.SkipBatchAsync(Job, BatchIdx, Graph, Reason);
