Missing file.

diff --git a/Engine/Source/Programs/Horde/HordeServer/Services/RpcService.cs b/Engine/Source/Programs/Horde/HordeServer/Services/RpcService.cs
--- a/Engine/Source/Programs/Horde/HordeServer/Services/RpcService.cs
+++ b/Engine/Source/Programs/Horde/HordeServer/Services/RpcService.cs
@@ -562,7 +562,7 @@ namespace HordeServer.Services
 					// If there aren't any more steps, send a complete message
 					if (StepIdx == Batch.Steps.Count)
 					{
-						Logger.LogDebug("All steps in batch complete");
+						Logger.LogDebug("Job {JobId} batch {BatchId} is complete", Job.Id, Batch.Id);
 						BeginStepResponse Response = new BeginStepResponse();
 						Response.State = BeginStepResponse.Types.Result.Complete;
 						return Response;
@@ -576,7 +576,7 @@ namespace HordeServer.Services
 					}
 					if (Step.State == JobStepState.Waiting)
 					{
-						Logger.LogDebug("Waiting for job {JobId}, batch {BatchId}, step {StepId}")
+						Logger.LogDebug("Waiting for job {JobId}, batch {BatchId}, step {StepId}", Job.Id, Batch.Id, Step.Id);
 						BeginStepResponse Response = new BeginStepResponse();
 						Response.State = BeginStepResponse.Types.Result.Waiting;
 						return Response;
