Horde: Fix situations where batches are unnecessarily restarted due to being treated as complete.

Previous method to check whether batch can be modified just checked to see if all steps were complete. This is incorrect, and assumes the batch cannot be updated while the agent is requesting the next batch to execute. Now has an explicit state for checking if the batch is shutting down.

diff --git a/Engine/Source/Programs/Horde/HordeServer/Api/Jobs.cs b/Engine/Source/Programs/Horde/HordeServer/Api/Jobs.cs
--- a/Engine/Source/Programs/Horde/HordeServer/Api/Jobs.cs
+++ b/Engine/Source/Programs/Horde/HordeServer/Api/Jobs.cs
@@ -586,14 +586,24 @@ namespace HordeServer.Api
 		Ready = 1,
 
 		/// <summary>
-		/// Currently running
+		/// Preparing to execute work
 		/// </summary>
-		Running = 2,
+		Starting = 2,
+
+		/// <summary>
+		/// Executing work
+		/// </summary>
+		Running = 3,
+
+		/// <summary>
+		/// Preparing to stop
+		/// </summary>
+		Stopping = 4,
 
 		/// <summary>
 		/// All steps have finished executing
 		/// </summary>
-		Complete = 3
+		Complete = 5
 	}
 
 	/// <summary>
diff --git a/Engine/Source/Programs/Horde/HordeServer/Collections/Impl/JobCollection.cs b/Engine/Source/Programs/Horde/HordeServer/Collections/Impl/JobCollection.cs
--- a/Engine/Source/Programs/Horde/HordeServer/Collections/Impl/JobCollection.cs
+++ b/Engine/Source/Programs/Horde/HordeServer/Collections/Impl/JobCollection.cs
@@ -617,7 +617,7 @@ namespace HordeServer.Collections.Impl
 				Batch.State = NewState.Value;
 				Updates.Add(UpdateBuilder.Set(x => x.Batches[BatchIdx].State, Batch.State));
 
-				if (NewState == JobStepBatchState.Running)
+				if (Batch.StartTime == null && NewState >= JobStepBatchState.Starting)
 				{
 					Batch.StartTime = DateTimeOffset.Now;
 					Updates.Add(UpdateBuilder.Set(x => x.Batches[BatchIdx].StartTime, Batch.StartTime));
@@ -1264,7 +1264,7 @@ namespace HordeServer.Collections.Impl
 			// Cancel any batches which are still running but are no longer required
 			foreach (JobStepBatchDocument Batch in Job.Batches)
 			{
-				if (Batch.State == JobStepBatchState.Running)
+				if (Batch.State == JobStepBatchState.Starting || Batch.State == JobStepBatchState.Running)
 				{
 					INodeGroup Group = Graph.Groups[Batch.GroupIdx];
 					if (!Batch.Steps.Any(x => NewNodesToExecute.Contains(Group.Nodes[x.NodeIdx])))
diff --git a/Engine/Source/Programs/Horde/HordeServer/Models/Job.cs b/Engine/Source/Programs/Horde/HordeServer/Models/Job.cs
--- a/Engine/Source/Programs/Horde/HordeServer/Models/Job.cs
+++ b/Engine/Source/Programs/Horde/HordeServer/Models/Job.cs
@@ -266,15 +266,7 @@ namespace HordeServer.Models
 		/// <returns>True if new steps can be appended to this batch</returns>
 		public static bool CanBeAppendedTo(this IJobStepBatch Batch)
 		{
-			for (int Idx = Batch.Steps.Count - 1; Idx >= 0; Idx--)
-			{
-				IJobStep Step = Batch.Steps[Idx];
-				if (Step.State != JobStepState.Skipped)
-				{
-					return Step.State != JobStepState.Completed && Step.State != JobStepState.Aborted;
-				}
-			}
-			return false;
+			return State <= JobStepBatchState.Running;
 		}
 
 		/// <summary>
@@ -726,7 +718,7 @@ namespace HordeServer.Models
 					}
 					else if (Step.State == JobStepState.Ready || Step.State == JobStepState.Waiting)
 					{
-						if (Batch.State == JobStepBatchState.Running)
+						if (Batch.State == JobStepBatchState.Starting || Batch.State == JobStepBatchState.Running)
 						{
 							return JobState.Running;
 						}
diff --git a/Engine/Source/Programs/Horde/HordeServer/Services/JobService.cs b/Engine/Source/Programs/Horde/HordeServer/Services/JobService.cs
--- a/Engine/Source/Programs/Horde/HordeServer/Services/JobService.cs
+++ b/Engine/Source/Programs/Horde/HordeServer/Services/JobService.cs
@@ -235,7 +235,7 @@ namespace HordeServer.Services
 
 			if (PreflightChange != null)
 			{
-//PF				ClonedPreflightChange = await CloneShelvedChangeAsync(ClonedPreflightChange ?? PreflightChange.Value);
+				//PF				ClonedPreflightChange = await CloneShelvedChangeAsync(ClonedPreflightChange ?? PreflightChange.Value);
 			}
 
 			Logger.LogInformation("Creating job at CL {Change}, code CL {CodeChange}, preflight CL {PreflightChange}, cloned CL {ClonedPreflightChange}", Change, CodeChange, PreflightChange, ClonedPreflightChange);
@@ -300,7 +300,7 @@ namespace HordeServer.Services
 				if (Job.TemplateHash != NewJob.TemplateHash) continue;
 				if (Job.StartedByUser != NewJob.StartedByUser) continue;
 				if (string.Join(",", Job.Arguments) != string.Join(",", NewJob.Arguments)) continue;
-				
+
 				bool bUpdateSuccessful = await UpdateJobAsync(Job, null, null, null, "horde.duplicated.by.newer.CL", null, null, null);
 				if (!bUpdateSuccessful)
 				{
@@ -308,7 +308,7 @@ namespace HordeServer.Services
 				}
 			}
 		}
-		
+
 		/// <summary>
 		/// Deletes a job
 		/// </summary>
@@ -318,7 +318,7 @@ namespace HordeServer.Services
 			using IDisposable Scope = Logger.BeginScope("DeleteJobAsync({JobId})", Job.Id);
 
 			// Delete the job
-			while(!await Jobs.RemoveAsync(Job))
+			while (!await Jobs.RemoveAsync(Job))
 			{
 				IJob? NewJob = await Jobs.GetAsync(Job.Id);
 				if (NewJob == null)
@@ -371,7 +371,7 @@ namespace HordeServer.Services
 		public async Task<bool> UpdateJobAsync(IJob Job, string? Name = null, Priority? Priority = null, bool? AutoSubmit = null, string? AbortedByUser = null, ObjectId? OnCompleteTriggerId = null, List<Report>? Reports = null, List<string>? Arguments = null, KeyValuePair<int, ObjectId>? LabelIdxToTriggerId = null)
 		{
 			using IDisposable Scope = Logger.BeginScope("UpdateJobAsync({JobId})", Job.Id);
-			for(; ;)
+			for (; ; )
 			{
 				IGraph Graph = await GetGraphAsync(Job);
 
@@ -385,9 +385,9 @@ namespace HordeServer.Services
 					await JobTaskSource.UpdateUgsBadges(Job, Graph, OldLabelStates);
 
 					// Cancel any leases which are no longer required
-					foreach(IJobStepBatch Batch in Job.Batches)
+					foreach (IJobStepBatch Batch in Job.Batches)
 					{
-						if (Batch.Error == JobStepBatchError.Cancelled && Batch.State == JobStepBatchState.Running && Batch.AgentId != null && Batch.LeaseId != null)
+						if (Batch.Error == JobStepBatchError.Cancelled && (Batch.State == JobStepBatchState.Starting || Batch.State == JobStepBatchState.Running) && Batch.AgentId != null && Batch.LeaseId != null)
 						{
 							await CancelLeaseAsync(Batch.AgentId.Value, Batch.LeaseId.Value);
 						}
@@ -397,7 +397,7 @@ namespace HordeServer.Services
 
 				// Get the new job state
 				IJob? NewJob = await GetJobAsync(Job.Id);
-				if(NewJob == null)
+				if (NewJob == null)
 				{
 					return false;
 				}
@@ -413,7 +413,7 @@ namespace HordeServer.Services
 		/// <returns></returns>
 		async Task CancelLeaseAsync(AgentId AgentId, ObjectId LeaseId)
 		{
-			for(; ;)
+			for (; ; )
 			{
 				IAgent? Agent = await Agents.GetAsync(AgentId);
 				if (Agent == null)
@@ -422,23 +422,23 @@ namespace HordeServer.Services
 				}
 
 				int Index = 0;
-				while(Index < Agent.Leases.Count && Agent.Leases[Index].Id != LeaseId)
+				while (Index < Agent.Leases.Count && Agent.Leases[Index].Id != LeaseId)
 				{
 					Index++;
 				}
-				if(Index == -1)
+				if (Index == -1)
 				{
 					break;
 				}
 
 				AgentLease Lease = Agent.Leases[Index];
-				if(Lease.State != LeaseState.Active)
+				if (Lease.State != LeaseState.Active)
 				{
 					break;
 				}
 
 				IAgent? NewAgent = await Agents.TryCancelLeaseAsync(Agent, Index);
-				if(NewAgent != null)
+				if (NewAgent != null)
 				{
 					JobTaskSource.CancelLongPollForAgent(Agent.Id);
 					break;
@@ -612,12 +612,12 @@ namespace HordeServer.Services
 				if (JobTiming == null)
 				{
 					IJobTiming? NewJobTiming = await JobTimings.TryAddAsync(Job.Id, NewSteps);
-					if(NewJobTiming != null)
+					if (NewJobTiming != null)
 					{
 						return NewJobTiming;
 					}
 				}
-				else if(NewSteps.Count == 0)
+				else if (NewSteps.Count == 0)
 				{
 					return JobTiming;
 				}
@@ -631,7 +631,7 @@ namespace HordeServer.Services
 				}
 
 				// Update the cached list of steps for the next iteration
-				foreach(JobStepTimingData NewStep in NewSteps)
+				foreach (JobStepTimingData NewStep in NewSteps)
 				{
 					CachedNewSteps[NewStep.Name] = NewStep;
 				}
@@ -669,7 +669,7 @@ namespace HordeServer.Services
 			}
 
 			// Compute the averages
-			if(Count == 0)
+			if (Count == 0)
 			{
 				return new JobStepTimingData(NodeName, null, null, null);
 			}
@@ -692,11 +692,11 @@ namespace HordeServer.Services
 			using IDisposable Scope = Logger.BeginScope("UpdateBatchAsync({JobId})", Job.Id);
 
 			bool bCheckForBadAgent = true;
-			for(; ;)
+			for (; ; )
 			{
 				// Find the index of the appropriate batch
 				int BatchIdx = Job.Batches.FindIndex(x => x.Id == BatchId);
-				if(BatchIdx == -1)
+				if (BatchIdx == -1)
 				{
 					break;
 				}
@@ -739,11 +739,8 @@ namespace HordeServer.Services
 					}
 				}
 
-				// Get the graph for this job
-				IGraph Graph = await GetGraphAsync(Job);
-
 				// Update the batch state
-				if (await Jobs.TryUpdateBatchAsync(Job, Graph, BatchId, NewLogId, NewState, NewError))
+				if (await TryUpdateBatchAsync(Job, BatchId, NewLogId, NewState, NewError))
 				{
 					return true;
 				}
@@ -759,6 +756,25 @@ namespace HordeServer.Services
 			return false;
 		}
 
+		/// <summary>
+		/// Attempts to update the state of a batch
+		/// </summary>
+		/// <param name="Job">Job to update</param>
+		/// <param name="BatchId">Unique id of the batch to update</param>
+		/// <param name="NewLogId">The new log file id</param>
+		/// <param name="NewState">New state of the jobstep</param>
+		/// <param name="NewError">New error state</param>
+		/// <returns>The updated job, otherwise null</returns>
+		public async Task<bool> TryUpdateBatchAsync(IJob Job, SubResourceId BatchId, ObjectId? NewLogId = null, JobStepBatchState? NewState = null, JobStepBatchError? NewError = null)
+		{
+			IGraph Graph = await GetGraphAsync(Job);
+			if (!await Jobs.TryUpdateBatchAsync(Job, Graph, BatchId, NewLogId, NewState, NewError))
+			{
+				return false;
+			}
+			return true;
+		}
+
 		/// <summary>
 		/// Update a jobstep state
 		/// </summary>
diff --git a/Engine/Source/Programs/Horde/HordeServer/Services/RpcService.cs b/Engine/Source/Programs/Horde/HordeServer/Services/RpcService.cs
--- a/Engine/Source/Programs/Horde/HordeServer/Services/RpcService.cs
+++ b/Engine/Source/Programs/Horde/HordeServer/Services/RpcService.cs
@@ -499,7 +499,7 @@ namespace HordeServer.Services
 			}
 
 			IJobStepBatch Batch = AuthorizeBatch(Job, Request.BatchId.ToSubResourceId(), Context);
-			if (!await JobService.UpdateBatchAsync(Job, BatchId, null, Api.JobStepBatchState.Running))
+			if (!await JobService.UpdateBatchAsync(Job, BatchId, null, Api.JobStepBatchState.Starting))
 			{
 				throw new StructuredRpcException(StatusCode.NotFound, "Batch {JobId}:{BatchId} not found for updating", Request.JobId, Request.BatchId);
 			}
@@ -539,89 +539,97 @@ namespace HordeServer.Services
 		/// <returns>Information about the new agent</returns>
 		public async override Task<BeginStepResponse> BeginStep(BeginStepRequest Request, ServerCallContext Context)
 		{
-			// Log file that's allocated for this step. We only need do this once.
-			ILogFile? Log = null;
-
-			// Get the job
+			Boxed<ILogFile?> Log = new Boxed<ILogFile?>();
 			for (; ; )
 			{
-				// Check the job exists and we can access it
-				IJob? Job = await JobService.GetJobAsync(Request.JobId.ToObjectId());
-				if (Job == null)
+				BeginStepResponse? Response = await TryBeginStep(Request, Log, Context);
+				if (Response != null)
 				{
-					throw new StructuredRpcException(StatusCode.NotFound, "Job {JobId} not found", Request.JobId);
+					return Response;
 				}
+			}
+		}
 
-				// Find the batch being executed
-				IJobStepBatch Batch = AuthorizeBatch(Job, Request.BatchId.ToSubResourceId(), Context);
+		async Task<BeginStepResponse?> TryBeginStep(BeginStepRequest Request, Boxed<ILogFile?> Log, ServerCallContext Context)
+		{
+			// Check the job exists and we can access it
+			IJob? Job = await JobService.GetJobAsync(Request.JobId.ToObjectId());
+			if (Job == null)
+			{
+				throw new StructuredRpcException(StatusCode.NotFound, "Job {JobId} not found", Request.JobId);
+			}
 
-				// Figure out which step to execute next
-				IJobStep? Step;
-				for (int StepIdx = 0; ; StepIdx++)
-				{
-					// If there aren't any more steps, send a complete message
-					if (StepIdx == Batch.Steps.Count)
-					{
-						Logger.LogDebug("Job {JobId} batch {BatchId} is complete", Job.Id, Batch.Id);
-						BeginStepResponse Response = new BeginStepResponse();
-						Response.State = BeginStepResponse.Types.Result.Complete;
-						return Response;
-					}
+			// Find the batch being executed
+			IJobStepBatch Batch = AuthorizeBatch(Job, Request.BatchId.ToSubResourceId(), Context);
 
-					// Check if this step is ready to be executed
-					Step = Batch.Steps[StepIdx];
-					if (Step.State == JobStepState.Ready)
-					{
-						break;
-					}
-					if (Step.State == JobStepState.Waiting)
+			// Figure out which step to execute next
+			IJobStep? Step;
+			for (int StepIdx = 0; ; StepIdx++)
+			{
+				// If there aren't any more steps, send a complete message
+				if (StepIdx == Batch.Steps.Count)
+				{
+					Logger.LogDebug("Job {JobId} batch {BatchId} is complete", Job.Id, Batch.Id);
+					if (!await JobService.TryUpdateBatchAsync(Job, Batch.Id, NewState: Api.JobStepBatchState.Stopping))
 					{
-						Logger.LogDebug("Waiting for job {JobId}, batch {BatchId}, step {StepId}", Job.Id, Batch.Id, Step.Id);
-						BeginStepResponse Response = new BeginStepResponse();
-						Response.State = BeginStepResponse.Types.Result.Waiting;
-						return Response;
+						return null;
 					}
+					return new BeginStepResponse { State = BeginStepResponse.Types.Result.Complete };
 				}
 
-				// Create a log file if necessary
-				if (Log == null)
+				// Check if this step is ready to be executed
+				Step = Batch.Steps[StepIdx];
+				if (Step.State == JobStepState.Ready)
+				{
+					break;
+				}
+				if (Step.State == JobStepState.Waiting)
 				{
-					Log = await LogFileService.CreateLogFileAsync(Job.Id, Batch.SessionId, Api.LogType.Json);
+					Logger.LogDebug("Waiting for job {JobId}, batch {BatchId}, step {StepId}", Job.Id, Batch.Id, Step.Id);
+					return new BeginStepResponse { State = BeginStepResponse.Types.Result.Waiting };
 				}
+			}
 
-				// Get the node for this step
-				IGraph Graph = await JobService.GetGraphAsync(Job);
-				INode Node = Graph.Groups[Batch.GroupIdx].Nodes[Step.NodeIdx];
-
-				// Figure out all the credentials for it (and check we can access them)
-				Dictionary<string, string> Credentials = new Dictionary<string, string>();
-				//				if (Node.Credentials != null)
-				//				{
-				//					ClaimsPrincipal Principal = new ClaimsPrincipal(new ClaimsIdentity(Job.Claims.Select(x => new Claim(x.Type, x.Value))));
-				//					if (!await GetCredentialsForStep(Principal, Node, Credentials, Message => FailStep(Job, Batch.Id, Step.Id, Log, Message)))
-				//					{
-				//						Log = null;
-				//						continue;
-				//					}
-				//				}
-
-				// Update the step state
-				if (await JobService.TryUpdateStepAsync(Job, Batch.Id, Step.Id, JobStepState.Running, JobStepOutcome.Unspecified, null, null, Log.Id, null, null, null, null))
+			// Create a log file if necessary
+			if (Log.Value == null)
+			{
+				Log = await LogFileService.CreateLogFileAsync(Job.Id, Batch.SessionId, Api.LogType.Json);
+			}
+
+			// Get the node for this step
+			IGraph Graph = await JobService.GetGraphAsync(Job);
+			INode Node = Graph.Groups[Batch.GroupIdx].Nodes[Step.NodeIdx];
+
+			// Figure out all the credentials for it (and check we can access them)
+			Dictionary<string, string> Credentials = new Dictionary<string, string>();
+			//				if (Node.Credentials != null)
+			//				{
+			//					ClaimsPrincipal Principal = new ClaimsPrincipal(new ClaimsIdentity(Job.Claims.Select(x => new Claim(x.Type, x.Value))));
+			//					if (!await GetCredentialsForStep(Principal, Node, Credentials, Message => FailStep(Job, Batch.Id, Step.Id, Log, Message)))
+			//					{
+			//						Log = null;
+			//						continue;
+			//					}
+			//				}
+
+			// Update the step state
+			if (await JobService.TryUpdateStepAsync(Job, Batch.Id, Step.Id, JobStepState.Running, JobStepOutcome.Unspecified, null, null, Log.Id, null, null, null, null))
+			{
+				BeginStepResponse Response = new BeginStepResponse();
+				Response.State = BeginStepResponse.Types.Result.Ready;
+				Response.LogId = Log.Id.ToString();
+				Response.StepId = Step.Id.ToString();
+				Response.Name = Node.Name;
+				Response.Credentials.Add(Credentials);
+				if (Node.Properties != null)
 				{
-					BeginStepResponse Response = new BeginStepResponse();
-					Response.State = BeginStepResponse.Types.Result.Ready;
-					Response.LogId = Log.Id.ToString();
-					Response.StepId = Step.Id.ToString();
-					Response.Name = Node.Name;
-					Response.Credentials.Add(Credentials);
-					if (Node.Properties != null)
-					{
-						Response.Properties.Add(Node.Properties);
-					}
-					Response.Warnings = Node.Warnings;
-					return Response;
+					Response.Properties.Add(Node.Properties);
 				}
+				Response.Warnings = Node.Warnings;
+				return Response;
 			}
+
+			return null;
 		}
 
 		/// <summary>
diff --git a/Engine/Source/Programs/Horde/HordeServer/Tasks/Impl/JobTaskSource.cs b/Engine/Source/Programs/Horde/HordeServer/Tasks/Impl/JobTaskSource.cs
--- a/Engine/Source/Programs/Horde/HordeServer/Tasks/Impl/JobTaskSource.cs
+++ b/Engine/Source/Programs/Horde/HordeServer/Tasks/Impl/JobTaskSource.cs
@@ -494,7 +494,7 @@ namespace HordeServer.Tasks.Impl
 				for (int BatchIdx = 0; NewJob != null && BatchIdx < NewJob.Batches.Count; BatchIdx++)
 				{
 					IJobStepBatch Batch = NewJob.Batches[BatchIdx];
-					IsRunning |= (Batch.State == JobStepBatchState.Ready || Batch.State == JobStepBatchState.Running);
+					IsRunning |= (Batch.State == JobStepBatchState.Ready || Batch.State == JobStepBatchState.Starting || Batch.State == JobStepBatchState.Running || Batch.State == JobStepBatchState.Stopping);
 
 					// Skip any batches which aren't ready to execute
 					if (Batch.State == JobStepBatchState.Ready)
