Missing files.

diff --git a/Engine/Source/Programs/Horde/HordeServer/Models/Job.cs b/Engine/Source/Programs/Horde/HordeServer/Models/Job.cs
--- a/Engine/Source/Programs/Horde/HordeServer/Models/Job.cs
+++ b/Engine/Source/Programs/Horde/HordeServer/Models/Job.cs
@@ -266,7 +266,7 @@ namespace HordeServer.Models
 		/// <returns>True if new steps can be appended to this batch</returns>
 		public static bool CanBeAppendedTo(this IJobStepBatch Batch)
 		{
-			return State <= JobStepBatchState.Running;
+			return Batch.State <= JobStepBatchState.Running;
 		}
 
 		/// <summary>
diff --git a/Engine/Source/Programs/Horde/HordeServer/Services/RpcService.cs b/Engine/Source/Programs/Horde/HordeServer/Services/RpcService.cs
--- a/Engine/Source/Programs/Horde/HordeServer/Services/RpcService.cs
+++ b/Engine/Source/Programs/Horde/HordeServer/Services/RpcService.cs
@@ -1,5 +1,6 @@
 // Copyright Epic Games, Inc. All Rights Reserved.
 
+using EpicGames.Core;
 using Google.Protobuf.WellKnownTypes;
 using Grpc.Core;
 using HordeCommon;
@@ -539,7 +540,7 @@ namespace HordeServer.Services
 		/// <returns>Information about the new agent</returns>
 		public async override Task<BeginStepResponse> BeginStep(BeginStepRequest Request, ServerCallContext Context)
 		{
-			Boxed<ILogFile?> Log = new Boxed<ILogFile?>();
+			Boxed<ILogFile?> Log = new Boxed<ILogFile?>(null);
 			for (; ; )
 			{
 				BeginStepResponse? Response = await TryBeginStep(Request, Log, Context);
@@ -593,7 +594,7 @@ namespace HordeServer.Services
 			// Create a log file if necessary
 			if (Log.Value == null)
 			{
-				Log = await LogFileService.CreateLogFileAsync(Job.Id, Batch.SessionId, Api.LogType.Json);
+				Log.Value = await LogFileService.CreateLogFileAsync(Job.Id, Batch.SessionId, Api.LogType.Json);
 			}
 
 			// Get the node for this step
@@ -613,11 +614,11 @@ namespace HordeServer.Services
 			//				}
 
 			// Update the step state
-			if (await JobService.TryUpdateStepAsync(Job, Batch.Id, Step.Id, JobStepState.Running, JobStepOutcome.Unspecified, null, null, Log.Id, null, null, null, null))
+			if (await JobService.TryUpdateStepAsync(Job, Batch.Id, Step.Id, JobStepState.Running, JobStepOutcome.Unspecified, null, null, Log.Value.Id, null, null, null, null))
 			{
 				BeginStepResponse Response = new BeginStepResponse();
 				Response.State = BeginStepResponse.Types.Result.Ready;
-				Response.LogId = Log.Id.ToString();
+				Response.LogId = Log.Value.Id.ToString();
 				Response.StepId = Step.Id.ToString();
 				Response.Name = Node.Name;
 				Response.Credentials.Add(Credentials);
