Horde: Import HordeAgent under alias to avoid clashing of imported classes

diff --git a/Engine/Source/Programs/Horde/HordeServerTests/HordeServerTests.csproj b/Engine/Source/Programs/Horde/HordeServerTests/HordeServerTests.csproj
--- a/Engine/Source/Programs/Horde/HordeServerTests/HordeServerTests.csproj
+++ b/Engine/Source/Programs/Horde/HordeServerTests/HordeServerTests.csproj
@@ -26,7 +26,9 @@
   </ItemGroup>
 
   <ItemGroup>
-    <ProjectReference Include="..\HordeAgent\HordeAgent.csproj" />
+    <ProjectReference Include="..\HordeAgent\HordeAgent.csproj" Label="HordeAgent">
+      <Aliases>HordeAgent</Aliases>
+    </ProjectReference>
     <ProjectReference Include="..\HordeServer\HordeServer.csproj" />
   </ItemGroup>
 
diff --git a/Engine/Source/Programs/Horde/HordeServerTests/IssueServiceTests.cs b/Engine/Source/Programs/Horde/HordeServerTests/IssueServiceTests.cs
--- a/Engine/Source/Programs/Horde/HordeServerTests/IssueServiceTests.cs
+++ b/Engine/Source/Programs/Horde/HordeServerTests/IssueServiceTests.cs
@@ -1,8 +1,9 @@
 // Copyright Epic Games, Inc. All Rights Reserved.
+extern alias HordeAgent;
 
 using EpicGames.Core;
-using HordeAgent.Parser;
-using HordeAgent.Utility;
+using HordeAgent::HordeAgent.Parser;
+using HordeAgent::HordeAgent.Utility;
 using HordeCommon;
 using HordeServer;
 using HordeServer.Collections;
