Horde: Add clustered ActionTaskExecutor

Enables a stateless REAPI to be implememented as well as allowing remote action task stealing between Horde Server processes (to which agents are attached).

diff --git a/Engine/Source/Programs/Horde/HordeServer/Startup.cs b/Engine/Source/Programs/Horde/HordeServer/Startup.cs
--- a/Engine/Source/Programs/Horde/HordeServer/Startup.cs
+++ b/Engine/Source/Programs/Horde/HordeServer/Startup.cs
@@ -271,6 +271,7 @@ namespace HordeServer
 			// The Connect() is sync as AddSingleton() does not support async operations.
 			// This *should* be fine as this is only established once during startup (reconnects happen async)
 			RedisConMux = ConnectionMultiplexer.Connect(Settings.RedisConnectionConfig);
+			Services.AddSingleton<ConnectionMultiplexer>(RedisConMux);
 			Services.AddSingleton<IDatabase>(Sp => RedisConMux.GetDatabase());
 			Services.AddSingleton(new RedisConnectionPool(20, Settings.RedisConnectionConfig));
 			
diff --git a/Engine/Source/Programs/Horde/HordeServer/Tasks/Impl/ClusteredActionTaskSource.cs b/Engine/Source/Programs/Horde/HordeServer/Tasks/Impl/ClusteredActionTaskSource.cs
new file mode 100644
--- /dev/null
+++ b/Engine/Source/Programs/Horde/HordeServer/Tasks/Impl/ClusteredActionTaskSource.cs
@@ -0,0 +1,483 @@
+// Copyright Epic Games, Inc. All Rights Reserved.
+
+using Build.Bazel.Remote.Execution.V2;
+using Google.LongRunning;
+using Google.Protobuf;
+using Google.Protobuf.Reflection;
+using Google.Protobuf.WellKnownTypes;
+using Grpc.Core;
+using HordeCommon;
+using HordeCommon.Rpc.Tasks;
+using HordeServer.Api;
+using HordeServer.Models;
+using HordeServer.Rpc;
+using HordeServer.Services;
+using Microsoft.Extensions.Logging;
+using MongoDB.Bson;
+using System;
+using System.Collections.Generic;
+using System.Linq;
+using System.Threading.Tasks;
+using Microsoft.Extensions.Options;
+using StackExchange.Redis;
+
+namespace HordeServer.Tasks.Impl
+{
+	/// <summary>
+	/// Dispatches remote actions across all agents that are connected to a Horde server in the cluster
+	/// </summary>
+	public class ClusteredActionTaskSource : ITaskSource
+	{
+		/// <summary>
+		/// Tracks an agent waiting for work
+		/// </summary>
+		class AgentSubscription : ITaskListener
+		{
+			/// <summary>
+			/// The owning instance
+			/// </summary>
+			ClusteredActionTaskSource Outer;
+
+			/// <summary>
+			/// Agent waiting for work
+			/// </summary>
+			public IAgent Agent { get; }
+
+			/// <summary>
+			/// The operation assigned to this subscription. May be set to null if the subscription is cancelled.
+			/// </summary>
+			public TaskCompletionSource<NewLeaseInfo?> LeaseTaskSource = new TaskCompletionSource<NewLeaseInfo?>();
+
+			/// <inheritdoc/>
+			public Task<NewLeaseInfo?> LeaseTask => LeaseTaskSource.Task;
+
+			/// <inheritdoc/>
+			public bool Accepted { get; set; }
+
+			/// <summary>
+			/// Constructor
+			/// </summary>
+			/// <param name="Outer">The task source</param>
+			/// <param name="Agent">The agent instance</param>
+			public AgentSubscription(ClusteredActionTaskSource Outer, IAgent Agent)
+			{
+				this.Outer = Outer;
+				this.Agent = Agent;
+
+				lock (Outer.AgentSubscriptions)
+				{
+					Outer.AgentSubscriptions.Add(this);
+				}
+			}
+
+			/// <summary>
+			/// Attempts to assign a lease to this agent
+			/// </summary>
+			/// <param name="Lease">The lease to assign</param>
+			/// <param name="OnConnectionLost"></param>
+			/// <returns>True if the lease was set, false otherwise</returns>
+			public bool TrySetLease(AgentLease Lease, System.Action OnConnectionLost)
+			{
+				return LeaseTaskSource.TrySetResult(new NewLeaseInfo(Lease, OnConnectionLost));
+			}
+
+			/// <summary>
+			/// Cancel of this subscription
+			/// </summary>
+			/// <returns></returns>
+			public async ValueTask DisposeAsync()
+			{
+				// Remove from the list of active subscriptions
+				lock (Outer.AgentSubscriptions)
+				{
+					Outer.AgentSubscriptions.Remove(this);
+				}
+
+				// If an operation has been assigned to this subscription and it wasn't accepted, attempt to reassign it
+				if (!LeaseTaskSource.TrySetResult(null) && !Accepted)
+				{
+					NewLeaseInfo? NewLeaseInfo = await LeaseTaskSource.Task;
+					if (NewLeaseInfo != null)
+					{
+						await Outer.CancelLeaseAsync(NewLeaseInfo.Lease);
+					}
+				}
+			}
+		}
+
+		private class RemoteExecOperationInternal : IActionExecuteOperation
+		{
+			public TaskCompletionSource<ActionResult?> ResultTaskSource { get; } = new TaskCompletionSource<ActionResult?>();
+			public ObjectId Id { get; }
+			public ObjectId LeaseId { get; }
+
+			public RemoteExecOperationInternal(ObjectId Id, ObjectId LeaseId)
+			{
+				this.Id = Id;
+				this.LeaseId = LeaseId;
+			}
+
+			public bool TrySetResult(ActionResult? Result)
+			{
+				return ResultTaskSource.TrySetResult(Result);
+			}
+
+			public IAsyncEnumerable<Operation> ReadStatusUpdatesAsync()
+			{
+				throw new NotImplementedException();
+			}
+		}
+
+		private readonly ActionCacheService ActionCacheService;
+		private readonly ILogFileService LogFileService;
+		private readonly HashSet<AgentSubscription> AgentSubscriptions = new HashSet<AgentSubscription>();
+		private readonly Dictionary<ObjectId, RemoteExecOperationInternal> ActiveOperations = new Dictionary<ObjectId, RemoteExecOperationInternal>();
+		private readonly ILogger Logger;
+		private readonly RemoteExecSettings RemoteExecSettings;
+		private readonly ConnectionMultiplexer RedisConMux;
+		private readonly IDatabase RedisDb;
+
+		internal static TimeSpan OperationTimeout = TimeSpan.FromMinutes(5);
+		internal static TimeSpan OperationRedisTtl = OperationTimeout + TimeSpan.FromMinutes(5);
+		
+		private static string KeyOp(string Id) => "re-op-" + Id;
+		private readonly string KeyPendingOps = "re-pending-ops";
+		private readonly string KeyInProgressOps = "re-inprogress-ops";
+		private readonly string KeyCompletedOps = "re-completed-ops";
+		private readonly string ChannelOps = "re-ops";
+
+		/// <inheritdoc/>
+		public MessageDescriptor Descriptor => ActionTask.Descriptor;
+
+		/// <summary>
+		/// Constructor
+		/// </summary>
+		/// <param name="RedisConMux"></param>
+		/// <param name="RedisDb"></param>
+		/// <param name="ActionCacheService"></param>
+		/// <param name="LogFileService">The log file service instance</param>
+		/// <param name="Logger">The logger instance</param>
+		/// <param name="Settings">For getting the remote exec settings</param>
+		public ClusteredActionTaskSource(ConnectionMultiplexer RedisConMux, IDatabase RedisDb, ActionCacheService ActionCacheService, ILogFileService LogFileService, ILogger<ClusteredActionTaskSource> Logger, IOptionsMonitor<ServerSettings> Settings)
+		{
+			this.RedisConMux = RedisConMux;
+			this.RedisDb = RedisDb;
+			this.ActionCacheService = ActionCacheService;
+			this.LogFileService = LogFileService;
+			this.Logger = Logger;
+			this.RemoteExecSettings = Settings.CurrentValue.RemoteExecSettings;
+
+			// FIXME: Add subscription for notifying
+		}
+
+		Task CancelLeaseAsync(AgentLease Lease)
+		{
+			lock (ActiveOperations)
+			{
+				foreach (var (_, OpInternal) in ActiveOperations)
+				{
+					if (OpInternal.LeaseId == Lease.Id)
+					{
+						OpInternal.TrySetResult(null);
+					}
+				}
+			}
+			
+			return Task.CompletedTask;
+		}
+
+		private async Task<bool> SetOperationStatusAsync(RemoteExecOperation Op, ExecutionStage.Types.Value Stage, ExecuteResponse? Response)
+		{
+			ExecuteOperationMetadata Metadata = new ExecuteOperationMetadata();
+			Metadata.Stage = Stage;
+			Metadata.ActionDigest = Op.Request.ActionDigest;
+		
+			Operation LongRunningOp = new Operation();
+			if (Response != null)
+			{
+				LongRunningOp.Response = Any.Pack(Response);
+				LongRunningOp.Done = true;
+				Op.Response = Response;
+			}
+			LongRunningOp.Metadata = Any.Pack(Metadata);
+			
+			Op.Operations.Add(LongRunningOp);
+			await SetOperationAsync(Op);
+			return true;
+		}
+
+		/// <summary>
+		/// Attempts to assign an action for execution
+		/// </summary>
+		/// <param name="Request">Execution request</param>
+		/// <returns>The operation created to execute this action</returns>
+		public async Task<RemoteExecOperation> ExecuteAsync(ExecuteRequest Request)
+		{
+			RemoteExecOperation Op = new RemoteExecOperation();
+			Op.Id = ObjectId.GenerateNewId().ToString();
+			Op.Request = Request;
+			await CreateOperationAsync(Op);
+			return Op;
+		}
+
+		private async Task SetOperationAsync(RemoteExecOperation Op)
+		{
+			Op.LastAccess = DateTimeOffset.UtcNow.ToUnixTimeMilliseconds();
+			await RedisDb.StringSetAsync(KeyOp(Op.Id), Op.ToByteArray(), OperationRedisTtl);
+		}
+
+		private async Task<bool> CreateOperationAsync(RemoteExecOperation Op)
+		{
+			await SetOperationAsync(Op);
+			ExecuteResponse? CachedResponse = await CheckForCachedResponseAsync(Op);
+			if (CachedResponse != null)
+			{
+				return true;
+			}
+
+			await SetOperationStatusAsync(Op, ExecutionStage.Types.Value.Queued, null);
+			await RedisDb.SortedSetAddAsync(KeyPendingOps, Op.Id, SortedSetUtcNow());
+			await RedisDb.PublishAsync(ChannelOps, "pending-ops-updated");
+
+			return false;
+		}
+
+		private async Task<ExecuteResponse?> CheckForCachedResponseAsync(RemoteExecOperation Op)
+		{
+			if (Op.Request.SkipCacheLookup)
+			{
+				return null;
+			}
+
+			await SetOperationStatusAsync(Op, ExecutionStage.Types.Value.CacheCheck, null);
+	
+			ActionResult? Result = await ActionCacheService.TryGetResult(Op.Request.ActionDigest);
+			if (Result == null) return null;
+			ExecuteResponse CachedResponse = new ExecuteResponse();
+			CachedResponse.Result = Result;
+			CachedResponse.CachedResult = true;
+			CachedResponse.Status = new Status(StatusCode.OK, String.Empty);
+			return CachedResponse;
+		}
+
+		private static double SortedSetUtcNow()
+		{
+			return Convert.ToDouble(DateTimeOffset.UtcNow.ToUnixTimeMilliseconds());
+		}
+
+		internal async Task<RemoteExecOperation?> GetOperationAsync(string Id)
+		{
+			RedisValue OpData = await RedisDb.StringGetAsync(KeyOp(Id));
+			return RemoteExecOperation.Parser.ParseFrom(OpData);
+		}
+		
+		internal async Task<ISet<string>> GetPendingOpsAsync()
+		{
+			return (await RedisDb.SortedSetRangeByScoreAsync(KeyPendingOps)).ToStringArray().ToHashSet();
+		}
+		
+		private async Task<string?> MarkRandomPendingOpAsInProgressAsync()
+		{
+			string OpId = (await GetPendingOpsAsync()).First();
+			if (await MoveOpId(OpId, KeyPendingOps, KeyInProgressOps))
+			{
+				return OpId;
+			}
+
+			return null;
+		}
+		
+		private async Task<bool> MoveOpId(string OpId, string FromSet, string ToSet)
+		{
+			if (await RedisDb.SortedSetRemoveAsync(FromSet, OpId))
+			{
+				await RedisDb.SortedSetAddAsync(ToSet, OpId, SortedSetUtcNow());
+				return true;
+			}
+
+			return false;
+		}
+
+		internal async Task UpdateOperations()
+		{
+			// Check for in-process available subscriptions
+			AgentSubscription? Subscription = AllocateSubscription();
+			if (Subscription == null)
+			{
+				Logger.LogDebug("No subscriptions (agents) available for work in this process.");
+				return;
+			}
+
+			// Steal a pending op from Redis
+			string? OpId = await MarkRandomPendingOpAsInProgressAsync();
+			if (OpId == null)
+			{
+				Logger.LogDebug("Unable to steal a pending operation and schedule it on attached subscription (agent)");
+				return;
+			}
+
+			RemoteExecOperation? Op = await GetOperationAsync(OpId);
+			if (Op == null)
+			{
+				Logger.LogError("Operation {OperationId} was not found prior to executing.", OpId);
+				// FIXME: Just delete the operation ID from pending set?
+				return;
+			}
+
+			ExecuteResponse? Response = await RunOperationOnAgentAsync(Subscription, Op);
+			if (Response == null)
+			{
+				Logger.LogWarning("Failed to execute operation {OperationId}", Op.Id);
+				// FIXME: Move operation ID back to pending set to let others try executing it? Check subscription.
+			}
+
+			await SetOperationStatusAsync(Op, ExecutionStage.Types.Value.Completed, Response);
+			if (!await MoveOpId(Op.Id, KeyInProgressOps, KeyCompletedOps))
+			{
+				Logger.LogWarning("Failed moving operation ID {OpId} from in-progress to completed set");
+			}
+		}
+
+		private AgentSubscription? AllocateSubscription()
+		{
+			lock (AgentSubscriptions)
+			{
+				AgentSubscription? Subscription = AgentSubscriptions.FirstOrDefault();
+				if (Subscription == null) return null;
+				AgentSubscriptions.Remove(Subscription);
+				return Subscription;
+			}
+		}
+
+		/// <summary>
+		/// Run the task as a background operation
+		/// </summary>
+		/// <returns></returns>
+		private async Task<ExecuteResponse?> RunOperationOnAgentAsync(AgentSubscription Subscription, RemoteExecOperation Op)
+		{
+			// Create a lease for it
+			string LeaseName = $"Remote action ({Op.Request.ActionDigest.Hash.Substring(0, 16)}...)";
+			ILogFile LogFile = await LogFileService.CreateLogFileAsync(ObjectId.Empty, Subscription.Agent.SessionId, LogType.Json);
+			
+			ActionTask ActionTask = CreateActionTask(Op.Request, LogFile);
+			
+			byte[] Payload = Any.Pack(ActionTask).ToByteArray();
+			ObjectId OpId = ObjectId.Parse(Op.Id);
+			AgentLease Lease = new AgentLease(OpId, LeaseName, null, null, LogFile.Id, LeaseState.Pending, Payload, new AgentRequirements(), null);
+			RemoteExecOperationInternal OpInternal = new RemoteExecOperationInternal(OpId, Lease.Id);
+
+			lock (ActiveOperations)
+			{
+				if (!ActiveOperations.TryAdd(OpId, OpInternal))
+				{
+					Logger.LogError("There's already an active operation for {OpId}. Another execution will not be started.");
+					return null;
+				}
+			}
+
+			void RemoveActiveOperation(ObjectId OpId)
+			{
+				lock (ActiveOperations)
+				{
+					ActiveOperations.Remove(OpId);
+				}
+			}
+
+			if (Subscription.TrySetLease(Lease, () => OpInternal.ResultTaskSource.TrySetResult(null)))
+			{
+				await SetOperationStatusAsync(Op, ExecutionStage.Types.Value.Executing, null);
+
+				Task<ActionResult?> ResultTask = OpInternal.ResultTaskSource.Task;
+				if (await Task.WhenAny(ResultTask, Task.Delay(OperationTimeout)) == ResultTask)
+				{
+					ActionResult? Result = ResultTask.Result;
+					if (Result != null)
+					{
+						ExecuteResponse SuccessResponse = new ExecuteResponse();
+						SuccessResponse.Status = new Status(StatusCode.OK, String.Empty);
+						SuccessResponse.Result = Result;
+						
+						RemoveActiveOperation(OpId);
+						return SuccessResponse;
+					}
+					
+					Logger.LogError("No result from remote execution under lease {LeaseId}", Lease.Id);
+				}
+				else
+				{
+					Logger.LogError("Remote execution timed out during lease {LeaseId}", Lease.Id);
+				}
+			}
+			else
+			{
+				Logger.LogError("Failed to assign lease {LeaseId} to agent for remote execution", Lease.Id);
+			}
+
+			RemoveActiveOperation(OpId);
+			return null;
+		}
+
+		private ActionTask CreateActionTask(ExecuteRequest Request, ILogFile LogFile)
+		{
+			ActionTask ActionTask = new ActionTask();
+			ActionTask.InstanceName = Request.InstanceName;
+			ActionTask.Digest = Request.ActionDigest;
+			ActionTask.LogId = LogFile.Id.ToString();
+		
+			if (Request.InstanceName != null && RemoteExecSettings.Instances.TryGetValue(Request.InstanceName, out var InstanceSettings))
+			{
+				if (InstanceSettings.CasUrl != null)
+				{
+					ActionTask.CasUrl = InstanceSettings.CasUrl.ToString();
+				}
+					
+				if (InstanceSettings.ActionCacheUrl != null)
+				{
+					ActionTask.ActionCacheUrl = InstanceSettings.ActionCacheUrl.ToString();
+				}
+					
+				if (InstanceSettings.ServiceAccountToken != null)
+				{
+					ActionTask.ServiceAccountToken = InstanceSettings.ServiceAccountToken;
+				}
+			}
+
+			return ActionTask;
+		}
+		
+		internal bool CheckIfActiveOperationExists(ObjectId OperationId)
+		{
+			lock (ActiveOperations)
+			{
+				return ActiveOperations.ContainsKey(OperationId);
+			}
+		}
+
+		internal bool SetResultForActiveOperation(ObjectId OperationId, ActionResult? Result)
+		{
+			lock (ActiveOperations)
+			{
+				if (ActiveOperations.TryGetValue(OperationId, out var OpInternal))
+				{
+					OpInternal.ResultTaskSource.TrySetResult(Result);
+					return true;
+				}
+
+				return false;
+			}
+		}
+
+		/// <inheritdoc/>
+		public Task<ITaskListener?> SubscribeAsync(IAgent Agent)
+		{
+			return Task.FromResult<ITaskListener?>(new AgentSubscription(this, Agent));
+		}
+
+		/// <inheritdoc/>
+		public Task AbortTaskAsync(IAgent Agent, ObjectId LeaseId, Any Payload)
+		{
+			// Remotely executed tasks cannot be aborted (for the time being)
+			return Task.CompletedTask;
+		}
+	}
+}
diff --git a/Engine/Source/Programs/Horde/HordeServerTests/RemoteExec/ClusteredActionTaskSourceTest.cs b/Engine/Source/Programs/Horde/HordeServerTests/RemoteExec/ClusteredActionTaskSourceTest.cs
new file mode 100644
--- /dev/null
+++ b/Engine/Source/Programs/Horde/HordeServerTests/RemoteExec/ClusteredActionTaskSourceTest.cs
@@ -0,0 +1,104 @@
+﻿extern alias HordeAgent;
+using System;
+using System.Collections.Generic;
+using System.Threading.Tasks;
+using Build.Bazel.Remote.Execution.V2;
+using HordeCommon;
+using HordeServer.Collections;
+using HordeServer.Models;
+using HordeServer.Services;
+using HordeServer.Tasks.Impl;
+using HordeServer.Utilities;
+using Microsoft.VisualStudio.TestTools.UnitTesting;
+using MongoDB.Bson;
+using RemoteExecOperation = HordeCommon.Rpc.Tasks.RemoteExecOperation;
+
+namespace HordeServerTests.RemoteExec
+{
+	public class FakeAgent : IAgent
+	{
+		public AgentId Id { get; set; }
+		public ObjectId? SessionId { get; set; }
+		public DateTime? SessionExpiresAt { get; set; }
+		public AgentStatus Status { get; set; }
+		public bool Enabled { get; set; }
+		public string? Comment { get; set; }
+		public bool Ephemeral { get; set; }
+		public bool Deleted { get; set; }
+		public StringId<IAgentSoftwareCollection>? Version { get; set; }
+		public StringId<AgentSoftwareChannels>? Channel { get; set; }
+		public string? LastUpgradeVersion { get; set; }
+		public DateTime? LastUpgradeTime { get; set; }
+		public IReadOnlyList<StringId<IPool>> ExplicitPools { get; set; } = new List<StringId<IPool>>();
+		public bool RequestConform { get; set; }
+		public bool RequestRestart { get; set; }
+		public bool RequestShutdown { get; set; }
+		public IReadOnlyList<AgentWorkspace> Workspaces { get; set; } = new List<AgentWorkspace>();
+		public DateTime LastConformTime { get; set; }
+		public int? ConformAttemptCount { get; set; }
+		public AgentCapabilities Capabilities { get; set; } = new AgentCapabilities();
+		public IReadOnlyList<AgentLease> Leases { get; set; } = new List<AgentLease>();
+		public Acl? Acl { get; set; }
+		public DateTime UpdateTime { get; set; }
+		public uint UpdateIndex { get; set; }
+	}
+	
+	[TestClass]
+	public class ActionRedisTaskSourceTest : DatabaseIntegrationTest
+	{
+		private readonly TestSetup TestSetup;
+
+		public ActionRedisTaskSourceTest()
+		{
+			TestSetup = GetTestSetup().Result;
+		}
+
+		[TestMethod]
+		public async Task ScheduleActionOnSeparateServer()
+		{
+			ClusteredActionTaskSource Server1 = CreateTaskSource();
+			ClusteredActionTaskSource Server2 = CreateTaskSource();
+
+			FakeAgent Agent = new FakeAgent();
+			await Server2.SubscribeAsync(Agent);
+			
+			ExecuteRequest ExecuteRequest = new ExecuteRequest();
+			ExecuteRequest.SkipCacheLookup = true;
+			ExecuteRequest.ActionDigest = new Digest { Hash = "my-bogus-hash-digest", SizeBytes = 20 };
+
+			// Schedule task server 1
+			RemoteExecOperation Op = await Server1.ExecuteAsync(ExecuteRequest);
+			
+			// Force server 2 to pick it up by ticking it
+#pragma warning disable 4014
+			Server2.UpdateOperations();
+#pragma warning restore 4014
+			Assert.IsTrue(await WaitForOperationToExistAsync(Server2, ObjectId.Parse(Op.Id)));
+			int ExitCode = 11122233;
+			Assert.IsTrue(Server2.SetResultForActiveOperation(ObjectId.Parse(Op.Id), new ActionResult { ExitCode = ExitCode }));
+			
+			// Verify result can fetched from server 1
+			RemoteExecOperation Op2 = (await Server1.GetOperationAsync(Op.Id))!;
+			Assert.IsNotNull(Op2);
+			Assert.AreEqual(ExitCode, Op2.Response.Result.ExitCode);
+		}
+
+		private ClusteredActionTaskSource CreateTaskSource()
+		{
+			ILogFileService LogFileService = (TestSetup.ServiceProvider.GetService(typeof(ILogFileService)) as ILogFileService)!;
+			return new ClusteredActionTaskSource(null!, GetRedisDatabase(), null!, LogFileService, null!, TestSetup.ServerSettingsMon);
+		}
+
+		private async Task<bool> WaitForOperationToExistAsync(ClusteredActionTaskSource TaskSource, ObjectId OperationId, int MaxWaitSecs = 5)
+		{
+			DateTimeOffset ExpireAt = DateTimeOffset.UtcNow + TimeSpan.FromSeconds(MaxWaitSecs);
+			while (DateTimeOffset.UtcNow < ExpireAt)
+			{
+				if (TaskSource.CheckIfActiveOperationExists(OperationId)) return true;
+				await Task.Delay(20);
+			}
+
+			return false;
+		}
+	}
+}
\ No newline at end of file
