Horde: Add missing copyright header

diff --git a/Engine/Source/Programs/Horde/HordeServerTests/RemoteExec/ClusteredActionTaskSourceTest.cs b/Engine/Source/Programs/Horde/HordeServerTests/RemoteExec/ClusteredActionTaskSourceTest.cs
--- a/Engine/Source/Programs/Horde/HordeServerTests/RemoteExec/ClusteredActionTaskSourceTest.cs
+++ b/Engine/Source/Programs/Horde/HordeServerTests/RemoteExec/ClusteredActionTaskSourceTest.cs
@@ -1,4 +1,6 @@
-﻿extern alias HordeAgent;
+﻿// Copyright Epic Games, Inc. All Rights Reserved.
+
+extern alias HordeAgent;
 using System;
 using System.Collections.Generic;
 using System.Threading.Tasks;
