Horde: Add ability to duplicate actions sent in REAPI test client for load testing purposes

diff --git a/Engine/Source/Programs/Horde/HordeAgent/Commands/Execution/ExecuteCommand.cs b/Engine/Source/Programs/Horde/HordeAgent/Commands/Execution/ExecuteCommand.cs
--- a/Engine/Source/Programs/Horde/HordeAgent/Commands/Execution/ExecuteCommand.cs
+++ b/Engine/Source/Programs/Horde/HordeAgent/Commands/Execution/ExecuteCommand.cs
@@ -21,6 +21,7 @@ using HordeAgent.Services;
 using Grpc.Net.Client;
 using System.Linq;
 using System.Net.Http.Headers;
+using System.Threading;
 using Grpc.Core;
 using Google.LongRunning;
 using Microsoft.Extensions.Hosting;
@@ -62,7 +63,7 @@ namespace HordeAgent.Commands
 			public async Task<Digest> AddFileAsync(string LocalFile)
 			{
 				BatchUpdateBlobsRequest.Types.Request Content = new BatchUpdateBlobsRequest.Types.Request();
-				using (FileStream Stream = File.Open(LocalFile, FileMode.Open, FileAccess.Read))
+				using (FileStream Stream = File.Open(LocalFile, FileMode.Open, FileAccess.Read, FileShare.Read))
 				{
 					Content.Data = await ByteString.FromStreamAsync(Stream);
 					Content.Digest = StorageExtensions.GetDigest(Content.Data);
@@ -210,6 +211,18 @@ namespace HordeAgent.Commands
 		/// </summary>
 		[CommandLine("-OutputDir")]
 		public string? OutputDir = null;
+		
+		/// <summary>
+		/// Number of concurrent executions to send (use with NumDuplicatedExecutions)
+		/// </summary>
+		[CommandLine("-NumConcurrentExecutions")]
+		public int NumConcurrentExecutions = 1;
+		
+		/// <summary>
+		/// Number of duplicated executions. Only useful for load testing purposes.
+		/// </summary>
+		[CommandLine("-NumDuplicatedExecutions")]
+		public int NumDuplicatedExecutions = 1;
 
 		/// <inheritdoc/>
 		public override async Task<int> ExecuteAsync(ILogger Logger)
@@ -254,169 +267,198 @@ namespace HordeAgent.Commands
 				UploadList UploadList = new UploadList();
 				Digest ActionDigest = await Action.BuildAsync(UploadList, SaltBytes, SkipCacheLookup);
 
-				using (GrpcChannel Channel = GrpcService.CreateGrpcChannel())
+				if (NumDuplicatedExecutions == 1)
+				{
+					await ExecuteAction(Logger, GrpcService, UploadList, ActionDigest);	
+				}
+				else
+				{
+					await ExecuteActionsConcurrently(Logger, GrpcService, UploadList, ActionDigest);
+				}
+				
+			}
+			return 0;
+		}
+
+		private async Task ExecuteActionsConcurrently(ILogger Logger, GrpcService GrpcService, UploadList UploadList, Digest ActionDigest)
+		{
+			Logger.LogInformation($"Running {NumDuplicatedExecutions} concurrently using {NumConcurrentExecutions} tasks...");
+			int Count = NumDuplicatedExecutions;
+			List<Task> TaskList = new List<Task>();
+			for (int i = 0; i < NumConcurrentExecutions; i++)
+			{
+				TaskList.Add(Task.Run(async () =>
 				{
-					GrpcChannel GetChannel()
+					while (Interlocked.Decrement(ref Count) > 0)
 					{
-						if (CasUrl == null) return Channel;
-						if (CasAuthToken != null) return GrpcService.CreateGrpcChannel(CasUrl, new AuthenticationHeaderValue("ServiceAccount", CasAuthToken));
-						return GrpcService.CreateGrpcChannel(CasUrl, null);
+						await ExecuteAction(Logger, GrpcService, UploadList, ActionDigest);
 					}
+				}));
+			}
 
-					using GrpcChannel CasChannel = GetChannel();
-					ContentAddressableStorage.ContentAddressableStorageClient StorageClient = new ContentAddressableStorage.ContentAddressableStorageClient(CasChannel);
+			await Task.WhenAll(TaskList.ToArray());
+		}
 
-					// Find the missing blobs and upload them
-					FindMissingBlobsRequest MissingBlobsRequest = new FindMissingBlobsRequest();
-					MissingBlobsRequest.InstanceName = InstanceName;
-					MissingBlobsRequest.BlobDigests.Add(UploadList.Blobs.Values.Select(x => x.Digest));
+		private async Task ExecuteAction(ILogger Logger, GrpcService GrpcService, UploadList UploadList, Digest ActionDigest)
+		{
+			using GrpcChannel Channel = GrpcService.CreateGrpcChannel();
 
-					FindMissingBlobsResponse MissingBlobs = await StorageClient.FindMissingBlobsAsync(MissingBlobsRequest);
-					if (MissingBlobs.MissingBlobDigests.Count > 0)
+			GrpcChannel GetChannel()
+			{
+				if (CasUrl == null) return Channel;
+				if (CasAuthToken != null) return GrpcService.CreateGrpcChannel(CasUrl, new AuthenticationHeaderValue("ServiceAccount", CasAuthToken));
+				return GrpcService.CreateGrpcChannel(CasUrl, null);
+			}
+
+			using GrpcChannel CasChannel = GetChannel();
+			ContentAddressableStorage.ContentAddressableStorageClient StorageClient = new ContentAddressableStorage.ContentAddressableStorageClient(CasChannel);
+
+			// Find the missing blobs and upload them
+			FindMissingBlobsRequest MissingBlobsRequest = new FindMissingBlobsRequest();
+			MissingBlobsRequest.InstanceName = InstanceName;
+			MissingBlobsRequest.BlobDigests.Add(UploadList.Blobs.Values.Select(x => x.Digest));
+
+			FindMissingBlobsResponse MissingBlobs = await StorageClient.FindMissingBlobsAsync(MissingBlobsRequest);
+			if (MissingBlobs.MissingBlobDigests.Count > 0)
+			{
+				foreach (Digest MissingBlobDigest in MissingBlobs.MissingBlobDigests)
+				{
+					BatchUpdateBlobsRequest UpdateRequest = new BatchUpdateBlobsRequest();
+					UpdateRequest.InstanceName = InstanceName;
+					UpdateRequest.Requests.Add(UploadList.Blobs[MissingBlobDigest.Hash]);
+					BatchUpdateBlobsResponse UpdateResponse = await StorageClient.BatchUpdateBlobsAsync(UpdateRequest);
+						
+					bool UpdateFailed = false;
+					foreach (var Response in UpdateResponse.Responses)
 					{
-						foreach (Digest MissingBlobDigest in MissingBlobs.MissingBlobDigests)
+						if (Response.Status.Code != (int) Google.Rpc.Code.Ok)
 						{
-							BatchUpdateBlobsRequest UpdateRequest = new BatchUpdateBlobsRequest();
-							UpdateRequest.InstanceName = InstanceName;
-							UpdateRequest.Requests.Add(UploadList.Blobs[MissingBlobDigest.Hash]);
-							BatchUpdateBlobsResponse UpdateResponse = await StorageClient.BatchUpdateBlobsAsync(UpdateRequest);
-							
-							bool UpdateFailed = false;
-							foreach (var Response in UpdateResponse.Responses)
-							{
-								if (Response.Status.Code != (int) Google.Rpc.Code.Ok)
-								{
-									Console.WriteLine($"Upload failed for {Response.Digest}. Code: {Response.Status.Code} Message: {Response.Status.Message}");
-									UpdateFailed = true;
-								}
-							}
-
-							if (UpdateFailed)
-							{
-								throw new Exception("Failed updating blobs in CAS service");
-							}
+							Console.WriteLine($"Upload failed for {Response.Digest}. Code: {Response.Status.Code} Message: {Response.Status.Message}");
+							UpdateFailed = true;
 						}
 					}
 
-					// Execute the action
-					RpcExecution.ExecutionClient ExecutionClient = new RpcExecution.ExecutionClient(Channel);
+					if (UpdateFailed)
+					{
+						throw new Exception("Failed updating blobs in CAS service");
+					}
+				}
+			}
+
+			// Execute the action
+			RpcExecution.ExecutionClient ExecutionClient = new RpcExecution.ExecutionClient(Channel);
+
+			ExecuteRequest ExecuteRequest = new ExecuteRequest();
+			ExecuteRequest.ActionDigest = ActionDigest;
+			ExecuteRequest.SkipCacheLookup = SkipCacheLookup;
+			if (InstanceName != null)
+			{
+				ExecuteRequest.InstanceName = InstanceName;
+			}
+
+			using AsyncServerStreamingCall<Operation> Call = ExecutionClient.Execute(ExecuteRequest);
+			while (await Call.ResponseStream.MoveNext())
+			{
+				Operation Operation = Call.ResponseStream.Current;
+
+				ExecuteOperationMetadata Metadata = Operation.Metadata.Unpack<ExecuteOperationMetadata>();
+				Logger.LogInformation("Execution state: {State}", Metadata.Stage);
 
-					ExecuteRequest ExecuteRequest = new ExecuteRequest();
-					ExecuteRequest.ActionDigest = ActionDigest;
-					ExecuteRequest.SkipCacheLookup = SkipCacheLookup;
-					if (InstanceName != null)
+				if (Operation.Done)
+				{
+					// Unpack the response
+					ExecuteResponse ExecuteResponse;
+					if(!Operation.Response.TryUnpack(out ExecuteResponse))
+					{
+						Logger.LogError("Unable to decode response");
+						break;
+					}
+					if((StatusCode)ExecuteResponse.Status.Code != StatusCode.OK)
 					{
-						ExecuteRequest.InstanceName = InstanceName;
+						Logger.LogError("Error {0}: {1}", ExecuteResponse.Status.Code, ExecuteResponse.Status.Message);
+						break;
 					}
 
-					using (AsyncServerStreamingCall<Operation> Call = ExecutionClient.Execute(ExecuteRequest))
+					ActionResult? ActionResult = ExecuteResponse.Result;
+					if (ExecuteResponse.Result == null)
 					{
-						while (await Call.ResponseStream.MoveNext())
-						{
-							Operation Operation = Call.ResponseStream.Current;
+						Logger.LogError("Empty result");
+						break;
+					}
 
-							ExecuteOperationMetadata Metadata = Operation.Metadata.Unpack<ExecuteOperationMetadata>();
-							Logger.LogInformation("Execution state: {State}", Metadata.Stage);
+					Logger.LogInformation("Cached: {CachedResult}", ExecuteResponse.CachedResult);
+							
+					Logger.LogInformation("Output directories:");
+					foreach (OutputDirectory Dir in ActionResult.OutputDirectories)
+					{
+						Logger.LogInformation("Path: {Path} Digest: {Digest}", Dir.Path, Dir.TreeDigest);
+					}
+							
+					Logger.LogInformation("Output files:");
+					foreach (OutputFile File in ActionResult.OutputFiles)
+					{
+						Logger.LogInformation("Path: {Path} Digest: {Digest}", File.Path, File.Digest);
+					}
 
-							if (Operation.Done)
+					// Print the result
+					BatchReadBlobsRequest BatchReadRequest = new BatchReadBlobsRequest();
+					BatchReadRequest.InstanceName = InstanceName;
+					if (ActionResult.StdoutDigest != null && ActionResult.StdoutDigest.SizeBytes > 0)
+					{
+						BatchReadRequest.Digests.Add(ActionResult.StdoutDigest);
+					}
+					if (ActionResult.StderrDigest != null && ActionResult.StderrDigest.SizeBytes > 0)
+					{
+						BatchReadRequest.Digests.Add(ActionResult.StderrDigest);
+					}
+					if (OutputDir != null && ActionResult.OutputFiles != null && ActionResult.OutputFiles.Count > 0)
+					{
+						BatchReadRequest.Digests.Add(ActionResult.OutputFiles.Select(x => x.Digest));
+					}
+
+					BatchReadBlobsResponse BatchReadResponse = await StorageClient.BatchReadBlobsAsync(BatchReadRequest);
+					if (ActionResult.StdoutDigest != null)
+					{
+						BatchReadBlobsResponse.Types.Response? StdoutResponse = BatchReadResponse.Responses.FirstOrDefault(x => x.Digest.Hash == ActionResult.StdoutDigest.Hash);
+						if (StdoutResponse != null && StdoutResponse.Data.Length > 0)
+						{
+							foreach (string Line in Encoding.UTF8.GetString(StdoutResponse.Data.ToByteArray()).Split('\n'))
 							{
-								// Unpack the response
-								ExecuteResponse ExecuteResponse;
-								if(!Operation.Response.TryUnpack(out ExecuteResponse))
-								{
-									Logger.LogError("Unable to decode response");
-									break;
-								}
-								if((StatusCode)ExecuteResponse.Status.Code != StatusCode.OK)
-								{
-									Logger.LogError("Error {0}: {1}", ExecuteResponse.Status.Code, ExecuteResponse.Status.Message);
-									break;
-								}
-
-								ActionResult? ActionResult = ExecuteResponse.Result;
-								if (ExecuteResponse.Result == null)
-								{
-									Logger.LogError("Empty result");
-									break;
-								}
-
-								Logger.LogInformation("Cached: {CachedResult}", ExecuteResponse.CachedResult);
-								
-								Logger.LogInformation("Output directories:");
-								foreach (OutputDirectory Dir in ActionResult.OutputDirectories)
-								{
-									Logger.LogInformation("Path: {Path} Digest: {Digest}", Dir.Path, Dir.TreeDigest);
-								}
-								
-								Logger.LogInformation("Output files:");
-								foreach (OutputFile File in ActionResult.OutputFiles)
-								{
-									Logger.LogInformation("Path: {Path} Digest: {Digest}", File.Path, File.Digest);
-								}
-
-								// Print the result
-								BatchReadBlobsRequest BatchReadRequest = new BatchReadBlobsRequest();
-								BatchReadRequest.InstanceName = InstanceName;
-								if (ActionResult.StdoutDigest != null && ActionResult.StdoutDigest.SizeBytes > 0)
-								{
-									BatchReadRequest.Digests.Add(ActionResult.StdoutDigest);
-								}
-								if (ActionResult.StderrDigest != null && ActionResult.StderrDigest.SizeBytes > 0)
-								{
-									BatchReadRequest.Digests.Add(ActionResult.StderrDigest);
-								}
-								if (OutputDir != null && ActionResult.OutputFiles != null && ActionResult.OutputFiles.Count > 0)
-								{
-									BatchReadRequest.Digests.Add(ActionResult.OutputFiles.Select(x => x.Digest));
-								}
-
-								BatchReadBlobsResponse BatchReadResponse = await StorageClient.BatchReadBlobsAsync(BatchReadRequest);
-								if (ActionResult.StdoutDigest != null)
-								{
-									BatchReadBlobsResponse.Types.Response? StdoutResponse = BatchReadResponse.Responses.FirstOrDefault(x => x.Digest.Hash == ActionResult.StdoutDigest.Hash);
-									if (StdoutResponse != null && StdoutResponse.Data.Length > 0)
-									{
-										foreach (string Line in Encoding.UTF8.GetString(StdoutResponse.Data.ToByteArray()).Split('\n'))
-										{
-											Logger.LogInformation("stdout: {Line}", Line);
-										}
-									}
-								}
-								if (ActionResult.StderrDigest != null)
-								{
-									BatchReadBlobsResponse.Types.Response? StderrResponse = BatchReadResponse.Responses.FirstOrDefault(x => x.Digest.Hash == ActionResult.StderrDigest.Hash);
-									if (StderrResponse != null)
-									{
-										foreach (string Line in Encoding.UTF8.GetString(StderrResponse.Data.ToByteArray()).Split('\n'))
-										{
-											Logger.LogError("stderr: {Line}", Line);
-										}
-									}
-								}
-								if (OutputDir != null && ActionResult.OutputFiles != null)
-								{
-									foreach (OutputFile OutputFile in ActionResult.OutputFiles)
-									{
-										BatchReadBlobsResponse.Types.Response? FileResponse = BatchReadResponse.Responses.FirstOrDefault(x => x.Digest.Hash == OutputFile.Digest.Hash);
-										if (FileResponse == null)
-										{
-											continue;
-										}
-										string OutputPath = Path.Join(OutputDir, OutputFile.Path);
-										System.IO.Directory.CreateDirectory(Path.GetDirectoryName(OutputPath));
-										await File.WriteAllBytesAsync(OutputPath, FileResponse.Data.ToByteArray());
-										Logger.LogInformation("Wrote {OutputFilePath}", OutputFile.Path);
-									}
-								}
-
-								Logger.LogInformation("exit: {ExitCode}", ActionResult.ExitCode);
-								break;
+								Logger.LogInformation("stdout: {Line}", Line);
 							}
 						}
 					}
+					if (ActionResult.StderrDigest != null)
+					{
+						BatchReadBlobsResponse.Types.Response? StderrResponse = BatchReadResponse.Responses.FirstOrDefault(x => x.Digest.Hash == ActionResult.StderrDigest.Hash);
+						if (StderrResponse != null)
+						{
+							foreach (string Line in Encoding.UTF8.GetString(StderrResponse.Data.ToByteArray()).Split('\n'))
+							{
+								Logger.LogError("stderr: {Line}", Line);
+							}
+						}
+					}
+					if (OutputDir != null && ActionResult.OutputFiles != null)
+					{
+						foreach (OutputFile OutputFile in ActionResult.OutputFiles)
+						{
+							BatchReadBlobsResponse.Types.Response? FileResponse = BatchReadResponse.Responses.FirstOrDefault(x => x.Digest.Hash == OutputFile.Digest.Hash);
+							if (FileResponse == null)
+							{
+								continue;
+							}
+							string OutputPath = Path.Join(OutputDir, OutputFile.Path);
+							System.IO.Directory.CreateDirectory(Path.GetDirectoryName(OutputPath));
+							await File.WriteAllBytesAsync(OutputPath, FileResponse.Data.ToByteArray());
+							Logger.LogInformation("Wrote {OutputFilePath}", OutputFile.Path);
+						}
+					}
+
+					Logger.LogInformation("exit: {ExitCode}", ActionResult.ExitCode);
+					break;
 				}
 			}
-			return 0;
 		}
 	}
 }
