Horde: Fix bug with agent event listener disposal

Dispose called itself, setting the timer to null which failed during the second call.
Exception manifested itself during REAPI testing for some reason.

diff --git a/Engine/Source/Programs/Horde/HordeServer/Services/AgentService.cs b/Engine/Source/Programs/Horde/HordeServer/Services/AgentService.cs
--- a/Engine/Source/Programs/Horde/HordeServer/Services/AgentService.cs
+++ b/Engine/Source/Programs/Horde/HordeServer/Services/AgentService.cs
@@ -609,7 +609,7 @@ namespace HordeServer.Services
 			Listener.Dispose();
 		}
 
-		class AgentEventListener : IDisposable
+		internal class AgentEventListener : IDisposable
 		{
 			object LockObject = new object();
 			AgentService AgentService;
@@ -630,9 +630,9 @@ namespace HordeServer.Services
 			{
 				if (Timer != null)
 				{
-					BeginUpdate(Array.Empty<ObjectId>());
 					Timer.Dispose();
 					Timer = null;
+					BeginUpdate(Array.Empty<ObjectId>());
 				}
 			}
 
diff --git a/Engine/Source/Programs/Horde/HordeServerTests/AgentServiceTest.cs b/Engine/Source/Programs/Horde/HordeServerTests/AgentServiceTest.cs
--- a/Engine/Source/Programs/Horde/HordeServerTests/AgentServiceTest.cs
+++ b/Engine/Source/Programs/Horde/HordeServerTests/AgentServiceTest.cs
@@ -12,6 +12,7 @@ using HordeServer.Services;
 using Microsoft.AspNetCore.Mvc;
 using AgentSoftwareVersion = HordeServer.Utilities.StringId<HordeServer.Collections.IAgentSoftwareCollection>;
 using HordeServer.Utilities;
+using MongoDB.Bson;
 
 namespace HordeServerTests
 {
@@ -50,6 +51,16 @@ namespace HordeServerTests
 	        Assert.IsFalse(AgentService.AuthorizeSession(Agent, GetUser(Agent)));
         }
 
+        [TestMethod]
+        public async Task AgentListenerDispose()
+        {
+	        // Test created to verify dispose bug with AgentEventListener
+	        TestSetup TestSetup = await GetTestSetup();
+	        AgentService AgentService = TestSetup.AgentService;
+	        AgentService.AgentEventListener Listener = new AgentService.AgentEventListener(AgentService, new AgentId("test"), ObjectId.GenerateNewId(), () => { });
+	        Listener.Dispose();
+        }
+
         private ClaimsPrincipal GetUser(IAgent Agent)
         {
 	        return new ClaimsPrincipal(new ClaimsIdentity(new List<Claim>
