Horde: Add default configuration for a project and stream.

diff --git a/Engine/Source/Programs/Horde/HordeServer/Defaults/global.json b/Engine/Source/Programs/Horde/HordeServer/Defaults/global.json
new file mode 100644
--- /dev/null
+++ b/Engine/Source/Programs/Horde/HordeServer/Defaults/global.json
@@ -0,0 +1,8 @@
+{
+	"projects": [
+		{
+			"id": "ue5",
+			"path": "project-ue5.json"
+		}
+	]
+}
diff --git a/Engine/Source/Programs/Horde/HordeServer/Defaults/project-ue5.json b/Engine/Source/Programs/Horde/HordeServer/Defaults/project-ue5.json
new file mode 100644
--- /dev/null
+++ b/Engine/Source/Programs/Horde/HordeServer/Defaults/project-ue5.json
@@ -0,0 +1,20 @@
+{
+	"name": "UE5",
+	"logo": "project-ue5.png",
+	"streams": [
+		{
+		  "id": "ue5-main",
+		  "path": "stream-ue5-main.json"
+		}
+	],
+	"categories": [
+		{
+			"name": "General",
+			"row": 0,
+			"showOnNavMenu": true,
+			"includePatterns": [
+				"//UE5/Main"
+			]
+		}
+	]
+}
diff --git a/Engine/Source/Programs/Horde/HordeServer/Defaults/stream-ue5-main.json b/Engine/Source/Programs/Horde/HordeServer/Defaults/stream-ue5-main.json
new file mode 100644
--- /dev/null
+++ b/Engine/Source/Programs/Horde/HordeServer/Defaults/stream-ue5-main.json
@@ -0,0 +1,71 @@
+{
+	"name": "//UE5/Main",
+	"tabs":
+	[
+		{
+			"title": "General",
+			"type": "Jobs",
+			"showNames": true,
+			"templates":
+			[
+				"custom"
+			],
+			"columns":
+			[
+				{
+					"heading": "Editors",
+					"category": "Editors",
+					"relativeWidth": 1
+				},
+				{
+					"heading": "Target Platforms",
+					"category": "Platforms",
+					"relativeWidth": 2
+				},
+				{
+					"heading": "Other",
+					"category": "Other",
+					"relativeWidth": 1
+				}
+			]
+		}
+	],
+	"agentTypes":
+	{
+		"Win-UE5":
+		{
+			"pool": "win-ue5",
+			"environment":
+			{
+			}
+		}
+	},
+	"workspaceTypes":
+	{
+	},
+	"templates":
+	[
+		{
+			"id": "custom",
+			"name": "Custom...",
+			"arguments":
+			[
+			],
+			"parameters":
+			[
+				{
+					"type": "Text",
+					"label": "Script",
+					"argument": "-Script="
+					"default": "Foo.xml"
+				},
+				{
+					"type": "Text",
+					"label": "Target",
+					"argument": "-Target=",
+					"default": ""
+				}
+			]
+		}
+	]
+}
diff --git a/Engine/Source/Programs/Horde/HordeServer/HordeServer.csproj b/Engine/Source/Programs/Horde/HordeServer/HordeServer.csproj
--- a/Engine/Source/Programs/Horde/HordeServer/HordeServer.csproj
+++ b/Engine/Source/Programs/Horde/HordeServer/HordeServer.csproj
@@ -1,4 +1,4 @@
-<Project Sdk="Microsoft.NET.Sdk.Web">
+﻿<Project Sdk="Microsoft.NET.Sdk.Web">
 
   <PropertyGroup>
     <TargetFramework>netcoreapp3.1</TargetFramework>
@@ -80,9 +80,15 @@
     <Content Include="..\..\..\ThirdParty\Perforce\P4Api.Net\native\windows\release\p4bridge.*" Condition=" '$(OS)' == 'Windows_NT' ">
       <CopyToOutputDirectory>Always</CopyToOutputDirectory>
     </Content>
-    <Content Include="..\..\..\ThirdParty\Perforce\P4Api.Net\native\linux\release\p4bridge.*" Condition="$([MSBuild]::IsOSPlatform('Linux'))" >
+    <Content Include="..\..\..\ThirdParty\Perforce\P4Api.Net\native\linux\release\p4bridge.*" Condition="$([MSBuild]::IsOSPlatform('Linux'))">
       <CopyToOutputDirectory>Always</CopyToOutputDirectory>
     </Content>	
   </ItemGroup>
 
+  <ItemGroup>
+    <Content Include="Defaults\project-ue5.png">
+      <CopyToOutputDirectory>PreserveNewest</CopyToOutputDirectory>
+    </Content>
+  </ItemGroup>
+
 </Project>
diff --git a/Engine/Source/Programs/Horde/HordeServer/appsettings.Development.json b/Engine/Source/Programs/Horde/HordeServer/appsettings.Development.json
--- a/Engine/Source/Programs/Horde/HordeServer/appsettings.Development.json
+++ b/Engine/Source/Programs/Horde/HordeServer/appsettings.Development.json
@@ -1,10 +1,6 @@
 {
 	"Horde":
 	{
-		"CorsEnabled": true,
-		"CorsOrigin": "https://localhost:5003;http://localhost:3000",
-		"LogFileProviderType": "Local",
-		"DisableAuth": false,
 		"DisableSchedules": true
 	}
 }
diff --git a/Engine/Source/Programs/Horde/HordeServer/appsettings.Production.json b/Engine/Source/Programs/Horde/HordeServer/appsettings.Production.json
--- a/Engine/Source/Programs/Horde/HordeServer/appsettings.Production.json
+++ b/Engine/Source/Programs/Horde/HordeServer/appsettings.Production.json
@@ -1,6 +1,5 @@
 {
     "Horde": 
 	{
-		"LogJsonToStdOut": true
     }
 }
diff --git a/Engine/Source/Programs/Horde/HordeServer/appsettings.json b/Engine/Source/Programs/Horde/HordeServer/appsettings.json
--- a/Engine/Source/Programs/Horde/HordeServer/appsettings.json
+++ b/Engine/Source/Programs/Horde/HordeServer/appsettings.json
@@ -4,7 +4,12 @@
 		"DatabaseConnectionString": "mongodb://localhost:27017",
 		"RedisConnectionConfig": "localhost",
 		"DatabaseName": "Horde",
-		"DashboardUrl": "https://localhost:3000"
+		"DashboardUrl": "https://localhost:3000",
+		"DisableAuth": true,
+		"ConfigPath": "Defaults/global.json",
+		"CorsEnabled": true,
+		"CorsOrigin": "https://localhost:5003;http://localhost:3000",
+		"LogFileProviderType": "Local"
 	},
 	"AllowedHosts": "*"
 }
