Horde: Set worker hostname in action result metadata

diff --git a/Engine/Source/Programs/Horde/HordeAgent/Execution/ActionExecutor.cs b/Engine/Source/Programs/Horde/HordeAgent/Execution/ActionExecutor.cs
--- a/Engine/Source/Programs/Horde/HordeAgent/Execution/ActionExecutor.cs
+++ b/Engine/Source/Programs/Horde/HordeAgent/Execution/ActionExecutor.cs
@@ -145,6 +145,10 @@ namespace HordeAgent
 					Result.StdoutDigest = await Storage.PutBulkDataAsync(InstanceName, StdOutData);
 					Result.StderrDigest = await Storage.PutBulkDataAsync(InstanceName, StdErrData);
 					Result.ExitCode = Process.ExitCode;
+					Result.ExecutionMetadata = new ExecutedActionMetadata
+					{
+						Worker = System.Net.Dns.GetHostName()
+					};
 
 					foreach (string Line in Encoding.UTF8.GetString(StdOutData).Split('\n'))
 					{
