Horde: Add support for generating JSON schemas for configuration classes.

* /api/v1/schemas/catalog.json also returns a schemastore.org compatible catalog of files for use by [REDACTED_PERSON].
* -UpdateSchemas argument when launching server will write schemas to disk. Specify -SchemaDir to override output folder.

diff --git a/Engine/Source/Programs/Horde/HordeServer/Controllers/SchemaController.cs b/Engine/Source/Programs/Horde/HordeServer/Controllers/SchemaController.cs
new file mode 100644
--- /dev/null
+++ b/Engine/Source/Programs/Horde/HordeServer/Controllers/SchemaController.cs
@@ -0,0 +1,90 @@
+// Copyright Epic Games, Inc. All Rights Reserved.
+
+using HordeServer.Models;
+using HordeServer.Utilities;
+using Json.Schema;
+using Json.Schema.Generation;
+
+using Microsoft.AspNetCore.Mvc;
+using Microsoft.Extensions.Primitives;
+using System;
+using System.Collections.Generic;
+using System.Linq;
+using System.Net;
+using System.Text.Json;
+using System.Text.Json.Serialization;
+using System.Threading.Tasks;
+
+namespace HordeServer.Controllers
+{
+	/// <summary>
+	/// Controller for the /api/v1/schema endpoint
+	/// </summary>
+	[ApiController]
+	[Route("[controller]")]
+	public class SchemaController : ControllerBase
+	{
+		class CatalogItem
+		{
+			public string? Name { get; set; }
+			public Uri? Url { get; set; }
+		}
+
+		class CatalogRoot
+		{
+			[JsonPropertyName("$schema")]
+			public string Schema { get; set; } = "https://json.schemastore.org/schema-catalog.json";
+			public int Version { get; set; } = 1;
+			public List<CatalogItem> Items { get; set; } = new List<CatalogItem>();
+		}
+
+		/// <summary>
+		/// Get the catalog for config schema
+		/// </summary>
+		/// <returns>Information about all the schedules</returns>
+		[HttpGet]
+		[Route("/api/v1/schema/catalog.json")]
+		public ActionResult GetCatalog()
+		{
+			string? Host = null;
+
+			StringValues Hosts;
+			if (Request.Headers.TryGetValue("Host", out Hosts))
+			{
+				Host = Hosts.FirstOrDefault();
+			}
+
+			if (Host == null)
+			{
+				Host = Dns.GetHostName();
+			}
+
+			CatalogRoot Root = new CatalogRoot();
+			foreach (SchemaInfo Schema in Program.ConfigSchemas)
+			{
+				Uri Url = new Uri($"https://{Host}/api/v1/schema/types/{Schema.Type.Name}.json");
+				Root.Items.Add(new CatalogItem { Name = Schema.Name, Url = Url });
+			}
+			return Ok(Root);
+		}
+
+		/// <summary>
+		/// Gets a specific schema
+		/// </summary>
+		/// <param name="TypeName">The type name</param>
+		/// <returns></returns>
+		[HttpGet]
+		[Route("/api/v1/schema/types/{TypeName}.json")]
+		public ActionResult GetSchema(string TypeName)
+		{
+			foreach (SchemaInfo Schema in Program.ConfigSchemas)
+			{
+				if (Schema.Type.Name.Equals(TypeName, StringComparison.OrdinalIgnoreCase))
+				{
+					return Ok(Schemas.CreateSchema(Schema.Id, Schema.Type));
+				}
+			}
+			return NotFound();
+		}
+	}
+}
diff --git a/Engine/Source/Programs/Horde/HordeServer/HordeServer.csproj b/Engine/Source/Programs/Horde/HordeServer/HordeServer.csproj
--- a/Engine/Source/Programs/Horde/HordeServer/HordeServer.csproj
+++ b/Engine/Source/Programs/Horde/HordeServer/HordeServer.csproj
@@ -23,6 +23,8 @@
       <PrivateAssets>all</PrivateAssets>
       <IncludeAssets>runtime; build; native; contentfiles; analyzers; buildtransitive</IncludeAssets>
     </PackageReference>
+    <PackageReference Include="JsonSchema.Net" Version="1.10.2" />
+    <PackageReference Include="JsonSchema.Net.Generation" Version="1.5.0" />
     <PackageReference Include="Microsoft.AspNetCore.Authentication.Cookies" Version="2.2.0" />
     <PackageReference Include="Microsoft.AspNetCore.Authentication.JwtBearer" Version="3.1.5" />
     <PackageReference Include="Microsoft.AspNetCore.Authentication.OpenIdConnect" Version="3.1.5" />
diff --git a/Engine/Source/Programs/Horde/HordeServer/Program.cs b/Engine/Source/Programs/Horde/HordeServer/Program.cs
--- a/Engine/Source/Programs/Horde/HordeServer/Program.cs
+++ b/Engine/Source/Programs/Horde/HordeServer/Program.cs
@@ -9,9 +9,14 @@ using System.Reflection;
 using System.Runtime.InteropServices;
 using System.Security.Cryptography;
 using System.Security.Cryptography.X509Certificates;
+using System.Text.Json;
 using System.Threading.Tasks;
 using EpicGames.Core;
 using HordeServer.Models;
+using HordeServer.Utilities;
+using Json.Schema;
+using Json.Schema.Generation;
+using Json.Schema.Generation.Intents;
 using Microsoft.AspNetCore.Hosting;
 using Microsoft.AspNetCore.Server.Kestrel.Core;
 using Microsoft.Extensions.Configuration;
@@ -56,8 +61,17 @@ namespace HordeServer
 	{
 		public static DirectoryReference AppDir { get; } = new FileReference(Assembly.GetExecutingAssembly().Location).Directory;
 
+		public static SchemaInfo[] ConfigSchemas =
+		{
+			new SchemaInfo("https://unrealengine.com/horde/global", "Horde root configuration file", typeof(GlobalConfig)),
+			new SchemaInfo("https://unrealengine.com/horde/project", "Horde project configuration file", typeof(ProjectConfig)),
+			new SchemaInfo("https://unrealengine.com/horde/stream", "Horde stream configuration file", typeof(StreamConfig))
+		};
+
 		public static void Main(string[] Args)
 		{
+			CommandLineArguments Arguments = new CommandLineArguments(Args);
+
 			IConfiguration Config = new ConfigurationBuilder()
 				.SetBasePath(AppDir.FullName)
 				.AddJsonFile("appsettings.json", optional: false)
@@ -86,6 +100,25 @@ namespace HordeServer
 				.WriteTo.File(new JsonFormatter(renderMessage: true), Path.Combine(AppDir.FullName, "Log.json"), rollingInterval: RollingInterval.Day, rollOnFileSizeLimit: true, fileSizeLimitBytes: 20 * 1024 * 1024, retainedFileCountLimit: 10)
 				.CreateLogger();
 
+			if (Arguments.HasOption("-UpdateSchemas"))
+			{
+				DirectoryReference? SchemaDir = Arguments.GetDirectoryReferenceOrDefault("-SchemaDir=", null);
+				if (SchemaDir == null)
+				{
+					SchemaDir = new DirectoryReference("Schemas");
+				}
+
+				Arguments.CheckAllArgumentsUsed();
+
+				DirectoryReference.CreateDirectory(SchemaDir);
+				foreach (SchemaInfo Schema in ConfigSchemas)
+				{
+					FileReference OutputFile = FileReference.Combine(SchemaDir, $"{Schema.Type.Name}.json");
+					Schemas.WriteSchema(Schema.Id, Schema.Type, OutputFile);
+				}
+				return;
+			}
+
 			try
 			{
 				using (X509Certificate2? GrpcCertificate = ReadGrpcCertificate(HordeSettings))
diff --git a/Engine/Source/Programs/Horde/HordeServer/Utilities/Schemas.cs b/Engine/Source/Programs/Horde/HordeServer/Utilities/Schemas.cs
new file mode 100644
--- /dev/null
+++ b/Engine/Source/Programs/Horde/HordeServer/Utilities/Schemas.cs
@@ -0,0 +1,126 @@
+// Copyright Epic Games, Inc. All Rights Reserved.
+
+using EpicGames.Core;
+using HordeServer.Models;
+using Json.Schema;
+using Json.Schema.Generation;
+using Json.Schema.Generation.Intents;
+using System;
+using System.Collections.Concurrent;
+using System.Collections.Generic;
+using System.IO;
+using System.Linq;
+using System.Text.Json;
+using System.Threading.Tasks;
+
+namespace HordeServer.Utilities
+{
+	/// <summary>
+	/// Information about a type to generate a schema for
+	/// </summary>
+	class SchemaInfo
+	{
+		/// <summary>
+		/// Identifier for the schema
+		/// </summary>
+		public string Id { get; set; }
+
+		/// <summary>
+		/// Name of the type
+		/// </summary>
+		public string Name { get; set; }
+
+		/// <summary>
+		/// The type to generate a schema for
+		/// </summary>
+		public Type Type { get; set; }
+
+		/// <summary>
+		/// Constructor
+		/// </summary>
+		/// <param name="Id"></param>
+		/// <param name="Name"></param>
+		/// <param name="Type"></param>
+		public SchemaInfo(string Id, string Name, Type Type)
+		{
+			this.Id = Id;
+			this.Name = Name;
+			this.Type = Type;
+		}
+	}
+
+	/// <summary>
+	/// Utility functions for generating schemas
+	/// </summary>
+	static class Schemas
+	{
+		class CamelCaseSchemaRefiner : ISchemaRefiner
+		{
+			public void Run(SchemaGeneratorContext Context)
+			{
+				PropertiesIntent? Intent = Context.Intents.OfType<PropertiesIntent>().FirstOrDefault();
+				if (Intent != null)
+				{
+					Dictionary<string, SchemaGeneratorContext> Properties = new Dictionary<string, SchemaGeneratorContext>(Intent.Properties);
+					foreach ((string Name, SchemaGeneratorContext PropertyContext) in Properties)
+					{
+						string NewName = Char.ToLowerInvariant(Name[0]) + Name.Substring(1);
+						Intent.Properties.Remove(Name);
+						Intent.Properties[NewName] = PropertyContext;
+					}
+				}
+
+			}
+
+			public bool ShouldRun(SchemaGeneratorContext Context)
+			{
+				return Context.Intents.OfType<PropertiesIntent>().Any();
+			}
+		}
+
+		static ConcurrentDictionary<(string, Type), JsonSchema> CachedSchemas = new ConcurrentDictionary<(string, Type), JsonSchema>();
+
+		/// <summary>
+		/// Creates a schema for the given type
+		/// </summary>
+		/// <param name="Id"></param>
+		/// <param name="Type"></param>
+		/// <returns></returns>
+		public static JsonSchema CreateSchema(string Id, Type Type)
+		{
+			JsonSchema? Schema;
+			if (!CachedSchemas.TryGetValue((Id, Type), out Schema))
+			{
+				SchemaGeneratorConfiguration Config = new SchemaGeneratorConfiguration();
+				Config.PropertyOrder = PropertyOrder.AsDeclared;
+				Config.Refiners.Add(new CamelCaseSchemaRefiner());
+
+				JsonSchemaBuilder Builder = new JsonSchemaBuilder().FromType(Type, Config);
+				Builder.Id(Id);
+
+				Schema = Builder.Build();
+				CachedSchemas.TryAdd((Id, Type), Schema);
+			}
+			return Schema;
+		}
+
+		/// <summary>
+		/// Writes a schema to a source file
+		/// </summary>
+		/// <param name="Id"></param>
+		/// <param name="Type"></param>
+		/// <param name="OutputFile"></param>
+		public static void WriteSchema(string Id, Type Type, FileReference OutputFile)
+		{
+			JsonSchema Schema = CreateSchema(Id, Type);
+
+			JsonSerializerOptions Options = new JsonSerializerOptions();
+			Options.WriteIndented = true;
+			Options.PropertyNamingPolicy = JsonNamingPolicy.CamelCase;
+			Options.DictionaryKeyPolicy = JsonNamingPolicy.CamelCase;
+
+			byte[] Result = JsonSerializer.SerializeToUtf8Bytes(Schema, Options);
+			FileReference.WriteAllBytes(OutputFile, Result);
+		}
+	}
+}
