Horde: Fix race condition in test

diff --git a/Engine/Source/Programs/Horde/HordeServerTests/RemoteExec/ClusteredActionTaskSourceTest.cs b/Engine/Source/Programs/Horde/HordeServerTests/RemoteExec/ClusteredActionTaskSourceTest.cs
--- a/Engine/Source/Programs/Horde/HordeServerTests/RemoteExec/ClusteredActionTaskSourceTest.cs
+++ b/Engine/Source/Programs/Horde/HordeServerTests/RemoteExec/ClusteredActionTaskSourceTest.cs
@@ -72,12 +72,11 @@ namespace HordeServerTests.RemoteExec
 			RemoteExecOperation Op = await Server1.ExecuteAsync(ExecuteRequest);
 			
 			// Force server 2 to pick it up by ticking it
-#pragma warning disable 4014
-			Server2.UpdateOperations();
-#pragma warning restore 4014
+			Task Server2UpdateTask = Server2.UpdateOperations();
 			Assert.IsTrue(await WaitForOperationToExistAsync(Server2, ObjectId.Parse(Op.Id)));
 			int ExitCode = 11122233;
 			Assert.IsTrue(Server2.SetResultForActiveOperation(ObjectId.Parse(Op.Id), new ActionResult { ExitCode = ExitCode }));
+			await Server2UpdateTask;
 			
 			// Verify result can fetched from server 1
 			RemoteExecOperation Op2 = (await Server1.GetOperationAsync(Op.Id))!;
