Horde: Add description to JSON schema catalog.

diff --git a/Engine/Source/Programs/Horde/HordeServer/Controllers/SchemaController.cs b/Engine/Source/Programs/Horde/HordeServer/Controllers/SchemaController.cs
--- a/Engine/Source/Programs/Horde/HordeServer/Controllers/SchemaController.cs
+++ b/Engine/Source/Programs/Horde/HordeServer/Controllers/SchemaController.cs
@@ -27,6 +27,7 @@ namespace HordeServer.Controllers
 		class CatalogItem
 		{
 			public string? Name { get; set; }
+			public string? Description { get; set; }
 			public Uri? Url { get; set; }
 		}
 
@@ -63,7 +64,7 @@ namespace HordeServer.Controllers
 			foreach (SchemaInfo Schema in Program.ConfigSchemas)
 			{
 				Uri Url = new Uri($"https://{Host}/api/v1/schema/types/{Schema.Type.Name}.json");
-				Root.Items.Add(new CatalogItem { Name = Schema.Name, Url = Url });
+				Root.Items.Add(new CatalogItem { Name = Schema.Name, Description = Schema.Description, Url = Url });
 			}
 			return Ok(Root);
 		}
diff --git a/Engine/Source/Programs/Horde/HordeServer/Program.cs b/Engine/Source/Programs/Horde/HordeServer/Program.cs
--- a/Engine/Source/Programs/Horde/HordeServer/Program.cs
+++ b/Engine/Source/Programs/Horde/HordeServer/Program.cs
@@ -63,9 +63,9 @@ namespace HordeServer
 
 		public static SchemaInfo[] ConfigSchemas =
 		{
-			new SchemaInfo("https://unrealengine.com/horde/global", "Horde root configuration file", typeof(GlobalConfig)),
-			new SchemaInfo("https://unrealengine.com/horde/project", "Horde project configuration file", typeof(ProjectConfig)),
-			new SchemaInfo("https://unrealengine.com/horde/stream", "Horde stream configuration file", typeof(StreamConfig))
+			new SchemaInfo("https://unrealengine.com/horde/global", "Horde Globals", "Horde root configuration file", typeof(GlobalConfig)),
+			new SchemaInfo("https://unrealengine.com/horde/project", "Horde Project", "Horde project configuration file", typeof(ProjectConfig)),
+			new SchemaInfo("https://unrealengine.com/horde/stream", "Horde Stream", "Horde stream configuration file", typeof(StreamConfig))
 		};
 
 		public static void Main(string[] Args)
diff --git a/Engine/Source/Programs/Horde/HordeServer/Utilities/Schemas.cs b/Engine/Source/Programs/Horde/HordeServer/Utilities/Schemas.cs
--- a/Engine/Source/Programs/Horde/HordeServer/Utilities/Schemas.cs
+++ b/Engine/Source/Programs/Horde/HordeServer/Utilities/Schemas.cs
@@ -30,6 +30,11 @@ namespace HordeServer.Utilities
 		/// </summary>
 		public string Name { get; set; }
 
+		/// <summary>
+		/// Description of the schema
+		/// </summary>
+		public string Description { get; set; }
+
 		/// <summary>
 		/// The type to generate a schema for
 		/// </summary>
@@ -40,11 +45,13 @@ namespace HordeServer.Utilities
 		/// </summary>
 		/// <param name="Id"></param>
 		/// <param name="Name"></param>
+		/// <param name="Description"></param>
 		/// <param name="Type"></param>
-		public SchemaInfo(string Id, string Name, Type Type)
+		public SchemaInfo(string Id, string Name, string Description, Type Type)
 		{
 			this.Id = Id;
 			this.Name = Name;
+			this.Description = Description;
 			this.Type = Type;
 		}
 	}
