Horde: Fix property name for schema catalog.

diff --git a/Engine/Source/Programs/Horde/HordeServer/Controllers/SchemaController.cs b/Engine/Source/Programs/Horde/HordeServer/Controllers/SchemaController.cs
--- a/Engine/Source/Programs/Horde/HordeServer/Controllers/SchemaController.cs
+++ b/Engine/Source/Programs/Horde/HordeServer/Controllers/SchemaController.cs
@@ -36,7 +36,7 @@ namespace HordeServer.Controllers
 			[JsonPropertyName("$schema")]
 			public string Schema { get; set; } = "https://json.schemastore.org/schema-catalog.json";
 			public int Version { get; set; } = 1;
-			public List<CatalogItem> Items { get; set; } = new List<CatalogItem>();
+			public List<CatalogItem> Schemas { get; set; } = new List<CatalogItem>();
 		}
 
 		/// <summary>
@@ -64,7 +64,7 @@ namespace HordeServer.Controllers
 			foreach (SchemaInfo Schema in Program.ConfigSchemas)
 			{
 				Uri Url = new Uri($"https://{Host}/api/v1/schema/types/{Schema.Type.Name}.json");
-				Root.Items.Add(new CatalogItem { Name = Schema.Name, Description = Schema.Description, Url = Url });
+				Root.Schemas.Add(new CatalogItem { Name = Schema.Name, Description = Schema.Description, Url = Url });
 			}
 			return Ok(Root);
 		}
