Horde: Fix streams not being undeleted when referenced by config file.

diff --git a/Engine/Source/Programs/Horde/HordeServer/Collections/Impl/StreamCollection.cs b/Engine/Source/Programs/Horde/HordeServer/Collections/Impl/StreamCollection.cs
--- a/Engine/Source/Programs/Horde/HordeServer/Collections/Impl/StreamCollection.cs
+++ b/Engine/Source/Programs/Horde/HordeServer/Collections/Impl/StreamCollection.cs
@@ -305,7 +305,10 @@ namespace HordeServer.Collections.Impl
 
 			Stream.Acl = Acl;
 			Updates.Add(UpdateBuilder.SetOrUnsetNullRef(x => x.Acl, Acl));
-			
+
+			Stream.Deleted = false;
+			Updates.Add(UpdateBuilder.Unset(x => x.Deleted));
+
 			Stream.Validate();
 			if(await TryUpdateStreamAsync(Stream, UpdateBuilder.Combine(Updates)))
 			{
