Horde: Json schema fixes:

* Now uses correct case for types (object/array/etc...)
* Add AdditionalProperties = false to all objects
* Add a generator for StringId classes
* Add attributes for specifying catalog metadata.

diff --git a/Engine/Source/Programs/Horde/HordeServer/Controllers/SchemaController.cs b/Engine/Source/Programs/Horde/HordeServer/Controllers/SchemaController.cs
--- a/Engine/Source/Programs/Horde/HordeServer/Controllers/SchemaController.cs
+++ b/Engine/Source/Programs/Horde/HordeServer/Controllers/SchemaController.cs
@@ -9,8 +9,10 @@ using Microsoft.AspNetCore.Mvc;
 using Microsoft.Extensions.Primitives;
 using System;
 using System.Collections.Generic;
+using System.IO;
 using System.Linq;
 using System.Net;
+using System.Reflection;
 using System.Text.Json;
 using System.Text.Json.Serialization;
 using System.Threading.Tasks;
@@ -28,6 +30,7 @@ namespace HordeServer.Controllers
 		{
 			public string? Name { get; set; }
 			public string? Description { get; set; }
+			public string[]? FileMatch { get; set; }
 			public Uri? Url { get; set; }
 		}
 
@@ -61,10 +64,19 @@ namespace HordeServer.Controllers
 			}
 
 			CatalogRoot Root = new CatalogRoot();
-			foreach (SchemaInfo Schema in Program.ConfigSchemas)
+			foreach (Type SchemaType in Program.ConfigSchemas)
 			{
-				Uri Url = new Uri($"https://{Host}/api/v1/schema/types/{Schema.Type.Name}.json");
-				Root.Schemas.Add(new CatalogItem { Name = Schema.Name, Description = Schema.Description, Url = Url });
+				JsonSchemaAttribute? SchemaAttribute = SchemaType.GetCustomAttribute<JsonSchemaAttribute>();
+				if (SchemaAttribute != null)
+				{
+					JsonSchemaCatalogAttribute? CatalogAttribute = SchemaType.GetCustomAttribute<JsonSchemaCatalogAttribute>();
+					if (CatalogAttribute != null)
+					{
+						Uri Url = new Uri($"https://{Host}/api/v1/schema/types/{SchemaType.Name}.json");
+						Root.Schemas.Add(new CatalogItem { Name = CatalogAttribute.Name, Description = CatalogAttribute.Description, FileMatch = new[] { "globals.json" }, Url = new Uri(SchemaAttribute.Id) });
+						break;
+					}
+				}
 			}
 			return Ok(Root);
 		}
@@ -78,11 +90,18 @@ namespace HordeServer.Controllers
 		[Route("/api/v1/schema/types/{TypeName}.json")]
 		public ActionResult GetSchema(string TypeName)
 		{
-			foreach (SchemaInfo Schema in Program.ConfigSchemas)
+			foreach (Type SchemaType in Program.ConfigSchemas)
 			{
-				if (Schema.Type.Name.Equals(TypeName, StringComparison.OrdinalIgnoreCase))
+				if (SchemaType.Name.Equals(TypeName, StringComparison.OrdinalIgnoreCase))
 				{
-					return Ok(Schemas.CreateSchema(Schema.Id, Schema.Type));
+					JsonSchema Schema = Schemas.CreateSchema(SchemaType);
+
+					using MemoryStream Stream = new MemoryStream();
+					using Utf8JsonWriter Writer = new Utf8JsonWriter(Stream);
+					JsonSerializer.Serialize(Writer, Schema, new JsonSerializerOptions());
+					Writer.Flush();
+
+					return new FileContentResult(Stream.ToArray(), "application/json");
 				}
 			}
 			return NotFound();
diff --git a/Engine/Source/Programs/Horde/HordeServer/Models/Config.cs b/Engine/Source/Programs/Horde/HordeServer/Models/Config.cs
--- a/Engine/Source/Programs/Horde/HordeServer/Models/Config.cs
+++ b/Engine/Source/Programs/Horde/HordeServer/Models/Config.cs
@@ -1,6 +1,7 @@
 // Copyright Epic Games, Inc. All Rights Reserved.
 
 using HordeServer.Api;
+using HordeServer.Utilities;
 using System;
 using System.Collections.Generic;
 using System.ComponentModel.DataAnnotations;
@@ -15,6 +16,8 @@ namespace HordeServer.Models
 	/// <summary>
 	/// Global configuration
 	/// </summary>
+	[JsonSchema("https://unrealengine.com/horde/global")]
+	[JsonSchemaCatalog("Horde Globals", "Horde global configuration file")]
 	public class GlobalConfig
 	{
 		/// <summary>
@@ -69,6 +72,8 @@ namespace HordeServer.Models
 	/// <summary>
 	/// Stores configuration for a project
 	/// </summary>
+	[JsonSchema("https://unrealengine.com/horde/project")]
+	[JsonSchemaCatalog("Horde Project", "Horde project configuration file")]
 	public class ProjectConfig
 	{
 		/// <summary>
@@ -119,6 +124,8 @@ namespace HordeServer.Models
 	/// <summary>
 	/// Config for a stream
 	/// </summary>
+	[JsonSchema("https://unrealengine.com/horde/stream")]
+	[JsonSchemaCatalog("Horde Stream", "Horde stream configuration file")]
 	public class StreamConfig
 	{
 		/// <summary>
diff --git a/Engine/Source/Programs/Horde/HordeServer/Program.cs b/Engine/Source/Programs/Horde/HordeServer/Program.cs
--- a/Engine/Source/Programs/Horde/HordeServer/Program.cs
+++ b/Engine/Source/Programs/Horde/HordeServer/Program.cs
@@ -61,12 +61,20 @@ namespace HordeServer
 	{
 		public static DirectoryReference AppDir { get; } = new FileReference(Assembly.GetExecutingAssembly().Location).Directory;
 
-		public static SchemaInfo[] ConfigSchemas =
+		public static Type[] ConfigSchemas = FindSchemaTypes();
+
+		static Type[] FindSchemaTypes()
 		{
-			new SchemaInfo("https://unrealengine.com/horde/global", "Horde Globals", "Horde root configuration file", typeof(GlobalConfig)),
-			new SchemaInfo("https://unrealengine.com/horde/project", "Horde Project", "Horde project configuration file", typeof(ProjectConfig)),
-			new SchemaInfo("https://unrealengine.com/horde/stream", "Horde Stream", "Horde stream configuration file", typeof(StreamConfig))
-		};
+			List<Type> SchemaTypes = new List<Type>();
+			foreach (Type Type in Assembly.GetExecutingAssembly().GetTypes())
+			{
+				if (Type.GetCustomAttribute<JsonSchemaAttribute>() != null)
+				{
+					SchemaTypes.Add(Type);
+				}
+			}
+			return SchemaTypes.ToArray();
+		}
 
 		public static void Main(string[] Args)
 		{
@@ -111,10 +119,10 @@ namespace HordeServer
 				Arguments.CheckAllArgumentsUsed();
 
 				DirectoryReference.CreateDirectory(SchemaDir);
-				foreach (SchemaInfo Schema in ConfigSchemas)
+				foreach (Type SchemaType in ConfigSchemas)
 				{
-					FileReference OutputFile = FileReference.Combine(SchemaDir, $"{Schema.Type.Name}.json");
-					Schemas.WriteSchema(Schema.Id, Schema.Type, OutputFile);
+					FileReference OutputFile = FileReference.Combine(SchemaDir, $"{SchemaType.Name}.json");
+					Schemas.WriteSchema(SchemaType, OutputFile);
 				}
 				return;
 			}
diff --git a/Engine/Source/Programs/Horde/HordeServer/Utilities/Schemas.cs b/Engine/Source/Programs/Horde/HordeServer/Utilities/Schemas.cs
--- a/Engine/Source/Programs/Horde/HordeServer/Utilities/Schemas.cs
+++ b/Engine/Source/Programs/Horde/HordeServer/Utilities/Schemas.cs
@@ -10,11 +10,61 @@ using System.Collections.Concurrent;
 using System.Collections.Generic;
 using System.IO;
 using System.Linq;
+using System.Reflection;
 using System.Text.Json;
 using System.Threading.Tasks;
 
 namespace HordeServer.Utilities
 {
+	/// <summary>
+	/// Attribute setting the identifier for a schema
+	/// </summary>
+	[AttributeUsage(AttributeTargets.Class)]
+	class JsonSchemaAttribute : Attribute
+	{
+		/// <summary>
+		/// The schema identifier
+		/// </summary>
+		public string Id { get; }
+
+		/// <summary>
+		/// Constructor
+		/// </summary>
+		/// <param name="Id"></param>
+		public JsonSchemaAttribute(string Id)
+		{
+			this.Id = Id;
+		}
+	}
+
+	/// <summary>
+	/// Attribute setting catalog entries for a schema
+	/// </summary>
+	[AttributeUsage(AttributeTargets.Class)]
+	public class JsonSchemaCatalogAttribute : Attribute
+	{
+		/// <summary>
+		/// The schema name
+		/// </summary>
+		public string Name { get; set; }
+
+		/// <summary>
+		/// Description of the schema
+		/// </summary>
+		public string Description { get; set; }
+
+		/// <summary>
+		/// Constructor
+		/// </summary>
+		/// <param name="Name"></param>
+		/// <param name="Description"></param>
+		public JsonSchemaCatalogAttribute(string Name, string Description)
+		{
+			this.Name = Name;
+			this.Description = Description;
+		}
+	}
+
 	/// <summary>
 	/// Information about a type to generate a schema for
 	/// </summary>
@@ -61,6 +111,27 @@ namespace HordeServer.Utilities
 	/// </summary>
 	static class Schemas
 	{
+		class NoAdditionalPropertiesIntent : ISchemaKeywordIntent
+		{
+			public void Apply(JsonSchemaBuilder builder)
+			{
+				builder.AdditionalProperties(false);
+			}
+		}
+
+		class NoAdditionalPropertiesRefiner : ISchemaRefiner
+		{
+			public void Run(SchemaGeneratorContext Context)
+			{
+				Context.Intents.Add(new NoAdditionalPropertiesIntent());
+			}
+
+			public bool ShouldRun(SchemaGeneratorContext Context)
+			{
+				return Context.Intents.OfType<PropertiesIntent>().Any();
+			}
+		}
+
 		class CamelCaseSchemaRefiner : ISchemaRefiner
 		{
 			public void Run(SchemaGeneratorContext Context)
@@ -85,28 +156,36 @@ namespace HordeServer.Utilities
 			}
 		}
 
-		static ConcurrentDictionary<(string, Type), JsonSchema> CachedSchemas = new ConcurrentDictionary<(string, Type), JsonSchema>();
+		/// <summary>
+		/// Cache of generated schemas
+		/// </summary>
+		static ConcurrentDictionary<Type, JsonSchema> CachedSchemas = new ConcurrentDictionary<Type, JsonSchema>();
 
 		/// <summary>
 		/// Creates a schema for the given type
 		/// </summary>
-		/// <param name="Id"></param>
 		/// <param name="Type"></param>
 		/// <returns></returns>
-		public static JsonSchema CreateSchema(string Id, Type Type)
+		public static JsonSchema CreateSchema(Type Type)
 		{
 			JsonSchema? Schema;
-			if (!CachedSchemas.TryGetValue((Id, Type), out Schema))
+			if (!CachedSchemas.TryGetValue(Type, out Schema))
 			{
+				JsonSchemaAttribute SchemaAttribute = Type.GetCustomAttribute<JsonSchemaAttribute>()!;
+
 				SchemaGeneratorConfiguration Config = new SchemaGeneratorConfiguration();
 				Config.PropertyOrder = PropertyOrder.AsDeclared;
+				Config.Generators.Add(new StringIdSchemaGenerator());
 				Config.Refiners.Add(new CamelCaseSchemaRefiner());
+				Config.Refiners.Add(new NoAdditionalPropertiesRefiner());
 
-				JsonSchemaBuilder Builder = new JsonSchemaBuilder().FromType(Type, Config);
-				Builder.Id(Id);
+				Schema = new JsonSchemaBuilder()
+					.Schema(MetaSchemas.Draft7Id)
+					.Id(SchemaAttribute.Id)
+					.FromType(Type, Config)
+					.Build();
 
-				Schema = Builder.Build();
-				CachedSchemas.TryAdd((Id, Type), Schema);
+				CachedSchemas.TryAdd(Type, Schema);
 			}
 			return Schema;
 		}
@@ -114,12 +193,11 @@ namespace HordeServer.Utilities
 		/// <summary>
 		/// Writes a schema to a source file
 		/// </summary>
-		/// <param name="Id"></param>
 		/// <param name="Type"></param>
 		/// <param name="OutputFile"></param>
-		public static void WriteSchema(string Id, Type Type, FileReference OutputFile)
+		public static void WriteSchema(Type Type, FileReference OutputFile)
 		{
-			JsonSchema Schema = CreateSchema(Id, Type);
+			JsonSchema Schema = CreateSchema(Type);
 
 			JsonSerializerOptions Options = new JsonSerializerOptions();
 			Options.WriteIndented = true;
diff --git a/Engine/Source/Programs/Horde/HordeServer/Utilities/StringId.cs b/Engine/Source/Programs/Horde/HordeServer/Utilities/StringId.cs
--- a/Engine/Source/Programs/Horde/HordeServer/Utilities/StringId.cs
+++ b/Engine/Source/Programs/Horde/HordeServer/Utilities/StringId.cs
@@ -1,6 +1,9 @@
 // Copyright Epic Games, Inc. All Rights Reserved.
 
 using HordeServer.Models;
+using Json.Schema;
+using Json.Schema.Generation;
+using Json.Schema.Generation.Generators;
 using MongoDB.Bson.Serialization;
 using MongoDB.Bson.Serialization.Attributes;
 using MongoDB.Bson.Serialization.Serializers;
@@ -17,6 +20,25 @@ using System.Threading.Tasks;
 
 namespace HordeServer.Utilities
 {
+	class StringIdSchemaIntent : ISchemaKeywordIntent
+	{
+		public void Apply(JsonSchemaBuilder Builder) => Builder.Type(SchemaValueType.String);
+	}
+
+
+	internal class StringIdSchemaGenerator : ISchemaGenerator
+	{
+		public bool Handles(Type Type)
+		{
+			return Type.IsGenericType && Type.GetGenericTypeDefinition() == typeof(StringId<>);
+		}
+
+		public void AddConstraints(SchemaGeneratorContext context)
+		{
+			context.Intents.Add(new StringIdSchemaIntent());
+		}
+	}
+
 	/// <summary>
 	/// Normalized string identifier for a resource
 	/// </summary>
