Horde: Fix URL for linked schemas.

diff --git a/Engine/Source/Programs/Horde/HordeServer/Controllers/SchemaController.cs b/Engine/Source/Programs/Horde/HordeServer/Controllers/SchemaController.cs
--- a/Engine/Source/Programs/Horde/HordeServer/Controllers/SchemaController.cs
+++ b/Engine/Source/Programs/Horde/HordeServer/Controllers/SchemaController.cs
@@ -73,8 +73,7 @@ namespace HordeServer.Controllers
 					if (CatalogAttribute != null)
 					{
 						Uri Url = new Uri($"https://{Host}/api/v1/schema/types/{SchemaType.Name}.json");
-						Root.Schemas.Add(new CatalogItem { Name = CatalogAttribute.Name, Description = CatalogAttribute.Description, FileMatch = new[] { "globals.json" }, Url = new Uri(SchemaAttribute.Id) });
-						break;
+						Root.Schemas.Add(new CatalogItem { Name = CatalogAttribute.Name, Description = CatalogAttribute.Description, Url = Url });
 					}
 				}
 			}
