Horde: Use comments from XML documentation as descriptions for JSON schemas.

diff --git a/Engine/Source/Programs/Horde/HordeServer/Utilities/Schemas.cs b/Engine/Source/Programs/Horde/HordeServer/Utilities/Schemas.cs
--- a/Engine/Source/Programs/Horde/HordeServer/Utilities/Schemas.cs
+++ b/Engine/Source/Programs/Horde/HordeServer/Utilities/Schemas.cs
@@ -13,6 +13,7 @@ using System.Linq;
 using System.Reflection;
 using System.Text.Json;
 using System.Threading.Tasks;
+using System.Xml;
 
 namespace HordeServer.Utilities
 {
@@ -132,7 +133,7 @@ namespace HordeServer.Utilities
 			}
 		}
 
-		class CamelCaseSchemaRefiner : ISchemaRefiner
+		class CamelCaseRefiner : ISchemaRefiner
 		{
 			public void Run(SchemaGeneratorContext Context)
 			{
@@ -156,11 +157,100 @@ namespace HordeServer.Utilities
 			}
 		}
 
+		class DescriptionIntent : ISchemaKeywordIntent
+		{
+			string Description;
+
+			public DescriptionIntent(string Description)
+			{
+				this.Description = Description;
+			}
+
+			public void Apply(JsonSchemaBuilder Builder)
+			{
+				Builder.Description(Description);
+			}
+
+			public override bool Equals(object? Obj)
+			{
+				return (Obj is DescriptionIntent Other) && Other.Description.Equals(Description, StringComparison.Ordinal);
+			}
+
+			public override int GetHashCode()
+			{
+				return Description.GetHashCode(StringComparison.Ordinal);
+			}
+		}
+	
+		[AttributeUsage(AttributeTargets.Class)]
+		class XmlDocAttribute : Attribute, IAttributeHandler
+		{
+			public string Selector;
+			public string Description;
+
+			public XmlDocAttribute(string Selector, string Description)
+			{
+				this.Selector = Selector;
+				this.Description = Description;
+			}
+
+			public void AddConstraints(SchemaGeneratorContext Context)
+			{
+				int Index = Math.Min(1, Context.Intents.Count);
+				Context.Intents.Insert(Index, new DescriptionIntent(Description));
+			}
+		}
+
+		class XmlDocRefiner : ISchemaRefiner
+		{
+			XmlDocument XmlDoc;
+
+			public XmlDocRefiner(XmlDocument XmlDoc)
+			{
+				this.XmlDoc = XmlDoc;
+			}
+
+			public void Run(SchemaGeneratorContext Context)
+			{
+				PropertiesIntent? Intent = Context.Intents.OfType<PropertiesIntent>().FirstOrDefault();
+				if (Intent != null)
+				{
+					KeyValuePair<string, SchemaGeneratorContext>[] Properties = Intent.Properties.ToArray();
+					foreach ((string PropertyName, SchemaGeneratorContext PropertyContext) in Properties)
+					{
+						string Selector = $"//member[@name='P:{Context.Type.FullName}.{PropertyName}']/summary";
+
+						XmlNode Node = XmlDoc.SelectSingleNode(Selector);
+						if (Node != null)
+						{
+							string Description = Node.InnerText.Trim().Replace("\r\n", "\n", StringComparison.Ordinal);
+
+							List<Attribute> Attributes = new List<Attribute>(PropertyContext.Attributes);
+							Attributes.Add(new XmlDocAttribute(Selector, Description));
+
+							SchemaGeneratorContext NewContext = SchemaGenerationContextCache.Get(PropertyContext.Type, Attributes, PropertyContext.Configuration);
+							Intent.Properties[PropertyName] = NewContext;
+						}
+					}
+				}
+			}
+
+			public bool ShouldRun(SchemaGeneratorContext Context)
+			{
+				return Context.Intents.OfType<PropertiesIntent>().Any();
+			}
+		}
+
 		/// <summary>
 		/// Cache of generated schemas
 		/// </summary>
 		static ConcurrentDictionary<Type, JsonSchema> CachedSchemas = new ConcurrentDictionary<Type, JsonSchema>();
 
+		/// <summary>
+		/// Xml documentation for this assembly
+		/// </summary>
+		static XmlDocument? XmlDoc { get; } = ReadXmlDocumentation();
+
 		/// <summary>
 		/// Creates a schema for the given type
 		/// </summary>
@@ -176,8 +266,12 @@ namespace HordeServer.Utilities
 				SchemaGeneratorConfiguration Config = new SchemaGeneratorConfiguration();
 				Config.PropertyOrder = PropertyOrder.AsDeclared;
 				Config.Generators.Add(new StringIdSchemaGenerator());
-				Config.Refiners.Add(new CamelCaseSchemaRefiner());
 				Config.Refiners.Add(new NoAdditionalPropertiesRefiner());
+				if (XmlDoc != null)
+				{
+					Config.Refiners.Add(new XmlDocRefiner(XmlDoc));
+				}
+				Config.Refiners.Add(new CamelCaseRefiner());
 
 				Schema = new JsonSchemaBuilder()
 					.Schema(MetaSchemas.Draft7Id)
@@ -190,6 +284,21 @@ namespace HordeServer.Utilities
 			return Schema;
 		}
 
+		static XmlDocument? ReadXmlDocumentation()
+		{
+			// Get the path to the XML documentation
+			FileReference InputDocumentationFile = new FileReference(Assembly.GetExecutingAssembly().Location).ChangeExtension(".xml");
+			if (!FileReference.Exists(InputDocumentationFile))
+			{
+				return null;
+			}
+
+			// Read the documentation
+			XmlDocument InputDocumentation = new XmlDocument();
+			InputDocumentation.Load(InputDocumentationFile.FullName);
+			return InputDocumentation;
+		}
+
 		/// <summary>
 		/// Writes a schema to a source file
 		/// </summary>
