Horde: Set the schema version to draft-04 (max version supported by [REDACTED_PERSON])

diff --git a/Engine/Source/Programs/Horde/HordeServer/Utilities/Schemas.cs b/Engine/Source/Programs/Horde/HordeServer/Utilities/Schemas.cs
--- a/Engine/Source/Programs/Horde/HordeServer/Utilities/Schemas.cs
+++ b/Engine/Source/Programs/Horde/HordeServer/Utilities/Schemas.cs
@@ -274,7 +274,7 @@ namespace HordeServer.Utilities
 				Config.Refiners.Add(new CamelCaseRefiner());
 
 				Schema = new JsonSchemaBuilder()
-					.Schema(MetaSchemas.Draft7Id)
+					.Schema("http://json-schema.org/draft-04/schema#")
 					.Id(SchemaAttribute.Id)
 					.FromType(Type, Config)
 					.Build();
