Horde: Add id for legacy JSON schemas.

diff --git a/Engine/Source/Programs/Horde/HordeServer/Utilities/Schemas.cs b/Engine/Source/Programs/Horde/HordeServer/Utilities/Schemas.cs
--- a/Engine/Source/Programs/Horde/HordeServer/Utilities/Schemas.cs
+++ b/Engine/Source/Programs/Horde/HordeServer/Utilities/Schemas.cs
@@ -1,18 +1,16 @@
 // Copyright Epic Games, Inc. All Rights Reserved.
 
 using EpicGames.Core;
-using HordeServer.Models;
 using Json.Schema;
 using Json.Schema.Generation;
 using Json.Schema.Generation.Intents;
 using System;
 using System.Collections.Concurrent;
 using System.Collections.Generic;
-using System.IO;
 using System.Linq;
 using System.Reflection;
 using System.Text.Json;
-using System.Threading.Tasks;
+using System.Text.Json.Serialization;
 using System.Xml;
 
 namespace HordeServer.Utilities
@@ -241,6 +239,50 @@ namespace HordeServer.Utilities
 			}
 		}
 
+		class LegacyIdKeywordJsonConverter : JsonConverter<LegacyIdKeyword>
+		{
+			public override LegacyIdKeyword Read(ref Utf8JsonReader reader, Type typeToConvert, JsonSerializerOptions options)
+			{
+				throw new NotImplementedException();
+			}
+
+			public override void Write(Utf8JsonWriter writer, LegacyIdKeyword value, JsonSerializerOptions options)
+			{
+				writer.WriteString("id", value.Id);
+			}
+		}
+
+		[SchemaKeyword("id")]
+		[JsonConverter(typeof(LegacyIdKeywordJsonConverter))]
+		class LegacyIdKeyword : IJsonSchemaKeyword, IEquatable<LegacyIdKeyword>
+		{
+			public string Id { get; set; }
+
+			public LegacyIdKeyword(string Id)
+			{
+				this.Id = Id;
+			}
+
+			public bool Equals(LegacyIdKeyword? other)
+			{
+				return other != null && other.Id.Equals(Id, StringComparison.Ordinal);
+			}
+
+			public override bool Equals(object? obj)
+			{
+				return Equals(obj as LegacyIdKeyword);
+			}
+
+			public override int GetHashCode()
+			{
+				return Id.GetHashCode(StringComparison.Ordinal);
+			}
+
+			public void Validate(ValidationContext context)
+			{
+			}
+		}
+
 		/// <summary>
 		/// Cache of generated schemas
 		/// </summary>
@@ -273,9 +315,10 @@ namespace HordeServer.Utilities
 				}
 				Config.Refiners.Add(new CamelCaseRefiner());
 
-				Schema = new JsonSchemaBuilder()
-					.Schema("http://json-schema.org/draft-04/schema#")
-					.Id(SchemaAttribute.Id)
+				JsonSchemaBuilder Builder = new JsonSchemaBuilder()
+					.Schema("http://json-schema.org/draft-04/schema#"); // VS2019 only supports draft-4 :(
+				Builder.Add(new LegacyIdKeyword(SchemaAttribute.Id));
+				Schema = Builder
 					.FromType(Type, Config)
 					.Build();
 
