Horde: Adding device management
#rnx

diff --git a/Engine/Source/Programs/Horde/HordeDashboard/src/backend/Api.ts b/Engine/Source/Programs/Horde/HordeDashboard/src/backend/Api.ts
--- a/Engine/Source/Programs/Horde/HordeDashboard/src/backend/Api.ts
+++ b/Engine/Source/Programs/Horde/HordeDashboard/src/backend/Api.ts
@@ -2690,3 +2690,158 @@ export type GetPerforceServerStatusResponse = {
 	status: string;
 	detail: string;
 }
+
+/** Get object response which describes a device platform */
+export type GetDevicePlatformResponse = {
+	
+	/** Unique id of device platform */	
+	id: string;
+	
+	/** Friendly name of device platform */
+	name:string;
+
+	/** Platform vendor models */
+	modelIds:string[];
+}
+
+
+/** Device pool response object */
+export type GetDevicePoolResponse = {
+	
+	/**  Id of  the device pool */
+	id: string;
+
+	/**  Name of the device pool */
+	name: string;
+
+}
+
+
+export type CreateDeviceRequest = {
+
+	/**  The platform of the device */
+	platformId: string;
+
+	/**  The pool to assign the device */
+	poolId: string;
+
+	/**  The friendly name of the device */
+	name: string;
+
+	/**  Whether to create the device in enabled state */
+	enabled?: boolean;
+
+	/**  The network address of the device */
+	address?: string;
+
+	/**  The vendor model id of the device */
+	modelId?: string;
+}
+
+export type CreateDeviceResponse = {
+
+	/** id of created device */
+	id: string;
+
+}
+	
+/** Get response object which describes a device */
+export type GetDeviceResponse = {
+	
+	/** The unique id of the device */
+	id: string;
+	
+	/** The platform of the device */
+	platformId: string;
+	
+	/** The pool the device belongs to */
+	poolId: string;
+	
+	/** The friendly name of the device */
+	name: string;
+	
+	/** Whether the device is currently enabled */
+	enabled: boolean;
+	
+	/**  The address of the device (if it allows network connections) */
+	address?: string;
+	
+	/** The vendor model id of the device */
+	modelId?: string;
+
+	/** Any notes provided for the device */
+	notes? : string;	
+	
+	/** The UTC time when a device problem was reported */		
+	problemTime?: Date | string;
+
+	/** The UTC time when a device was set for maintenance */		
+	maintenanceTime?: Date | string;
+	
+}
+
+/** Request object for updating a device */		
+export type UpdateDeviceRequest = {
+
+	/** The pool to assign device to */
+	poolId?: string;
+
+	/** The new name of the device */
+	name?: string;
+
+	/** Whether the device is enabled */
+	enabled?: boolean;
+
+	/** Whether to clear any problem state */ 
+	problem?:boolean;
+
+	/** Whether the device should be put into maintenance mode */
+	maintenance?:boolean;
+
+	/** New address or hostname of device */
+	address?: string;
+
+	/** The new model id of device */
+	modelId?: string;
+
+	/** Markdown notes */
+	notes?: string;
+}
+
+	
+/** A reservation containing one or more devices */
+
+export type GetDeviceReservationResponse = {
+		
+	/** Randomly generated unique id for this reservation */
+	id: string;
+
+	/** Which device pool the reservation is in	*/
+	poolId: string;
+	
+	/** The reserved device ids	*/
+	devices:string[];
+	
+	/** JobID holding reservation */		
+	jobId?: string;
+	
+	/** Job step id holding reservation	*/
+	stepId?: string;
+	
+	/** Reservations held by a user	*/
+	userId?: string;
+	
+	/** The hostname of machine holding reservation	*/
+	hostname?:string;
+
+	/** The optional details of the reservation	*/
+	reservationDetails?:string;
+	
+	/** The UTC time when the reservation was created */		
+	createTimeUtc: Date | string;
+	
+	/** The legacy reservation system guid, to be removed once can update Gauntlet client in all streams */
+	legacyGuid: string;
+}
+
+
diff --git a/Engine/Source/Programs/Horde/HordeDashboard/src/backend/Backend.ts b/Engine/Source/Programs/Horde/HordeDashboard/src/backend/Backend.ts
--- a/Engine/Source/Programs/Horde/HordeDashboard/src/backend/Backend.ts
+++ b/Engine/Source/Programs/Horde/HordeDashboard/src/backend/Backend.ts
@@ -3,7 +3,7 @@
 import { configure } from "mobx";
 import { isNumber } from 'util';
 import templateCache from '../backend/TemplateCache';
-import { AgentData, ArtifactData, BatchUpdatePoolRequest, ChangeSummaryData, CreateJobRequest, CreateJobResponse, CreatePoolRequest, CreateSoftwareResponse, CreateSubscriptionRequest, CreateSubscriptionResponse, DashboardPreference, EventData, GetArtifactZipRequest, GetUserResponse, GetGraphResponse, GetIssueStreamResponse, GetJobStepRefResponse, GetJobStepTraceResponse, GetJobTimingResponse, GetLogEventResponse, GetNotificationResponse, GetSoftwareResponse, GetSubscriptionResponse, IssueData, IssueQuery, JobData, JobQuery, LeaseData, LogData, LogLineData, PoolData, ProjectData, ScheduleData, ScheduleQuery, SearchLogFileResponse, SessionData, StreamData, TemplateData, TestData, UpdateAgentRequest, UpdateIssueRequest, UpdateJobRequest, UpdateNotificationsRequest, UpdatePoolRequest, UpdateStepRequest, UpdateStepResponse, UpdateUserRequest, GetUtilizationTelemetryResponse, TabType, GetJobsTabResponse, JobsTabColumnType, GetPerforceServerStatusResponse } from './Api';
+import { AgentData, ArtifactData, BatchUpdatePoolRequest, ChangeSummaryData, CreateJobRequest, CreateJobResponse, CreatePoolRequest, CreateSoftwareResponse, CreateSubscriptionRequest, CreateSubscriptionResponse, DashboardPreference, EventData, GetArtifactZipRequest, GetUserResponse, GetGraphResponse, GetIssueStreamResponse, GetJobStepRefResponse, GetJobStepTraceResponse, GetJobTimingResponse, GetLogEventResponse, GetNotificationResponse, GetSoftwareResponse, GetSubscriptionResponse, IssueData, IssueQuery, JobData, JobQuery, LeaseData, LogData, LogLineData, PoolData, ProjectData, ScheduleData, ScheduleQuery, SearchLogFileResponse, SessionData, StreamData, TemplateData, TestData, UpdateAgentRequest, UpdateIssueRequest, UpdateJobRequest, UpdateNotificationsRequest, UpdatePoolRequest, UpdateStepRequest, UpdateStepResponse, UpdateUserRequest, GetUtilizationTelemetryResponse, TabType, GetJobsTabResponse, JobsTabColumnType, GetPerforceServerStatusResponse, GetDeviceResponse, GetDevicePlatformResponse, CreateDeviceResponse, CreateDeviceRequest, UpdateDeviceRequest, GetDevicePoolResponse, GetDeviceReservationResponse } from './Api';
 import dashboard from './Dashboard';
 import { ChallengeStatus, Fetch } from './Fetch';
 import graphCache, { GraphQuery } from './GraphCache';
@@ -1078,6 +1078,84 @@ export class Backend {
         
     }
 
+    getDevices(): Promise<GetDeviceResponse[]> {
+
+        return new Promise<GetDeviceResponse[]>((resolve, reject) => {
+            this.backend.get(`/api/v2/devices`).then((value) => {
+                resolve(value.data as GetDeviceResponse[]);
+            }).catch(reason => {
+                reject(reason);
+            });
+        });        
+    }
+
+    addDevice(request: CreateDeviceRequest): Promise<CreateDeviceResponse> {
+        return new Promise<CreateDeviceResponse>((resolve, reject) => {
+            this.backend.post(`/api/v2/devices`, request).then((value) => {
+                resolve(value.data as CreateDeviceResponse);
+            }).catch(reason => {
+                reject(reason);
+            });
+        });        
+    }
+
+    modifyDevice(deviceId: string, request: UpdateDeviceRequest): Promise<void> {
+        return new Promise<void>((resolve, reject) => {
+            this.backend.put(`/api/v2/devices/${deviceId}`, request).then((value) => {
+                resolve();
+            }).catch(reason => {
+                reject(reason);
+            });
+        });        
+    }
+
+
+    deleteDevice(deviceId: string): Promise<void> {
+        return new Promise<void>((resolve, reject) => {
+            this.backend.delete(`/api/v2/devices/${deviceId}`).then((response) => {
+                resolve();
+            }).catch(reason => {
+                reject(reason);
+            });
+        });        
+    }
+
+    getDevicePlatforms(): Promise<GetDevicePlatformResponse[]> {
+
+        return new Promise<GetDevicePlatformResponse[]>((resolve, reject) => {
+            this.backend.get(`/api/v2/devices/platforms`).then((value) => {
+                resolve(value.data as GetDevicePlatformResponse[]);
+            }).catch(reason => {
+                reject(reason);
+            });
+        });
+        
+    }
+
+    getDevicePools(): Promise<GetDevicePoolResponse[]> {
+
+        return new Promise<GetDevicePoolResponse[]>((resolve, reject) => {
+            this.backend.get(`/api/v2/devices/pools`).then((value) => {
+                resolve(value.data as GetDevicePoolResponse[]);
+            }).catch(reason => {
+                reject(reason);
+            });
+        });
+        
+    }
+
+    getDeviceReservations(): Promise<GetDeviceReservationResponse[]> {
+
+        return new Promise<GetDeviceReservationResponse[]>((resolve, reject) => {
+            this.backend.get(`/api/v2/devices/reservations`).then((value) => {
+                resolve(value.data as GetDeviceReservationResponse[]);
+            }).catch(reason => {
+                reject(reason);
+            });
+        });
+        
+    }
+
     private async update() {
 
         if (this.updateID === "updating") {
diff --git a/Engine/Source/Programs/Horde/HordeServer/Api/Devices.cs b/Engine/Source/Programs/Horde/HordeServer/Api/Devices.cs
new file mode 100644
--- /dev/null
+++ b/Engine/Source/Programs/Horde/HordeServer/Api/Devices.cs
@@ -0,0 +1,580 @@
+// Copyright Epic Games, Inc. All Rights Reserved.
+
+using System;
+using System.Collections.Generic;
+using System.ComponentModel.DataAnnotations;
+using System.Diagnostics.CodeAnalysis;
+
+namespace HordeServer.Api
+{
+
+	/// <summary>
+	/// Create device platform request
+	/// </summary>
+	public class CreateDevicePlatformRequest
+	{
+		/// <summary>
+		/// The name of the platform
+		/// </summary>
+		[Required]
+		public string Name { get; set; } = null!;
+
+	}
+
+	/// <summary>
+	/// Create device platform response object
+	/// </summary>
+	public class CreateDevicePlatformResponse
+	{
+		/// <summary>
+		/// Id of newly created platform
+		/// </summary>
+		public string Id { get; set; }
+
+		/// <summary>
+		/// Constructor
+		/// </summary>
+		public CreateDevicePlatformResponse(string Id)
+		{
+			this.Id = Id;
+		}
+	}
+
+	/// <summary>
+	/// Update requesty object for a device platform
+	/// </summary>
+	public class UpdateDevicePlatformRequest
+	{
+		/// <summary>
+		/// The vendor model ids for the platform
+		/// </summary>
+		public string[] ModelIds { get; set; } = null!;
+	}
+
+
+	/// <summary>
+	/// Get object response which describes a device platform
+	/// </summary>
+	public class GetDevicePlatformResponse
+	{
+		/// <summary>
+		/// Unique id of device platform
+		/// </summary>
+		public string Id { get; set; }
+
+		/// <summary>
+		/// Friendly name of device platform
+		/// </summary>
+		public string Name { get; set; }
+
+		/// <summary>
+		/// Platform vendor models
+		/// </summary>
+		public string[] ModelIds { get; set; }
+
+		/// <summary>
+		/// Response constructor
+		/// </summary>
+		public GetDevicePlatformResponse(string Id, string Name, string[] ModelIds)
+		{
+			this.Id = Id;
+			this.Name = Name;
+			this.ModelIds = ModelIds;
+		}
+
+	}
+
+	/// <summary>
+	/// Device pool creation request object
+	/// </summary>
+	public class CreateDevicePoolRequest
+	{
+		/// <summary>
+		/// The name for the new pool
+		/// </summary>
+		[Required]
+		public string Name { get; set; } = null!;
+
+	}
+
+	/// <summary>
+	/// Device pool creation response object
+	/// </summary>
+	public class CreateDevicePoolResponse
+	{
+		/// <summary>
+		/// Id of the newly created device pool
+		/// </summary>
+		public string Id { get; set; }
+
+		/// <summary>
+		/// Constructor
+		/// </summary>
+		public CreateDevicePoolResponse(string Id)
+		{
+			this.Id = Id;
+		}
+	}
+
+	/// <summary>
+	/// Device pool response object
+	/// </summary>
+	public class GetDevicePoolResponse
+	{
+		/// <summary>
+		/// Id of  the device pool
+		/// </summary>
+		public string Id { get; set; }
+
+		/// <summary>
+		/// Name of the device pool
+		/// </summary>
+		public string Name { get; set; }
+
+		/// <summary>
+		/// Constructor
+		/// </summary>
+		public GetDevicePoolResponse(string Id, string Name)
+		{
+			this.Id = Id;
+			this.Name = Name;
+		}
+	}
+
+	// Devices 
+
+	/// <summary>
+	/// Device creation request object
+	/// </summary>
+	public class CreateDeviceRequest
+	{
+		/// <summary>
+		/// The platform of the device 
+		/// </summary>
+		[Required]
+		public string PlatformId { get; set; } = null!;
+
+		/// <summary>
+		/// The pool to assign the device
+		/// </summary>
+		[Required]
+		public string PoolId { get; set; } = null!;
+
+		/// <summary>
+		/// The friendly name of the device
+		/// </summary>
+		[Required]
+		public string Name { get; set; } = null!;
+
+		/// <summary>
+		/// Whether to create the device in enabled state
+		/// </summary>
+		public bool? Enabled { get; set; }
+
+		/// <summary>
+		/// The network address of the device
+		/// </summary>
+		public string? Address { get; set; }
+
+		/// <summary>
+		/// The vendor model id of the device
+		/// </summary>
+		public string? ModelId { get; set; }
+
+	}
+
+
+	/// <summary>
+	/// Device creation response object
+	/// </summary>
+	public class CreateDeviceResponse
+	{
+		/// <summary>
+		/// The id of the newly created device
+		/// </summary>
+		[Required]
+		public string Id { get; set; } = null!;
+
+		/// <summary>
+		/// Constructor
+		/// </summary>
+		/// <param name="Id"></param>
+		public CreateDeviceResponse(string Id)
+		{
+			this.Id = Id;
+		}
+	}
+	/// <summary>
+	/// Get response object which describes a device
+	/// </summary>
+	public class GetDeviceResponse
+	{
+		/// <summary>
+		/// The unique id of the device
+		/// </summary>
+		public string Id { get; set; }
+
+		/// <summary>
+		/// The platform of the device
+		/// </summary>
+		public string PlatformId { get; set; }
+
+		/// <summary>
+		/// The pool the device belongs to
+		/// </summary>
+		public string PoolId { get; set; }
+
+		/// <summary>
+		/// The friendly name of the device
+		/// </summary>
+		public string Name { get; set; }
+
+		/// <summary>
+		/// Whether the device is currently enabled
+		/// </summary>
+		public bool Enabled { get; set; }
+
+		/// <summary>
+		///  The address of the device (if it allows network connections)
+		/// </summary>
+		public string? Address { get; set; }
+
+		/// <summary>
+		/// The vendor model id of the device
+		/// </summary>
+		public string? ModelId { get; set; }
+
+		/// <summary>
+		/// Any notes provided for the device
+		/// </summary>
+		public string? Notes { get; set; }
+
+		/// <summary>
+		/// If the device has a marked problem
+		/// </summary>
+		public DateTime? ProblemTime { get; set; }
+
+		/// <summary>
+		/// If the device is in maintenance mode
+		/// </summary>
+		public DateTime? MaintenanceTime { get; set; }
+
+
+		/// <summary>
+		/// Device response constructor
+		/// </summary>
+		public GetDeviceResponse(string Id, string PlatformId, string PoolId, string Name, bool Enabled, string? Address, string? ModelId, string? Notes, DateTime? ProblemTime, DateTime? MaintenanceTime)
+		{
+			this.Id = Id;
+			this.Name = Name;
+			this.PlatformId = PlatformId;
+			this.PoolId = PoolId;
+			this.Enabled = Enabled;
+			this.Address = Address;
+			this.ModelId = ModelId;
+			this.Notes = Notes;
+			this.ProblemTime = ProblemTime;
+			this.MaintenanceTime = MaintenanceTime;
+		}
+	}
+
+	/// <summary>
+	/// Device update request object
+	/// </summary>
+	public class UpdateDeviceRequest
+	{
+		/// <summary>
+		/// The device pool id
+		/// </summary>
+		public string? PoolId { get; set; }
+
+		/// <summary>
+		/// The device name
+		/// </summary>
+		public string? Name { get; set; }
+
+		/// <summary>
+		/// IP address or hostname of device
+		/// </summary>
+		public string? Address { get; set; }
+
+		/// <summary>
+		/// Device vendor model id
+		/// </summary>
+		public string? ModelId { get; set; }
+
+		/// <summary>
+		/// Markdown notes
+		/// </summary>
+		public string? Notes { get; set; }
+
+		/// <summary>
+		/// Whether device is enabled
+		/// </summary>
+		public bool? Enabled { get; set; }
+
+		/// <summary>
+		/// Whether the device is in maintenance mode
+		/// </summary>
+		public bool? Maintenance { get; set; }
+
+		/// <summary>
+		/// Whether to set or clear any device problem state
+		/// </summary>
+		public bool? Problem { get; set; }
+
+	}
+
+	/// <summary>
+	/// Device reservation request object
+	/// </summary>
+	public class DeviceReservationRequest
+	{
+		/// <summary>
+		/// Device reservation platform id
+		/// </summary>
+		[Required]
+		public string PlatformId { get; set; } = null!;
+
+		/// <summary>
+		/// The optional vendor model ids to include for this device
+		/// </summary>
+		public List<string>? IncludeModels { get; set; }
+
+		/// <summary>
+		/// The optional vendor model ids to exclude for this device
+		/// </summary>
+		public List<string>? ExcludeModels { get; set; }
+
+	}
+
+	/// <summary>
+	/// Reservation request object
+	/// </summary>
+	public class CreateDeviceReservationRequest
+	{
+		/// <summary>
+		/// What pool to reserve devices in
+		/// </summary>
+		[Required]
+		public string PoolId { get; set; } = null!;
+
+		/// <summary>
+		/// Devices to reserve
+		/// </summary>
+		[Required]
+		public List<DeviceReservationRequest> Devices { get; set; } = null!;
+
+	}
+
+	/// <summary>
+	/// Device reservation response object
+	/// </summary>
+	public class CreateDeviceReservationResponse
+	{
+		/// <summary>
+		/// The reservation id of newly created reservation
+		/// </summary>
+		public string Id { get; set; } = null!;
+
+		/// <summary>
+		/// The devices that were reserved
+		/// </summary>
+		public List<GetDeviceResponse> Devices { get; set; } = null!;
+
+	}
+
+	/// <summary>
+	/// A reservation containing one or more devices
+	/// </summary>
+	public class GetDeviceReservationResponse
+	{
+		/// <summary>
+		/// Randomly generated unique id for this reservation
+		/// </summary>
+		[Required]
+		public string Id { get; set; } = null!;
+
+		/// <summary>
+		/// Which device pool the reservation is in
+		/// </summary>
+		[Required]
+		public string PoolId { get; set; } = null!;
+
+		/// <summary>
+		/// The reserved devices
+		/// </summary>
+		[Required]
+		public List<string> Devices { get; set; } = null!;
+
+		/// <summary>
+		/// JobID holding reservation
+		/// </summary>
+		public string? JobId { get; set; }
+
+		/// <summary>
+		/// Job step id holding reservation
+		/// </summary>
+		public string? StepId { get; set; }
+
+		/// <summary>
+		/// Reservations held by a user, requires a token
+		/// </summary>
+		public string? UserId { get; set; }
+
+		/// <summary>
+		/// The hostname of machine holding reservation
+		/// </summary>
+		public string? Hostname { get; set; }
+
+		/// <summary>
+		/// The optional reservation details 
+		/// </summary>
+		public string? ReservationDetails { get; set; }
+
+		/// <summary>
+		/// The UTC time when the reservation was created
+		/// </summary>
+		public DateTime CreateTimeUtc { get; set; }
+
+		/// <summary>
+		/// The legacy reservation system guid, to be removed once can update Gauntlet client in all streams
+		/// </summary>
+		public string LegacyGuid { get; set; } = null!;
+	}
+
+
+	// Legacy clients
+
+	/// <summary>
+	/// Reservation request for legacy clients
+	/// </summary>
+	public class LegacyCreateReservationRequest
+	{
+		/// <summary>
+		/// The device types to reserve, these are mapped to platforms
+		/// </summary>
+		[Required]
+		public string[] DeviceTypes { get; set; } = null!;
+
+		/// <summary>
+		/// The hostname of machine reserving devices 
+		/// </summary>
+		[Required]
+		public string Hostname { get; set; } = null!;
+
+		/// <summary>
+		/// The duration of reservation
+		/// </summary>
+		[Required]
+		public string Duration { get; set; } = null!;
+
+		/// <summary>
+		/// Reservation details string
+		/// </summary>
+		public string? ReservationDetails { get; set; } = null;
+
+		/// <summary>
+		/// The PoolId of reservation request 
+		/// </summary>
+		public string? PoolId { get; set; } = null;
+	}
+
+	/// <summary>
+	/// Reservation response for legacy clients
+	/// </summary>
+	public class GetLegacyReservationResponse
+	{
+		/// <summary>
+		/// The names of the devices that were reserved
+		/// </summary>
+		[Required]
+		public string[] DeviceNames { get; set; } = null!;
+
+		/// <summary>
+		/// The corresponding perf specs of the reserved devices
+		/// </summary>
+		[Required]
+		public string[] DevicePerfSpecs { get; set; } = null!;
+
+		/// <summary>
+		/// The host name of the machine making the reservation
+		/// </summary>
+		[Required]
+		public string HostName { get; set; } = null!;
+
+		/// <summary>
+		/// The start time of the reservation
+		/// </summary>
+		[Required]
+		public string StartDateTime { get; set; } = null!;
+
+		/// <summary>
+		/// The duration of the reservation (before renew)
+		/// </summary>
+		[Required]
+		public string Duration { get; set; } = null!;
+
+		/// <summary>
+		/// The legacy guid of the reservation
+		/// </summary>
+		[Required]
+		[SuppressMessage("Compiler", "CA1720: Identifier 'Guid' contains type name")]
+		public string Guid { get; set; } = null!;
+	}
+
+	/// <summary>
+	/// Device response for legacy clients
+	/// </summary>
+	public class GetLegacyDeviceResponse
+	{
+		/// <summary>
+		/// The  name of the reserved device
+		/// </summary>
+		[Required]
+		public string Name { get; set; } = null!;
+
+		/// <summary>
+		/// The (legacy) type of the device, mapped from platform id
+		/// </summary>
+		[Required]
+		public string Type { get; set; } = null!;
+
+		/// <summary>
+		/// The IP or hostname of device
+		/// </summary>
+		[Required]
+		public string IPOrHostName { get; set; } = null!;
+
+		/// <summary>
+		/// The (legacy) perf spec of the device
+		/// </summary>
+		[Required]
+		public string PerfSpec { get; set; } = null!;
+
+		/// <summary>
+		/// The available start time which is parsed client side
+		/// </summary>
+		[Required]
+		public string AvailableStartTime { get; set; } = null!;
+
+		/// <summary>
+		/// The available end time which is parsed client side
+		/// </summary>
+		[Required]
+		public string AvailableEndTime { get; set; } = null!;
+
+		/// <summary>
+		/// Whether device is enabled
+		/// </summary>
+		[Required]
+		public bool Enabled { get; set; }
+
+		/// <summary>
+		/// Associated device data
+		/// </summary>
+		[Required]
+		public string DeviceData { get; set; } = null!;
+	}
+
+}
diff --git a/Engine/Source/Programs/Horde/HordeServer/Collections/IDeviceCollection.cs b/Engine/Source/Programs/Horde/HordeServer/Collections/IDeviceCollection.cs
new file mode 100644
--- /dev/null
+++ b/Engine/Source/Programs/Horde/HordeServer/Collections/IDeviceCollection.cs
@@ -0,0 +1,189 @@
+// Copyright Epic Games, Inc. All Rights Reserved.
+
+using System.Threading.Tasks;
+
+using HordeServer.Models;
+using MongoDB.Bson;
+using System.Collections.Generic;
+
+using DeviceId = HordeServer.Utilities.StringId<HordeServer.Models.IDevice>;
+using DevicePlatformId = HordeServer.Utilities.StringId<HordeServer.Models.IDevicePlatform>;
+using DevicePoolId = HordeServer.Utilities.StringId<HordeServer.Models.IDevicePool>;
+using System;
+
+namespace HordeServer.Collections
+{
+	/// <summary>
+	/// Device reservation request data
+	/// </summary>
+	public class DeviceRequestData
+	{
+		/// <summary>
+		/// The platform of the device to reserve
+		/// </summary>
+		public DevicePlatformId PlatformId { get; set; }
+
+
+		/// <summary>
+		/// Models to include for this request
+		/// </summary>
+		public List<string> IncludeModels { get; set; }
+
+		/// <summary>
+		/// Models to exclude for this request
+		/// </summary>
+		public List<string> ExcludeModels { get; set; }
+
+		/// <summary>
+		/// Constructor
+		/// </summary>
+		public DeviceRequestData(DevicePlatformId PlatformId, List<string>? IncludeModels = null, List<string>? ExcludeModels = null)
+		{
+			this.PlatformId = PlatformId;
+			this.IncludeModels = IncludeModels ?? new List<string>();
+			this.ExcludeModels = ExcludeModels ?? new List<string>();
+		}
+	}
+
+	/// <summary>
+	/// A collection for device management
+	/// </summary>
+	public interface IDeviceCollection
+	{
+		// PLATFORMS
+
+		/// <summary>
+		/// Add a platform 
+		/// </summary>
+		/// <param name="Id">The id of the platform</param>
+		/// <param name="Name">The friendly name of the platform</param>
+		Task<IDevicePlatform?> TryAddPlatformAsync(DevicePlatformId Id, string Name);
+
+		/// <summary>
+		/// Get a list of all available device platforms
+		/// </summary>		
+		Task<List<IDevicePlatform>> FindAllPlatformsAsync();
+
+		/// <summary>
+		/// Get a specific platform by id
+		/// </summary>
+		Task<IDevicePlatform?> GetPlatformAsync(DevicePlatformId PlatformId);
+
+		/// <summary>
+		/// Update a device platform
+		/// </summary>
+		/// <param name="platformId">The id of the device</param>
+		/// <param name="ModelIds">The available model ids for the platform</param>
+		Task<bool> UpdatePlatformAsync(DevicePlatformId platformId, string[]? ModelIds);
+
+		// POOLS
+
+		/// <summary>
+		/// Add a new device pool to the collection
+		/// </summary>
+		/// <param name="Id">The id of the new pool</param>
+		/// <param name="Name">The friendly name of the new pool</param>
+		Task<IDevicePool?> TryAddPoolAsync(DevicePoolId Id, string Name);
+
+		/// <summary>
+		/// Get a pool by id
+		/// </summary>
+		Task<IDevicePool?> GetPoolAsync(DevicePoolId PoolId);
+
+		/// <summary>
+		/// Gets a list of existing device pools
+		/// </summary>		
+		Task<List<IDevicePool>> FindAllPoolsAsync();
+
+		// DEVICES
+
+		/// <summary>
+		/// Get a device by id
+		/// </summary>
+		Task<IDevice?> GetDeviceAsync(DeviceId DeviceId);
+
+		/// <summary>
+		/// Get a device by name
+		/// </summary>
+		Task<IDevice?> GetDeviceByNameAsync(string DeviceName);
+
+		/// <summary>
+		/// Get a list of all devices
+		/// </summary>
+		/// <param name="DeviceIds">Optional list of device ids to get</param>
+		Task<List<IDevice>> FindAllDevicesAsync(List<DeviceId>? DeviceIds);
+
+		/// <summary>
+		/// Adds a new device
+		/// </summary>
+		/// <param name="Id">The device id</param>
+		/// <param name="Name">The name of the device (unique)</param>
+		/// <param name="PlatformId">The platform of the device</param>
+		/// <param name="PoolId">Which pool to add it to</param>
+		/// <param name="Enabled">Whather the device is enabled by default</param>
+		/// <param name="Address">The network address or hostname the device can be reached at</param>
+		/// <param name="ModelId">The vendor model id of the device</param>
+		Task<IDevice?> TryAddDeviceAsync(DeviceId Id, string Name, DevicePlatformId PlatformId, DevicePoolId PoolId, bool? Enabled, string? Address, string? ModelId);
+
+		/// <summary>
+		/// Update a device
+		/// </summary>
+		/// <param name="DeviceId">The id of the device to update</param>
+		/// <param name="NewPoolId">The new pool to assign</param>
+		/// <param name="NewName">The new name to assign</param>		
+		/// <param name="NewAddress">The new device address or hostname</param>
+		/// <param name="NewModelId">The new model id</param>
+		/// <param name="NewNotes">The devices markdown notes</param>
+		/// <param name="NewEnabled">Whether the device is enabled or not</param>
+		/// <param name="NewProblem">Whether to set or clear problem state</param>
+		/// <param name="NewMaintenance">Whether to set or clear maintenance state</param>
+		Task UpdateDeviceAsync(DeviceId DeviceId, DevicePoolId? NewPoolId, string? NewName, string? NewAddress, string? NewModelId, string? NewNotes, bool? NewEnabled, bool? NewProblem, bool? NewMaintenance);
+
+		/// <summary>
+		/// Delete a device from the collection
+		/// </summary>
+		/// <param name="DeviceId">The id of the device to delete</param>
+		Task<bool> DeleteDeviceAsync(DeviceId DeviceId);
+
+		// RESERVATIONS
+
+		/// <summary>
+		/// Create a new reseveration in the pool with the specified devices
+		/// </summary>
+		/// <param name="PoolId">The pool of devices to use for the new reservation</param>
+		/// <param name="Request">The requested devices for the reservation</param>
+		/// <param name="Hostname">The hostname of the machine making the reservation</param>
+		/// <param name="ReservationDetails">The details of the reservation</param>
+		Task<IDeviceReservation?> TryAddReservationAsync(DevicePoolId PoolId, List<DeviceRequestData> Request, string? Hostname, string? ReservationDetails);
+
+		/// <summary>
+		/// Gets a reservation by guid for legacy clients
+		/// </summary>
+		/// <param name="LegacyGuid">YThe legacy guid of the reservation</param>
+		Task<IDeviceReservation?> TryGetReservationFromLegacyGuidAsync(string LegacyGuid);
+
+		/// <summary>
+		/// Get a list of all reservations
+		/// </summary>
+		Task<List<IDeviceReservation>> FindAllReservationsAsync();
+
+		/// <summary>
+		/// Updates a reservation to the current time, for expiration
+		/// </summary>
+		/// <param name="Id">The id of the reservation to update</param>
+		public Task<bool> TryUpdateReservationAsync(ObjectId Id);
+
+		/// <summary>
+		/// Deletes a reservation and releases reserved devices
+		/// </summary>
+		/// <param name="Id">The id of the reservation to delete</param>
+		public Task<bool> DeleteReservationAsync(ObjectId Id);
+
+		/// <summary>
+		/// Deletes expired reservations
+		/// </summary>
+		public Task<bool> ExpireReservationsAsync();
+
+	}
+
+}
\ No newline at end of file
diff --git a/Engine/Source/Programs/Horde/HordeServer/Collections/Impl/DeviceCollection.cs b/Engine/Source/Programs/Horde/HordeServer/Collections/Impl/DeviceCollection.cs
new file mode 100644
--- /dev/null
+++ b/Engine/Source/Programs/Horde/HordeServer/Collections/Impl/DeviceCollection.cs
@@ -0,0 +1,608 @@
+// Copyright Epic Games, Inc. All Rights Reserved.
+
+using HordeServer.Services;
+using MongoDB.Bson.Serialization.Attributes;
+using System.Threading.Tasks;
+
+using HordeServer.Models;
+using MongoDB.Driver;
+using System;
+using HordeServer.Utilities;
+using System.Linq;
+
+using MongoDB.Bson;
+using System.Collections.Generic;
+
+using DeviceId = HordeServer.Utilities.StringId<HordeServer.Models.IDevice>;
+using DevicePlatformId = HordeServer.Utilities.StringId<HordeServer.Models.IDevicePlatform>;
+using DevicePoolId = HordeServer.Utilities.StringId<HordeServer.Models.IDevicePool>;
+
+namespace HordeServer.Collections.Impl
+{
+
+	/// <summary>
+	/// Collection of device documents
+	/// </summary>
+	public class DeviceCollection : IDeviceCollection
+	{
+
+		/// <summary>
+		/// Document representing a device platform
+		/// </summary>
+		class DevicePlatformDocument : IDevicePlatform
+		{
+			[BsonRequired, BsonId]
+			public DevicePlatformId Id { get; set; }
+
+			public string Name { get; set; }
+
+			public List<string> Models { get; set; } = new List<string>();
+
+			IReadOnlyList<string> IDevicePlatform.Models => Models;
+
+			[BsonConstructor]
+			private DevicePlatformDocument()
+			{
+				Name = null!;
+			}
+
+			public DevicePlatformDocument(DevicePlatformId Id, string Name)
+			{
+				this.Id = Id;
+				this.Name = Name;
+			}
+
+		}
+
+		/// <summary>
+		/// Document representing a pool of devices
+		/// </summary>
+		class DevicePoolDocument : IDevicePool
+		{
+			[BsonRequired, BsonId]
+			public DevicePoolId Id { get; set; }
+
+			[BsonRequired]
+			public string Name { get; set; } = null!;
+
+			[BsonIgnoreIfNull]
+			public Acl? Acl { get; set; }
+
+			[BsonConstructor]
+			private DevicePoolDocument()
+			{
+
+			}
+
+			public DevicePoolDocument(DevicePoolId Id, string Name)
+			{
+				this.Id = Id;
+				this.Name = Name;
+			}
+
+
+		}
+
+		/// <summary>
+		/// Document representing a reservation of devices
+		/// </summary>
+		class DeviceReservationDocument : IDeviceReservation
+		{
+			/// <summary>
+			/// Randomly generated unique id for this reservation.
+			/// </summary>
+			[BsonRequired, BsonId]
+			public ObjectId Id { get; set; }
+
+			[BsonRequired]
+			public DevicePoolId PoolId { get; set; }
+
+			[BsonIgnoreIfNull]
+			public ObjectId? JobId { get; }
+
+			[BsonIgnoreIfNull]
+			public SubResourceId? StepId { get; }
+
+			/// <summary>
+			/// Reservations held by a user, requires a token like download code
+			/// </summary>
+			[BsonIgnoreIfNull]
+			public ObjectId? UserId { get; }
+
+			/// <summary>
+			/// The hostname of the machine which has made the reservation
+			/// </summary>
+			[BsonIgnoreIfNull]
+			public string? Hostname { get; set; }
+
+			/// <summary>
+			/// Optional string holding details about the reservation
+			/// </summary>
+			[BsonIgnoreIfNull]
+			public string? ReservationDetails { get; set; }
+
+			/// <summary>
+			/// 
+			/// </summary>
+			public List<DeviceId> Devices { get; set; } = new List<DeviceId>();
+
+			public DateTime CreateTimeUtc { get; set; }
+
+			public DateTime UpdateTimeUtc { get; set; }
+
+			// Legacy Guid
+			public string LegacyGuid { get; set; } = null!;
+
+			[BsonConstructor]
+			private DeviceReservationDocument()
+			{
+
+			}
+
+			public DeviceReservationDocument(ObjectId Id, DevicePoolId PoolId, List<DeviceId> Devices, DateTime CreateTimeUtc, string? Hostname, string? ReservationDetails)
+			{
+				this.Id = Id;
+				this.PoolId = PoolId;
+				this.Devices = Devices;
+				this.CreateTimeUtc = CreateTimeUtc;
+				this.UpdateTimeUtc = CreateTimeUtc;
+				this.Hostname = Hostname;
+				this.ReservationDetails = ReservationDetails;
+
+				this.LegacyGuid = Guid.NewGuid().ToString();
+
+			}
+
+		}
+
+
+		/// <summary>
+		/// Concrete implementation of an device document
+		/// </summary>
+		class DeviceDocument : IDevice
+		{
+			[BsonRequired, BsonId]
+			public DeviceId Id { get; set; }
+
+			public DevicePlatformId PlatformId { get; set; }
+
+			public DevicePoolId PoolId { get; set; }
+
+			public string Name { get; set; } = null!;
+
+			public bool Enabled { get; set; }
+
+			[BsonIgnoreIfNull]
+			public string? ModelId { get; set; }
+
+			[BsonIgnoreIfNull]
+			public string? Address { get; set; }
+
+			[BsonIgnoreIfNull]
+			public DateTime? ProblemTimeUtc { get; set; }
+
+			[BsonIgnoreIfNull]
+			public DateTime? MaintenanceTimeUtc { get; set; }
+
+			/// <summary>
+			///  The last time this device was reserved, used to cycle devices for reservations
+			/// </summary>
+			[BsonIgnoreIfNull]
+			public DateTime? ReservationTimeUtc { get; set; }
+
+			[BsonIgnoreIfNull]
+			public string? ModifiedByUser { get; set; }
+
+			[BsonIgnoreIfNull]
+			public string? Notes { get; set; }
+
+			[BsonIgnoreIfNull]
+			public Acl? Acl { get; set; }
+
+			[BsonConstructor]
+			private DeviceDocument()
+			{
+
+			}
+
+			public DeviceDocument(DeviceId Id, DevicePlatformId PlatformId, DevicePoolId PoolId, string Name, bool Enabled, string? Address, string? ModelId)
+			{
+				this.Id = Id;
+				this.PlatformId = PlatformId;
+				this.PoolId = PoolId;
+				this.Name = Name;
+				this.Enabled = Enabled;
+				this.Address = Address;
+				this.ModelId = ModelId;
+
+			}
+
+		}
+
+		readonly IMongoCollection<DevicePlatformDocument> Platforms;
+
+		readonly IMongoCollection<DeviceDocument> Devices;
+
+		readonly IMongoCollection<DevicePoolDocument> Pools;
+
+		readonly IMongoCollection<DeviceReservationDocument> Reservations;
+
+		/// <summary>
+		/// Constructor
+		/// </summary>
+		public DeviceCollection(DatabaseService DatabaseService)
+		{
+			Devices = DatabaseService.GetCollection<DeviceDocument>("Devices");
+			Platforms = DatabaseService.GetCollection<DevicePlatformDocument>("Devices.Platforms");
+			Pools = DatabaseService.GetCollection<DevicePoolDocument>("Devices.Pools");
+			Reservations = DatabaseService.GetCollection<DeviceReservationDocument>("Devices.Reservations");
+
+			if (!DatabaseService.ReadOnlyMode)
+			{
+				Platforms.Indexes.CreateOne(new CreateIndexModel<DevicePlatformDocument>(Builders<DevicePlatformDocument>.IndexKeys.Ascending(x => x.Name), new CreateIndexOptions { Unique = true }));
+				Pools.Indexes.CreateOne(new CreateIndexModel<DevicePoolDocument>(Builders<DevicePoolDocument>.IndexKeys.Ascending(x => x.Name), new CreateIndexOptions { Unique = true }));
+				Devices.Indexes.CreateOne(new CreateIndexModel<DeviceDocument>(Builders<DeviceDocument>.IndexKeys.Ascending(x => x.Name), new CreateIndexOptions { Unique = true }));
+			}
+		}
+
+		/// <inheritdoc/>
+		public async Task<IDevice?> TryAddDeviceAsync(DeviceId Id, string Name, DevicePlatformId PlatformId, DevicePoolId PoolId, bool? Enabled, string? Address, string? ModelId)
+		{
+			DeviceDocument NewDevice = new DeviceDocument(Id, PlatformId, PoolId, Name, Enabled ?? true, Address, ModelId);
+			await Devices.InsertOneAsync(NewDevice);
+			return NewDevice;
+		}
+
+
+		/// <inheritdoc/>
+		public async Task<IDevicePlatform?> TryAddPlatformAsync(DevicePlatformId Id, string Name)
+		{
+			DevicePlatformDocument NewPlatform = new DevicePlatformDocument(Id, Name);
+
+			try
+			{
+				await Platforms.InsertOneAsync(NewPlatform);
+				return NewPlatform;
+			}
+			catch (MongoWriteException Ex)
+			{
+				if (Ex.WriteError.Category == ServerErrorCategory.DuplicateKey)
+				{
+					return null;
+				}
+				else
+				{
+					throw;
+				}
+			}
+		}
+
+		/// <inheritdoc/>
+		public async Task<List<IDevicePlatform>> FindAllPlatformsAsync()
+		{
+			List<DevicePlatformDocument> Results = await Platforms.Find(x => true).ToListAsync();
+			return Results.OrderBy(x => x.Name).Select<DevicePlatformDocument, IDevicePlatform>(x => x).ToList();
+		}
+
+		/// <inheritdoc/>
+		public async Task<bool> UpdatePlatformAsync(DevicePlatformId PlatformId, string[]? ModelIds)
+		{
+			UpdateDefinitionBuilder<DevicePlatformDocument> UpdateBuilder = Builders<DevicePlatformDocument>.Update;
+
+			List<UpdateDefinition<DevicePlatformDocument>> Updates = new List<UpdateDefinition<DevicePlatformDocument>>();
+
+			if (ModelIds != null)
+			{
+				Updates.Add(UpdateBuilder.Set(x => x.Models, ModelIds.ToList()));
+			}
+
+			if (Updates.Count > 0)
+			{
+				await Platforms.FindOneAndUpdateAsync<DevicePlatformDocument>(x => x.Id == PlatformId, UpdateBuilder.Combine(Updates));
+			}
+
+			return true;
+
+		}
+
+		/// <inheritdoc/>
+		public async Task<IDevicePool?> TryAddPoolAsync(DevicePoolId Id, string Name)
+		{
+			DevicePoolDocument NewPool = new DevicePoolDocument(Id, Name);
+
+			try
+			{
+				await Pools.InsertOneAsync(NewPool);
+				return NewPool;
+			}
+			catch (MongoWriteException Ex)
+			{
+				if (Ex.WriteError.Category == ServerErrorCategory.DuplicateKey)
+				{
+					return null;
+				}
+				else
+				{
+					throw;
+				}
+			}
+		}
+
+		/// <inheritdoc/>
+		public async Task<List<IDevice>> FindAllDevicesAsync(List<DeviceId>? DeviceIds = null)
+		{
+			FilterDefinition<DeviceDocument> Filter = Builders<DeviceDocument>.Filter.Empty;
+			if (DeviceIds != null)
+			{
+				Filter &= Builders<DeviceDocument>.Filter.In(x => x.Id, DeviceIds);
+			}
+			List<DeviceDocument> Results = await Devices.Find(Filter).ToListAsync();
+			return Results.OrderBy(x => x.Name).Select<DeviceDocument, IDevice>(x => x).ToList();
+		}
+
+		/// <inheritdoc/>
+		public async Task<List<IDevicePool>> FindAllPoolsAsync()
+		{
+			List<DevicePoolDocument> Results = await Pools.Find(x => true).ToListAsync();
+			return Results.OrderBy(x => x.Name).Select<DevicePoolDocument, IDevicePool>(x => x).ToList();
+		}
+
+		/// <inheritdoc/>
+		public async Task<IDevicePlatform?> GetPlatformAsync(DevicePlatformId PlatformId)
+		{
+			return await Platforms.Find<DevicePlatformDocument>(x => x.Id == PlatformId).FirstOrDefaultAsync();
+		}
+
+		/// <inheritdoc/>
+		public async Task<IDevicePool?> GetPoolAsync(DevicePoolId PoolId)
+		{
+			return await Pools.Find<DevicePoolDocument>(x => x.Id == PoolId).FirstOrDefaultAsync();
+		}
+
+		/// <inheritdoc/>
+		public async Task<IDevice?> GetDeviceAsync(DeviceId DeviceId)
+		{
+			return await Devices.Find<DeviceDocument>(x => x.Id == DeviceId).FirstOrDefaultAsync();
+		}
+
+		/// <inheritdoc/>
+		public async Task<IDevice?> GetDeviceByNameAsync(string DeviceName)
+		{
+			return await Devices.Find<DeviceDocument>(x => x.Name == DeviceName).FirstOrDefaultAsync();
+		}
+
+		/// <inheritdoc/>
+		public async Task UpdateDeviceAsync(DeviceId DeviceId, DevicePoolId? NewPoolId, string? NewName, string? NewAddress, string? NewModelId, string? NewNotes, bool? NewEnabled, bool? NewProblem, bool? NewMaintenance)
+		{
+			UpdateDefinitionBuilder<DeviceDocument> UpdateBuilder = Builders<DeviceDocument>.Update;
+
+			List<UpdateDefinition<DeviceDocument>> Updates = new List<UpdateDefinition<DeviceDocument>>();
+
+			DateTime UtcNow = DateTime.UtcNow;
+
+			if (NewPoolId != null)
+			{
+				Updates.Add(UpdateBuilder.Set(x => x.PoolId, NewPoolId.Value));
+			}
+
+			if (NewName != null)
+			{
+				Updates.Add(UpdateBuilder.Set(x => x.Name, NewName));
+			}
+
+			if (NewEnabled != null)
+			{
+				Updates.Add(UpdateBuilder.Set(x => x.Enabled, NewEnabled));
+			}
+
+			if (NewAddress != null)
+			{
+				Updates.Add(UpdateBuilder.Set(x => x.Address, NewAddress));
+			}
+
+			if (string.IsNullOrEmpty(NewModelId))
+			{
+				Updates.Add(UpdateBuilder.Set(x => x.ModelId, null));
+			}
+			else
+			{
+				Updates.Add(UpdateBuilder.Set(x => x.ModelId, NewModelId));
+			}
+
+			if (NewNotes != null)
+			{
+				Updates.Add(UpdateBuilder.Set(x => x.Notes, NewNotes));
+			}
+
+			if (NewProblem.HasValue)
+			{
+				DateTime? ProblemTime = null;
+				if (NewProblem.Value)
+				{
+					ProblemTime = UtcNow;
+				}
+
+				Updates.Add(UpdateBuilder.Set(x => x.ProblemTimeUtc, ProblemTime));
+			}
+
+			if (NewMaintenance.HasValue)
+			{
+				DateTime? MaintenanceTime = null;
+				if (NewMaintenance.Value)
+				{
+					MaintenanceTime = UtcNow;
+				}
+
+				Updates.Add(UpdateBuilder.Set(x => x.MaintenanceTimeUtc, MaintenanceTime));
+			}
+
+
+			if (Updates.Count > 0)
+			{
+				await Devices.FindOneAndUpdateAsync<DeviceDocument>(x => x.Id == DeviceId, UpdateBuilder.Combine(Updates));
+			}
+		}
+
+		/// <inheritdoc/>
+		public async Task<bool> DeleteDeviceAsync(DeviceId DeviceId)
+		{
+			FilterDefinition<DeviceDocument> Filter = Builders<DeviceDocument>.Filter.Eq(x => x.Id, DeviceId);
+			DeleteResult Result = await Devices.DeleteOneAsync(Filter);
+			return Result.DeletedCount > 0;
+
+		}
+
+		/// <inheritdoc/>
+		public async Task<List<IDeviceReservation>> FindAllDeviceReservationsAsync(DevicePoolId? PoolId = null)
+		{
+			return await Reservations.Find(A => PoolId == null || A.PoolId == PoolId).ToListAsync<DeviceReservationDocument, IDeviceReservation>();
+		}
+
+		/// <inheritdoc/>
+		public async Task<IDeviceReservation?> TryAddReservationAsync(DevicePoolId PoolId, List<DeviceRequestData> Request, string? Hostname, string? ReservationDetails)
+		{
+
+			if (Request.Count == 0)
+			{
+				return null;
+			}
+
+
+			HashSet<DeviceId> Allocated = new HashSet<DeviceId>();
+
+			List<DeviceReservationDocument> PoolReservations = await Reservations.Find(x => x.PoolId == PoolId).ToListAsync();
+
+			DateTime UtcNow = DateTime.UtcNow;
+
+			// Get available devices
+			List<DeviceDocument> PoolDevices = await Devices.Find(x =>
+				x.PoolId == PoolId &&
+				x.Enabled &&				
+				x.MaintenanceTimeUtc == null).ToListAsync();
+
+			// filter out problem devices
+			PoolDevices = PoolDevices.FindAll(x => (x.ProblemTimeUtc == null || ((UtcNow - x.ProblemTimeUtc).Value.TotalMinutes > 30)));			
+
+			// filter out currently reserved devices
+			PoolDevices = PoolDevices.FindAll(x => PoolReservations.FirstOrDefault(p => p.Devices.Contains(x.Id)) == null);
+
+			// sort to use last reserved first to cycle devices
+			PoolDevices.Sort((A, B) => {
+				DateTime? ATime = A.ReservationTimeUtc;
+				DateTime? BTime = B.ReservationTimeUtc;
+
+				if (ATime == BTime)
+				{
+					return 0;
+				}
+
+				if (ATime == null && BTime != null)
+				{
+					return -1;
+				}
+				if (ATime != null && BTime == null)
+				{
+					return 1;
+				}
+
+				return ATime < BTime ? -1 : 1;
+
+			});
+
+			foreach (DeviceRequestData Data in Request)
+			{
+				DeviceDocument Device = PoolDevices.FirstOrDefault(A =>
+				{
+
+					if (Allocated.Contains(A.Id) || A.PlatformId != Data.PlatformId)
+					{
+						return false;
+					}
+
+					if (Data.IncludeModels.Count > 0 && (A.ModelId == null || !Data.IncludeModels.Contains(A.ModelId)))
+					{
+						return false;
+					}
+
+					if (Data.ExcludeModels.Count > 0 && (A.ModelId != null && Data.ExcludeModels.Contains(A.ModelId)))
+					{
+						return false;
+					}
+
+
+					return true;
+				});
+
+				if (Device == null)
+				{
+					// can't fulfill request
+					return null;
+				}
+
+				Allocated.Add(Device.Id);
+			}
+
+			// update reservation time for allocated devices
+			await Devices.UpdateManyAsync(Builders<DeviceDocument>.Filter.In(x => x.Id, Allocated), Builders<DeviceDocument>.Update.Set(x => x.ReservationTimeUtc, UtcNow));
+
+			// Create new reservation
+			DeviceReservationDocument NewReservation = new DeviceReservationDocument(ObjectId.GenerateNewId(), PoolId, Allocated.ToList(), DateTime.UtcNow, Hostname, ReservationDetails);
+			await Reservations.InsertOneAsync(NewReservation);
+			return NewReservation;
+
+		}
+
+		/// <inheritdoc/>
+		public async Task<bool> TryUpdateReservationAsync(ObjectId Id)
+		{
+			UpdateResult Result = await Reservations.UpdateOneAsync(x => x.Id == Id, Builders<DeviceReservationDocument>.Update.Set(x => x.UpdateTimeUtc, DateTime.UtcNow));
+			return Result.ModifiedCount == 1;
+		}
+
+		/// <inheritdoc/>
+		public async Task<bool> DeleteReservationAsync(ObjectId Id)
+		{
+			FilterDefinition<DeviceReservationDocument> Filter = Builders<DeviceReservationDocument>.Filter.Eq(x => x.Id, Id);
+			DeleteResult Result = await Reservations.DeleteOneAsync(Filter);
+			return Result.DeletedCount > 0;
+		}
+
+		/// <summary>
+		/// Deletes expired reservations
+		/// </summary>
+		public async Task<bool> ExpireReservationsAsync()
+		{
+			List<IDeviceReservation> Reserves = await Reservations.Find(A => true).ToListAsync<DeviceReservationDocument, IDeviceReservation>();
+		
+			DateTime UtcNow = DateTime.UtcNow;
+
+			Reserves = Reserves.FindAll(R => ( UtcNow - R.UpdateTimeUtc ).TotalMinutes > 10).ToList();			
+
+			if (Reserves.Count > 0)
+			{
+				DeleteResult Result = await Reservations.DeleteManyAsync(Builders<DeviceReservationDocument>.Filter.In(x => x.Id, Reserves.Select(y => y.Id)));
+				return (Result.DeletedCount > 9);
+			}			
+
+			return false;
+
+		}
+
+		/// <inheritdoc/>
+		public async Task<List<IDeviceReservation>> FindAllReservationsAsync()
+		{
+			List<DeviceReservationDocument> Results = await Reservations.Find(x => true).ToListAsync();
+			return Results.OrderBy(x => x.CreateTimeUtc.Ticks).Select<DeviceReservationDocument, IDeviceReservation>(x => x).ToList();
+		}
+
+
+		/// <inheritdoc/>
+		public async Task<IDeviceReservation?> TryGetReservationFromLegacyGuidAsync(string LegacyGuid)
+		{
+			return await Reservations.Find<DeviceReservationDocument>(x => x.LegacyGuid == LegacyGuid).FirstOrDefaultAsync();
+		}
+
+	}
+
+}
+
diff --git a/Engine/Source/Programs/Horde/HordeServer/Controllers/DevicesController.cs b/Engine/Source/Programs/Horde/HordeServer/Controllers/DevicesController.cs
new file mode 100644
--- /dev/null
+++ b/Engine/Source/Programs/Horde/HordeServer/Controllers/DevicesController.cs
@@ -0,0 +1,828 @@
+// Copyright Epic Games, Inc. All Rights Reserved.
+
+using HordeServer.Services;
+using System.Threading.Tasks;
+
+using HordeServer.Api;
+using HordeServer.Models;
+using MongoDB.Driver;
+using Microsoft.AspNetCore.Mvc;
+using HordeServer.Collections;
+using System;
+using System.Linq;
+
+using Microsoft.Extensions.Logging;
+using MongoDB.Bson;
+using System.Collections.Generic;
+
+using DeviceId = HordeServer.Utilities.StringId<HordeServer.Models.IDevice>;
+using DevicePlatformId = HordeServer.Utilities.StringId<HordeServer.Models.IDevicePlatform>;
+using DevicePoolId = HordeServer.Utilities.StringId<HordeServer.Models.IDevicePool>;
+using Microsoft.AspNetCore.Authorization;
+using System.Text.Json;
+
+
+namespace HordeServer.Controllers
+{
+	/// <summary>
+	/// Controller for device service
+	/// </summary>
+	public class DevicesController : ControllerBase
+	{
+
+		/// <summary>
+		/// Singleton instance of the device service
+		/// </summary>
+		DeviceService DeviceService;
+
+		/// <summary>
+		///  Logger for controller
+		/// </summary>
+		private readonly ILogger<DevicesController> Logger;
+
+		/// <summary>
+		/// Constructor
+		/// </summary>
+		public DevicesController(DeviceService DeviceService, ILogger<DevicesController> Logger)
+		{
+			this.DeviceService = DeviceService;
+			this.Logger = Logger;
+		}
+
+		// DEVICES
+
+		/// <summary>
+		/// Create a new device
+		/// </summary>
+		[HttpPost]
+		[Authorize]
+		[Route("/api/v2/devices")]
+		public async Task<ActionResult<CreateDeviceResponse>> CreateDeviceAsync([FromBody] CreateDeviceRequest DeviceRequest)
+		{
+
+			if (!await DeviceService.AuthorizeAsync(AclAction.DeviceWrite, User))
+			{
+				return Forbid();
+			}
+
+			IDevicePlatform? Platform = await DeviceService.GetPlatformAsync(new DevicePlatformId(DeviceRequest.PlatformId!));
+
+			if (Platform == null)
+			{
+				return BadRequest($"Bad platform id {DeviceRequest.PlatformId} on request");
+			}
+
+			IDevicePool? Pool = await DeviceService.GetPoolAsync(new DevicePoolId(DeviceRequest.PoolId!));
+
+			if (Pool == null)
+			{
+				return BadRequest($"Bad pool id {DeviceRequest.PoolId} on request");
+			}
+
+
+			string? ModelId = null;
+
+			if (!string.IsNullOrEmpty(DeviceRequest.ModelId))
+			{
+				ModelId = new string(DeviceRequest.ModelId);
+
+				if (Platform.Models.FirstOrDefault(x => x == ModelId) == null)
+				{
+					return BadRequest($"Bad model id {ModelId} for platform {Platform.Id} on request");
+				}
+			}
+
+			string Name = DeviceRequest.Name.Trim();
+			string? Address = DeviceRequest.Address?.Trim();
+
+			IDevice? Device = await DeviceService.TryCreateDeviceAsync(DeviceId.Sanitize(Name), Name, Platform.Id, Pool.Id, DeviceRequest.Enabled, Address, ModelId);
+
+			if (Device == null)
+			{
+				return BadRequest($"Unable to create device");
+			}
+
+			return new CreateDeviceResponse(Device.Id.ToString());
+		}
+
+		/// <summary>
+		/// Get list of devices
+		/// </summary>
+		[HttpGet]
+		[Authorize]
+		[Route("/api/v2/devices")]
+		[ProducesResponseType(typeof(List<GetDeviceResponse>), 200)]
+		public async Task<ActionResult<List<object>>> GetDevicesAsync()
+		{
+
+			if (!await DeviceService.AuthorizeAsync(AclAction.DeviceRead, User))
+			{
+				return Forbid();
+			}
+
+			List<IDevice> Devices = await DeviceService.GetDevicesAsync();
+
+			List<object> Responses = new List<object>();
+
+			foreach (IDevice Device in Devices)
+			{
+				Responses.Add(new GetDeviceResponse(Device.Id.ToString(), Device.PlatformId.ToString(), Device.PoolId.ToString(), Device.Name, Device.Enabled, Device.Address, Device.ModelId?.ToString(), Device.Notes, Device.ProblemTimeUtc, Device.MaintenanceTimeUtc));
+			}
+
+			return Responses;
+		}
+
+		/// <summary>
+		/// Get a specific device
+		/// </summary>
+		[HttpGet]
+		[Authorize]
+		[Route("/api/v2/devices/{DeviceId}")]
+		[ProducesResponseType(typeof(GetDeviceResponse), 200)]
+		public async Task<ActionResult<object>> GetDeviceAsync(string DeviceId)
+		{
+
+			if (!await DeviceService.AuthorizeAsync(AclAction.DeviceRead, User))
+			{
+				return Forbid();
+			}
+
+			DeviceId DeviceIdValue = new DeviceId(DeviceId);
+
+			IDevice? Device = await DeviceService.GetDeviceAsync(DeviceIdValue);
+
+			if (Device == null)
+			{
+				return BadRequest($"Unable to find device with id {DeviceId}");
+			}
+
+			return new GetDeviceResponse(Device.Id.ToString(), Device.PlatformId.ToString(), Device.PoolId.ToString(), Device.Name, Device.Enabled, Device.Address, Device.ModelId?.ToString(), Device.Notes, Device.ProblemTimeUtc, Device.MaintenanceTimeUtc);
+		}
+
+		/// <summary>
+		/// Update a specific device
+		/// </summary>
+		[HttpPut]
+		[Authorize]
+		[Route("/api/v2/devices/{DeviceId}")]
+		[ProducesResponseType(typeof(List<GetDeviceResponse>), 200)]
+		public async Task<ActionResult> UpdateDeviceAsync(string DeviceId, [FromBody] UpdateDeviceRequest Update)
+		{
+			if (!await DeviceService.AuthorizeAsync(AclAction.DeviceWrite, User))
+			{
+				return Forbid();
+			}
+
+			DeviceId DeviceIdValue = new DeviceId(DeviceId);
+
+			IDevice? Device = await DeviceService.GetDeviceAsync(DeviceIdValue);
+
+			if (Device == null)
+			{
+				return BadRequest($"Device with id ${DeviceId} does not exist");
+			}
+
+			IDevicePlatform? Platform = await DeviceService.GetPlatformAsync(Device.PlatformId);
+
+			if (Platform == null)
+			{
+				return BadRequest($"Platform id ${Device.PlatformId} does not exist");
+			}
+
+			DevicePoolId? PoolIdValue = null;
+
+			if (Update.PoolId != null)
+			{
+				PoolIdValue = new DevicePoolId(Update.PoolId);
+			}
+
+			string? ModelIdValue = null;
+
+			if (Update.ModelId != null)
+			{
+				ModelIdValue = new string(Update.ModelId);
+
+				if (Platform.Models.FirstOrDefault(x => x == ModelIdValue) == null)
+				{
+					return BadRequest($"Bad model id {Update.ModelId} for platform {Platform.Id} on request");
+				}
+
+			}
+
+			string? Name = Update.Name?.Trim();
+			string? Address = Update.Address?.Trim();
+
+			await DeviceService.UpdateDeviceAsync(DeviceIdValue, PoolIdValue, Name, Address, ModelIdValue, Update.Notes, Update.Enabled, Update.Problem, Update.Maintenance);
+
+			return Ok();
+		}
+
+		/// <summary>
+		/// Delete a specific device
+		/// </summary>
+		[HttpDelete]
+		[Authorize]
+		[Route("/api/v2/devices/{DeviceId}")]
+		public async Task<ActionResult> DeleteDeviceAsync(string DeviceId)
+		{
+			if (!await DeviceService.AuthorizeAsync(AclAction.DeviceWrite, User))
+			{
+				return Forbid();
+			}
+
+			DeviceId DeviceIdValue = new DeviceId(DeviceId);
+
+			IDevice? Device = await DeviceService.GetDeviceAsync(DeviceIdValue);
+			if (Device == null)
+			{
+				return NotFound();
+			}
+
+			await DeviceService.DeleteDeviceAsync(DeviceIdValue);
+			return Ok();
+		}
+
+		// PLATFORMS
+
+		/// <summary>
+		/// Create a new device platform
+		/// </summary>
+		[HttpPost]
+		[Authorize]
+		[Route("/api/v2/devices/platforms")]
+		public async Task<ActionResult<CreateDevicePlatformResponse>> CreatePlatformAsync([FromBody] CreateDevicePlatformRequest Request)
+		{
+			if (!await DeviceService.AuthorizeAsync(AclAction.DeviceWrite, User))
+			{
+				return Forbid();
+			}
+
+			string Name = Request.Name.Trim();
+
+			IDevicePlatform? Platform = await DeviceService.TryCreatePlatformAsync(DevicePlatformId.Sanitize(Name), Name);
+
+			if (Platform == null)
+			{
+				return BadRequest($"Unable to create platform for {Name}");
+			}
+
+			return new CreateDevicePlatformResponse(Platform.Id.ToString());
+
+
+		}
+
+		/// <summary>
+		/// Update a device platform
+		/// </summary>
+		[HttpPut]
+		[Authorize]
+		[Route("/api/v2/devices/platforms/{PlatformId}")]
+		public async Task<ActionResult<CreateDevicePlatformResponse>> UpdatePlatformAsync(string PlatformId, [FromBody] UpdateDevicePlatformRequest Request)
+		{
+			if (!await DeviceService.AuthorizeAsync(AclAction.DeviceWrite, User))
+			{
+				return Forbid();
+			}
+
+
+			await DeviceService.UpdatePlatformAsync(new DevicePlatformId(PlatformId), Request.ModelIds);
+
+			return Ok();
+
+
+		}
+
+		/// <summary>
+		/// Get a list of supported device platforms
+		/// </summary>
+		[HttpGet]
+		[Authorize]
+		[Route("/api/v2/devices/platforms")]
+		[ProducesResponseType(typeof(List<GetDevicePlatformResponse>), 200)]
+		public async Task<ActionResult<List<object>>> GetDevicePlatformsAsync()
+		{
+
+			if (!await DeviceService.AuthorizeAsync(AclAction.DeviceRead, User))
+			{
+				return Forbid();
+			}
+
+			List<IDevicePlatform> Platforms = await DeviceService.GetPlatformsAsync();
+
+			List<object> Responses = new List<object>();
+
+			foreach (IDevicePlatform Platform in Platforms)
+			{
+				// @todo: ACL per platform
+				Responses.Add(new GetDevicePlatformResponse(Platform.Id.ToString(), Platform.Name, Platform.Models.ToArray()));
+			}
+
+			return Responses;
+		}
+
+		// POOLS
+
+		/// <summary>
+		/// Create a new device pool
+		/// </summary>
+		[HttpPost]
+		[Authorize]
+		[Route("/api/v2/devices/pools")]
+		public async Task<ActionResult<CreateDevicePoolResponse>> CreatePoolAsync([FromBody] CreateDevicePoolRequest Request)
+		{
+			if (!await DeviceService.AuthorizeAsync(AclAction.DeviceWrite, User))
+			{
+				return Forbid();
+			}
+
+			string Name = Request.Name.Trim();
+
+			IDevicePool? Pool = await DeviceService.TryCreatePoolAsync(DevicePoolId.Sanitize(Name), Name);
+
+			if (Pool == null)
+			{
+				return BadRequest($"Unable to create pool for {Request.Name}");
+			}
+
+			return new CreateDevicePoolResponse(Pool.Id.ToString());
+
+
+		}
+
+		/// <summary>
+		/// Get a list of existing device pools
+		/// </summary>
+		[HttpGet]
+		[Authorize]
+		[Route("/api/v2/devices/pools")]
+		[ProducesResponseType(typeof(List<GetDevicePoolResponse>), 200)]
+		public async Task<ActionResult<List<object>>> GetDevicePoolsAsync()
+		{
+
+			if (!await DeviceService.AuthorizeAsync(AclAction.DeviceRead, User))
+			{
+				return Forbid();
+			}
+
+			List<IDevicePool> Pools = await DeviceService.GetPoolsAsync();
+
+			List<object> Responses = new List<object>();
+
+			foreach (IDevicePool Pool in Pools)
+			{
+				// @todo: ACL per platform
+				Responses.Add(new GetDevicePoolResponse(Pool.Id.ToString(), Pool.Name));
+			}
+
+			return Responses;
+		}
+
+		// RESERVATIONS
+
+		/// <summary>
+		/// Create a new device reservation
+		/// </summary>
+		[HttpPost]
+		[Authorize]
+		[Route("/api/v2/devices/reservations")]
+		[ProducesResponseType(typeof(CreateDeviceReservationResponse), 200)]
+		public async Task<ActionResult<CreateDeviceReservationResponse>> CreateDeviceReservation([FromBody] CreateDeviceReservationRequest Request)
+		{
+
+			if (!await DeviceService.AuthorizeAsync(AclAction.DeviceWrite, User))
+			{
+				return Forbid();
+			}
+
+			List<IDevicePool> Pools = await DeviceService.GetPoolsAsync();
+			List<IDevicePlatform> Platforms = await DeviceService.GetPlatformsAsync();
+
+			DevicePoolId PoolIdValue = new DevicePoolId(Request.PoolId);
+			IDevicePool? Pool = Pools.FirstOrDefault(x => x.Id == PoolIdValue);
+			if (Pool == null)
+			{
+				return BadRequest($"Unknown pool {Request.PoolId} on device reservation request");
+			}
+
+			List<DeviceRequestData> RequestedDevices = new List<DeviceRequestData>();
+
+			foreach (DeviceReservationRequest DeviceRequest in Request.Devices)
+			{
+				DevicePlatformId PlatformIdValue = new DevicePlatformId(DeviceRequest.PlatformId);
+				IDevicePlatform? Platform = Platforms.FirstOrDefault(x => x.Id == PlatformIdValue);
+
+				if (Platform == null)
+				{
+					return BadRequest($"Unknown platform {DeviceRequest.PlatformId} on device reservation request");
+				}
+
+				if (DeviceRequest.IncludeModels != null)
+				{
+					foreach (string Model in DeviceRequest.IncludeModels)
+					{
+						if (Platform.Models.FirstOrDefault(x => x == Model) == null)
+						{
+							return BadRequest($"Unknown model {Model} for platform {DeviceRequest.PlatformId} on device reservation request");
+						}
+
+					}
+
+				}
+
+				if (DeviceRequest.ExcludeModels != null)
+				{
+					foreach (string Model in DeviceRequest.ExcludeModels)
+					{
+						if (Platform.Models.FirstOrDefault(x => x == Model) == null)
+						{
+							return BadRequest($"Unknown model {Model} for platform {DeviceRequest.PlatformId} on device reservation request");
+						}
+
+					}
+
+				}					
+
+				RequestedDevices.Add(new DeviceRequestData(PlatformIdValue, DeviceRequest.IncludeModels, DeviceRequest.ExcludeModels));
+
+			}
+
+			IDeviceReservation? Reservation = await DeviceService.TryCreateReservationAsync(PoolIdValue, RequestedDevices);
+
+			if (Reservation == null)
+			{
+				return Conflict("Unable to allocated devices for reservation");
+			}
+
+			List<IDevice> Devices = await DeviceService.GetDevicesAsync(Reservation.Devices);
+
+
+			CreateDeviceReservationResponse Response = new CreateDeviceReservationResponse();
+
+			Response.Id = Reservation.Id.ToString();
+
+			foreach (IDevice Device in Devices)
+			{
+				Response.Devices.Add(new GetDeviceResponse(Device.Id.ToString(), Device.PlatformId.ToString(), Device.PoolId.ToString(), Device.Name, Device.Enabled, Device.Address, Device.ModelId?.ToString(), Device.Notes, Device.ProblemTimeUtc, Device.MaintenanceTimeUtc));
+			}
+
+			return Response;
+		}
+
+		/// <summary>
+		/// Get active device reservations
+		/// </summary>		
+		[HttpGet]
+		[Authorize]
+		[Route("/api/v2/devices/reservations")]
+		[ProducesResponseType(typeof(List<GetDeviceReservationResponse>), 200)]
+		public async Task<ActionResult<List<GetDeviceReservationResponse>>> GetDeviceReservations()
+		{
+			if (!await DeviceService.AuthorizeAsync(AclAction.DeviceRead, User))
+			{
+				return Forbid();
+			}
+
+			List<IDeviceReservation> Reservations = await DeviceService.GetReservationsAsync();
+
+			List<GetDeviceReservationResponse> Response = new List<GetDeviceReservationResponse>();
+
+			foreach (IDeviceReservation Reservation in Reservations)
+			{
+				GetDeviceReservationResponse ReservationResponse = new GetDeviceReservationResponse()
+				{
+					Id = Reservation.Id.ToString(),
+					PoolId = Reservation.PoolId.ToString(),
+					Devices = Reservation.Devices.Select(x => x.ToString()).ToList(),
+					JobId = Reservation.JobId?.ToString(),
+					StepId = Reservation.StepId?.ToString(),
+					UserId = Reservation.UserId?.ToString(),
+					Hostname = Reservation.Hostname,
+					ReservationDetails = Reservation.ReservationDetails,
+					CreateTimeUtc = Reservation.CreateTimeUtc,
+					LegacyGuid = Reservation.LegacyGuid
+				};
+
+				Response.Add(ReservationResponse);
+
+			}
+
+			return Response;
+		}
+
+		/// <summary>
+		/// Renew an existing reservation
+		/// </summary>
+		[HttpPut]
+		[Authorize]
+		[Route("/api/v2/devices/reservations/{ReservationId}")]
+		public async Task<ActionResult> UpdateReservationAsync(string ReservationId)
+		{
+			if (!await DeviceService.AuthorizeAsync(AclAction.DeviceWrite, User))
+			{
+				return Forbid();
+			}
+
+			ObjectId ReservationIdValue = new ObjectId(ReservationId);
+
+			bool Updated = await DeviceService.TryUpdateReservationAsync(ReservationIdValue);
+
+			if (!Updated)
+			{
+				return BadRequest("Failed to update reservation");
+			}
+
+			return Ok();
+		}
+
+		/// <summary>
+		/// Delete a reservation
+		/// </summary>
+		[HttpDelete]
+		[Authorize]
+		[Route("/api/v2/devices/reservations/{ReservationId}")]
+		public async Task<ActionResult> DeleteReservationAsync(string ReservationId)
+		{
+			if (!await DeviceService.AuthorizeAsync(AclAction.DeviceWrite, User))
+			{
+				return Forbid();
+			}
+
+			ObjectId ReservationIdValue = new ObjectId(ReservationId);
+
+			bool Deleted = await DeviceService.DeleteReservationAsync(ReservationIdValue);
+
+			if (!Deleted)
+			{
+				return BadRequest("Failed to delete reservation");
+			}
+
+			return Ok();
+		}
+
+		// LEGACY V1 API
+
+		enum LegacyPerfSpec
+		{
+			Unspecified,
+			Minimum,
+			Recommended,
+			High
+		};
+
+		static Dictionary<string, DevicePlatformId> PlatformMap = new Dictionary<string, DevicePlatformId>()
+			{
+				{ "PS4", new DevicePlatformId("ps4") },
+				{ "PS4-DevKit", new DevicePlatformId("ps4") },
+				{ "XboxOne", new DevicePlatformId("xboxone") },
+				{ "XboxOneGDK", new DevicePlatformId("xboxone") },
+				{ "XboxOne-DevKit", new DevicePlatformId("xboxone") },
+				{ "Switch", new DevicePlatformId("switch") },
+				{ "XSX", new DevicePlatformId("xsx") },
+				{ "PS5", new DevicePlatformId("ps5") }
+			};
+
+		static Dictionary<DevicePlatformId, string> PlatformReverseMap = new Dictionary<DevicePlatformId, string>()
+			{
+				{ new DevicePlatformId("ps4"),  "PS4-DevKit"},
+				{ new DevicePlatformId("ps5"),  "PS5"},
+				{ new DevicePlatformId("xboxone"),  "XBoxOneGDK"},
+				{ new DevicePlatformId("xsx"),  "XSX"},
+				{ new DevicePlatformId("switch"),  "Switch"},
+			};
+
+		static Dictionary<DevicePlatformId, string> PerfSpecHighMap = new Dictionary<DevicePlatformId, string>()
+			{
+				{ new DevicePlatformId("ps4"),  "PS4Pro"},
+				{ new DevicePlatformId("xboxone"),  "XboxOneX"}
+			};
+
+
+		/// <summary>
+		/// Create a device reservation
+		/// </summary>
+		[HttpPost]
+		[Route("/api/v1/reservations")]
+		[ProducesResponseType(typeof(GetLegacyReservationResponse), 200)]
+		public async Task<ActionResult<GetLegacyReservationResponse>> CreateDeviceReservationV1Async([FromBody] LegacyCreateReservationRequest Request)
+		{
+
+			List<IDevicePool> Pools = await DeviceService.GetPoolsAsync();
+			List<IDevicePlatform> Platforms = await DeviceService.GetPlatformsAsync();
+
+			//@todo: pull this fromm request
+			string PoolId = "fnue5";
+
+			DevicePoolId PoolIdValue = new DevicePoolId(PoolId);
+			IDevicePool? Pool = Pools.FirstOrDefault(x => x.Id == PoolIdValue);
+			if (Pool == null)
+			{
+				return BadRequest($"Unknown pool {PoolId} on device reservation request");
+			}
+
+			List<DeviceRequestData> RequestedDevices = new List<DeviceRequestData>();
+
+			List<string> PerfSpecs = new List<string>();
+
+			foreach (string DeviceType in Request.DeviceTypes)
+			{
+
+				string PlatformName = DeviceType;
+				string PerfSpecName = "Unspecified";
+
+				if (DeviceType.Contains(":", StringComparison.OrdinalIgnoreCase))
+				{
+					string[] Tokens = DeviceType.Split(":");
+					PlatformName = Tokens[0];
+					PerfSpecName = Tokens[1];
+				}
+
+				PerfSpecs.Add(PerfSpecName);
+
+				DevicePlatformId PlatformId;
+
+				if (!PlatformMap.TryGetValue(PlatformName, out PlatformId))
+				{
+					return BadRequest($"Unknown platform for device type {PlatformName} on device reservation request");
+				}
+
+				IDevicePlatform? Platform = Platforms.FirstOrDefault(x => x.Id == PlatformId);
+
+				if (Platform == null)
+				{
+					return BadRequest($"Unknown platform {PlatformId} on device reservation request");
+				}
+
+				List<string> IncludeModels = new List<string>();
+				List<string> ExcludeModels = new List<string>();
+
+				if (PerfSpecName == "High")
+				{
+					string? Model = null;
+					if (PerfSpecHighMap.TryGetValue(PlatformId, out Model))
+					{
+						IncludeModels.Add(Model);
+					}
+				}
+
+				if (PerfSpecName == "Minimum" || PerfSpecName == "Recommended")
+				{
+					string? Model = null;
+					if (PerfSpecHighMap.TryGetValue(PlatformId, out Model))
+					{
+						ExcludeModels.Add(Model);
+					}
+				}
+
+
+				RequestedDevices.Add(new DeviceRequestData(PlatformId, IncludeModels, ExcludeModels));
+
+			}
+
+			IDeviceReservation? Reservation = await DeviceService.TryCreateReservationAsync(PoolIdValue, RequestedDevices, Request.Hostname, Request.ReservationDetails);
+
+			if (Reservation == null)
+			{
+				return Conflict("Unable to allocated devices for reservation");
+			}
+
+			List<IDevice> Devices = await DeviceService.GetDevicesAsync(Reservation.Devices);
+
+			GetLegacyReservationResponse Response = new GetLegacyReservationResponse();
+
+			Response.Guid = Reservation.LegacyGuid;
+			Response.DeviceNames = Devices.Select(x => x.Name).ToArray();
+			Response.DevicePerfSpecs = PerfSpecs.ToArray();
+			Response.HostName = Request.Hostname;
+			Response.StartDateTime = Reservation.CreateTimeUtc.ToString("O", System.Globalization.CultureInfo.InvariantCulture);
+			Response.Duration = $"{Request.Duration}";
+
+			return new JsonResult(Response, new JsonSerializerOptions() { PropertyNamingPolicy = null });
+
+		}
+
+		/// <summary>
+		/// Renew a reservation
+		/// </summary>
+		[HttpPut]
+		[Route("/api/v1/reservations/{ReservationGuid}")]
+		public async Task<ActionResult> UpdateReservationV1Async(string ReservationGuid, [FromBody] string Duration)
+		{
+			IDeviceReservation? Reservation = await DeviceService.TryGetReservationFromLegacyGuidAsync(ReservationGuid);
+
+			if (Reservation == null)
+			{
+				return BadRequest($"Unable to find reservation for guid {ReservationGuid}");
+			}
+
+			bool Updated = await DeviceService.TryUpdateReservationAsync(Reservation.Id);
+
+			if (!Updated)
+			{
+				return BadRequest("Failed to update reservation");
+			}
+
+			Logger.LogDebug($"Ignoring update duration of {Duration}");
+
+			return Ok();
+		}
+
+		/// <summary>
+		/// Delete a reservation
+		/// </summary>
+		[HttpDelete]
+		[Route("/api/v1/reservations/{ReservationGuid}")]
+		public async Task<ActionResult> DeleteReservationV1Async(string ReservationGuid)
+		{
+			IDeviceReservation? Reservation = await DeviceService.TryGetReservationFromLegacyGuidAsync(ReservationGuid);
+
+			if (Reservation == null)
+			{
+				return BadRequest($"Unable to find reservation for guid {ReservationGuid}");
+			}
+
+			bool Deleted = await DeviceService.DeleteReservationAsync(Reservation.Id);
+
+			if (!Deleted)
+			{
+				return BadRequest("Failed to delete reservation");
+			}
+
+			return Ok();
+		}
+
+
+		/// <summary>
+		/// Get device info for a reserved device
+		/// </summary>
+		[HttpGet]
+		[Route("/api/v1/devices/{DeviceName}")]
+		[ProducesResponseType(typeof(GetLegacyDeviceResponse), 200)]
+		public async Task<ActionResult<GetLegacyDeviceResponse>> GetDeviceV1Async(string DeviceName)
+		{
+			IDevice? Device = await DeviceService.GetDeviceByNameAsync(DeviceName);
+
+			if (Device == null)
+			{
+				return BadRequest($"Unknown device {DeviceName}");
+			}
+
+			string? Platform;
+			if (!PlatformReverseMap.TryGetValue(Device.PlatformId, out Platform))
+			{
+				return BadRequest($"Unable to map platform for {DeviceName} : {Device.PlatformId}");
+			}
+
+			GetLegacyDeviceResponse Response = new GetLegacyDeviceResponse();
+
+			Response.Name = Device.Name;
+			Response.Type = Platform;
+			Response.IPOrHostName = Device.Address ?? "";
+			Response.AvailableStartTime = "00:00:00";
+			Response.AvailableEndTime = "00:00:00";
+			Response.Enabled = true;
+			Response.DeviceData = "";
+						
+			Response.PerfSpec = "Minimum";
+
+			if (Device.ModelId != null)
+			{
+				string? Model = null;
+				if (PerfSpecHighMap.TryGetValue(Device.PlatformId, out Model))
+				{
+					if (Model == Device.ModelId)
+					{
+						Response.PerfSpec = "High";
+					}
+					
+				}
+			}
+
+			return new JsonResult(Response, new JsonSerializerOptions() { PropertyNamingPolicy = null });
+
+		}
+
+		/// <summary>
+		/// Mark a problem device
+		/// </summary>
+		[HttpPut]
+		[Route("/api/v1/deviceerror/{DeviceName}")]
+		public async Task<ActionResult> PutDeviceErrorAsync(string DeviceName)
+		{
+
+			IDevice? Device = await DeviceService.GetDeviceByNameAsync(DeviceName);
+
+			if (Device == null)
+			{
+				return BadRequest($"Unknown device {DeviceName}");
+			}
+
+			await DeviceService.UpdateDeviceAsync(Device.Id, NewProblem: true);
+
+
+			return Ok();
+
+		}
+
+
+	}
+}
\ No newline at end of file
diff --git a/Engine/Source/Programs/Horde/HordeServer/Models/Acl.cs b/Engine/Source/Programs/Horde/HordeServer/Models/Acl.cs
--- a/Engine/Source/Programs/Horde/HordeServer/Models/Acl.cs
+++ b/Engine/Source/Programs/Horde/HordeServer/Models/Acl.cs
@@ -410,6 +410,19 @@ namespace HordeServer.Models
 		/// Ability to subscribe to notifications
 		/// </summary>
 		CreateSubscription,
+
+		//// DEVICES ////
+		
+		/// <summary>
+		/// Ability to read devices
+		/// </summary>
+		DeviceRead,
+
+		/// <summary>
+		/// Ability to write devices
+		/// </summary>
+		DeviceWrite
+
 	}
 
 	/// <summary>
diff --git a/Engine/Source/Programs/Horde/HordeServer/Models/Device.cs b/Engine/Source/Programs/Horde/HordeServer/Models/Device.cs
new file mode 100644
--- /dev/null
+++ b/Engine/Source/Programs/Horde/HordeServer/Models/Device.cs
@@ -0,0 +1,193 @@
+// Copyright Epic Games, Inc. All Rights Reserved.
+
+using System;
+using MongoDB.Bson;
+using System.Collections.Generic;
+
+using DeviceId = HordeServer.Utilities.StringId<HordeServer.Models.IDevice>;
+using DevicePlatformId = HordeServer.Utilities.StringId<HordeServer.Models.IDevicePlatform>;
+using DevicePoolId = HordeServer.Utilities.StringId<HordeServer.Models.IDevicePool>;
+
+namespace HordeServer.Models
+{
+	
+	/// <summary>
+	/// A reservation containing one or more devices
+	/// </summary>
+	public interface IDeviceReservation
+	{
+		/// <summary>
+		/// Randomly generated unique id for this reservation
+		/// </summary>
+		public ObjectId Id { get; }
+
+		/// <summary>
+		/// Which device pool the reservation is in
+		/// </summary>
+		public DevicePoolId PoolId { get; }
+
+		/// <summary>
+		/// JobID holding reservation
+		/// </summary>
+		public ObjectId? JobId { get; }
+
+		/// <summary>
+		/// Job step id holding reservation
+		/// </summary>
+		public SubResourceId? StepId { get; }
+
+		/// <summary>
+		/// Reservations held by a user, requires a token
+		/// </summary>
+		public ObjectId? UserId { get; }
+
+		/// <summary>
+		/// The hostname of machine holding reservation
+		/// </summary>
+		string? Hostname { get; }
+
+		/// <summary>
+		/// The hostname of machine holding reservation
+		/// </summary>
+		string? ReservationDetails { get; }
+
+		/// <summary>
+		/// The UTC time when the reservation was created
+		/// </summary>
+		DateTime CreateTimeUtc { get; }
+
+		/// <summary>
+		/// The last update time for reservation renewal (for expiration)
+		/// </summary>
+		DateTime UpdateTimeUtc { get; }
+
+		/// <summary>
+		/// The reserved devices
+		/// </summary>
+		public List<DeviceId> Devices { get; }
+
+		/// <summary>
+		/// The legacy reservation system guid, to be removed once can update Gauntlet client in all streams
+		/// </summary>
+		public string LegacyGuid { get; }
+	}
+
+
+	/// <summary>
+	/// A device platform 
+	/// </summary>
+	public interface IDevicePlatform
+	{
+		/// <summary>
+		/// Unique identifier of device platform 
+		/// </summary>
+		public DevicePlatformId Id { get; }
+		
+		/// <summary>
+		/// Device platform name, for example Android, PS5, etc
+		/// </summary>
+		string Name { get; }
+
+		/// <summary>
+		/// A list of valid models for the platform
+		/// </summary>
+		public IReadOnlyList<string>? Models { get; }
+
+	}
+
+
+	/// <summary>
+	/// A logical pool of devices
+	/// </summary>
+	public interface IDevicePool
+	{
+		/// <summary>
+		/// Unique identifier of pool
+		/// </summary>
+		public DevicePoolId Id { get; }
+
+		/// <summary>
+		/// Friendly name of the pool
+		/// </summary>
+		string Name { get; }
+
+		/// <summary>
+		/// Acl access to the pool
+		/// </summary>
+		public Acl? Acl { get; }
+	}
+
+	/// <summary>
+	/// A physical device
+	/// </summary>
+	public interface IDevice
+	{
+		/// <summary>
+		/// The unique id of the device
+		/// </summary>
+		public DeviceId Id { get; }
+
+		/// <summary>
+		/// The platform of the device
+		/// </summary>
+		public DevicePlatformId PlatformId { get; }
+
+		/// <summary>
+		/// Which pool the device belongs to
+		/// </summary>
+		public DevicePoolId PoolId { get; }
+
+		/// <summary>
+		/// Friendly name of the device
+		/// </summary>
+		public string Name { get; }
+
+		/// <summary>
+		/// The model of the device
+		/// </summary>
+		public string? ModelId { get; }
+
+
+		/// <summary>
+		/// Address if the device supports shared network connections
+		/// </summary>
+		public string? Address { get; }
+
+		/// <summary>
+		/// Whether the device is currently enabled
+		/// </summary>
+		public bool Enabled { get; }
+
+		/// <summary>
+		/// Name of the user that last modified this device
+		/// </summary>
+		public string? ModifiedByUser { get; }
+
+		/// <summary>
+		/// The last time a problem was reported
+		/// </summary>
+		public DateTime? ProblemTimeUtc { get; }
+
+		/// <summary>
+		/// The last time a problem was reported
+		/// </summary>
+		public DateTime? MaintenanceTimeUtc { get; }
+
+		/// <summary>
+		/// Markdown notes for device if any 
+		/// </summary>
+		public string? Notes { get; }
+
+		/// <summary>
+		/// ACL for modifying this device
+		/// </summary>
+		public Acl? Acl { get; }
+	}
+}
+
+
+
+
+
+
+
diff --git a/Engine/Source/Programs/Horde/HordeServer/Services/DeviceService.cs b/Engine/Source/Programs/Horde/HordeServer/Services/DeviceService.cs
new file mode 100644
--- /dev/null
+++ b/Engine/Source/Programs/Horde/HordeServer/Services/DeviceService.cs
@@ -0,0 +1,240 @@
+// Copyright Epic Games, Inc. All Rights Reserved.
+
+using System.Threading.Tasks;
+
+using HordeServer.Models;
+using HordeServer.Collections;
+using System;
+using HordeServer.Utilities;
+using System.Threading;
+
+using Microsoft.Extensions.Logging;
+using HordeCommon;
+using MongoDB.Bson;
+using System.Collections.Generic;
+
+using DeviceId = HordeServer.Utilities.StringId<HordeServer.Models.IDevice>;
+using DevicePlatformId = HordeServer.Utilities.StringId<HordeServer.Models.IDevicePlatform>;
+using DevicePoolId = HordeServer.Utilities.StringId<HordeServer.Models.IDevicePool>;
+using System.Security.Claims;
+using System.Diagnostics.CodeAnalysis;
+
+namespace HordeServer.Services
+{
+	/// <summary>
+	/// Device management service
+	/// </summary>
+	public class DeviceService : TickedBackgroundService
+	{
+		/// <summary>
+		/// The ACL service instance
+		/// </summary>
+		AclService AclService;
+
+		/// <summary>
+		/// Log output writer
+		/// </summary>
+		ILogger<DeviceService> Logger;
+
+		/// <summary>
+		/// Device collection
+		/// </summary>
+		IDeviceCollection Devices;
+
+		/// <summary>
+		/// Device service constructor
+		/// </summary>
+		public DeviceService(IDeviceCollection Devices, AclService AclService, ILogger<DeviceService> Logger)
+			: base(TimeSpan.FromMinutes(1.0), Logger)
+		{
+			this.Devices = Devices;
+			this.AclService = AclService;
+			this.Logger = Logger;
+		}
+
+		/// <summary>
+		/// Ticks service
+		/// </summary>
+		protected override async Task TickAsync(CancellationToken StoppingToken)
+		{
+
+			if (!StoppingToken.IsCancellationRequested)
+			{
+				await Devices.ExpireReservationsAsync();
+			}
+
+		}
+
+		/// <summary>
+		/// Create a new device platform
+		/// </summary>
+		public Task<IDevicePlatform?> TryCreatePlatformAsync(DevicePlatformId Id, string Name)
+		{
+			return Devices.TryAddPlatformAsync(Id, Name);
+		}
+
+		/// <summary>
+		/// Get a list of existing device platforms
+		/// </summary>
+		public Task<List<IDevicePlatform>> GetPlatformsAsync()
+		{
+			return Devices.FindAllPlatformsAsync();
+		}
+
+		/// <summary>
+		/// Update an existing platform
+		/// </summary>
+		public Task<bool> UpdatePlatformAsync(DevicePlatformId PlatformId, string[]? ModelIds)
+		{
+			return Devices.UpdatePlatformAsync(PlatformId, ModelIds);
+		}
+
+		/// <summary>
+		/// Get a specific device platform
+		/// </summary>
+		public Task<IDevicePlatform?> GetPlatformAsync(DevicePlatformId Id)
+		{
+			return Devices.GetPlatformAsync(Id);
+		}
+
+		/// <summary>
+		/// Get a device pool by id
+		/// </summary>
+		public Task<IDevicePool?> GetPoolAsync(DevicePoolId Id)
+		{
+			return Devices.GetPoolAsync(Id);
+		}
+
+		/// <summary>
+		/// Create a new device pool
+		/// </summary>
+		public Task<IDevicePool?> TryCreatePoolAsync(DevicePoolId Id, string Name)
+		{
+			return Devices.TryAddPoolAsync(Id, Name);
+		}
+
+		/// <summary>
+		/// Get a list of existing device pools
+		/// </summary>
+		public Task<List<IDevicePool>> GetPoolsAsync()
+		{
+			return Devices.FindAllPoolsAsync();
+		}
+
+		/// <summary>
+		/// Get a list of devices, optionally filtered to provided ids
+		/// </summary>
+		public Task<List<IDevice>> GetDevicesAsync(List<DeviceId>? DeviceIds = null)
+		{
+			return Devices.FindAllDevicesAsync(DeviceIds);
+		}
+
+		/// <summary>
+		/// Get a specific device
+		/// </summary>
+		public Task<IDevice?> GetDeviceAsync(DeviceId Id)
+		{
+			return Devices.GetDeviceAsync(Id);
+		}
+
+		/// <summary>
+		/// Get a device by name
+		/// </summary>
+		public Task<IDevice?> GetDeviceByNameAsync(string DeviceName)
+		{
+			return Devices.GetDeviceByNameAsync(DeviceName);
+		}
+
+		/// <summary>
+		/// Delete a device
+		/// </summary>
+		public async Task<bool> DeleteDeviceAsync(DeviceId Id)
+		{
+			return await Devices.DeleteDeviceAsync(Id);
+		}
+
+		/// <summary>
+		/// Create a new device
+		/// </summary>
+		/// <param name="Id">Unique id of the device</param>
+		/// <param name="Name">Friendly name of the device</param>
+		/// <param name="PlatformId">The device platform</param>
+		/// <param name="PoolId">Which pool to add the device</param>
+		/// <param name="Enabled">Whether the device is enabled</param>
+		/// <param name="Address">Address or hostname of device</param>
+		/// <param name="ModelId">Vendor model id</param>
+		/// <returns></returns>
+		public Task<IDevice?> TryCreateDeviceAsync(DeviceId Id, string Name, DevicePlatformId PlatformId, DevicePoolId PoolId, bool? Enabled, string? Address, string? ModelId)
+		{
+			return Devices.TryAddDeviceAsync(Id, Name, PlatformId, PoolId, Enabled, Address, ModelId);
+		}
+
+		/// <summary>
+		/// Update a device
+		/// </summary>
+		public Task UpdateDeviceAsync(DeviceId DeviceId, DevicePoolId? NewPoolId = null, string? NewName = null, string? NewAddress = null, string? NewModelId = null, string? NewNotes = null, bool? NewEnabled = null, bool? NewProblem = null, bool? NewMaintenance = null)
+		{
+			return Devices.UpdateDeviceAsync(DeviceId, NewPoolId, NewName, NewAddress, NewModelId, NewNotes, NewEnabled, NewProblem, NewMaintenance);
+		}
+
+		/// <summary>
+		/// Try to create a reservation satisfying the specified device platforms and models
+		/// </summary>
+		public Task<IDeviceReservation?> TryCreateReservationAsync(DevicePoolId Pool, List<DeviceRequestData> Request, string? Hostname = null, string? ReservationDetails = null)
+		{
+			return Devices.TryAddReservationAsync(Pool, Request, Hostname, ReservationDetails);
+		}
+
+		/// <summary>
+		/// Update/renew an existing reservation
+		/// </summary>
+		public Task<bool> TryUpdateReservationAsync(ObjectId Id)
+		{
+			return Devices.TryUpdateReservationAsync(Id);
+		}
+
+		/// <summary>
+		///  Delete an existing reservation
+		/// </summary>
+		public Task<bool> DeleteReservationAsync(ObjectId Id)
+		{
+			return Devices.DeleteReservationAsync(Id);
+		}
+
+		/// <summary>
+		/// Get a reservation from a legacy guid
+		/// </summary>
+		public Task<IDeviceReservation?> TryGetReservationFromLegacyGuidAsync(string LegacyGuid)
+		{
+			return Devices.TryGetReservationFromLegacyGuidAsync(LegacyGuid);
+		}
+
+		/// <summary>
+		/// Get a list of existing device reservations
+		/// </summary>
+		public Task<List<IDeviceReservation>> GetReservationsAsync()
+		{
+			return Devices.FindAllReservationsAsync();
+		}
+
+		/// <summary>
+		/// Authorize device action
+		/// </summary>
+		/// <param name="Action"></param>
+		/// <param name="User"></param>
+		/// <returns></returns>
+		[SuppressMessage("Usage", "CA1801:Review unused parameters")]
+		[SuppressMessage("Performance", "CA1822: Can be static ")]
+		public Task<bool> AuthorizeAsync(AclAction Action, ClaimsPrincipal User)
+		{
+			// Setup ACL's for pool and platform access, https://jira.it.epicgames.com/browse/UE-117224/			
+			// allow reads, though restrict writing to internal employees
+			if (Action == AclAction.DeviceRead)
+			{
+				return Task.FromResult(true);
+			}
+			return Task.FromResult(User.IsInRole("Internal-Employees"));			
+		}
+
+	}
+}
\ No newline at end of file
diff --git a/Engine/Source/Programs/Horde/HordeServer/Startup.cs b/Engine/Source/Programs/Horde/HordeServer/Startup.cs
--- a/Engine/Source/Programs/Horde/HordeServer/Startup.cs
+++ b/Engine/Source/Programs/Horde/HordeServer/Startup.cs
@@ -321,6 +321,7 @@ namespace HordeServer
 			Services.AddSingleton<ITelemetryCollection, TelemetryCollection>();
 			Services.AddSingleton<IUgsMetadataCollection, UgsMetadataCollection>();
 			Services.AddSingleton<IUserCollection, UserCollectionV2>();
+			Services.AddSingleton<IDeviceCollection, DeviceCollection>();
 
 			Services.AddSingleton(typeof(ISingletonDocument<>), typeof(SingletonDocument<>));
 
@@ -411,6 +412,8 @@ namespace HordeServer
 			Services.AddSingleton<FileSystemStorageBackend>();
 			Services.AddSingleton<IStorageService, SimpleStorageService>(SP => new SimpleStorageService(SP.GetRequiredService<FileSystemStorageBackend>()));
 
+			Services.AddSingleton<DeviceService>();
+
 			ConfigureStorageProvider(Services, Settings);
 			ConfigureLogStorage(Services);
 			//			ConfigureLogFileWriteCache(Services, Settings);
@@ -514,6 +517,7 @@ namespace HordeServer
 			Services.AddHostedService(Provider => Provider.GetRequiredService<SlackNotificationSink>());
 			Services.AddHostedService<ConfigService>();
 			Services.AddHostedService<TelemetryService>();
+			Services.AddHostedService(Provider => Provider.GetRequiredService<DeviceService>());
 
 			// Task sources. Order of registration is important here; it dictates the order in which sources are served.
 			Services.AddSingleton<JobTaskSource>();
