Horde: Add explicit directories for storing logs and artifacts, and default to C:\ProgramData on Windows.

diff --git a/Engine/Source/Programs/Horde/HordeServer/Program.cs b/Engine/Source/Programs/Horde/HordeServer/Program.cs
--- a/Engine/Source/Programs/Horde/HordeServer/Program.cs
+++ b/Engine/Source/Programs/Horde/HordeServer/Program.cs
@@ -188,6 +188,23 @@ namespace HordeServer
 				return new X509Certificate2(FileReference.ReadAllBytes(FileReference.Combine(AppDir, HordeSettings.ServerPrivateCert)));
 			}
 		}
+
+		/// <summary>
+		/// Gets the default directory for storing application data
+		/// </summary>
+		/// <returns>The default data directory</returns>
+		public static DirectoryReference GetDefaultDataDir()
+		{
+			if (RuntimeInformation.IsOSPlatform(OSPlatform.Windows))
+			{
+				DirectoryReference? Dir = DirectoryReference.GetSpecialFolder(Environment.SpecialFolder.CommonApplicationData);
+				if (Dir != null)
+				{
+					return DirectoryReference.Combine(Dir, "Horde");
+				}
+			}
+			return DirectoryReference.Combine(AppDir, "Data");
+		}
 	}
 }
 
diff --git a/Engine/Source/Programs/Horde/HordeServer/ServerSettings.cs b/Engine/Source/Programs/Horde/HordeServer/ServerSettings.cs
--- a/Engine/Source/Programs/Horde/HordeServer/ServerSettings.cs
+++ b/Engine/Source/Programs/Horde/HordeServer/ServerSettings.cs
@@ -1,11 +1,13 @@
 // Copyright Epic Games, Inc. All Rights Reserved.
 
+using EpicGames.Core;
 using HordeServer.Services;
 using HordeServer.Utilities;
 using System;
 using System.Collections.Generic;
 using System.ComponentModel.DataAnnotations;
 using System.Linq;
+using System.Runtime.InteropServices;
 using System.Threading.Tasks;
 using TimeZoneConverter;
 
@@ -100,14 +102,12 @@ namespace HordeServer
 		/// <summary>
 		/// MongoDB connection string
 		/// </summary>
-		[Required]
-		public string DatabaseConnectionString { get; set; } = null!;
+		public string? DatabaseConnectionString { get; set; }
 
 		/// <summary>
 		/// MongoDB database name
 		/// </summary>
-		[Required]
-		public string DatabaseName { get; set; } = null!;
+		public string DatabaseName { get; set; } = "Horde";
 
 		/// <summary>
 		/// The claim type for administrators
@@ -226,9 +226,24 @@ namespace HordeServer
 		public StorageProviderType ExternalStorageProviderType { get; set; } = StorageProviderType.FileSystem;
 
 		/// <summary>
-		/// local log file storage directory, if using type filesystem
+		/// Local log/artifact storage directory, if using type filesystem
+		/// </summary>
+		public string LocalLogsDir { get; set; } = "Logs";
+
+		/// <summary>
+		/// Gets the full path referred to by LocalStorageDir
+		/// </summary>
+		public DirectoryReference LocalLogsDirRef => DirectoryReference.Combine(Program.GetDefaultDataDir(), LocalLogsDir);
+
+		/// <summary>
+		/// Local artifact storage directory, if using type filesystem
+		/// </summary>
+		public string LocalArtifactsDir { get; set; } = "Artifacts";
+
+		/// <summary>
+		/// Gets the full path referred to by LocalStorageDir
 		/// </summary>
-		public string LocalStorageDir { get; set; } = "D:\\HordeLogs";
+		public DirectoryReference LocalArtifactsDirRef => DirectoryReference.Combine(Program.GetDefaultDataDir(), LocalArtifactsDir);
 
 		/// <summary>
 		/// S3 bucket region for logfile storage
diff --git a/Engine/Source/Programs/Horde/HordeServer/Services/ArtifactService.cs b/Engine/Source/Programs/Horde/HordeServer/Services/ArtifactService.cs
--- a/Engine/Source/Programs/Horde/HordeServer/Services/ArtifactService.cs
+++ b/Engine/Source/Programs/Horde/HordeServer/Services/ArtifactService.cs
@@ -72,7 +72,7 @@ namespace HordeServer.Services
 		public FSExternalArtifactStorage(ServerSettings CurrentSettings, ILogger<ArtifactService> Logger)
 		{
 			this.Logger = Logger;
-			this.BaseDir = new DirectoryInfo(CurrentSettings.LocalStorageDir);
+			this.BaseDir = new DirectoryInfo(CurrentSettings.LocalArtifactsDirRef.FullName);
 
 			Directory.CreateDirectory(this.BaseDir.FullName);
 		}
diff --git a/Engine/Source/Programs/Horde/HordeServer/Storage/Backends/FileSystemStorageBackend.cs b/Engine/Source/Programs/Horde/HordeServer/Storage/Backends/FileSystemStorageBackend.cs
--- a/Engine/Source/Programs/Horde/HordeServer/Storage/Backends/FileSystemStorageBackend.cs
+++ b/Engine/Source/Programs/Horde/HordeServer/Storage/Backends/FileSystemStorageBackend.cs
@@ -27,7 +27,7 @@ namespace HordeServer.Storage.Impl
 		public FileSystemStorageBackend(IOptionsMonitor<ServerSettings> Settings)
 		{
 			ServerSettings CurrentSettings = Settings.CurrentValue;
-			BaseDir = new DirectoryReference(CurrentSettings.LocalStorageDir);
+			BaseDir = CurrentSettings.LocalLogsDirRef;
 			DirectoryReference.CreateDirectory(BaseDir);
 		}
 
