Horde: Automatically spawn a bundled copy of MongoDB if available. Data is stored in C:\ProgramData\Horde\Mongo by default.

diff --git a/Engine/Source/Programs/Horde/HordeServer/Program.cs b/Engine/Source/Programs/Horde/HordeServer/Program.cs
--- a/Engine/Source/Programs/Horde/HordeServer/Program.cs
+++ b/Engine/Source/Programs/Horde/HordeServer/Program.cs
@@ -59,6 +59,10 @@ namespace HordeServer
 
 	class Program
 	{
+		public static DirectoryReference DataDir = GetDefaultDataDir();
+
+		public static FileReference UserConfigFile { get; } = FileReference.Combine(GetDefaultDataDir(), "Horde.json");
+
 		public static DirectoryReference AppDir { get; } = new FileReference(Assembly.GetExecutingAssembly().Location).Directory;
 
 		public static Type[] ConfigSchemas = FindSchemaTypes();
@@ -85,6 +89,7 @@ namespace HordeServer
 				.AddJsonFile("appsettings.json", optional: false)
 				.AddJsonFile($"appsettings.{Environment.GetEnvironmentVariable("ASPNETCORE_ENVIRONMENT")}.json", optional: true)
 				.AddJsonFile("appsettings.User.json", optional: true)
+				.AddJsonFile(UserConfigFile.FullName, optional: true)
 				.AddEnvironmentVariables()
 				.Build();
 
@@ -94,6 +99,7 @@ namespace HordeServer
 			Serilog.Log.Logger = new LoggerConfiguration()
 				.MinimumLevel.Debug()
 //				.MinimumLevel.Override("HordeServer.Services.DatabaseService", LogEventLevel.Verbose) // For MongoDB query tracing
+				.MinimumLevel.Override("MongoDB", LogEventLevel.Warning) // For bundled MongoDB output
 				.MinimumLevel.Override("Microsoft", LogEventLevel.Warning)
 				.MinimumLevel.Override("Microsoft.AspNetCore", LogEventLevel.Warning)
 				.MinimumLevel.Override("HordeServer.Authentication.OktaHandler", LogEventLevel.Warning)
@@ -193,7 +199,7 @@ namespace HordeServer
 		/// Gets the default directory for storing application data
 		/// </summary>
 		/// <returns>The default data directory</returns>
-		public static DirectoryReference GetDefaultDataDir()
+		static DirectoryReference GetDefaultDataDir()
 		{
 			if (RuntimeInformation.IsOSPlatform(OSPlatform.Windows))
 			{
diff --git a/Engine/Source/Programs/Horde/HordeServer/ServerSettings.cs b/Engine/Source/Programs/Horde/HordeServer/ServerSettings.cs
--- a/Engine/Source/Programs/Horde/HordeServer/ServerSettings.cs
+++ b/Engine/Source/Programs/Horde/HordeServer/ServerSettings.cs
@@ -233,7 +233,7 @@ namespace HordeServer
 		/// <summary>
 		/// Gets the full path referred to by LocalStorageDir
 		/// </summary>
-		public DirectoryReference LocalLogsDirRef => DirectoryReference.Combine(Program.GetDefaultDataDir(), LocalLogsDir);
+		public DirectoryReference LocalLogsDirRef => DirectoryReference.Combine(Program.DataDir, LocalLogsDir);
 
 		/// <summary>
 		/// Local artifact storage directory, if using type filesystem
@@ -243,7 +243,7 @@ namespace HordeServer
 		/// <summary>
 		/// Gets the full path referred to by LocalStorageDir
 		/// </summary>
-		public DirectoryReference LocalArtifactsDirRef => DirectoryReference.Combine(Program.GetDefaultDataDir(), LocalArtifactsDir);
+		public DirectoryReference LocalArtifactsDirRef => DirectoryReference.Combine(Program.DataDir, LocalArtifactsDir);
 
 		/// <summary>
 		/// S3 bucket region for logfile storage
diff --git a/Engine/Source/Programs/Horde/HordeServer/Services/DatabaseService.cs b/Engine/Source/Programs/Horde/HordeServer/Services/DatabaseService.cs
--- a/Engine/Source/Programs/Horde/HordeServer/Services/DatabaseService.cs
+++ b/Engine/Source/Programs/Horde/HordeServer/Services/DatabaseService.cs
@@ -33,13 +33,16 @@ using MongoDB.Driver.Core.Events;
 using System.Threading;
 using Datadog.Trace;
 using HordeServer.Collections;
+using System.Diagnostics;
+using System.Text.RegularExpressions;
+using System.Net.NetworkInformation;
 
 namespace HordeServer.Services
 {
 	/// <summary>
 	/// Singleton for accessing the database
 	/// </summary>
-	public class DatabaseService
+	public sealed class DatabaseService : IDisposable
 	{
 		/// <summary>
 		/// The database instance
@@ -64,7 +67,7 @@ namespace HordeServer.Services
 		/// <summary>
 		/// Settings for the application
 		/// </summary>
-		IOptionsMonitor<ServerSettings> Settings;
+		ServerSettings Settings { get; }
 
 		/// <summary>
 		/// Name of the JWT issuer
@@ -86,6 +89,31 @@ namespace HordeServer.Services
 		/// </summary>
 		public bool ReadOnlyMode { get; }
 
+		/// <summary>
+		/// The mongo process group
+		/// </summary>
+		ManagedProcessGroup? MongoProcessGroup;
+
+		/// <summary>
+		/// The mongo process
+		/// </summary>
+		ManagedProcess? MongoProcess;
+
+		/// <summary>
+		/// Task to read from the mongo process
+		/// </summary>
+		Task? MongoOutputTask;
+
+		/// <summary>
+		/// Factory for creating logger instances
+		/// </summary>
+		ILoggerFactory LoggerFactory;
+
+		/// <summary>
+		/// Default port for MongoDB connections
+		/// </summary>
+		const int DefaultMongoPort = 27017;
+
 		/// <summary>
 		/// Trace MongoDB commands in Datadog for observability
 		/// The built-in version from Datadog tracing library doesn't work for us (missing queries).
@@ -162,25 +190,25 @@ namespace HordeServer.Services
 		/// <summary>
 		/// Constructor
 		/// </summary>
-		/// <param name="Settings">The settings instance</param>
-		/// <param name="Logger">Instance of the logger for this service</param>
-		public DatabaseService(IOptionsMonitor<ServerSettings> Settings, ILogger<DatabaseService> Logger)
+		/// <param name="SettingsSnapshot">The settings instance</param>
+		/// <param name="LoggerFactory">Instance of the logger for this service</param>
+		public DatabaseService(IOptionsSnapshot<ServerSettings> SettingsSnapshot, ILoggerFactory LoggerFactory)
 		{
-			this.Settings = Settings;
-			this.Logger = Logger;
+			this.Settings = SettingsSnapshot.Value;
+			this.Logger = LoggerFactory.CreateLogger<DatabaseService>();
+			this.LoggerFactory = LoggerFactory;
 
 			try
 			{
-				ServerSettings CurrentSettings = Settings.CurrentValue;
-				ReadOnlyMode = CurrentSettings.DatabaseReadOnlyMode;
-				if (CurrentSettings.DatabasePublicCert != null)
+				ReadOnlyMode = Settings.DatabaseReadOnlyMode;
+				if (Settings.DatabasePublicCert != null)
 				{
 					X509Store LocalTrustStore = new X509Store(StoreName.Root);
 					try
 					{
 						LocalTrustStore.Open(OpenFlags.ReadWrite);
 
-						X509Certificate2Collection Collection = ImportCertificateBundle(CurrentSettings.DatabasePublicCert);
+						X509Certificate2Collection Collection = ImportCertificateBundle(Settings.DatabasePublicCert);
 						foreach (X509Certificate2 Certificate in Collection)
 						{
 							Logger.LogInformation("Importing certificate for {Subject}", Certificate.Subject);
@@ -194,7 +222,24 @@ namespace HordeServer.Services
 					}
 				}
 
-				MongoClientSettings MongoSettings = MongoClientSettings.FromConnectionString(CurrentSettings.DatabaseConnectionString);
+				string? ConnectionString = Settings.DatabaseConnectionString;
+				if (ConnectionString == null)
+				{
+					if (IsPortInUse(DefaultMongoPort))
+					{
+						ConnectionString = "mongodb://localhost:27017";
+					}
+					else if (TryStartMongoServer(Logger))
+					{
+						ConnectionString = "mongodb://localhost:27017/?readPreference=primary&appname=Horde&ssl=false";
+					}
+					else
+					{
+						throw new Exception($"Unable to connect to MongoDB server. Setup a MongoDB server and set the connection string in {Program.UserConfigFile}");
+					}
+				}
+
+				MongoClientSettings MongoSettings = MongoClientSettings.FromConnectionString(ConnectionString);
 				MongoSettings.ClusterConfigurator = ClusterBuilder =>
 				{
 					if (Logger.IsEnabled(LogLevel.Trace))
@@ -211,7 +256,7 @@ namespace HordeServer.Services
 				//TestSslConnection(MongoSettings.Server.Host, MongoSettings.Server.Port, Logger);
 
 				MongoClient Client = new MongoClient(MongoSettings);
-				Database = Client.GetDatabase(CurrentSettings.DatabaseName);
+				Database = Client.GetDatabase(Settings.DatabaseName);
 
 				Singletons = GetCollection<BsonDocument>("Singletons");
 
@@ -247,14 +292,14 @@ namespace HordeServer.Services
 					}
 				}
 
-				JwtIssuer = Settings.CurrentValue.JwtIssuer ?? Dns.GetHostName();
-				if (String.IsNullOrEmpty(CurrentSettings.JwtSecret))
+				JwtIssuer = Settings.JwtIssuer ?? Dns.GetHostName();
+				if (String.IsNullOrEmpty(Settings.JwtSecret))
 				{
 					JwtSigningKey = new SymmetricSecurityKey(Globals.JwtSigningKey);
 				}
 				else
 				{
-					JwtSigningKey = new SymmetricSecurityKey(Convert.FromBase64String(CurrentSettings.JwtSecret));
+					JwtSigningKey = new SymmetricSecurityKey(Convert.FromBase64String(Settings.JwtSecret));
 				}
 
 				Credentials = GetCollection<Credential>("Credentials");
@@ -267,6 +312,165 @@ namespace HordeServer.Services
 			}
 		}
 
+		internal const int CTRL_C_EVENT = 0;
+
+		[DllImport("kernel32.dll")]
+		internal static extern bool GenerateConsoleCtrlEvent(int Event, int ProcessGroupId);
+
+		/// <inheritdoc/>
+		public void Dispose()
+		{
+			if (MongoProcess != null)
+			{
+				GenerateConsoleCtrlEvent(CTRL_C_EVENT, MongoProcess.Id);
+
+				MongoOutputTask?.Wait();
+				MongoOutputTask = null;
+
+				MongoProcess.WaitForExit();
+				MongoProcess.Dispose();
+				MongoProcess = null;
+			}
+			if(MongoProcessGroup != null)
+			{
+				MongoProcessGroup?.Dispose();
+				MongoProcessGroup = null;
+			}
+		}
+
+		/// <summary>
+		/// Checks if the given port is in use
+		/// </summary>
+		/// <param name="PortNumber"></param>
+		/// <returns></returns>
+		static bool IsPortInUse(int PortNumber)
+		{
+			IPGlobalProperties IpGlobalProperties = IPGlobalProperties.GetIPGlobalProperties();
+
+			IPEndPoint[] Listeners = IpGlobalProperties.GetActiveTcpListeners();
+			if (Listeners.Any(x => x.Port == PortNumber))
+			{
+				return true;
+			}
+
+			return false;
+		}
+
+		/// <summary>
+		/// Attempts to start a local instance of MongoDB
+		/// </summary>
+		/// <param name="Logger"></param>
+		/// <returns></returns>
+		bool TryStartMongoServer(ILogger Logger)
+		{
+			if (!RuntimeInformation.IsOSPlatform(OSPlatform.Windows))
+			{
+				return false;
+			}
+
+			FileReference MongoExe = FileReference.Combine(Program.AppDir, "ThirdParty", "Mongo", "mongod.exe");
+			if (!FileReference.Exists(MongoExe))
+			{
+				Logger.LogWarning("Unable to find Mongo executable.");
+				return false;
+			}
+
+			DirectoryReference MongoDir = DirectoryReference.Combine(Program.DataDir, "Mongo");
+
+			DirectoryReference MongoDataDir = DirectoryReference.Combine(MongoDir, "Data");
+			DirectoryReference.CreateDirectory(MongoDataDir);
+
+			FileReference MongoLogFile = FileReference.Combine(MongoDir, "mongod.log");
+
+			FileReference ConfigFile = FileReference.Combine(MongoDir, "mongod.conf");
+			if(!FileReference.Exists(ConfigFile))
+			{
+				DirectoryReference.CreateDirectory(ConfigFile.Directory);
+				using (StreamWriter Writer = new StreamWriter(ConfigFile.FullName))
+				{
+					Writer.WriteLine("# mongod.conf");
+					Writer.WriteLine();
+					Writer.WriteLine("# for documentation of all options, see:");
+					Writer.WriteLine("# http://docs.mongodb.org/manual/reference/configuration-options/");
+					Writer.WriteLine();
+					Writer.WriteLine("storage:");
+					Writer.WriteLine("    dbPath: {0}", MongoDataDir.FullName);
+					Writer.WriteLine();
+					Writer.WriteLine("net:");
+					Writer.WriteLine("    port: {0}", DefaultMongoPort);
+					Writer.WriteLine("    bindIp: 127.0.0.1");
+				}
+			}
+
+			MongoProcessGroup = new ManagedProcessGroup();
+			try
+			{
+				MongoProcess = new ManagedProcess(MongoProcessGroup, MongoExe.FullName, $"--config \"{ConfigFile}\"", null, null, ProcessPriorityClass.Normal);
+				MongoProcess.StdIn.Close();
+				MongoOutputTask = Task.Run(() => RelayMongoOutput());
+				return true;
+			}
+			catch (Exception Ex)
+			{
+				Logger.LogWarning(Ex, "Unable to start Mongo server process");
+				return false;
+			}
+		}
+
+		/// <summary>
+		/// Copies output from the mongo process to the logger
+		/// </summary>
+		/// <returns></returns>
+		async Task RelayMongoOutput()
+		{
+			ILogger MongoLogger = LoggerFactory.CreateLogger("MongoDB");
+
+			Dictionary<string, ILogger> ChannelLoggers = new Dictionary<string, ILogger>();
+			for (; ; )
+			{
+				string? Line = await MongoProcess!.ReadLineAsync();
+				if (Line == null)
+				{
+					break;
+				}
+				if (Line.Length > 0)
+				{
+					Match Match = Regex.Match(Line, @"^\s*[^\s]+\s+([A-Z])\s+([^\s]+)\s+(.*)");
+					if (Match.Success)
+					{
+						ILogger? ChannelLogger;
+						if (!ChannelLoggers.TryGetValue(Match.Groups[2].Value, out ChannelLogger))
+						{
+							ChannelLogger = LoggerFactory.CreateLogger($"MongoDB.{Match.Groups[2].Value}");
+							ChannelLoggers.Add(Match.Groups[2].Value, ChannelLogger);
+						}
+						ChannelLogger.Log(ParseMongoLogLevel(Match.Groups[1].Value), Match.Groups[3].Value.TrimEnd());
+					}
+					else
+					{
+						MongoLogger.Log(LogLevel.Information, Line);
+					}
+				}
+			}
+			MongoLogger.LogInformation("Exit code {ExitCode}", MongoProcess.ExitCode);
+		}
+
+		static LogLevel ParseMongoLogLevel(string Text)
+		{
+			if (Text.Equals("I", StringComparison.Ordinal))
+			{
+				return LogLevel.Information;
+			}
+			else if (Text.Equals("E", StringComparison.Ordinal))
+			{
+				return LogLevel.Error;
+			}
+			else
+			{
+				return LogLevel.Warning;
+			}
+		}
+
 		/// <summary>
 		/// Logs a mongodb command, removing any fields we don't care about
 		/// </summary>
diff --git a/Engine/Source/Programs/Horde/HordeServer/appsettings.json b/Engine/Source/Programs/Horde/HordeServer/appsettings.json
--- a/Engine/Source/Programs/Horde/HordeServer/appsettings.json
+++ b/Engine/Source/Programs/Horde/HordeServer/appsettings.json
@@ -1,9 +1,7 @@
 {
 	"Horde": 
 	{
-		"DatabaseConnectionString": "mongodb://localhost:27017",
 		"RedisConnectionConfig": "localhost",
-		"DatabaseName": "Horde",
 		"DashboardUrl": "https://localhost:3000",
 		"DisableAuth": true,
 		"ConfigPath": "Defaults/global.json",
diff --git a/Engine/Source/Programs/Shared/EpicGames.Core/ManagedProcess.cs b/Engine/Source/Programs/Shared/EpicGames.Core/ManagedProcess.cs
--- a/Engine/Source/Programs/Shared/EpicGames.Core/ManagedProcess.cs
+++ b/Engine/Source/Programs/Shared/EpicGames.Core/ManagedProcess.cs
@@ -262,6 +262,11 @@ namespace EpicGames.Core
 
 		const int ERROR_ACCESS_DENIED = 5;
 
+		/// <summary>
+		/// The process id
+		/// </summary>
+		public int Id { get; private set; }
+
 		/// <summary>
 		/// Handle for the child process.
 		/// </summary>
@@ -493,6 +498,9 @@ namespace EpicGames.Core
 
 								Thread.Sleep(RetryDelay[AttemptIdx]);
 							}
+
+							// Save the process id
+							Id = (int)ProcessInfo.dwProcessId;
 						}
 						finally
 						{
@@ -620,6 +628,7 @@ namespace EpicGames.Core
 
 			}
 
+			Id = FrameworkProcess.Id;
 			StdIn = FrameworkProcess.StandardInput.BaseStream;
 			StdOut = FrameworkProcess.StandardOutput.BaseStream;
 		}
