Fix case issue in plaform mapping
#rnx

diff --git a/Engine/Source/Programs/Horde/HordeServer/Controllers/DevicesController.cs b/Engine/Source/Programs/Horde/HordeServer/Controllers/DevicesController.cs
--- a/Engine/Source/Programs/Horde/HordeServer/Controllers/DevicesController.cs
+++ b/Engine/Source/Programs/Horde/HordeServer/Controllers/DevicesController.cs
@@ -585,7 +585,7 @@ namespace HordeServer.Controllers
 			{
 				{ new DevicePlatformId("ps4"),  "PS4-DevKit"},
 				{ new DevicePlatformId("ps5"),  "PS5"},
-				{ new DevicePlatformId("xboxone"),  "XBoxOneGDK"},
+				{ new DevicePlatformId("xboxone"),  "XboxOneGDK"},
 				{ new DevicePlatformId("xsx"),  "XSX"},
 				{ new DevicePlatformId("switch"),  "Switch"},
 			};
