Horde: Output logs to ProgramData directory, and configure verbosity from config file.

diff --git a/Engine/Source/Programs/Horde/HordeServer/Program.cs b/Engine/Source/Programs/Horde/HordeServer/Program.cs
--- a/Engine/Source/Programs/Horde/HordeServer/Program.cs
+++ b/Engine/Source/Programs/Horde/HordeServer/Program.cs
@@ -96,22 +96,18 @@ namespace HordeServer
 			ServerSettings HordeSettings = new ServerSettings();
 			Config.GetSection("Horde").Bind(HordeSettings);
 
+			DirectoryReference LogDir = AppDir;
+			if (RuntimeInformation.IsOSPlatform(OSPlatform.Windows))
+			{
+				LogDir = DirectoryReference.Combine(DataDir);
+			}
+
 			Serilog.Log.Logger = new LoggerConfiguration()
-				.MinimumLevel.Debug()
-//				.MinimumLevel.Override("HordeServer.Services.DatabaseService", LogEventLevel.Verbose) // For MongoDB query tracing
-				.MinimumLevel.Override("MongoDB", LogEventLevel.Warning) // For bundled MongoDB output
-				.MinimumLevel.Override("Microsoft", LogEventLevel.Warning)
-				.MinimumLevel.Override("Microsoft.AspNetCore", LogEventLevel.Warning)
-				.MinimumLevel.Override("HordeServer.Authentication.OktaHandler", LogEventLevel.Warning)
-				.MinimumLevel.Override("HordeServer.Authentication.OktaViaJwtHandler", LogEventLevel.Warning)
-				.MinimumLevel.Override("HordeServer.Authentication.HordeJwtBearerHandler", LogEventLevel.Warning)
-				.MinimumLevel.Override("HordeServer.Authentication.ServiceAccountAuthHandler", LogEventLevel.Warning)
-				.MinimumLevel.Override("System.Net.Http.HttpClient", LogEventLevel.Warning)
-				.MinimumLevel.Override("Grpc", LogEventLevel.Warning)
 				.Enrich.FromLogContext()
 				.WriteTo.Console(HordeSettings)
-				.WriteTo.File(Path.Combine(AppDir.FullName, "Log.txt"), outputTemplate: "[{Timestamp:HH:mm:ss} {Level:w3}] {Indent}{Message:l}{NewLine}{Exception} [{SourceContext}]", rollingInterval: RollingInterval.Day, rollOnFileSizeLimit: true, fileSizeLimitBytes: 20 * 1024 * 1024, retainedFileCountLimit: 10)
-				.WriteTo.File(new JsonFormatter(renderMessage: true), Path.Combine(AppDir.FullName, "Log.json"), rollingInterval: RollingInterval.Day, rollOnFileSizeLimit: true, fileSizeLimitBytes: 20 * 1024 * 1024, retainedFileCountLimit: 10)
+				.WriteTo.File(Path.Combine(LogDir.FullName, "Log.txt"), outputTemplate: "[{Timestamp:HH:mm:ss} {Level:w3}] {Indent}{Message:l}{NewLine}{Exception} [{SourceContext}]", rollingInterval: RollingInterval.Day, rollOnFileSizeLimit: true, fileSizeLimitBytes: 20 * 1024 * 1024, retainedFileCountLimit: 10)
+				.WriteTo.File(new JsonFormatter(renderMessage: true), Path.Combine(LogDir.FullName, "Log.json"), rollingInterval: RollingInterval.Day, rollOnFileSizeLimit: true, fileSizeLimitBytes: 20 * 1024 * 1024, retainedFileCountLimit: 10)
+				.ReadFrom.Configuration(Config)
 				.CreateLogger();
 
 			if (Arguments.HasOption("-UpdateSchemas"))
diff --git a/Engine/Source/Programs/Horde/HordeServer/appsettings.json b/Engine/Source/Programs/Horde/HordeServer/appsettings.json
--- a/Engine/Source/Programs/Horde/HordeServer/appsettings.json
+++ b/Engine/Source/Programs/Horde/HordeServer/appsettings.json
@@ -8,5 +8,25 @@
 		"CorsOrigin": "https://localhost:5003;http://localhost:3000",
 		"LogFileProviderType": "Local"
 	},
+	"Serilog": 
+	{
+		"MinimumLevel": 
+		{
+			"Default": "Debug",
+			"Override": 
+			{
+				"MongoDB": "Warning", // For bundled MongoDB output
+				"Redis" : "Warning", // For bundled Redis output
+				"Microsoft": "Warning",
+				"Microsoft.AspNetCore": "Warning",
+				"HordeServer.Authentication.OktaHandler": "Warning",
+				"HordeServer.Authentication.OktaViaJwtHandler": "Warning",
+				"HordeServer.Authentication.HordeJwtBearerHandler": "Warning",
+				"HordeServer.Authentication.ServiceAccountAuthHandler": "Warning",
+				"System.Net.Http.HttpClient": "Warning",
+				"Grpc": "Warning"
+			}
+		}
+	},
 	"AllowedHosts": "*"
 }
