Horde: Missing file.

diff --git a/Engine/Source/Programs/Horde/HordeServer/Services/RedisService.cs b/Engine/Source/Programs/Horde/HordeServer/Services/RedisService.cs
new file mode 100644
--- /dev/null
+++ b/Engine/Source/Programs/Horde/HordeServer/Services/RedisService.cs
@@ -0,0 +1,212 @@
+// Copyright Epic Games, Inc. All Rights Reserved.
+
+using EpicGames.Core;
+using HordeServer.Utilities;
+using Microsoft.Extensions.Logging;
+using Microsoft.Extensions.Options;
+using StackExchange.Redis;
+using System;
+using System.Collections.Generic;
+using System.Diagnostics;
+using System.IO;
+using System.Linq;
+using System.Net;
+using System.Net.NetworkInformation;
+using System.Runtime.InteropServices;
+using System.Threading.Tasks;
+
+namespace HordeServer.Services
+{
+	/// <summary>
+	/// Manages the lifetime of a bundled Redis instance
+	/// </summary>
+	public sealed class RedisService : IDisposable
+	{
+		/// <summary>
+		/// Default Redis port
+		/// </summary>
+		const int RedisPort = 6379;
+
+		/// <summary>
+		/// The managed process group containing the Redis server
+		/// </summary>
+		ManagedProcessGroup? RedisProcessGroup;
+
+		/// <summary>
+		/// The server process
+		/// </summary>
+		ManagedProcess? RedisProcess;
+
+		/// <summary>
+		/// Connection multiplexer
+		/// </summary>
+		public ConnectionMultiplexer Multiplexer { get; }
+
+		/// <summary>
+		/// The database interface
+		/// </summary>
+		public IDatabase Database { get; }
+
+		/// <summary>
+		/// Connection pool
+		/// </summary>
+		public RedisConnectionPool ConnectionPool { get; }
+
+		/// <summary>
+		/// Logger factory
+		/// </summary>
+		ILoggerFactory LoggerFactory;
+
+		/// <summary>
+		/// Logging instance
+		/// </summary>
+		ILogger<RedisService> Logger;
+
+		/// <summary>
+		/// The output task
+		/// </summary>
+		Task? OutputTask;
+
+		/// <summary>
+		/// Hack to initialize RedisService early enough to use data protection
+		/// </summary>
+		/// <param name="Settings"></param>
+		public RedisService(ServerSettings Settings)
+			: this(Options.Create<ServerSettings>(Settings), new Serilog.Extensions.Logging.SerilogLoggerFactory(Serilog.Log.Logger))
+		{
+		}
+
+		/// <summary>
+		/// Constructor
+		/// </summary>
+		/// <param name="Options"></param>
+		/// <param name="LoggerFactory"></param>
+		public RedisService(IOptions<ServerSettings> Options, ILoggerFactory LoggerFactory)
+		{
+			this.LoggerFactory = LoggerFactory;
+			this.Logger = LoggerFactory.CreateLogger<RedisService>();
+
+			ServerSettings Settings = Options.Value;
+
+			string? ConnectionString = Settings.RedisConnectionConfig;
+			if (ConnectionString == null)
+			{
+				if (IsRunningOnDefaultPort())
+				{
+					ConnectionString = "localhost";
+				}
+				else if (TryStartRedisServer())
+				{
+					ConnectionString = "localhost";
+				}
+				else
+				{
+					throw new Exception($"Unable to connect to Redis. Please set {nameof(Settings.RedisConnectionConfig)} in {Program.UserConfigFile}");
+				}
+			}
+
+			Multiplexer = ConnectionMultiplexer.Connect(ConnectionString);
+			Database = Multiplexer.GetDatabase();
+			ConnectionPool = new RedisConnectionPool(20, ConnectionString);
+		}
+
+		/// <inheritdoc/>
+		public void Dispose()
+		{
+			if (RedisProcess != null)
+			{
+				RedisProcess.Dispose();
+				RedisProcess = null;
+			}
+			if (RedisProcessGroup != null)
+			{
+				RedisProcessGroup.Dispose();
+				RedisProcessGroup = null;
+			}
+		}
+
+		/// <summary>
+		/// Checks if Redis is already running on the default port
+		/// </summary>
+		/// <returns></returns>
+		static bool IsRunningOnDefaultPort()
+		{
+			IPGlobalProperties IpGlobalProperties = IPGlobalProperties.GetIPGlobalProperties();
+
+			IPEndPoint[] Listeners = IpGlobalProperties.GetActiveTcpListeners();
+			if (Listeners.Any(x => x.Port == RedisPort))
+			{
+				return true;
+			}
+
+			return false;
+		}
+
+		/// <summary>
+		/// Attempts to start a local instance of Redis
+		/// </summary>
+		/// <returns></returns>
+		bool TryStartRedisServer()
+		{
+			if (!RuntimeInformation.IsOSPlatform(OSPlatform.Windows))
+			{
+				return false;
+			}
+
+			FileReference RedisExe = FileReference.Combine(Program.AppDir, "ThirdParty", "Redis", "redis-server.exe");
+			if (!FileReference.Exists(RedisExe))
+			{
+				Logger.LogDebug("Redis executable does not exist at {ExePath}", RedisExe);
+				return false;
+			}
+
+			DirectoryReference RedisDir = DirectoryReference.Combine(Program.DataDir, "Redis");
+			DirectoryReference.CreateDirectory(RedisDir);
+
+			FileReference RedisConfigFile = FileReference.Combine(RedisDir, "redis.conf");
+			if(!FileReference.Exists(RedisConfigFile))
+			{
+				using (StreamWriter Writer = new StreamWriter(RedisConfigFile.FullName))
+				{
+					Writer.WriteLine("# redis.conf");
+				}
+			}
+
+			RedisProcessGroup = new ManagedProcessGroup();
+			try
+			{
+				RedisProcess = new ManagedProcess(RedisProcessGroup, RedisExe.FullName, "", null, null, ProcessPriorityClass.Normal);
+				RedisProcess.StdIn.Close();
+				OutputTask = Task.Run(() => RelayRedisOutput());
+				return true;
+			}
+			catch (Exception Ex)
+			{
+				Logger.LogWarning(Ex, "Unable to start Redis server process");
+				return false;
+			}
+		}
+
+		/// <summary>
+		/// Copies output from the redis process to the logger
+		/// </summary>
+		/// <returns></returns>
+		async Task RelayRedisOutput()
+		{
+			ILogger RedisLogger = LoggerFactory.CreateLogger("Redis");
+			for (; ; )
+			{
+				string? Line = await RedisProcess!.ReadLineAsync();
+				if (Line == null)
+				{
+					break;
+				}
+				if (Line.Length > 0)
+				{
+					RedisLogger.Log(LogLevel.Information, Line);
+				}
+			}
+			RedisLogger.LogInformation("Exit code {ExitCode}", RedisProcess.ExitCode);
+		}
+	}
+}
