Horde: Additional logging for lease errors.

diff --git a/Engine/Source/Programs/Horde/HordeServer/Services/AgentService.cs b/Engine/Source/Programs/Horde/HordeServer/Services/AgentService.cs
--- a/Engine/Source/Programs/Horde/HordeServer/Services/AgentService.cs
+++ b/Engine/Source/Programs/Horde/HordeServer/Services/AgentService.cs
@@ -754,7 +754,15 @@ namespace HordeServer.Services
 		/// <returns>New lease document</returns>
 		private Task<ILease> CreateLeaseAsync(IAgent Agent, AgentLease AgentLease)
 		{
-			return Leases.AddAsync(AgentLease.Id, AgentLease.Name, Agent.Id, Agent.SessionId!.Value, AgentLease.StreamId, AgentLease.PoolId, AgentLease.LogId, AgentLease.StartTime, AgentLease.Payload!);
+			try
+			{
+				return Leases.AddAsync(AgentLease.Id, AgentLease.Name, Agent.Id, Agent.SessionId!.Value, AgentLease.StreamId, AgentLease.PoolId, AgentLease.LogId, AgentLease.StartTime, AgentLease.Payload!);
+			}
+			catch (MongoWriteException Ex)
+			{
+				Logger.LogError(Ex, "Unable to create lease {LeaseId} for agent {AgentId}; lease already exists", AgentLease.Id, Agent.Id);
+				throw;
+			}
 		}
 
 		/// <summary>
