Horde: Fix null reference exception starting up on Linux.

diff --git a/Engine/Source/Programs/Horde/HordeServer/Program.cs b/Engine/Source/Programs/Horde/HordeServer/Program.cs
--- a/Engine/Source/Programs/Horde/HordeServer/Program.cs
+++ b/Engine/Source/Programs/Horde/HordeServer/Program.cs
@@ -59,11 +59,11 @@ namespace HordeServer
 
 	class Program
 	{
-		public static DirectoryReference DataDir = GetDefaultDataDir();
+		public static DirectoryReference AppDir { get; } = GetAppDir();
 
-		public static FileReference UserConfigFile { get; } = FileReference.Combine(GetDefaultDataDir(), "Horde.json");
+		public static DirectoryReference DataDir { get; } = GetDefaultDataDir();
 
-		public static DirectoryReference AppDir { get; } = new FileReference(Assembly.GetExecutingAssembly().Location).Directory;
+		public static FileReference UserConfigFile { get; } = FileReference.Combine(GetDefaultDataDir(), "Horde.json");
 
 		public static Type[] ConfigSchemas = FindSchemaTypes();
 
@@ -191,6 +191,15 @@ namespace HordeServer
 			}
 		}
 
+		/// <summary>
+		/// Get the application directory
+		/// </summary>
+		/// <returns></returns>
+		static DirectoryReference GetAppDir()
+		{
+			return new FileReference(Assembly.GetExecutingAssembly().Location).Directory;
+		}
+
 		/// <summary>
 		/// Gets the default directory for storing application data
 		/// </summary>
@@ -205,7 +214,7 @@ namespace HordeServer
 					return DirectoryReference.Combine(Dir, "Horde");
 				}
 			}
-			return DirectoryReference.Combine(AppDir, "Data");
+			return DirectoryReference.Combine(GetAppDir(), "Data");
 		}
 	}
 }
