Horde: Go back to assuming a Redis server is always running.

diff --git a/Engine/Source/Programs/Horde/HordeServer/Startup.cs b/Engine/Source/Programs/Horde/HordeServer/Startup.cs
--- a/Engine/Source/Programs/Horde/HordeServer/Startup.cs
+++ b/Engine/Source/Programs/Horde/HordeServer/Startup.cs
@@ -263,6 +263,7 @@ namespace HordeServer
 			RedisService RedisService = new RedisService(Settings);
 #pragma warning restore CA2000 // Dispose objects before losing scope
 			Services.AddSingleton<RedisService>(SP => RedisService);
+			Services.AddSingleton<IDatabase>(RedisService.Database);
 			Services.AddSingleton<ConnectionMultiplexer>(SP => RedisService.Multiplexer);
 			Services.AddDataProtection().PersistKeysToStackExchangeRedis(() => RedisService.Database, "aspnet-data-protection");
 
diff --git a/Engine/Source/Programs/Horde/HordeServer/appsettings.json b/Engine/Source/Programs/Horde/HordeServer/appsettings.json
--- a/Engine/Source/Programs/Horde/HordeServer/appsettings.json
+++ b/Engine/Source/Programs/Horde/HordeServer/appsettings.json
@@ -1,6 +1,7 @@
 {
 	"Horde": 
 	{
+		"RedisConnectionConfig": "localhost",
 		"DashboardUrl": "https://localhost:3000",
 		"DisableAuth": true,
 		"ConfigPath": "Defaults/global.json",
