Horde: Additional logging.

diff --git a/Engine/Source/Programs/Horde/HordeServer/Services/AgentService.cs b/Engine/Source/Programs/Horde/HordeServer/Services/AgentService.cs
--- a/Engine/Source/Programs/Horde/HordeServer/Services/AgentService.cs
+++ b/Engine/Source/Programs/Horde/HordeServer/Services/AgentService.cs
@@ -758,9 +758,9 @@ namespace HordeServer.Services
 			{
 				return Leases.AddAsync(AgentLease.Id, AgentLease.Name, Agent.Id, Agent.SessionId!.Value, AgentLease.StreamId, AgentLease.PoolId, AgentLease.LogId, AgentLease.StartTime, AgentLease.Payload!);
 			}
-			catch (MongoWriteException Ex)
+			catch (Exception Ex)
 			{
-				Logger.LogError(Ex, "Unable to create lease {LeaseId} for agent {AgentId}; lease already exists", AgentLease.Id, Agent.Id);
+				Logger.LogError(Ex, "Unable to create lease {LeaseId} for agent {AgentId}; lease already exists?", AgentLease.Id, Agent.Id);
 				throw;
 			}
 		}
