Horde: Fix typo when returning result of reservation delete
#rnx

diff --git a/Engine/Source/Programs/Horde/HordeServer/Collections/Impl/DeviceCollection.cs b/Engine/Source/Programs/Horde/HordeServer/Collections/Impl/DeviceCollection.cs
--- a/Engine/Source/Programs/Horde/HordeServer/Collections/Impl/DeviceCollection.cs
+++ b/Engine/Source/Programs/Horde/HordeServer/Collections/Impl/DeviceCollection.cs
@@ -581,7 +581,7 @@ namespace HordeServer.Collections.Impl
 			if (Reserves.Count > 0)
 			{
 				DeleteResult Result = await Reservations.DeleteManyAsync(Builders<DeviceReservationDocument>.Filter.In(x => x.Id, Reserves.Select(y => y.Id)));
-				return (Result.DeletedCount > 9);
+				return (Result.DeletedCount > 0);
 			}			
 
 			return false;
