Horde: Copy Redis/Mongo (if available) to output directory on Windows, allowing them to be launched automatically if not already present on the host machine.

diff --git a/Engine/Source/Programs/Horde/HordeServer/HordeServer.csproj b/Engine/Source/Programs/Horde/HordeServer/HordeServer.csproj
--- a/Engine/Source/Programs/Horde/HordeServer/HordeServer.csproj
+++ b/Engine/Source/Programs/Horde/HordeServer/HordeServer.csproj
@@ -1,4 +1,4 @@
-﻿<Project Sdk="Microsoft.NET.Sdk.Web">
+<Project Sdk="Microsoft.NET.Sdk.Web">
 
   <PropertyGroup>
     <TargetFramework>netcoreapp3.1</TargetFramework>
@@ -84,7 +84,15 @@
     </Content>
     <Content Include="..\..\..\ThirdParty\Perforce\P4Api.Net\native\linux\release\p4bridge.*" Condition="$([MSBuild]::IsOSPlatform('Linux'))">
       <CopyToOutputDirectory>Always</CopyToOutputDirectory>
-    </Content>	
+    </Content>
+    <ContentWithTargetPath Include="..\..\..\..\Restricted\NotForLicensees\Source\Programs\Horde\ThirdParty\Mongo\mongod.exe" Condition="$([MSBuild]::IsOSPlatform('Windows')) And Exists('..\..\..\..\Restricted\NotForLicensees\Source\Programs\Horde\ThirdParty\Mongo\mongod.exe')">
+      <CopyToOutputDirectory>PreserveNewest</CopyToOutputDirectory>
+      <TargetPath>ThirdParty/Mongo/mongod.exe</TargetPath>
+    </ContentWithTargetPath>
+    <ContentWithTargetPath Include="..\ThirdParty\Redis\redis-server.exe" Condition="$([MSBuild]::IsOSPlatform('Windows'))">
+      <CopyToOutputDirectory>PreserveNewest</CopyToOutputDirectory>
+      <TargetPath>ThirdParty/Redis/redis-server.exe</TargetPath>
+    </ContentWithTargetPath>
   </ItemGroup>
 
   <ItemGroup>
diff --git a/Engine/Source/Programs/Horde/HordeServerTests/DatabaseRunner.cs b/Engine/Source/Programs/Horde/HordeServerTests/DatabaseRunner.cs
--- a/Engine/Source/Programs/Horde/HordeServerTests/DatabaseRunner.cs
+++ b/Engine/Source/Programs/Horde/HordeServerTests/DatabaseRunner.cs
@@ -9,6 +9,7 @@ using System.Net.Sockets;
 using System.Reflection;
 using System.Runtime.InteropServices;
 using Castle.Core.Internal;
+using EpicGames.Core;
 
 namespace HordeServerTests
 {
@@ -117,30 +118,9 @@ namespace HordeServerTests
 		
 		private string GetBinaryPath()
 		{
-			// Slightly sketchy way of grabbing hold of binaries in repo
-			
-			string File = (new Uri(Assembly.GetExecutingAssembly().CodeBase!)).LocalPath;
-			string BinPath = "";
-
-			List<string> TempPath = new List<string>();
-			foreach (string Part in File.Replace("\\", "/").Split("/"))
-			{
-				TempPath.Add(Part);
-				if (Part == "HordeServerTests")
-				{
-					TempPath.Add("ThirdParty");
-					TempPath.Add(BinName);
-					BinPath = string.Join("/", TempPath.ToArray());
-					break;
-				}
-			}
-
-			if (BinPath.IsNullOrEmpty())
-			{
-				throw new Exception("Failed locating directory containing third-party binaries for testing!");
-			}
-
-			return BinPath;
+			FileReference File = new FileReference(new Uri(Assembly.GetExecutingAssembly().CodeBase!).LocalPath);
+			FileReference BinPath = FileReference.Combine(File.Directory, BinName);
+			return BinPath.FullName;
 		}
 		
 		private string GetTemporaryDirectory()
@@ -198,7 +178,7 @@ namespace HordeServerTests
 	{
 		private readonly string DatabaseName;
 		
-		public MongoDbRunnerLocal(string DatabaseName) : base("mongodb", "mongod.exe", 27017, true)
+		public MongoDbRunnerLocal(string DatabaseName) : base("mongodb", "ThirdParty/Mongo/mongod.exe", 27017, true)
 		{
 			this.DatabaseName = DatabaseName;
 		}
@@ -217,7 +197,7 @@ namespace HordeServerTests
 	
 	public class RedisRunner : DatabaseRunner
 	{
-		public RedisRunner() : base("redis", "redis-server.exe", 6379, true)
+		public RedisRunner() : base("redis", "ThirdParty/Redis/redis-server.exe", 6379, true)
 		{
 		}
 
