Horde: Missing comma in default stream config file.

diff --git a/Engine/Source/Programs/Horde/HordeServer/Defaults/stream-ue5-main.json b/Engine/Source/Programs/Horde/HordeServer/Defaults/stream-ue5-main.json
--- a/Engine/Source/Programs/Horde/HordeServer/Defaults/stream-ue5-main.json
+++ b/Engine/Source/Programs/Horde/HordeServer/Defaults/stream-ue5-main.json
@@ -56,7 +56,7 @@
 				{
 					"type": "Text",
 					"label": "Script",
-					"argument": "-Script="
+					"argument": "-Script=",
 					"default": "Foo.xml"
 				},
 				{
