Horde: Serve default files (index.html etc...) when hitting static file endpoints.

diff --git a/Engine/Source/Programs/Horde/HordeServer/Startup.cs b/Engine/Source/Programs/Horde/HordeServer/Startup.cs
--- a/Engine/Source/Programs/Horde/HordeServer/Startup.cs
+++ b/Engine/Source/Programs/Horde/HordeServer/Startup.cs
@@ -698,6 +698,7 @@ namespace HordeServer
 				App.UseMiddleware<RequestTrackerMiddleware>();	
 			}
 
+			App.UseDefaultFiles();
 			App.UseStaticFiles();
 			App.UseRouting();
 
