Horde: Remove default Redis connection string, allowing the server to spawn its own instance.

diff --git a/Engine/Source/Programs/Horde/HordeServer/appsettings.json b/Engine/Source/Programs/Horde/HordeServer/appsettings.json
--- a/Engine/Source/Programs/Horde/HordeServer/appsettings.json
+++ b/Engine/Source/Programs/Horde/HordeServer/appsettings.json
@@ -1,7 +1,6 @@
 {
 	"Horde": 
 	{
-		"RedisConnectionConfig": "localhost",
 		"DashboardUrl": "https://localhost:3000",
 		"DisableAuth": true,
 		"ConfigPath": "Defaults/global.json",
