Horde: Device note editor fix and run perforce api service in task threads now
#rnx

diff --git a/Engine/Source/Programs/Horde/HordeServer/Services/PerforceServiceApi.cs b/Engine/Source/Programs/Horde/HordeServer/Services/PerforceServiceApi.cs
--- a/Engine/Source/Programs/Horde/HordeServer/Services/PerforceServiceApi.cs
+++ b/Engine/Source/Programs/Horde/HordeServer/Services/PerforceServiceApi.cs
@@ -379,179 +379,187 @@ namespace HordeServer.Services
 		public override async Task<List<ChangeSummary>> GetChangesAsync(string StreamName, int? MinChange, int? MaxChange, int Results, string? ImpersonateUser)
 		{
 
-			await Task.Yield();
-
-			using (P4.Repository Repository = GetConnection(StreamName, ImpersonateUser))
+			return await Task.Run(() =>
 			{
 
-				P4.ChangesCmdOptions Options = new P4.ChangesCmdOptions(P4.ChangesCmdFlags.IncludeTime | P4.ChangesCmdFlags.FullDescription, null, Results, P4.ChangeListStatus.Submitted, null);
+				using (P4.Repository Repository = GetConnection(StreamName, ImpersonateUser))
+				{
 
-				string Filter = GetFilter($"//{Repository.Connection.Client.Name}/...", MinChange, MaxChange);
+					P4.ChangesCmdOptions Options = new P4.ChangesCmdOptions(P4.ChangesCmdFlags.IncludeTime | P4.ChangesCmdFlags.FullDescription, null, Results, P4.ChangeListStatus.Submitted, null);
 
-				P4.FileSpec FileSpec = new P4.FileSpec(new P4.DepotPath(Filter), null, null, null);
+					string Filter = GetFilter($"//{Repository.Connection.Client.Name}/...", MinChange, MaxChange);
 
-				IList<P4.Changelist> Changelists = Repository.GetChangelists(Options, FileSpec);
+					P4.FileSpec FileSpec = new P4.FileSpec(new P4.DepotPath(Filter), null, null, null);
 
-				List<ChangeSummary> Changes = new List<ChangeSummary>();
-				foreach (P4.Changelist Changelist in Changelists)
-				{
-					Changes.Add(new ChangeSummary(Changelist.Id, Changelist.OwnerName, Changelist.Description));
-				}
+					IList<P4.Changelist> Changelists = Repository.GetChangelists(Options, FileSpec);
 
-				return Changes;
-			}
+					List<ChangeSummary> Changes = new List<ChangeSummary>();
+					foreach (P4.Changelist Changelist in Changelists)
+					{
+						Changes.Add(new ChangeSummary(Changelist.Id, Changelist.OwnerName, Changelist.Description));
+					}
+
+					return Changes;
+				}
+			});
 
 		}
 
 		/// <inheritdoc/>
 		public override async Task<List<ChangeDetails>> GetChangeDetailsAsync(string StreamName, IReadOnlyList<int> ChangeNumbers, string? ImpersonateUser)
 		{
-			await Task.Yield();
-
-			using (P4.Repository Repository = GetConnection(StreamName, ImpersonateUser))
+			return await Task.Run(() =>
 			{
-				List<ChangeDetails> Results = new List<ChangeDetails>();
-
-				List<string> Args = new List<string> { "-s", "-S" };
-				Args.AddRange(ChangeNumbers.Select(Change => Change.ToString(CultureInfo.InvariantCulture)));
-
-				using (P4.P4Command Command = new P4.P4Command(Repository, "describe", true, Args.ToArray()))
+				using (P4.Repository Repository = GetConnection(StreamName, ImpersonateUser))
 				{
-					P4.P4CommandResult Result = Command.Run();
+					List<ChangeDetails> Results = new List<ChangeDetails>();
 
-					if (!Result.Success)
-					{
-						throw new Exception("Unable to get changes");
-					}
+					List<string> Args = new List<string> { "-s", "-S" };
+					Args.AddRange(ChangeNumbers.Select(Change => Change.ToString(CultureInfo.InvariantCulture)));
 
-					if (Result.TaggedOutput == null || Result.TaggedOutput.Count <= 0)
+					using (P4.P4Command Command = new P4.P4Command(Repository, "describe", true, Args.ToArray()))
 					{
-						return Results;
-					}
-
-					List<P4.Changelist> Changelists = new List<P4.Changelist>() { };
+						P4.P4CommandResult Result = Command.Run();
 
-					bool DSTMismatch = false;
-					string Offset = string.Empty;
+						if (!Result.Success)
+						{
+							throw new Exception("Unable to get changes");
+						}
 
-					if (Server != null && Server.Metadata != null)
-					{
-						Offset = Server.Metadata.DateTimeOffset;
-						DSTMismatch = P4.FormBase.DSTMismatch(Server.Metadata);
-					}
+						if (Result.TaggedOutput == null || Result.TaggedOutput.Count <= 0)
+						{
+							return Results;
+						}
 
-					foreach (P4.TaggedObject TaggedObject in Result.TaggedOutput)
-					{
-						List<string> Files = new List<string>();
+						List<P4.Changelist> Changelists = new List<P4.Changelist>() { };
 
-						P4.Changelist Change = new P4.Changelist();
-						Change.FromChangeCmdTaggedOutput(TaggedObject, true, Offset, DSTMismatch);
+						bool DSTMismatch = false;
+						string Offset = string.Empty;
 
-						foreach (P4.FileMetaData DescribeFile in Change.Files)
+						if (Server != null && Server.Metadata != null)
 						{
-							string? RelativePath;
-							if (TryGetStreamRelativePath(DescribeFile.DepotPath.Path, StreamName, out RelativePath))
-							{
-								Files.Add(RelativePath);
-							}
+							Offset = Server.Metadata.DateTimeOffset;
+							DSTMismatch = P4.FormBase.DSTMismatch(Server.Metadata);
 						}
 
-						if (Change.ShelvedFiles != null && Change.ShelvedFiles.Count > 0)
+						foreach (P4.TaggedObject TaggedObject in Result.TaggedOutput)
 						{
-							foreach (P4.ShelvedFile ShelvedFile in Change.ShelvedFiles)
+							List<string> Files = new List<string>();
+
+							P4.Changelist Change = new P4.Changelist();
+							Change.FromChangeCmdTaggedOutput(TaggedObject, true, Offset, DSTMismatch);
+
+							foreach (P4.FileMetaData DescribeFile in Change.Files)
 							{
 								string? RelativePath;
-								if (TryGetStreamRelativePath(ShelvedFile.Path.ToString(), StreamName, out RelativePath))
+								if (TryGetStreamRelativePath(DescribeFile.DepotPath.Path, StreamName, out RelativePath))
 								{
 									Files.Add(RelativePath);
 								}
 							}
 
+							if (Change.ShelvedFiles != null && Change.ShelvedFiles.Count > 0)
+							{
+								foreach (P4.ShelvedFile ShelvedFile in Change.ShelvedFiles)
+								{
+									string? RelativePath;
+									if (TryGetStreamRelativePath(ShelvedFile.Path.ToString(), StreamName, out RelativePath))
+									{
+										Files.Add(RelativePath);
+									}
+								}
+
+							}
+
+							Results.Add(new ChangeDetails(Change.Id, Change.OwnerName, Change.Description, Files));
 						}
 
-						Results.Add(new ChangeDetails(Change.Id, Change.OwnerName, Change.Description, Files));
 					}
 
+					return Results;
 				}
-
-				return Results;
-			}
+			});
 		}
 
 		/// <inheritdoc />
 		public override async Task<string> CreateTicket(string ImpersonateUser)
 		{
-			await Task.Yield();
+			return await Task.Run(() =>
+			{
 
-			P4.Credential Credential = GetImpersonateCredential(ImpersonateUser);
+				P4.Credential Credential = GetImpersonateCredential(ImpersonateUser);
 
-			if (Credential == null)
-			{
-				throw new Exception($"Unable to get ticket for user {ImpersonateUser}");
-			}
+				if (Credential == null)
+				{
+					throw new Exception($"Unable to get ticket for user {ImpersonateUser}");
+				}
 
-			return Credential.Ticket;
+				return Credential.Ticket;
+			});
 
 		}
 
 		/// <inheritdoc/>
 		public override async Task<List<FileSummary>> FindFilesAsync(IEnumerable<string> Paths)
 		{
-			await Task.Yield();
-
-			List<FileSummary> Results = new List<FileSummary>();
-
-			using (P4.Repository Repository = GetConnection(NoClient: true))
+			return await Task.Run(() =>
 			{
-				P4.GetFileMetaDataCmdOptions Options = new P4.GetFileMetaDataCmdOptions(P4.GetFileMetadataCmdFlags.ExcludeClientData, "", "", 0, "", "", "");
-
-				IList<P4.FileMetaData>? Files = Repository.GetFileMetaData(Options, Paths.Select(Path => { P4.FileSpec FileSpec = new P4.FileSpec(new P4.DepotPath(Path)); return FileSpec; }).ToArray());
+				List<FileSummary> Results = new List<FileSummary>();
 
-				if (Files == null)
+				using (P4.Repository Repository = GetConnection(NoClient: true))
 				{
-					Files = new List<P4.FileMetaData>();
-				}
+					P4.GetFileMetaDataCmdOptions Options = new P4.GetFileMetaDataCmdOptions(P4.GetFileMetadataCmdFlags.ExcludeClientData, "", "", 0, "", "", "");
 
-				foreach (string Path in Paths)
-				{
-					P4.FileMetaData Meta = Files.FirstOrDefault(File => File.DepotPath.Path == Path);
+					IList<P4.FileMetaData>? Files = Repository.GetFileMetaData(Options, Paths.Select(Path => { P4.FileSpec FileSpec = new P4.FileSpec(new P4.DepotPath(Path)); return FileSpec; }).ToArray());
 
-					if (Meta == null)
+					if (Files == null)
 					{
-						Results.Add(new FileSummary(Path, false, 0));
+						Files = new List<P4.FileMetaData>();
 					}
-					else
+
+					foreach (string Path in Paths)
 					{
-						Results.Add(new FileSummary(Path, Meta.HeadAction != P4.FileAction.Delete && Meta.HeadAction != P4.FileAction.MoveDelete, Meta.HeadChange));
+						P4.FileMetaData Meta = Files.FirstOrDefault(File => File.DepotPath.Path == Path);
+
+						if (Meta == null)
+						{
+							Results.Add(new FileSummary(Path, false, 0));
+						}
+						else
+						{
+							Results.Add(new FileSummary(Path, Meta.HeadAction != P4.FileAction.Delete && Meta.HeadAction != P4.FileAction.MoveDelete, Meta.HeadChange));
+						}
+
 					}
 
+					return Results;
 				}
-
-				return Results;
-			}
+			});
 		}
 
 		/// <inheritdoc/>
 		public override async Task<byte[]> PrintAsync(string DepotPath)
 		{
-			await Task.Yield();
-
-			if (DepotPath.EndsWith("...", StringComparison.OrdinalIgnoreCase) || DepotPath.EndsWith("*", StringComparison.OrdinalIgnoreCase))
+			return await Task.Run(() =>
 			{
-				throw new Exception("PrintAsync requires exactly one file to be specified");
-			}
 
-			using (P4.Repository Repository = GetConnection(NoClient: true))
-			{
-				using (P4.P4Command Command = new P4.P4Command(Repository, "print", false, new string[] { "-q", DepotPath }))
+				if (DepotPath.EndsWith("...", StringComparison.OrdinalIgnoreCase) || DepotPath.EndsWith("*", StringComparison.OrdinalIgnoreCase))
 				{
-					P4.P4CommandResult Result = Command.Run();
+					throw new Exception("PrintAsync requires exactly one file to be specified");
+				}
 
-					return Encoding.Default.GetBytes(Result.TextOutput);
+				using (P4.Repository Repository = GetConnection(NoClient: true))
+				{
+					using (P4.P4Command Command = new P4.P4Command(Repository, "print", false, new string[] { "-q", DepotPath }))
+					{
+						P4.P4CommandResult Result = Command.Run();
 
-				}
+						return Encoding.Default.GetBytes(Result.TextOutput);
 
-			}
+					}
+
+				}
+			});
 
 		}
 
@@ -576,341 +584,349 @@ namespace HordeServer.Services
 		/// <inheritdoc/>
 		public override async Task DeleteShelvedChangeAsync(int ShelvedChange)
 		{
-			await Task.Yield();
-
-			using (P4.Repository Repository = GetConnection(NoClient: true))
+			await Task.Run(() =>
 			{
 
-				P4.Changelist? Changelist = Repository.GetChangelist(ShelvedChange, new P4.DescribeCmdOptions(P4.DescribeChangelistCmdFlags.Omit, 0, 0));
-
-				if (Changelist == null)
+				using (P4.Repository Repository = GetConnection(NoClient: true))
 				{
-					throw new Exception($"Unable to get shelved changelist for delete: {ShelvedChange}");
-				}
 
-				string ClientId = Changelist.ClientId;
+					P4.Changelist? Changelist = Repository.GetChangelist(ShelvedChange, new P4.DescribeCmdOptions(P4.DescribeChangelistCmdFlags.Omit, 0, 0));
 
-				if (String.IsNullOrEmpty(ClientId))
-				{
-					throw new Exception($"Unable to get shelved changelist client id for delete: {ShelvedChange}");
-				}
+					if (Changelist == null)
+					{
+						throw new Exception($"Unable to get shelved changelist for delete: {ShelvedChange}");
+					}
 
-				using (P4.Repository ClientRepository = GetConnection(Changelist.Stream, Changelist.OwnerName, ClientId: ClientId, UseClientFromChange: true, UsePortFromChange: true))
-				{
+					string ClientId = Changelist.ClientId;
 
-					if (Changelist.Shelved)
+					if (String.IsNullOrEmpty(ClientId))
 					{
-						IList<P4.FileSpec> Files = ClientRepository.Connection.Client.ShelveFiles(new P4.ShelveFilesCmdOptions(P4.ShelveFilesCmdFlags.Delete, null, ShelvedChange));
+						throw new Exception($"Unable to get shelved changelist client id for delete: {ShelvedChange}");
 					}
 
-					ClientRepository.DeleteChangelist(Changelist, null);
+					using (P4.Repository ClientRepository = GetConnection(Changelist.Stream, Changelist.OwnerName, ClientId: ClientId, UseClientFromChange: true, UsePortFromChange: true))
+					{
+
+						if (Changelist.Shelved)
+						{
+							IList<P4.FileSpec> Files = ClientRepository.Connection.Client.ShelveFiles(new P4.ShelveFilesCmdOptions(P4.ShelveFilesCmdFlags.Delete, null, ShelvedChange));
+						}
+
+						ClientRepository.DeleteChangelist(Changelist, null);
+					}
 				}
-			}
+			});
 		}
 
 		/// <inheritdoc/>
 		public override async Task UpdateChangelistDescription(int Change, string Description)
 		{
-			await Task.Yield();
-
-			using (P4.Repository Repository = GetConnection(NoClient: true))
+			await Task.Run(() =>
 			{
-				P4.Changelist Changelist = Repository.GetChangelist(Change);
+				using (P4.Repository Repository = GetConnection(NoClient: true))
+				{
+					P4.Changelist Changelist = Repository.GetChangelist(Change);
 
-				Repository.Connection.Disconnect();
+					Repository.Connection.Disconnect();
 
-				// the client must exist for the change list, otherwise will fail (for example, CreateNewChangeAsync deletes the client before returning)
-				using (P4.Repository UpdateRepository = GetConnection(ClientFromChange: Change, UseClientFromChange: true, Username: Changelist.OwnerName))
-				{
-					P4.Changelist UpdatedChangelist = UpdateRepository.GetChangelist(Change);
-					UpdatedChangelist.Description = Description;
-					UpdateRepository.UpdateChangelist(UpdatedChangelist);
+					// the client must exist for the change list, otherwise will fail (for example, CreateNewChangeAsync deletes the client before returning)
+					using (P4.Repository UpdateRepository = GetConnection(ClientFromChange: Change, UseClientFromChange: true, Username: Changelist.OwnerName))
+					{
+						P4.Changelist UpdatedChangelist = UpdateRepository.GetChangelist(Change);
+						UpdatedChangelist.Description = Description;
+						UpdateRepository.UpdateChangelist(UpdatedChangelist);
+					}
 				}
-			}
+			});
 		}
 
 		/// <inheritdoc/>
 		public async override Task<int> CreateNewChangeAsync(string StreamName, string FilePath)
 		{
-			await Task.Yield();
-
-			ServerSettings Settings = this.Settings.CurrentValue;
+			return await Task.Run(() =>
+			{
 
-			string Username = CanImpersonate ? "buildmachine" : Settings.P4BridgeServiceUsername!;
+				ServerSettings Settings = this.Settings.CurrentValue;
 
-			P4.SubmitResults? SubmitResults = null;
+				string Username = CanImpersonate ? "buildmachine" : Settings.P4BridgeServiceUsername!;
 
-			using (P4.Repository Repository = GetConnection(Stream: StreamName, Username: Username, ReadOnly: false, CreateChange: true))
-			{
-				P4.Client Client = Repository.Connection.Client;
+				P4.SubmitResults? SubmitResults = null;
 
-				string WorkspaceFilePath = $"//{Client.Name}/{FilePath.TrimStart('/')}";
-				string DiskFilePath = $"{Client.Root + FilePath.TrimStart('/')}";
+				using (P4.Repository Repository = GetConnection(Stream: StreamName, Username: Username, ReadOnly: false, CreateChange: true))
+				{
+					P4.Client Client = Repository.Connection.Client;
 
-				P4.FileSpec WorkspaceFileSpec = new P4.FileSpec(new P4.ClientPath(WorkspaceFilePath));
+					string WorkspaceFilePath = $"//{Client.Name}/{FilePath.TrimStart('/')}";
+					string DiskFilePath = $"{Client.Root + FilePath.TrimStart('/')}";
 
-				IList<P4.File> Files = Repository.GetFiles(new P4.FilesCmdOptions(P4.FilesCmdFlags.None, 1), WorkspaceFileSpec);
+					P4.FileSpec WorkspaceFileSpec = new P4.FileSpec(new P4.ClientPath(WorkspaceFilePath));
 
-				P4.Changelist? SubmitChangelist = null;
-				P4.DepotPath? DepotPath = null;
+					IList<P4.File> Files = Repository.GetFiles(new P4.FilesCmdOptions(P4.FilesCmdFlags.None, 1), WorkspaceFileSpec);
 
-				const int MaxRetries = 10;
-				int Retry = 0;
+					P4.Changelist? SubmitChangelist = null;
+					P4.DepotPath? DepotPath = null;
 
-				for (; ; )
-				{
-					DepotPath = null;
+					const int MaxRetries = 10;
+					int Retry = 0;
 
-					if (Retry == MaxRetries)
+					for (; ; )
 					{
-						break;
-					}
+						DepotPath = null;
 
-					try
-					{
-						if (Files == null || Files.Count == 0)
+						if (Retry == MaxRetries)
 						{
-							// File does not exist, create it
-							string? DirectoryName = Path.GetDirectoryName(DiskFilePath);
+							break;
+						}
 
-							if (string.IsNullOrEmpty(DirectoryName))
+						try
+						{
+							if (Files == null || Files.Count == 0)
 							{
-								throw new Exception($"Invalid directory name for local client file, disk file path: {DiskFilePath}");
-							}
+								// File does not exist, create it
+								string? DirectoryName = Path.GetDirectoryName(DiskFilePath);
 
-							// Create the directory
-							if (!Directory.Exists(DirectoryName))
-							{
-								Directory.CreateDirectory(DirectoryName);
+								if (string.IsNullOrEmpty(DirectoryName))
+								{
+									throw new Exception($"Invalid directory name for local client file, disk file path: {DiskFilePath}");
+								}
 
+								// Create the directory
 								if (!Directory.Exists(DirectoryName))
 								{
-									throw new Exception($"Unable to create directrory: {DirectoryName}");
+									Directory.CreateDirectory(DirectoryName);
+
+									if (!Directory.Exists(DirectoryName))
+									{
+										throw new Exception($"Unable to create directrory: {DirectoryName}");
+									}
 								}
-							}
 
-							// Create the file
-							if (!File.Exists(DiskFilePath))
-							{
-								using (FileStream FileStream = File.OpenWrite(DiskFilePath))
+								// Create the file
+								if (!File.Exists(DiskFilePath))
 								{
-									FileStream.Close();
+									using (FileStream FileStream = File.OpenWrite(DiskFilePath))
+									{
+										FileStream.Close();
+									}
+
+									if (!File.Exists(DiskFilePath))
+									{
+										throw new Exception($"Unable to create local change file: {DiskFilePath}");
+									}
+
 								}
 
-								if (!File.Exists(DiskFilePath))
+								IList<P4.FileSpec> DepotFiles = Client.AddFiles(new P4.Options(), WorkspaceFileSpec);
+
+								if (DepotFiles == null || DepotFiles.Count != 1)
 								{
-									throw new Exception($"Unable to create local change file: {DiskFilePath}");
+									throw new Exception($"Unable to add local change file,  local: {DiskFilePath} : workspace: {WorkspaceFileSpec}");
 								}
 
+								DepotPath = DepotFiles[0].DepotPath;
+
 							}
+							else
+							{
+								IList<P4.FileSpec> SyncResults = Client.SyncFiles(new P4.SyncFilesCmdOptions(P4.SyncFilesCmdFlags.Force), WorkspaceFileSpec);
 
-							IList<P4.FileSpec> DepotFiles = Client.AddFiles(new P4.Options(), WorkspaceFileSpec);
+								if (SyncResults == null || SyncResults.Count != 1)
+								{
+									throw new Exception($"Unable to sync file, workspace: {WorkspaceFileSpec}");
+								}
 
-							if (DepotFiles == null || DepotFiles.Count != 1)
-							{
-								throw new Exception($"Unable to add local change file,  local: {DiskFilePath} : workspace: {WorkspaceFileSpec}");
-							}
+								IList<P4.FileSpec> EditResults = Client.EditFiles(new P4.FileSpec[] { new P4.FileSpec(SyncResults[0].DepotPath) }, new P4.Options());
 
-							DepotPath = DepotFiles[0].DepotPath;
+								if (EditResults == null || EditResults.Count != 1)
+								{
+									throw new Exception($"Unable to edit file, workspace: {WorkspaceFileSpec}");
+								}
 
-						}
-						else
-						{
-							IList<P4.FileSpec> SyncResults = Client.SyncFiles(new P4.SyncFilesCmdOptions(P4.SyncFilesCmdFlags.Force), WorkspaceFileSpec);
+								DepotPath = EditResults[0].DepotPath;
 
-							if (SyncResults == null || SyncResults.Count != 1)
-							{
-								throw new Exception($"Unable to sync file, workspace: {WorkspaceFileSpec}");
 							}
 
-							IList<P4.FileSpec> EditResults = Client.EditFiles(new P4.FileSpec[] { new P4.FileSpec(SyncResults[0].DepotPath) }, new P4.Options());
-
-							if (EditResults == null || EditResults.Count != 1)
+							if (DepotPath == null || string.IsNullOrEmpty(DepotPath.Path))
 							{
-								throw new Exception($"Unable to edit file, workspace: {WorkspaceFileSpec}");
+								throw new Exception($"Unable to get depot path for: {WorkspaceFileSpec}");
 							}
 
-							DepotPath = EditResults[0].DepotPath;
-
-						}
+							// create a new change
+							if (SubmitChangelist == null)
+							{
+								P4.Changelist Changelist = new P4.Changelist();
+								Changelist.Description = "New change for Horde job";
+								Changelist.Files.Add(new P4.FileMetaData(new P4.FileSpec(DepotPath)));
+								SubmitChangelist = Repository.CreateChangelist(Changelist);
 
-						if (DepotPath == null || string.IsNullOrEmpty(DepotPath.Path))
-						{
-							throw new Exception($"Unable to get depot path for: {WorkspaceFileSpec}");
-						}
+								if (SubmitChangelist == null)
+								{
+									throw new Exception($"Unable to create a changelist for: {DepotPath}");
+								}
+							}
 
-						// create a new change
-						if (SubmitChangelist == null)
-						{
-							P4.Changelist Changelist = new P4.Changelist();
-							Changelist.Description = "New change for Horde job";
-							Changelist.Files.Add(new P4.FileMetaData(new P4.FileSpec(DepotPath)));
-							SubmitChangelist = Repository.CreateChangelist(Changelist);
+							SubmitResults = SubmitChangelist.Submit(null);
 
-							if (SubmitChangelist == null)
+							if (SubmitResults == null)
 							{
-								throw new Exception($"Unable to create a changelist for: {DepotPath}");
+								throw new Exception($"Unable to submit changelist for: {DepotPath}");
 							}
+
+							break;
 						}
+						catch
+						{
+							Retry++;
 
-						SubmitResults = SubmitChangelist.Submit(null);
+							Client.RevertFiles(new P4.Options(), WorkspaceFileSpec);
 
-						if (SubmitResults == null)
-						{
-							throw new Exception($"Unable to submit changelist for: {DepotPath}");
+							continue;
 						}
 
-						break;
+					}
+
+					string ClientName = Client.Name;
+					try
+					{
+						Repository.DeleteClient(Client, new P4.Options());
 					}
 					catch
 					{
-						Retry++;
-
-						Client.RevertFiles(new P4.Options(), WorkspaceFileSpec);
-
-						continue;
+						Logger.LogError($"Unable to delete client {ClientName}");
 					}
 
 				}
 
-				string ClientName = Client.Name;
-				try
-				{
-					Repository.DeleteClient(Client, new P4.Options());
-				}
-				catch
+				if (SubmitResults == null)
 				{
-					Logger.LogError($"Unable to delete client {ClientName}");
+					throw new Exception($"Unable to submit change for {StreamName} {FilePath}");
 				}
 
-			}
-
-			if (SubmitResults == null)
-			{
-				throw new Exception($"Unable to submit change for {StreamName} {FilePath}");
-			}
-
-			return SubmitResults.ChangeIdAfterSubmit;
+				return SubmitResults.ChangeIdAfterSubmit;
+			});
 		}
 
 		/// <inheritdoc/>
 		public override async Task<(int? Change, string Message)> SubmitShelvedChangeAsync(int Change, int OriginalChange)
 		{
-			await Task.Yield();
-
-			using (P4.Repository Repository = GetConnection(NoClient: true))
+			return await Task.Run(() => 
 			{
+			   using (P4.Repository Repository = GetConnection(NoClient: true))
+			   {
 
-				List<string> Args = new List<string> { "-S", Change.ToString(CultureInfo.InvariantCulture), OriginalChange.ToString(CultureInfo.InvariantCulture) };
+				   List<string> Args = new List<string> { "-S", Change.ToString(CultureInfo.InvariantCulture), OriginalChange.ToString(CultureInfo.InvariantCulture) };
 
-				using (P4.P4Command Command = new P4.P4Command(Repository, "describe", true, Args.ToArray()))
-				{
-					P4.P4CommandResult Result = Command.Run();
+				   using (P4.P4Command Command = new P4.P4Command(Repository, "describe", true, Args.ToArray()))
+				   {
+					   P4.P4CommandResult Result = Command.Run();
 
-					if (!Result.Success)
-					{
-						throw new Exception("Unable to get changes");
-					}
+					   if (!Result.Success)
+					   {
+						   throw new Exception("Unable to get changes");
+					   }
 
-					if (Result.TaggedOutput == null || Result.TaggedOutput.Count != 2)
-					{
+					   if (Result.TaggedOutput == null || Result.TaggedOutput.Count != 2)
+					   {
 
-						throw new Exception($"Unable to get tagged output for change: {Change} and original change: {OriginalChange}");
-					}
+						   throw new Exception($"Unable to get tagged output for change: {Change} and original change: {OriginalChange}");
+					   }
 
-					bool DSTMismatch = false;
-					string Offset = string.Empty;
+					   bool DSTMismatch = false;
+					   string Offset = string.Empty;
 
-					if (Server != null && Server.Metadata != null)
-					{
-						Offset = Server.Metadata.DateTimeOffset;
-						DSTMismatch = P4.FormBase.DSTMismatch(Server.Metadata);
-					}
+					   if (Server != null && Server.Metadata != null)
+					   {
+						   Offset = Server.Metadata.DateTimeOffset;
+						   DSTMismatch = P4.FormBase.DSTMismatch(Server.Metadata);
+					   }
 
-					P4.Changelist Changelist = new P4.Changelist();
-					Changelist.FromChangeCmdTaggedOutput(Result.TaggedOutput[0], true, Offset, DSTMismatch);
+					   P4.Changelist Changelist = new P4.Changelist();
+					   Changelist.FromChangeCmdTaggedOutput(Result.TaggedOutput[0], true, Offset, DSTMismatch);
 
-					P4.Changelist OriginalChangelist = new P4.Changelist();
-					OriginalChangelist.FromChangeCmdTaggedOutput(Result.TaggedOutput[1], true, Offset, DSTMismatch);
+					   P4.Changelist OriginalChangelist = new P4.Changelist();
+					   OriginalChangelist.FromChangeCmdTaggedOutput(Result.TaggedOutput[1], true, Offset, DSTMismatch);
 
-					if (OriginalChangelist.ShelvedFiles.Count != Changelist.ShelvedFiles.Count)
-					{
-						throw new Exception($"Mismatched number of shelved files for change: {Change} and original change: {OriginalChange}");
-					}
+					   if (OriginalChangelist.ShelvedFiles.Count != Changelist.ShelvedFiles.Count)
+					   {
+						   throw new Exception($"Mismatched number of shelved files for change: {Change} and original change: {OriginalChange}");
+					   }
 
-					if (OriginalChangelist.ShelvedFiles.Count == 0)
-					{
-						throw new Exception($"No shelved file for change: {Change} and original change: {OriginalChange}");
-					}
+					   if (OriginalChangelist.ShelvedFiles.Count == 0)
+					   {
+						   throw new Exception($"No shelved file for change: {Change} and original change: {OriginalChange}");
+					   }
 
-					// todo, test straight up delete
-					foreach (P4.ShelvedFile ShelvedFile in Changelist.ShelvedFiles)
-					{
-						P4.ShelvedFile? Found = OriginalChangelist.ShelvedFiles.FirstOrDefault(Original => Original.Digest == ShelvedFile.Digest && Original.Action == ShelvedFile.Action);
+						// todo, test straight up delete
+						foreach (P4.ShelvedFile ShelvedFile in Changelist.ShelvedFiles)
+					   {
+						   P4.ShelvedFile? Found = OriginalChangelist.ShelvedFiles.FirstOrDefault(Original => Original.Digest == ShelvedFile.Digest && Original.Action == ShelvedFile.Action);
 
-						if (Found == null)
-						{
-							throw new Exception($"Mismatch in shelved file digest or action for {ShelvedFile.Path}");
-						}
-					}
+						   if (Found == null)
+						   {
+							   throw new Exception($"Mismatch in shelved file digest or action for {ShelvedFile.Path}");
+						   }
+					   }
 
-					Repository.Connection.Disconnect();
+					   Repository.Connection.Disconnect();
 
-					using (P4.Repository SubmitRepository = GetConnection(ReadOnly: false, ClientFromChange: Change, UseClientFromChange: true, Username: Changelist.OwnerName))
-					{
-						// we might not need a client here, possibly -e below facilitates this, check!
-						using (P4.P4Command SubmitCommand = new P4.P4Command(SubmitRepository, "submit", null, true, new string[] { "-e", Change.ToString(CultureInfo.InvariantCulture) }))
-						{
-							Result = SubmitCommand.Run();
+					   using (P4.Repository SubmitRepository = GetConnection(ReadOnly: false, ClientFromChange: Change, UseClientFromChange: true, Username: Changelist.OwnerName))
+					   {
+							// we might not need a client here, possibly -e below facilitates this, check!
+							using (P4.P4Command SubmitCommand = new P4.P4Command(SubmitRepository, "submit", null, true, new string[] { "-e", Change.ToString(CultureInfo.InvariantCulture) }))
+						   {
+							   Result = SubmitCommand.Run();
 
-							if (!Result.Success)
-							{
-								throw new Exception($"Unable to submit {Change}");
-							}
+							   if (!Result.Success)
+							   {
+								   throw new Exception($"Unable to submit {Change}");
+							   }
 
-							int? SubmittedChangeId = null;
+							   int? SubmittedChangeId = null;
 
-							foreach (P4.TaggedObject TaggedObject in Result.TaggedOutput)
-							{
-								string? Submitted;
-								if (TaggedObject.TryGetValue("submittedChange", out Submitted))
-								{
-									SubmittedChangeId = int.Parse(Submitted, CultureInfo.InvariantCulture);
-								}
-							}
+							   foreach (P4.TaggedObject TaggedObject in Result.TaggedOutput)
+							   {
+								   string? Submitted;
+								   if (TaggedObject.TryGetValue("submittedChange", out Submitted))
+								   {
+									   SubmittedChangeId = int.Parse(Submitted, CultureInfo.InvariantCulture);
+								   }
+							   }
 
-							if (SubmittedChangeId == null)
-							{
-								throw new Exception($"Unable to get submitted changelist");
-							}
+							   if (SubmittedChangeId == null)
+							   {
+								   throw new Exception($"Unable to get submitted changelist");
+							   }
 
-							return (SubmittedChangeId, Changelist.Description);
+							   return (SubmittedChangeId, Changelist.Description);
 
-						}
-					}
-				}
-			}
+						   }
+					   }
+				   }
+			   }
 
-			throw new Exception($"Unable to get shelve change: {Change} for original change: {OriginalChange}");
+			   throw new Exception($"Unable to get shelve change: {Change} for original change: {OriginalChange}");
+		   });
 
 		}
 
 		/// <inheritdoc/>		
 		public override async Task<PerforceUserInfo?> GetUserInfoAsync(string UserName)
 		{
-			await Task.Yield();
-
-			using (P4.Repository Repository = GetConnection(NoClient: true))
+			return await Task.Run(() =>
 			{
-				P4.User User = Repository.GetUser(UserName, new P4.UserCmdOptions(P4.UserCmdFlags.Output));
 
-				if (User == null)
+				using (P4.Repository Repository = GetConnection(NoClient: true))
 				{
-					return null;
-				}
+					P4.User User = Repository.GetUser(UserName, new P4.UserCmdOptions(P4.UserCmdFlags.Output));
 
-				return new PerforceUserInfo { Name = UserName, Email = User.EmailAddress };
-			}
+					if (User == null)
+					{
+						return null;
+					}
+
+					return new PerforceUserInfo { Name = UserName, Email = User.EmailAddress };
+				}
+			});
 		}
 
 		void RequireEdgeServer()
