Horde: Various improvements to installation experience.

* Add a shortcut to the start menu, which launches the server and opens a browser
* Add an icon to the HordeServer executable
* Allow overriding configuration via command line

diff --git a/Engine/Source/Programs/Horde/HordeServer/HordeServer.csproj b/Engine/Source/Programs/Horde/HordeServer/HordeServer.csproj
--- a/Engine/Source/Programs/Horde/HordeServer/HordeServer.csproj
+++ b/Engine/Source/Programs/Horde/HordeServer/HordeServer.csproj
@@ -6,6 +6,7 @@
     <Deterministic>true</Deterministic>
     <AssemblyName>HordeServer</AssemblyName>
     <DocumentationFile>HordeServer.xml</DocumentationFile>
+    <ApplicationIcon>HordeServer.ico</ApplicationIcon>
   </PropertyGroup>
 
   <ItemGroup>
diff --git a/Engine/Source/Programs/Horde/HordeServer/Program.cs b/Engine/Source/Programs/Horde/HordeServer/Program.cs
--- a/Engine/Source/Programs/Horde/HordeServer/Program.cs
+++ b/Engine/Source/Programs/Horde/HordeServer/Program.cs
@@ -91,6 +91,7 @@ namespace HordeServer
 				.AddJsonFile("appsettings.User.json", optional: true)
 				.AddJsonFile(UserConfigFile.FullName, optional: true)
 				.AddEnvironmentVariables()
+				.AddCommandLine(Args)
 				.Build();
 
 			ServerSettings HordeSettings = new ServerSettings();
diff --git a/Engine/Source/Programs/Horde/HordeServer/ServerSettings.cs b/Engine/Source/Programs/Horde/HordeServer/ServerSettings.cs
--- a/Engine/Source/Programs/Horde/HordeServer/ServerSettings.cs
+++ b/Engine/Source/Programs/Horde/HordeServer/ServerSettings.cs
@@ -425,6 +425,11 @@ namespace HordeServer
 
 		private TimeZoneInfo? CachedTimeZoneInfo;
 
+		/// <summary>
+		/// Whether to open a browser on startup
+		/// </summary>
+		public bool OpenBrowser { get; set; } = false;
+
 		/// <summary>
 		/// Check if Helix Swarm is enabled
 		/// </summary>
diff --git a/Engine/Source/Programs/Horde/HordeServer/Startup.cs b/Engine/Source/Programs/Horde/HordeServer/Startup.cs
--- a/Engine/Source/Programs/Horde/HordeServer/Startup.cs
+++ b/Engine/Source/Programs/Horde/HordeServer/Startup.cs
@@ -69,6 +69,8 @@ using StatsdClient;
 using Status = Grpc.Core.Status;
 using HordeServer.Notifications.Impl;
 using HordeServer.Notifications;
+using Microsoft.AspNetCore.Hosting.Server.Features;
+using System.Runtime.InteropServices;
 
 namespace HordeServer
 {
@@ -657,7 +659,7 @@ namespace HordeServer
 		}
 
 		// This method gets called by the runtime. Use this method to configure the HTTP request pipeline.
-		public void Configure(IApplicationBuilder App, IWebHostEnvironment Env)
+		public void Configure(IApplicationBuilder App, IWebHostEnvironment Env, Microsoft.Extensions.Hosting.IHostApplicationLifetime Lifetime)
 		{
 			App.UseForwardedHeaders();
 
@@ -717,6 +719,20 @@ namespace HordeServer
 
 				Endpoints.MapControllers();
 			});
+
+			if (Settings.OpenBrowser && RuntimeInformation.IsOSPlatform(OSPlatform.Windows))
+			{
+				Lifetime.ApplicationStarted.Register(() => LaunchBrowser(App));
+			}
+		}
+
+		static void LaunchBrowser(IApplicationBuilder App)
+		{
+			IServerAddressesFeature Feature = App.ServerFeatures.Get<IServerAddressesFeature>();
+			if (Feature.Addresses.Count > 0)
+			{
+				Process.Start(new ProcessStartInfo { FileName = Feature.Addresses.First(), UseShellExecute = true });
+			}
 		}
 	}
 }
