Horde: Register a user when using anonymous auth handler, to prevent 404s when querying /api/v1/user endpoint.

diff --git a/Engine/Source/Programs/Horde/HordeServer/Authentication/AnonymousHandler.cs b/Engine/Source/Programs/Horde/HordeServer/Authentication/AnonymousHandler.cs
--- a/Engine/Source/Programs/Horde/HordeServer/Authentication/AnonymousHandler.cs
+++ b/Engine/Source/Programs/Horde/HordeServer/Authentication/AnonymousHandler.cs
@@ -1,6 +1,9 @@
 // Copyright Epic Games, Inc. All Rights Reserved.
 
+using HordeServer.Collections;
+using HordeServer.Models;
 using HordeServer.Services;
+using HordeServer.Utilities;
 using Microsoft.AspNetCore.Authentication;
 using Microsoft.Extensions.Logging;
 using Microsoft.Extensions.Options;
@@ -23,25 +26,32 @@ namespace HordeServer.Authentication
 	{
 		public const string AuthenticationScheme = "Anonymous";
 
-		public AnonymousAuthenticationHandler(IOptionsMonitor<AnonymousAuthenticationOptions> Options,
+		IUserCollection UserCollection;
+
+		public AnonymousAuthenticationHandler(IUserCollection UserCollection, IOptionsMonitor<AnonymousAuthenticationOptions> Options,
 			ILoggerFactory Logger, UrlEncoder Encoder, ISystemClock Clock)
 			: base(Options, Logger, Encoder, Clock)
 		{
+			this.UserCollection = UserCollection;
 		}
 
-		protected override Task<AuthenticateResult> HandleAuthenticateAsync()
+		protected override async Task<AuthenticateResult> HandleAuthenticateAsync()
 		{
 			Claim[] Claims =
 			{
 				new Claim(ClaimTypes.Name, AuthenticationScheme),
 				new Claim(Options.AdminClaimType, Options.AdminClaimValue)
 			};
+
+			IUser User = await UserCollection.FindOrAddUserByLoginAsync("anonymous", "Anonymous", "[REDACTED_EMAIL]");
+
 			ClaimsIdentity Identity = new ClaimsIdentity(Claims, Scheme.Name);
+			Identity.AddClaim(new Claim(HordeClaimTypes.UserId, User.Id.ToString()));
 
 			ClaimsPrincipal Principal = new ClaimsPrincipal(Identity);
 			AuthenticationTicket Ticket = new AuthenticationTicket(Principal, Scheme.Name);
 
-			return Task.FromResult(AuthenticateResult.Success(Ticket));
+			return AuthenticateResult.Success(Ticket);
 		}
 	}
 }
