Horde: Add an /api/v1/server/version endpoint for querying the server version.

diff --git a/Engine/Source/Programs/Horde/HordeServer/Controllers/ServerController.cs b/Engine/Source/Programs/Horde/HordeServer/Controllers/ServerController.cs
new file mode 100644
--- /dev/null
+++ b/Engine/Source/Programs/Horde/HordeServer/Controllers/ServerController.cs
@@ -0,0 +1,54 @@
+// Copyright Epic Games, Inc. All Rights Reserved.
+
+using EpicGames.Core;
+using HordeServer.Api;
+using HordeServer.Collections;
+using HordeCommon;
+using HordeServer.Models;
+using HordeServer.Services;
+using HordeServer.Utilities;
+using Microsoft.AspNetCore.Authorization;
+using Microsoft.AspNetCore.Authorization.Infrastructure;
+using Microsoft.AspNetCore.Mvc;
+using Microsoft.Extensions.Options;
+using MongoDB.Bson;
+using System;
+using System.Collections.Generic;
+using System.IdentityModel.Tokens.Jwt;
+using System.IO;
+using System.Linq;
+using System.Runtime.InteropServices;
+using System.Security.Claims;
+using System.Threading.Tasks;
+using System.Reflection;
+using System.Diagnostics;
+
+namespace HordeServer.Controllers
+{
+	/// <summary>
+	/// Controller managing account status
+	/// </summary>
+	[ApiController]
+	[Authorize]
+	[Route("[controller]")]
+	public class ServerController : ControllerBase
+	{
+		/// <summary>
+		/// Constructor
+		/// </summary>
+		public ServerController()
+		{
+		}
+
+		/// <summary>
+		/// Force a reset on the database
+		/// </summary>
+		[HttpGet]
+		[Route("/api/v1/server/version")]
+		public ActionResult GetVersionAsync()
+		{
+			FileVersionInfo FileVersionInfo = FileVersionInfo.GetVersionInfo(Assembly.GetExecutingAssembly().Location);
+			return Ok(FileVersionInfo.ProductVersion);
+		}
+	}
+}
