Horde: Move script for building installer into a public folder.

diff --git a/Engine/Source/Programs/Horde/Installer/Installer.xml b/Engine/Source/Programs/Horde/Installer/Installer.xml
new file mode 100644
--- /dev/null
+++ b/Engine/Source/Programs/Horde/Installer/Installer.xml
@@ -0,0 +1,51 @@
+<?xml version='1.0' ?>
+<BuildGraph xmlns="http://www.epicgames.com/BuildGraph" xmlns:xsi="http://www.w3.org/2001/XMLSchema-instance" xsi:schemaLocation="http://www.epicgames.com/BuildGraph ../../Engine/Build/Graph/Schema.xsd" >
+
+	<Option Name="PublishDir" DefaultValue="" Description="Directory to publish files to"/>
+	<Option Name="VersionArguments" DefaultValue="" Description="Version arguments to pass through to MSBuild"/>
+
+	<Agent Name="HordeInstaller" Type="Win64_Docker">
+		<Property Name="StagingDir" Value="$(RootDir)/Staging"/>
+
+		<Node Name="Stage HordeInstaller">
+			<Delete Files="$(StagingDir)/..."/>
+			<DotNet Arguments="publish &quot;$(RootDir)/Engine/Source/Programs/Horde/HordeServer/HordeServer.csproj&quot; --output &quot;$(StagingDir)/Server&quot; --runtime win-x64 $(VersionArguments)"/>
+			<DotNet Arguments="publish &quot;$(RootDir)/Engine/Source/Programs/Horde/HordeAgent/HordeAgent.csproj&quot; --output &quot;$(StagingDir)/Agent&quot; $(VersionArguments)"/>
+	
+			<Spawn Exe="cmd.exe" Arguments="/C yarn install" WorkingDir="$(RootDir)/Engine/Source/Programs/Horde/HordeDashboard"/>
+			<Spawn Exe="cmd.exe" Arguments="/C yarn run build" WorkingDir="$(RootDir)/Engine/Source/Programs/Horde/HordeDashboard"/>
+	
+			<Copy From="$(RootDir)/Engine/Source/Programs/Horde/HordeDashboard/build/..." To="$(StagingDir)/Server/wwwroot/..."/>
+		</Node>
+
+		<Property Name="WixDir" Value="$(RootDir)/Engine/Source/ThirdParty/WiX/3.8"/>
+		<Property Name="SourceDir" Value="$(RootDir)/Engine/Source/Programs/Horde/Installer"/>
+		<Property Name="ObjDir" Value="obj/Debug"/>
+		<Property Name="BinDir" Value="bin/Debug"/>
+		
+		<Node Name="Build HordeInstaller" Requires="Stage HordeInstaller">
+			<Property Name="CommonArgs" Value="-dConfiguration=Debug -dPlatform=x64 -arch x64"/>
+
+			<Delete Files="$(StagingDir)/Server-Bulk/..."/>
+			<Tag Files="$(StagingDir)/Server/..." Except="HordeServer.exe" With="#ServerFiles"/>
+			<Copy Files="#ServerFiles" From="$(StagingDir)/Server" To="$(StagingDir)/Server-Bulk/Server"/>
+			<Spawn Exe="$(WixDir)/heat.exe" Arguments="dir &quot;$(StagingDir)/Server-Bulk&quot; -cg HordeServerFiles -scom -sreg -gg -dr InstallDir -srd -var var.SourceDir -out &quot;$(SourceDir)/HordeServerFiles.wxs&quot;"/>
+			<Spawn Exe="$(WixDir)/candle.exe" Arguments="$(CommonArgs) HordeServerFiles.wxs -dSourceDir=&quot;$(StagingDir)&quot; -out $(ObjDir)/HordeServerFiles.wixobj" WorkingDir="$(SourceDir)"/>
+
+			<Delete Files="$(StagingDir)/Agent-Bulk/..."/>
+			<Tag Files="$(StagingDir)/Agent/..." Except="HordeAgent.exe" With="#AgentFiles"/>
+			<Copy Files="#AgentFiles" From="$(StagingDir)/Agent" To="$(StagingDir)/Agent-Bulk/Agent"/>
+			<Spawn Exe="$(WixDir)/heat.exe" Arguments="dir &quot;$(StagingDir)/Agent-Bulk&quot; -cg HordeAgentFiles -scom -sreg -gg -dr InstallDir -srd -var var.SourceDir -out &quot;$(SourceDir)/HordeAgentFiles.wxs&quot;"/>
+			<Spawn Exe="$(WixDir)/candle.exe" Arguments="$(CommonArgs) HordeAgentFiles.wxs -dSourceDir=&quot;$(StagingDir)&quot; -out $(ObjDir)/HordeAgentFiles.wixobj" WorkingDir="$(SourceDir)"/>
+
+			<Spawn Exe="$(WixDir)/candle.exe" Arguments="$(CommonArgs) Installer.wxs -dStagingDir=&quot;$(StagingDir)&quot; -out $(ObjDir)/Installer.wixobj" WorkingDir="$(SourceDir)"/>
+			<Spawn Exe="$(WixDir)/light.exe" Arguments="-out $(BinDir)/Horde.msi -sw1076 -pdbout $(BinDir)/Horde.pdb $(ObjDir)/Installer.wixobj $(ObjDir)/HordeServerFiles.wixobj $(ObjDir)/HordeAgentFiles.wixobj" WorkingDir="$(SourceDir)"/>
+		</Node>
+		
+		<Node Name="Publish HordeInstaller" Requires="Build HordeInstaller">
+			<Error Message="PublishDir is not set." If="'$(PublishDir)' == ''"/>
+			<Copy From="$(SourceDir)/$(BinDir)/..." To="$(PublishDir)/..."/>
+		</Node>
+	</Agent>
+
+</BuildGraph>
\ No newline at end of file
