Horde: Use version numbers from the agent executable for versioning software.

diff --git a/Engine/Source/Programs/Horde/HordeAgent/AgentSettings.cs b/Engine/Source/Programs/Horde/HordeAgent/AgentSettings.cs
--- a/Engine/Source/Programs/Horde/HordeAgent/AgentSettings.cs
+++ b/Engine/Source/Programs/Horde/HordeAgent/AgentSettings.cs
@@ -156,11 +156,6 @@ namespace HordeAgent
 		[Required]
 		public string Server { get; set; } = null!;
 
-		/// <summary>
-		/// The unique version identifier for this client. This is inserted into the appsettings.json file when a client is uploaded to the server.
-		/// </summary>
-		public string? Version { get; set; }
-
 		/// <summary>
 		/// The executor to use for jobs. Defaults to the Perforce executor.
 		/// </summary>
diff --git a/Engine/Source/Programs/Horde/HordeAgent/Commands/Service/RunCommand.cs b/Engine/Source/Programs/Horde/HordeAgent/Commands/Service/RunCommand.cs
--- a/Engine/Source/Programs/Horde/HordeAgent/Commands/Service/RunCommand.cs
+++ b/Engine/Source/Programs/Horde/HordeAgent/Commands/Service/RunCommand.cs
@@ -103,9 +103,9 @@ namespace HordeAgent.Modes.Service
 					ConfigSection.Bind(Settings);
 
 					ServerProfile ServerProfile = Settings.GetCurrentServerProfile();
-					ConfigureTracing(ServerProfile.Environment, Settings.Version ?? "unknown");
-					
-					Logging.Version = Settings.Version ?? "";
+					ConfigureTracing(ServerProfile.Environment, Program.Version);
+
+					Logging.Version = Program.Version;
 
 					Services.AddHttpClient(Program.HordeServerClientName, Config =>
 					{
diff --git a/Engine/Source/Programs/Horde/HordeAgent/Program.cs b/Engine/Source/Programs/Horde/HordeAgent/Program.cs
--- a/Engine/Source/Programs/Horde/HordeAgent/Program.cs
+++ b/Engine/Source/Programs/Horde/HordeAgent/Program.cs
@@ -40,7 +40,7 @@ namespace HordeAgent
 		/// <summary>
 		/// Path to the root application directory
 		/// </summary>
-		public static readonly string AppDir = Path.GetFullPath(Path.GetDirectoryName(Assembly.GetExecutingAssembly().Location)!);
+		public static string AppDir { get; } = Path.GetFullPath(Path.GetDirectoryName(Assembly.GetExecutingAssembly().Location)!);
 
 		/// <summary>
 		/// Name of the saved application data directory
@@ -50,13 +50,18 @@ namespace HordeAgent
 		/// <summary>
 		/// Path to the saved application data directory
 		/// </summary>
-		public static readonly string SavedDir = Path.Combine(AppDir, SavedDirName);
+		public static string SavedDir { get; } = Path.Combine(AppDir, SavedDirName);
 
 		/// <summary>
 		/// The launch arguments
 		/// </summary>
 		public static string[] Args { get; private set; } = null!;
 
+		/// <summary>
+		/// The current application version
+		/// </summary>
+		public static string Version { get; } = GetVersion();
+
 		/// <summary>
 		/// Width to use for printing out help
 		/// </summary>
@@ -237,5 +242,21 @@ namespace HordeAgent
 			}
 			HelpUtils.PrintTable(Commands.OrderBy(x => x.Key).ToList(), 4, 20, HelpWidth, Logger);
 		}
+
+		/// <summary>
+		/// Gets the version of the current assembly
+		/// </summary>
+		/// <returns></returns>
+		static string GetVersion()
+		{
+			try
+			{
+				return FileVersionInfo.GetVersionInfo(Assembly.GetExecutingAssembly().Location).ProductVersion;
+			}
+			catch
+			{
+				return "unknown";
+			}
+		}
 	}
 }
diff --git a/Engine/Source/Programs/Horde/HordeAgent/Services/WorkerService.cs b/Engine/Source/Programs/Horde/HordeAgent/Services/WorkerService.cs
--- a/Engine/Source/Programs/Horde/HordeAgent/Services/WorkerService.cs
+++ b/Engine/Source/Programs/Horde/HordeAgent/Services/WorkerService.cs
@@ -251,7 +251,7 @@ namespace HordeAgent.Services
 			Logger.LogInformation("Arguments: {Arguments}", Environment.CommandLine);
 
 			// Show the current client id
-			string Version = Settings.Version ?? "(unknown)";
+			string Version = Program.Version;
 			Logger.LogInformation("Version: {Version}", Version);
 
 			// Keep trying to start an agent session with the server
@@ -372,7 +372,7 @@ namespace HordeAgent.Services
 				SessionRequest.Name = Environment.MachineName;
 				SessionRequest.Status = AgentStatus.Ok;
 				SessionRequest.Capabilities = Capabilities;
-				SessionRequest.Version = Settings.Version ?? "";
+				SessionRequest.Version = Program.Version;
 
 				// Create a session
 				CreateSessionResponse = await RpcClient.CreateSessionAsync(SessionRequest, null, null, StoppingToken);
@@ -1134,9 +1134,9 @@ namespace HordeAgent.Services
 			string RequiredVersion = UpgradeTask.SoftwareId;
 
 			// Check if we're running the right version
-			if (RequiredVersion != null && RequiredVersion != Settings.Version)
+			if (RequiredVersion != null && RequiredVersion != Program.Version)
 			{
-				Logger.LogInformation("Upgrading from {CurrentVersion} to {TargetVersion}", String.IsNullOrEmpty(Settings.Version)? "(unknown)" : Settings.Version, RequiredVersion);
+				Logger.LogInformation("Upgrading from {CurrentVersion} to {TargetVersion}", Program.Version, RequiredVersion);
 
 				// Clear out the working directory
 				DirectoryReference UpgradeDir = DirectoryReference.Combine(WorkingDir, "Upgrade");
diff --git a/Engine/Source/Programs/Horde/HordeServer/Api/Agents.cs b/Engine/Source/Programs/Horde/HordeServer/Api/Agents.cs
--- a/Engine/Source/Programs/Horde/HordeServer/Api/Agents.cs
+++ b/Engine/Source/Programs/Horde/HordeServer/Api/Agents.cs
@@ -260,7 +260,7 @@ namespace HordeServer.Api
 			this.StartTime = Session.StartTime;
 			this.FinishTime = Session.FinishTime;
 			this.Capabilities = (Session.Capabilities != null)? new GetAgentCapabilitiesResponse(Session.Capabilities) : null;
-			this.Version = Session.Version.ToString();
+			this.Version = Session.Version?.ToString();
 		}
 	}
 
@@ -474,7 +474,7 @@ namespace HordeServer.Api
 			{
 				this.Version += $" ({Agent.Channel})";
 			}
-			this.Version = Agent.Version.ToString();
+			this.Version = Agent.Version?.ToString();
 			this.Channel = Agent.Channel?.ToString();
 			this.UpdateTime = Agent.UpdateTime;
 			this.Pools = Pools.Select(x => x.ToString()).ToList();
diff --git a/Engine/Source/Programs/Horde/HordeServer/Collections/IAgentCollection.cs b/Engine/Source/Programs/Horde/HordeServer/Collections/IAgentCollection.cs
--- a/Engine/Source/Programs/Horde/HordeServer/Collections/IAgentCollection.cs
+++ b/Engine/Source/Programs/Horde/HordeServer/Collections/IAgentCollection.cs
@@ -17,7 +17,6 @@ using System.Linq.Expressions;
 using System.Threading.Tasks;
 
 using PoolId = HordeServer.Utilities.StringId<HordeServer.Models.IPool>;
-using AgentSoftwareVersion = HordeServer.Utilities.StringId<HordeServer.Collections.IAgentSoftwareCollection>;
 using AgentSoftwareChannelName = HordeServer.Utilities.StringId<HordeServer.Services.AgentSoftwareChannels>;
 
 namespace HordeServer.Collections
@@ -112,7 +111,7 @@ namespace HordeServer.Collections
 		/// <param name="Capabilities">Capabilities for the agent</param>
 		/// <param name="Version">Current version of the agent software</param>
 		/// <returns>New agent state</returns>
-		Task<IAgent?> TryStartSessionAsync(IAgent Agent, ObjectId SessionId, DateTime SessionExpiresAt, AgentStatus Status, AgentCapabilities Capabilities, AgentSoftwareVersion? Version);
+		Task<IAgent?> TryStartSessionAsync(IAgent Agent, ObjectId SessionId, DateTime SessionExpiresAt, AgentStatus Status, AgentCapabilities Capabilities, string? Version);
 
 		/// <summary>
 		/// Attempt to update the agent state
diff --git a/Engine/Source/Programs/Horde/HordeServer/Collections/IAgentSoftwareCollection.cs b/Engine/Source/Programs/Horde/HordeServer/Collections/IAgentSoftwareCollection.cs
--- a/Engine/Source/Programs/Horde/HordeServer/Collections/IAgentSoftwareCollection.cs
+++ b/Engine/Source/Programs/Horde/HordeServer/Collections/IAgentSoftwareCollection.cs
@@ -16,8 +16,6 @@ using System.Linq;
 using System.Text.Json;
 using System.Threading.Tasks;
 
-using AgentSoftwareVersion = HordeServer.Utilities.StringId<HordeServer.Collections.IAgentSoftwareCollection>;
-
 namespace HordeServer.Collections
 {
 	/// <summary>
@@ -31,19 +29,26 @@ namespace HordeServer.Collections
 		/// <param name="Version">The version number</param>
 		/// <param name="Data">Zip file containing the new software</param>
 		/// <returns>New software instance</returns>
-		Task AddAsync(AgentSoftwareVersion Version, byte[] Data);
+		Task<bool> AddAsync(string Version, byte[] Data);
+
+		/// <summary>
+		/// Tests whether a given version exists
+		/// </summary>
+		/// <param name="Version">Version of the software</param>
+		/// <returns>True if it exists, false otherwise</returns>
+		Task<bool> ExistsAsync(string Version);
 
 		/// <summary>
 		/// Removes a software archive 
 		/// </summary>
 		/// <param name="Version">Version of the software to delete</param>
-		Task<bool> RemoveAsync(AgentSoftwareVersion Version);
+		Task<bool> RemoveAsync(string Version);
 
 		/// <summary>
 		/// Downloads software of a given revision
 		/// </summary>
 		/// <param name="Version">Version of the software</param>
 		/// <returns>Data for the given software</returns>
-		Task<byte[]?> GetAsync(AgentSoftwareVersion Version);
+		Task<byte[]?> GetAsync(string Version);
 	}
 }
diff --git a/Engine/Source/Programs/Horde/HordeServer/Collections/ISessionCollection.cs b/Engine/Source/Programs/Horde/HordeServer/Collections/ISessionCollection.cs
--- a/Engine/Source/Programs/Horde/HordeServer/Collections/ISessionCollection.cs
+++ b/Engine/Source/Programs/Horde/HordeServer/Collections/ISessionCollection.cs
@@ -12,8 +12,6 @@ using System.Collections.Generic;
 using System.Linq;
 using System.Threading.Tasks;
 
-using AgentSoftwareVersion = HordeServer.Utilities.StringId<HordeServer.Collections.IAgentSoftwareCollection>;
-
 namespace HordeServer.Collections
 {
 	/// <summary>
@@ -29,7 +27,7 @@ namespace HordeServer.Collections
 		/// <param name="StartTime">Start time of this session</param>
 		/// <param name="Capabilities">Properties of this agent at the time the session started</param>
 		/// <param name="Version">Version of the agent software</param>
-		Task<ISession> AddAsync(ObjectId Id, AgentId AgentId, DateTime StartTime, AgentCapabilities? Capabilities, AgentSoftwareVersion? Version);
+		Task<ISession> AddAsync(ObjectId Id, AgentId AgentId, DateTime StartTime, AgentCapabilities? Capabilities, string? Version);
 
 		/// <summary>
 		/// Gets information about a particular session
diff --git a/Engine/Source/Programs/Horde/HordeServer/Collections/Impl/AgentCollection.cs b/Engine/Source/Programs/Horde/HordeServer/Collections/Impl/AgentCollection.cs
--- a/Engine/Source/Programs/Horde/HordeServer/Collections/Impl/AgentCollection.cs
+++ b/Engine/Source/Programs/Horde/HordeServer/Collections/Impl/AgentCollection.cs
@@ -17,7 +17,6 @@ using System.Linq.Expressions;
 using System.Threading.Tasks;
 
 using PoolId = HordeServer.Utilities.StringId<HordeServer.Models.IPool>;
-using AgentSoftwareVersion = HordeServer.Utilities.StringId<HordeServer.Collections.IAgentSoftwareCollection>;
 using AgentSoftwareChannelName = HordeServer.Utilities.StringId<HordeServer.Services.AgentSoftwareChannels>;
 
 namespace HordeServer.Collections.Impl
@@ -51,7 +50,7 @@ namespace HordeServer.Collections.Impl
 			public bool Deleted { get; set; }
 
 			[BsonElement("Version2")]
-			public AgentSoftwareVersion? Version { get; set; }
+			public string? Version { get; set; }
 
 			[BsonIgnoreIfNull]
 			public AgentSoftwareChannelName? Channel { get; set; }
@@ -415,7 +414,7 @@ namespace HordeServer.Collections.Impl
 		}
 
 		/// <inheritdoc/>
-		public async Task<IAgent?> TryStartSessionAsync(IAgent AgentInterface, ObjectId SessionId, DateTime SessionExpiresAt, AgentStatus Status, AgentCapabilities Capabilities, AgentSoftwareVersion? Version)
+		public async Task<IAgent?> TryStartSessionAsync(IAgent AgentInterface, ObjectId SessionId, DateTime SessionExpiresAt, AgentStatus Status, AgentCapabilities Capabilities, string? Version)
 		{
 			AgentDocument Agent = (AgentDocument)AgentInterface;
 
diff --git a/Engine/Source/Programs/Horde/HordeServer/Collections/Impl/AgentSoftwareCollection.cs b/Engine/Source/Programs/Horde/HordeServer/Collections/Impl/AgentSoftwareCollection.cs
--- a/Engine/Source/Programs/Horde/HordeServer/Collections/Impl/AgentSoftwareCollection.cs
+++ b/Engine/Source/Programs/Horde/HordeServer/Collections/Impl/AgentSoftwareCollection.cs
@@ -4,6 +4,7 @@ using EpicGames.Core;
 using HordeCommon;
 using HordeServer.Models;
 using HordeServer.Services;
+using HordeServer.Utilities;
 using MongoDB.Bson;
 using MongoDB.Bson.Serialization.Attributes;
 using MongoDB.Driver;
@@ -31,16 +32,17 @@ namespace HordeServer.Collections.Impl
 		class AgentSoftwareDocument
 		{
 			[BsonId]
-			public AgentSoftwareVersion Id { get; set; }
+			public string Id { get; set; }
 			public byte[] Data { get; set; } = null!;
 
 			[BsonConstructor]
 			private AgentSoftwareDocument()
 			{
+				this.Id = null!;
 				this.Data = null!;
 			}
 
-			public AgentSoftwareDocument(AgentSoftwareVersion Id, byte[] Data)
+			public AgentSoftwareDocument(string Id, byte[] Data)
 			{
 				this.Id = Id;
 				this.Data = Data;
@@ -62,75 +64,27 @@ namespace HordeServer.Collections.Impl
 		}
 
 		/// <inheritdoc/>
-		public async Task AddAsync(AgentSoftwareVersion Version, byte[] Data)
+		public async Task<bool> AddAsync(string Version, byte[] Data)
 		{
-			byte[] NewData = UpdateVersion(Data, Version);
-
-			AgentSoftwareDocument Document = new AgentSoftwareDocument(Version, NewData);
-			await Collection.InsertOneAsync(Document);
+			AgentSoftwareDocument Document = new AgentSoftwareDocument(Version, Data);
+			return await Collection.InsertOneIgnoreDuplicatesAsync(Document);
 		}
 
-		/// <summary>
-		/// Updates the client id within the archive in the data array
-		/// </summary>
-		/// <param name="Data">Data for the zip archive</param>
-		/// <param name="Version">The new version identifier</param>
-		/// <returns>New agent app data</returns>
-		public static byte[] UpdateVersion(byte[] Data, AgentSoftwareVersion Version)
+		/// <inheritdoc/>
+		public async Task<bool> ExistsAsync(string Version)
 		{
-			bool bWrittenClientId = false;
-
-			MemoryStream OutputStream = new MemoryStream();
-			using (ZipArchive OutputArchive = new ZipArchive(OutputStream, ZipArchiveMode.Create, true))
-			{
-				MemoryStream InputStream = new MemoryStream(Data);
-				using (ZipArchive InputArchive = new ZipArchive(InputStream, ZipArchiveMode.Read, true))
-				{
-					foreach (ZipArchiveEntry InputEntry in InputArchive.Entries)
-					{
-						ZipArchiveEntry OutputEntry = OutputArchive.CreateEntry(InputEntry.FullName);
-
-						using System.IO.Stream InputEntryStream = InputEntry.Open();
-						using System.IO.Stream OutputEntryStream = OutputEntry.Open();
-
-						if (InputEntry.FullName.Equals("appsettings.json", StringComparison.OrdinalIgnoreCase))
-						{
-							using MemoryStream MemoryStream = new MemoryStream();
-							InputEntryStream.CopyTo(MemoryStream);
-
-							Dictionary<string, Dictionary<string, object>> Document = JsonSerializer.Deserialize<Dictionary<string, Dictionary<string, object>>>(MemoryStream.ToArray());
-							Document["Horde"]["Version"] = Version.ToString();
-
-							using Utf8JsonWriter Writer = new Utf8JsonWriter(OutputEntryStream, new JsonWriterOptions { Indented = true });
-							JsonSerializer.Serialize<Dictionary<string, Dictionary<string, object>>>(Writer, Document, new JsonSerializerOptions { WriteIndented = true });
-
-							bWrittenClientId = true;
-						}
-						else
-						{
-							InputEntryStream.CopyTo(OutputEntryStream);
-						}
-					}
-				}
-			}
-
-			if (!bWrittenClientId)
-			{
-				throw new InvalidDataException("Missing appsettings.json file from zip archive");
-			}
-
-			return OutputStream.ToArray();
+			return await Collection.Find(x => x.Id == Version).AnyAsync();
 		}
 
 		/// <inheritdoc/>
-		public async Task<bool> RemoveAsync(AgentSoftwareVersion Version)
+		public async Task<bool> RemoveAsync(string Version)
 		{
 			DeleteResult Result = await Collection.DeleteOneAsync(x => x.Id == Version);
 			return Result.DeletedCount > 0;
 		}
 
 		/// <inheritdoc/>
-		public async Task<byte[]?> GetAsync(AgentSoftwareVersion Version)
+		public async Task<byte[]?> GetAsync(string Version)
 		{
 			IFindFluent<AgentSoftwareDocument, AgentSoftwareDocument> Query = Collection.Find(x => x.Id == Version);
 			AgentSoftwareDocument? Software = await Query.FirstOrDefaultAsync();
diff --git a/Engine/Source/Programs/Horde/HordeServer/Collections/Impl/SessionCollection.cs b/Engine/Source/Programs/Horde/HordeServer/Collections/Impl/SessionCollection.cs
--- a/Engine/Source/Programs/Horde/HordeServer/Collections/Impl/SessionCollection.cs
+++ b/Engine/Source/Programs/Horde/HordeServer/Collections/Impl/SessionCollection.cs
@@ -12,8 +12,6 @@ using System.Collections.Generic;
 using System.Linq;
 using System.Threading.Tasks;
 
-using AgentSoftwareVersion = HordeServer.Utilities.StringId<HordeServer.Collections.IAgentSoftwareCollection>;
-
 namespace HordeServer.Collections.Impl
 {
 	/// <summary>
@@ -35,14 +33,14 @@ namespace HordeServer.Collections.Impl
 			public DateTime StartTime { get; set; }
 			public DateTime? FinishTime { get; set; }
 			public AgentCapabilities? Capabilities { get; set; }
-			public AgentSoftwareVersion? Version { get; set; }
+			public string? Version { get; set; }
 
 			[BsonConstructor]
 			private SessionDocument()
 			{
 			}
 
-			public SessionDocument(ObjectId Id, AgentId AgentId, DateTime StartTime, AgentCapabilities? Capabilities, AgentSoftwareVersion? Version)
+			public SessionDocument(ObjectId Id, AgentId AgentId, DateTime StartTime, AgentCapabilities? Capabilities, string? Version)
 			{
 				this.Id = Id;
 				this.AgentId = AgentId;
@@ -74,7 +72,7 @@ namespace HordeServer.Collections.Impl
 		}
 
 		/// <inheritdoc/>
-		public async Task<ISession> AddAsync(ObjectId Id, AgentId AgentId, DateTime StartTime, AgentCapabilities? Capabilities, AgentSoftwareVersion? Version)
+		public async Task<ISession> AddAsync(ObjectId Id, AgentId AgentId, DateTime StartTime, AgentCapabilities? Capabilities, string? Version)
 		{
 			SessionDocument NewSession = new SessionDocument(Id, AgentId, StartTime, Capabilities, Version);
 			await Sessions.InsertOneAsync(NewSession);
diff --git a/Engine/Source/Programs/Horde/HordeServer/Models/Agent.cs b/Engine/Source/Programs/Horde/HordeServer/Models/Agent.cs
--- a/Engine/Source/Programs/Horde/HordeServer/Models/Agent.cs
+++ b/Engine/Source/Programs/Horde/HordeServer/Models/Agent.cs
@@ -448,7 +448,7 @@ namespace HordeServer.Models
 		/// <summary>
 		/// Version of the software running on this agent
 		/// </summary>
-		public AgentSoftwareVersion? Version { get; }
+		public string? Version { get; }
 
 		/// <summary>
 		/// Channel for the software running on this agent. Uses <see cref="AgentSoftwareService.DefaultChannelName"/> if not specified
diff --git a/Engine/Source/Programs/Horde/HordeServer/Models/AgentSoftware.cs b/Engine/Source/Programs/Horde/HordeServer/Models/AgentSoftware.cs
--- a/Engine/Source/Programs/Horde/HordeServer/Models/AgentSoftware.cs
+++ b/Engine/Source/Programs/Horde/HordeServer/Models/AgentSoftware.cs
@@ -41,6 +41,6 @@ namespace HordeServer.Models
 		/// <summary>
 		/// The software version number
 		/// </summary>
-		public AgentSoftwareVersion Version { get; set; }
+		public string Version { get; set; }
 	}
 }
diff --git a/Engine/Source/Programs/Horde/HordeServer/Models/Session.cs b/Engine/Source/Programs/Horde/HordeServer/Models/Session.cs
--- a/Engine/Source/Programs/Horde/HordeServer/Models/Session.cs
+++ b/Engine/Source/Programs/Horde/HordeServer/Models/Session.cs
@@ -9,8 +9,6 @@ using System.Collections.Generic;
 using System.Linq;
 using System.Threading.Tasks;
 
-using AgentSoftwareVersion = HordeServer.Utilities.StringId<HordeServer.Collections.IAgentSoftwareCollection>;
-
 namespace HordeServer.Models
 {
 	/// <summary>
@@ -46,6 +44,6 @@ namespace HordeServer.Models
 		/// <summary>
 		/// Version of the agent software
 		/// </summary>
-		public AgentSoftwareVersion? Version { get; }
+		public string? Version { get; }
 	}
 }
diff --git a/Engine/Source/Programs/Horde/HordeServer/Services/AgentService.cs b/Engine/Source/Programs/Horde/HordeServer/Services/AgentService.cs
--- a/Engine/Source/Programs/Horde/HordeServer/Services/AgentService.cs
+++ b/Engine/Source/Programs/Horde/HordeServer/Services/AgentService.cs
@@ -268,7 +268,7 @@ namespace HordeServer.Services
 		/// <param name="Capabilities">Capabilities of the agent software</param>
 		/// <param name="Version">Version of the software that's running</param>
 		/// <returns>New agent state</returns>
-		public async Task<IAgent> CreateSessionAsync(IAgent Agent, AgentStatus Status, AgentCapabilities Capabilities, AgentSoftwareVersion? Version)
+		public async Task<IAgent> CreateSessionAsync(IAgent Agent, AgentStatus Status, AgentCapabilities Capabilities, string? Version)
 		{
 			for (; ; )
 			{
diff --git a/Engine/Source/Programs/Horde/HordeServer/Services/AgentSoftwareService.cs b/Engine/Source/Programs/Horde/HordeServer/Services/AgentSoftwareService.cs
--- a/Engine/Source/Programs/Horde/HordeServer/Services/AgentSoftwareService.cs
+++ b/Engine/Source/Programs/Horde/HordeServer/Services/AgentSoftwareService.cs
@@ -15,8 +15,12 @@ using System.Linq;
 using System.Threading.Channels;
 using System.Threading.Tasks;
 
-using AgentSoftwareVersion = HordeServer.Utilities.StringId<HordeServer.Collections.IAgentSoftwareCollection>;
 using AgentSoftwareChannelName = HordeServer.Utilities.StringId<HordeServer.Services.AgentSoftwareChannels>;
+using EpicGames.Core;
+using Microsoft.Extensions.Logging;
+using System.Runtime.InteropServices;
+using Microsoft.Extensions.Options;
+using System.IO.Compression;
 
 namespace HordeServer.Services
 {
@@ -43,7 +47,7 @@ namespace HordeServer.Services
 		/// <summary>
 		/// The software revision number
 		/// </summary>
-		public AgentSoftwareVersion Version { get; set; }
+		public string Version { get; set; } = String.Empty;
 
 		/// <summary>
 		/// Private constructor for serialization
@@ -78,6 +82,22 @@ namespace HordeServer.Services
 		/// List of channels
 		/// </summary>
 		public List<AgentSoftwareChannel> Channels { get; set; } = new List<AgentSoftwareChannel>();
+
+		/// <summary>
+		/// Finds an existing channel by the given name, or adds a new one
+		/// </summary>
+		/// <param name="Name"></param>
+		/// <returns></returns>
+		public AgentSoftwareChannel FindOrAddChannel(AgentSoftwareChannelName Name)
+		{
+			AgentSoftwareChannel? Channel = Channels.FirstOrDefault(x => x.Name == Name);
+			if (Channel == null)
+			{
+				Channel = new AgentSoftwareChannel(Name);
+				Channels.Add(Channel);
+			}
+			return Channel;
+		}
 	}
 
 	/// <summary>
@@ -165,7 +185,7 @@ namespace HordeServer.Services
 					break;
 				}
 
-				AgentSoftwareVersion Version = Current.Channels[ChannelIdx].Version;
+				string Version = Current.Channels[ChannelIdx].Version;
 				Current.Channels.RemoveAt(ChannelIdx);
 
 				if (await Singleton.TryUpdateAsync(Current))
@@ -180,29 +200,46 @@ namespace HordeServer.Services
 		}
 
 		/// <summary>
-		/// Updates a new software revision
+		/// Reads the version number from an archive
 		/// </summary>
-		/// <param name="Name">Name of the channel</param>
-		/// <param name="Author">Name of the user uploading this file</param>
-		/// <param name="Data">The input data stream. This should be a zip archive containing the HordeAgent executable.</param>
-		/// <returns>Unique id for the file</returns>
-		public async Task<AgentSoftwareVersion> SetArchiveAsync(AgentSoftwareChannelName Name, string? Author, byte[] Data)
+		/// <param name="Data">The archive data</param>
+		/// <returns></returns>
+		static string ReadVersion(byte[] Data)
 		{
-			// Allocate a new version number
-			int VersionNumber;
-			for(; ;)
+			MemoryStream InputStream = new MemoryStream(Data);
+			using (ZipArchive InputArchive = new ZipArchive(InputStream, ZipArchiveMode.Read, true))
 			{
-				AgentSoftwareChannels Instance = await Singleton.GetAsync();
-				VersionNumber = Instance.NextVersion++;
-
-				if(await Singleton.TryUpdateAsync(Instance))
+				foreach (ZipArchiveEntry InputEntry in InputArchive.Entries)
 				{
-					break;
+					if (InputEntry.FullName.Equals("HordeAgent.dll", StringComparison.OrdinalIgnoreCase))
+					{
+						string TempFile = Path.GetTempFileName();
+						try
+						{
+							InputEntry.ExtractToFile(TempFile);
+							return FileVersionInfo.GetVersionInfo(TempFile).ProductVersion;
+						}
+						finally
+						{
+							File.Delete(TempFile);
+						}
+					}
 				}
 			}
+			throw new Exception("Unable to find HordeAgent.dll in archive");
+		}
 
+		/// <summary>
+		/// Updates a new software revision
+		/// </summary>
+		/// <param name="Name">Name of the channel</param>
+		/// <param name="Author">Name of the user uploading this file</param>
+		/// <param name="Data">The input data stream. This should be a zip archive containing the HordeAgent executable.</param>
+		/// <returns>Unique id for the file</returns>
+		public async Task<string> SetArchiveAsync(AgentSoftwareChannelName Name, string? Author, byte[] Data)
+		{
 			// Upload the software
-			AgentSoftwareVersion Version = new AgentSoftwareVersion($"r{VersionNumber}");
+			string Version = ReadVersion(Data);
 			await Collection.AddAsync(Version, Data);
 
 			// Update the channel
@@ -210,13 +247,7 @@ namespace HordeServer.Services
 			{
 				AgentSoftwareChannels Instance = await Singleton.GetAsync();
 
-				AgentSoftwareChannel? Channel = Instance.Channels.FirstOrDefault(x => x.Name == Name);
-				if (Channel == null)
-				{
-					Channel = new AgentSoftwareChannel(Name);
-					Instance.Channels.Add(Channel);
-				}
-
+				AgentSoftwareChannel Channel = Instance.FindOrAddChannel(Name);
 				Channel.ModifiedBy = Author;
 				Channel.ModifiedTime = DateTime.UtcNow;
 				Channel.Version = Version;
@@ -230,7 +261,7 @@ namespace HordeServer.Services
 		}
 
 		/// <summary>
-		/// Gets the zip file with a given version number
+		/// Gets the zip file for a given channel
 		/// </summary>
 		/// <param name="Name">The channel name</param>
 		/// <returns>Data for the given archive</returns>
@@ -254,7 +285,7 @@ namespace HordeServer.Services
 		/// </summary>
 		/// <param name="Version">The version</param>
 		/// <returns>Data for the given archive</returns>
-		public Task<byte[]?> GetArchiveAsync(AgentSoftwareVersion Version)
+		public Task<byte[]?> GetArchiveAsync(string Version)
 		{
 			return Collection.GetAsync(Version);
 		}
diff --git a/Engine/Source/Programs/Horde/HordeServer/Services/RpcService.cs b/Engine/Source/Programs/Horde/HordeServer/Services/RpcService.cs
--- a/Engine/Source/Programs/Horde/HordeServer/Services/RpcService.cs
+++ b/Engine/Source/Programs/Horde/HordeServer/Services/RpcService.cs
@@ -324,7 +324,7 @@ namespace HordeServer.Services
 			}
 
 			// Create a new session
-			Agent = await AgentService.CreateSessionAsync(Agent, Request.Status, new AgentCapabilities(Request.Capabilities), AgentSoftwareVersion.FromNullable(Request.Version));
+			Agent = await AgentService.CreateSessionAsync(Agent, Request.Status, new AgentCapabilities(Request.Capabilities), Request.Version);
 			if (Agent == null)
 			{
 				throw new StructuredRpcException(StatusCode.NotFound, "Agent {AgentId} not found", Request.Name);
@@ -886,10 +886,10 @@ namespace HordeServer.Services
 				throw new StructuredRpcException(StatusCode.PermissionDenied, "Access to software is forbidden");
 			}
 
-			AgentSoftwareVersion Version = await AgentSoftwareService.SetArchiveAsync(new AgentSoftwareChannelName(Request.Channel), null, Request.Data.ToArray());
+			string Version = await AgentSoftwareService.SetArchiveAsync(new AgentSoftwareChannelName(Request.Channel), null, Request.Data.ToArray());
 
 			UploadSoftwareResponse Response = new UploadSoftwareResponse();
-			Response.Version = Version.ToString();
+			Response.Version = Version;
 			return Response;
 		}
 
@@ -907,9 +907,7 @@ namespace HordeServer.Services
 				throw new StructuredRpcException(StatusCode.NotFound, "Access to software is forbidden");
 			}
 
-			AgentSoftwareVersion Version = new AgentSoftwareVersion(Request.Version);
-
-			byte[]? Data = await AgentSoftwareService.GetArchiveAsync(Version);
+			byte[]? Data = await AgentSoftwareService.GetArchiveAsync(Request.Version);
 			if (Data == null)
 			{
 				throw new StructuredRpcException(StatusCode.NotFound, "Missing version {Version}");
diff --git a/Engine/Source/Programs/Horde/HordeServer/Tasks/Impl/UpgradeTaskSource.cs b/Engine/Source/Programs/Horde/HordeServer/Tasks/Impl/UpgradeTaskSource.cs
--- a/Engine/Source/Programs/Horde/HordeServer/Tasks/Impl/UpgradeTaskSource.cs
+++ b/Engine/Source/Programs/Horde/HordeServer/Tasks/Impl/UpgradeTaskSource.cs
@@ -14,7 +14,6 @@ using System.Collections.Generic;
 using System.Linq;
 using System.Threading.Tasks;
 
-using AgentSoftwareVersion = HordeServer.Utilities.StringId<HordeServer.Collections.IAgentSoftwareCollection>;
 using AgentSoftwareChannelName = HordeServer.Utilities.StringId<HordeServer.Services.AgentSoftwareChannels>;
 
 namespace HordeServer.Tasks.Impl
@@ -36,7 +35,7 @@ namespace HordeServer.Tasks.Impl
 
 		public async Task<ITaskListener?> SubscribeAsync(IAgent Agent)
 		{
-			AgentSoftwareVersion? RequiredVersion = await GetRequiredSoftwareVersion(Agent);
+			string? RequiredVersion = await GetRequiredSoftwareVersion(Agent);
 			if (RequiredVersion != null && Agent.Version != RequiredVersion)
 			{
 				AgentLease? Lease = null;
@@ -67,7 +66,7 @@ namespace HordeServer.Tasks.Impl
 		/// </summary>
 		/// <param name="Agent">The agent instance</param>
 		/// <returns>Unique id of the client version this agent should be running</returns>
-		public async Task<AgentSoftwareVersion?> GetRequiredSoftwareVersion(IAgent Agent)
+		public async Task<string?> GetRequiredSoftwareVersion(IAgent Agent)
 		{
 			AgentSoftwareChannelName ChannelName = Agent.Channel ?? AgentSoftwareService.DefaultChannelName;
 			IAgentSoftwareChannel? Channel = await AgentSoftwareService.GetCachedChannelAsync(ChannelName);
