Horde: Fix tests.

diff --git a/Engine/Source/Programs/Horde/HordeServerTests/AgentServiceTest.cs b/Engine/Source/Programs/Horde/HordeServerTests/AgentServiceTest.cs
--- a/Engine/Source/Programs/Horde/HordeServerTests/AgentServiceTest.cs
+++ b/Engine/Source/Programs/Horde/HordeServerTests/AgentServiceTest.cs
@@ -44,7 +44,7 @@ namespace HordeServerTests
 	        AgentService AgentService = TestSetup.AgentService;
 
 	        IAgent Agent = await AgentService.CreateSessionAsync(TestSetup.Fixture!.Agent1, AgentStatus.Ok, new AgentCapabilities(),
-		        AgentSoftwareVersion.Sanitize("test"));
+		        "test");
 	        
 	        Assert.IsTrue(AgentService.AuthorizeSession(Agent, GetUser(Agent)));
 	        TestSetup.Clock.Advance(TimeSpan.FromMinutes(20));
diff --git a/Engine/Source/Programs/Horde/HordeServerTests/RemoteExec/ClusteredActionTaskSourceTest.cs b/Engine/Source/Programs/Horde/HordeServerTests/RemoteExec/ClusteredActionTaskSourceTest.cs
--- a/Engine/Source/Programs/Horde/HordeServerTests/RemoteExec/ClusteredActionTaskSourceTest.cs
+++ b/Engine/Source/Programs/Horde/HordeServerTests/RemoteExec/ClusteredActionTaskSourceTest.cs
@@ -27,7 +27,7 @@ namespace HordeServerTests.RemoteExec
 		public string? Comment { get; set; }
 		public bool Ephemeral { get; set; }
 		public bool Deleted { get; set; }
-		public StringId<IAgentSoftwareCollection>? Version { get; set; }
+		public string? Version { get; set; }
 		public StringId<AgentSoftwareChannels>? Channel { get; set; }
 		public string? LastUpgradeVersion { get; set; }
 		public DateTime? LastUpgradeTime { get; set; }
