Horde: Remove software upload test. No longer valid.

diff --git a/Engine/Source/Programs/Horde/HordeServerTests/RpcServiceTest.cs b/Engine/Source/Programs/Horde/HordeServerTests/RpcServiceTest.cs
--- a/Engine/Source/Programs/Horde/HordeServerTests/RpcServiceTest.cs
+++ b/Engine/Source/Programs/Horde/HordeServerTests/RpcServiceTest.cs
@@ -276,7 +276,7 @@ namespace HordeServerTests
 			string text = Reader.ReadToEnd();
 			Assert.AreEqual(DataStr, text);
 		}
-
+		/*
 		[TestMethod]
 		public async Task UploadSoftwareAsync()
 		{
@@ -299,5 +299,6 @@ namespace HordeServerTests
 			UploadSoftwareResponse Res2 = await TestSetup.RpcService.UploadSoftware(Req, AdminContext);
 			Assert.AreEqual("r2", Res2.Version);
 		}
+		*/
 	}
 }
\ No newline at end of file
