Horde: Put data in C:\ProgramData by default, and set other default settings to be able to connect to a server on the same machine.

diff --git a/Engine/Source/Programs/Horde/HordeAgent/AgentSettings.cs b/Engine/Source/Programs/Horde/HordeAgent/AgentSettings.cs
--- a/Engine/Source/Programs/Horde/HordeAgent/AgentSettings.cs
+++ b/Engine/Source/Programs/Horde/HordeAgent/AgentSettings.cs
@@ -7,6 +7,7 @@ using System.ComponentModel.DataAnnotations;
 using System.IO;
 using System.Linq;
 using System.Reflection;
+using System.Runtime.InteropServices;
 using System.Text;
 using System.Text.Json;
 using System.Threading.Tasks;
@@ -153,8 +154,7 @@ namespace HordeAgent
 		/// <summary>
 		/// The default server, unless overridden from the command line
 		/// </summary>
-		[Required]
-		public string Server { get; set; } = null!;
+		public string? Server { get; set; }
 
 		/// <summary>
 		/// The executor to use for jobs. Defaults to the Perforce executor.
@@ -174,7 +174,7 @@ namespace HordeAgent
 		/// <summary>
 		/// Working directory
 		/// </summary>
-		public string? WorkingDir { get; set; }
+		public string? WorkingDir { get; set; } = Program.DataDir.FullName;
 
 		/// <summary>
 		/// List of network shares to mount
diff --git a/Engine/Source/Programs/Horde/HordeAgent/Commands/Service/UpgradeCommand.cs b/Engine/Source/Programs/Horde/HordeAgent/Commands/Service/UpgradeCommand.cs
--- a/Engine/Source/Programs/Horde/HordeAgent/Commands/Service/UpgradeCommand.cs
+++ b/Engine/Source/Programs/Horde/HordeAgent/Commands/Service/UpgradeCommand.cs
@@ -67,10 +67,10 @@ namespace HordeAgent.Commands
 				TargetFiles.ExceptWith(SourceFiles);
 
 				// Ignore any files that are in the saved directory
-				string SourceDataDir = Path.Combine(SourceDir.FullName, Program.SavedDirName) + Path.DirectorySeparatorChar;
+				string SourceDataDir = Path.Combine(SourceDir.FullName, "Saved") + Path.DirectorySeparatorChar;
 				SourceFiles.RemoveWhere(x => x.StartsWith(SourceDataDir, StringComparison.OrdinalIgnoreCase));
 
-				string TargetDataDir = Path.Combine(TargetDir.FullName, Program.SavedDirName) + Path.DirectorySeparatorChar;
+				string TargetDataDir = Path.Combine(TargetDir.FullName, "Saved") + Path.DirectorySeparatorChar;
 				TargetFiles.RemoveWhere(x => x.StartsWith(TargetDataDir, StringComparison.OrdinalIgnoreCase));
 
 				// Copy all the files into the target directory
diff --git a/Engine/Source/Programs/Horde/HordeAgent/Commands/Utilities/PublishCommand.cs b/Engine/Source/Programs/Horde/HordeAgent/Commands/Utilities/PublishCommand.cs
--- a/Engine/Source/Programs/Horde/HordeAgent/Commands/Utilities/PublishCommand.cs
+++ b/Engine/Source/Programs/Horde/HordeAgent/Commands/Utilities/PublishCommand.cs
@@ -85,7 +85,7 @@ namespace HordeAgent.Commands.Utilities
 
 			// Read the application settings
 			IConfigurationRoot Config = new ConfigurationBuilder()
-				.SetBasePath(Program.AppDir)
+				.SetBasePath(Program.AppDir.FullName)
 				.AddJsonFile("appsettings.json", true)
 				.AddJsonFile("appsettings.Production.json", true)
 				.Build();
diff --git a/Engine/Source/Programs/Horde/HordeAgent/Logging.cs b/Engine/Source/Programs/Horde/HordeAgent/Logging.cs
--- a/Engine/Source/Programs/Horde/HordeAgent/Logging.cs
+++ b/Engine/Source/Programs/Horde/HordeAgent/Logging.cs
@@ -55,7 +55,7 @@ namespace HordeAgent
 
 		static Serilog.ILogger CreateSerilogLogger()
 		{
-			Directory.CreateDirectory(Program.SavedDir);
+			DirectoryReference.CreateDirectory(Program.DataDir);
 
 			ConsoleTheme Theme;
 			if (RuntimeInformation.IsOSPlatform(OSPlatform.Windows) && Environment.OSVersion.Version < new Version(10, 0))
@@ -74,8 +74,8 @@ namespace HordeAgent
 				.Enrich.FromLogContext()
 				.Enrich.With<VersionLogEnricher>()
 				.WriteTo.Console(outputTemplate: "[{Timestamp:HH:mm:ss} {Level:w3}] {Indent}{Message:l}{NewLine}{Exception}", theme: Theme)
-				.WriteTo.File(Path.Combine(Program.SavedDir, "Log-.txt"), fileSizeLimitBytes: 50 * 1024 * 1024, rollingInterval: RollingInterval.Day, rollOnFileSizeLimit: true, retainedFileCountLimit: 10)
-				.WriteTo.File(new JsonFormatter(renderMessage: true), Path.Combine(Program.SavedDir, "Log-.json"), fileSizeLimitBytes: 50 * 1024 * 1024, rollingInterval: RollingInterval.Day, rollOnFileSizeLimit: true, retainedFileCountLimit: 10)
+				.WriteTo.File(FileReference.Combine(Program.DataDir, "Log-.txt").FullName, fileSizeLimitBytes: 50 * 1024 * 1024, rollingInterval: RollingInterval.Day, rollOnFileSizeLimit: true, retainedFileCountLimit: 10)
+				.WriteTo.File(new JsonFormatter(renderMessage: true), FileReference.Combine(Program.DataDir, "Log-.json").FullName, fileSizeLimitBytes: 50 * 1024 * 1024, rollingInterval: RollingInterval.Day, rollOnFileSizeLimit: true, retainedFileCountLimit: 10)
 				.CreateLogger();
 		}
 	}
diff --git a/Engine/Source/Programs/Horde/HordeAgent/Program.cs b/Engine/Source/Programs/Horde/HordeAgent/Program.cs
--- a/Engine/Source/Programs/Horde/HordeAgent/Program.cs
+++ b/Engine/Source/Programs/Horde/HordeAgent/Program.cs
@@ -40,17 +40,12 @@ namespace HordeAgent
 		/// <summary>
 		/// Path to the root application directory
 		/// </summary>
-		public static string AppDir { get; } = Path.GetFullPath(Path.GetDirectoryName(Assembly.GetExecutingAssembly().Location)!);
+		public static DirectoryReference AppDir { get; } = GetAppDir();
 
 		/// <summary>
-		/// Name of the saved application data directory
+		/// Path to the default data directory
 		/// </summary>
-		public const string SavedDirName = "Saved";
-
-		/// <summary>
-		/// Path to the saved application data directory
-		/// </summary>
-		public static string SavedDir { get; } = Path.Combine(AppDir, SavedDirName);
+		public static DirectoryReference DataDir { get; } = GetDataDir();
 
 		/// <summary>
 		/// The launch arguments
@@ -258,5 +253,31 @@ namespace HordeAgent
 				return "unknown";
 			}
 		}
+
+		/// <summary>
+		/// Gets the application directory
+		/// </summary>
+		/// <returns></returns>
+		static DirectoryReference GetAppDir()
+		{
+			return new DirectoryReference(Path.GetDirectoryName(Assembly.GetExecutingAssembly().Location)!);
+		}
+
+		/// <summary>
+		/// Gets the default data directory
+		/// </summary>
+		/// <returns></returns>
+		static DirectoryReference GetDataDir()
+		{
+			if (RuntimeInformation.IsOSPlatform(OSPlatform.Windows))
+			{
+				DirectoryReference? ProgramDataDir = DirectoryReference.GetSpecialFolder(Environment.SpecialFolder.CommonApplicationData);
+				if (ProgramDataDir != null)
+				{
+					return DirectoryReference.Combine(ProgramDataDir, "HordeAgent");
+				}
+			}
+			return GetAppDir();
+		}
 	}
 }
diff --git a/Engine/Source/Programs/Horde/HordeServer/Program.cs b/Engine/Source/Programs/Horde/HordeServer/Program.cs
--- a/Engine/Source/Programs/Horde/HordeServer/Program.cs
+++ b/Engine/Source/Programs/Horde/HordeServer/Program.cs
@@ -212,7 +212,7 @@ namespace HordeServer
 				DirectoryReference? Dir = DirectoryReference.GetSpecialFolder(Environment.SpecialFolder.CommonApplicationData);
 				if (Dir != null)
 				{
-					return DirectoryReference.Combine(Dir, "Horde");
+					return DirectoryReference.Combine(Dir, "HordeServer");
 				}
 			}
 			return DirectoryReference.Combine(GetAppDir(), "Data");
