Horde: Fix error when upgrading agent software due to temp file already existing.

diff --git a/Engine/Source/Programs/Horde/HordeServer/Services/AgentSoftwareService.cs b/Engine/Source/Programs/Horde/HordeServer/Services/AgentSoftwareService.cs
--- a/Engine/Source/Programs/Horde/HordeServer/Services/AgentSoftwareService.cs
+++ b/Engine/Source/Programs/Horde/HordeServer/Services/AgentSoftwareService.cs
@@ -216,7 +216,7 @@ namespace HordeServer.Services
 						string TempFile = Path.GetTempFileName();
 						try
 						{
-							InputEntry.ExtractToFile(TempFile);
+							InputEntry.ExtractToFile(TempFile, true);
 							return FileVersionInfo.GetVersionInfo(TempFile).ProductVersion;
 						}
 						finally
