Horde: Add different extensions to config files, so we can automatically match the correct schema.

diff --git a/Engine/Source/Programs/Horde/HordeServer/Controllers/SchemaController.cs b/Engine/Source/Programs/Horde/HordeServer/Controllers/SchemaController.cs
--- a/Engine/Source/Programs/Horde/HordeServer/Controllers/SchemaController.cs
+++ b/Engine/Source/Programs/Horde/HordeServer/Controllers/SchemaController.cs
@@ -73,7 +73,7 @@ namespace HordeServer.Controllers
 					if (CatalogAttribute != null)
 					{
 						Uri Url = new Uri($"https://{Host}/api/v1/schema/types/{SchemaType.Name}.json");
-						Root.Schemas.Add(new CatalogItem { Name = CatalogAttribute.Name, Description = CatalogAttribute.Description, Url = Url });
+						Root.Schemas.Add(new CatalogItem { Name = CatalogAttribute.Name, Description = CatalogAttribute.Description, FileMatch = CatalogAttribute.FileMatch, Url = Url });
 					}
 				}
 			}
diff --git a/Engine/Source/Programs/Horde/HordeServer/Defaults/global.json b/Engine/Source/Programs/Horde/HordeServer/Defaults/global.json
deleted file mode 100644
diff --git a/Engine/Source/Programs/Horde/HordeServer/Defaults/globals.json b/Engine/Source/Programs/Horde/HordeServer/Defaults/globals.json
new file mode 100644
--- /dev/null
+++ b/Engine/Source/Programs/Horde/HordeServer/Defaults/globals.json
@@ -0,0 +1,8 @@
+{
+	"projects": [
+		{
+			"id": "ue5",
+			"path": "project-ue5.json"
+		}
+	]
+}
diff --git a/Engine/Source/Programs/Horde/HordeServer/Defaults/project-ue5.json b/Engine/Source/Programs/Horde/HordeServer/Defaults/project-ue5.json
deleted file mode 100644
diff --git a/Engine/Source/Programs/Horde/HordeServer/Defaults/stream-ue5-main.json b/Engine/Source/Programs/Horde/HordeServer/Defaults/stream-ue5-main.json
deleted file mode 100644
diff --git a/Engine/Source/Programs/Horde/HordeServer/Defaults/ue5-main.stream.json b/Engine/Source/Programs/Horde/HordeServer/Defaults/ue5-main.stream.json
new file mode 100644
--- /dev/null
+++ b/Engine/Source/Programs/Horde/HordeServer/Defaults/ue5-main.stream.json
@@ -0,0 +1,71 @@
+{
+	"name": "//UE5/Main",
+	"tabs":
+	[
+		{
+			"title": "General",
+			"type": "Jobs",
+			"showNames": true,
+			"templates":
+			[
+				"custom"
+			],
+			"columns":
+			[
+				{
+					"heading": "Editors",
+					"category": "Editors",
+					"relativeWidth": 1
+				},
+				{
+					"heading": "Target Platforms",
+					"category": "Platforms",
+					"relativeWidth": 2
+				},
+				{
+					"heading": "Other",
+					"category": "Other",
+					"relativeWidth": 1
+				}
+			]
+		}
+	],
+	"agentTypes":
+	{
+		"Win-UE5":
+		{
+			"pool": "win-ue5",
+			"environment":
+			{
+			}
+		}
+	},
+	"workspaceTypes":
+	{
+	},
+	"templates":
+	[
+		{
+			"id": "custom",
+			"name": "Custom...",
+			"arguments":
+			[
+			],
+			"parameters":
+			[
+				{
+					"type": "Text",
+					"label": "Script",
+					"argument": "-Script=",
+					"default": "Foo.xml"
+				},
+				{
+					"type": "Text",
+					"label": "Target",
+					"argument": "-Target=",
+					"default": ""
+				}
+			]
+		}
+	]
+}
diff --git a/Engine/Source/Programs/Horde/HordeServer/Defaults/ue5.project.json b/Engine/Source/Programs/Horde/HordeServer/Defaults/ue5.project.json
new file mode 100644
--- /dev/null
+++ b/Engine/Source/Programs/Horde/HordeServer/Defaults/ue5.project.json
@@ -0,0 +1,21 @@
+{
+	"$schema": "http://unrealengine.com/horde/ProjectConfig",
+	"name": "UE5",
+	"logo": "project-ue5.png",
+	"streams": [
+		{
+		  "id": "ue5-main",
+		  "path": "stream-ue5-main.json"
+		}
+	],
+	"categories": [
+		{
+			"name": "General",
+			"row": 0,
+			"showOnNavMenu": true,
+			"includePatterns": [
+				"//UE5/Main"
+			]
+		}
+	]
+}
diff --git a/Engine/Source/Programs/Horde/HordeServer/Models/Config.cs b/Engine/Source/Programs/Horde/HordeServer/Models/Config.cs
--- a/Engine/Source/Programs/Horde/HordeServer/Models/Config.cs
+++ b/Engine/Source/Programs/Horde/HordeServer/Models/Config.cs
@@ -17,7 +17,7 @@ namespace HordeServer.Models
 	/// Global configuration
 	/// </summary>
 	[JsonSchema("https://unrealengine.com/horde/global")]
-	[JsonSchemaCatalog("Horde Globals", "Horde global configuration file")]
+	[JsonSchemaCatalog("Horde Globals", "Horde global configuration file", "globals.json")]
 	public class GlobalConfig
 	{
 		/// <summary>
@@ -73,7 +73,7 @@ namespace HordeServer.Models
 	/// Stores configuration for a project
 	/// </summary>
 	[JsonSchema("https://unrealengine.com/horde/project")]
-	[JsonSchemaCatalog("Horde Project", "Horde project configuration file")]
+	[JsonSchemaCatalog("Horde Project", "Horde project configuration file", "*.project.json")]
 	public class ProjectConfig
 	{
 		/// <summary>
@@ -125,7 +125,7 @@ namespace HordeServer.Models
 	/// Config for a stream
 	/// </summary>
 	[JsonSchema("https://unrealengine.com/horde/stream")]
-	[JsonSchemaCatalog("Horde Stream", "Horde stream configuration file")]
+	[JsonSchemaCatalog("Horde Stream", "Horde stream configuration file", "*.stream.json")]
 	public class StreamConfig
 	{
 		/// <summary>
diff --git a/Engine/Source/Programs/Horde/HordeServer/Utilities/Schemas.cs b/Engine/Source/Programs/Horde/HordeServer/Utilities/Schemas.cs
--- a/Engine/Source/Programs/Horde/HordeServer/Utilities/Schemas.cs
+++ b/Engine/Source/Programs/Horde/HordeServer/Utilities/Schemas.cs
@@ -52,15 +52,33 @@ namespace HordeServer.Utilities
 		/// </summary>
 		public string Description { get; set; }
 
+		/// <summary>
+		/// File patterns to match
+		/// </summary>
+		public string[]? FileMatch { get; set; }
+
+		/// <summary>
+		/// Constructor
+		/// </summary>
+		/// <param name="Name"></param>
+		/// <param name="Description"></param>
+		/// <param name="FileMatch">File patterns to match</param>
+		public JsonSchemaCatalogAttribute(string Name, string Description, string? FileMatch)
+			: this(Name, Description, (FileMatch == null) ? (string[]?)null : new[] { FileMatch })
+		{
+		}
+
 		/// <summary>
 		/// Constructor
 		/// </summary>
 		/// <param name="Name"></param>
 		/// <param name="Description"></param>
-		public JsonSchemaCatalogAttribute(string Name, string Description)
+		/// <param name="FileMatch">File patterns to match</param>
+		public JsonSchemaCatalogAttribute(string Name, string Description, string[]? FileMatch)
 		{
 			this.Name = Name;
 			this.Description = Description;
+			this.FileMatch = FileMatch;
 		}
 	}
 
diff --git a/Engine/Source/Programs/Horde/HordeServer/appsettings.json b/Engine/Source/Programs/Horde/HordeServer/appsettings.json
--- a/Engine/Source/Programs/Horde/HordeServer/appsettings.json
+++ b/Engine/Source/Programs/Horde/HordeServer/appsettings.json
@@ -3,7 +3,7 @@
 	{
 		"DashboardUrl": "https://localhost:3000",
 		"DisableAuth": true,
-		"ConfigPath": "Defaults/global.json",
+		"ConfigPath": "Defaults/globals.json",
 		"CorsEnabled": true,
 		"CorsOrigin": "https://localhost:5003;http://localhost:3000",
 		"LogFileProviderType": "Local"
