Horde: Fix staging of project logo.

diff --git a/Engine/Source/Programs/Horde/HordeServer/HordeServer.csproj b/Engine/Source/Programs/Horde/HordeServer/HordeServer.csproj
--- a/Engine/Source/Programs/Horde/HordeServer/HordeServer.csproj
+++ b/Engine/Source/Programs/Horde/HordeServer/HordeServer.csproj
@@ -97,7 +97,7 @@
   </ItemGroup>
 
   <ItemGroup>
-    <Content Include="Defaults\project-ue5.png">
+    <Content Include="Defaults\ue5.project.png">
       <CopyToOutputDirectory>PreserveNewest</CopyToOutputDirectory>
     </Content>
   </ItemGroup>
