Horde: Remove incorrect schema from project config file.

diff --git a/Engine/Source/Programs/Horde/HordeServer/Defaults/ue5.project.json b/Engine/Source/Programs/Horde/HordeServer/Defaults/ue5.project.json
--- a/Engine/Source/Programs/Horde/HordeServer/Defaults/ue5.project.json
+++ b/Engine/Source/Programs/Horde/HordeServer/Defaults/ue5.project.json
@@ -1,5 +1,4 @@
 {
-	"$schema": "http://unrealengine.com/horde/ProjectConfig",
 	"name": "UE5",
 	"logo": "project-ue5.png",
 	"streams": [
