Horde: Add patterns to match Projects/*.json and Streams/*.json for schemas.

diff --git a/Engine/Source/Programs/Horde/HordeServer/Models/Config.cs b/Engine/Source/Programs/Horde/HordeServer/Models/Config.cs
--- a/Engine/Source/Programs/Horde/HordeServer/Models/Config.cs
+++ b/Engine/Source/Programs/Horde/HordeServer/Models/Config.cs
@@ -73,7 +73,7 @@ namespace HordeServer.Models
 	/// Stores configuration for a project
 	/// </summary>
 	[JsonSchema("https://unrealengine.com/horde/project")]
-	[JsonSchemaCatalog("Horde Project", "Horde project configuration file", "*.project.json")]
+	[JsonSchemaCatalog("Horde Project", "Horde project configuration file", new[] { "*.project.json", "Projects/*.json" })]
 	public class ProjectConfig
 	{
 		/// <summary>
@@ -125,7 +125,7 @@ namespace HordeServer.Models
 	/// Config for a stream
 	/// </summary>
 	[JsonSchema("https://unrealengine.com/horde/stream")]
-	[JsonSchemaCatalog("Horde Stream", "Horde stream configuration file", "*.stream.json")]
+	[JsonSchemaCatalog("Horde Stream", "Horde stream configuration file", new[] { "*.stream.json", "Streams/*.json" })]
 	public class StreamConfig
 	{
 		/// <summary>
