Horde: Allow anonymous auth for gRPC connections. Also add local server profile.

diff --git a/Engine/Source/Programs/Horde/HordeAgent/AgentSettings.cs b/Engine/Source/Programs/Horde/HordeAgent/AgentSettings.cs
--- a/Engine/Source/Programs/Horde/HordeAgent/AgentSettings.cs
+++ b/Engine/Source/Programs/Horde/HordeAgent/AgentSettings.cs
@@ -77,8 +77,7 @@ namespace HordeAgent
 		/// <summary>
 		/// Bearer token to use to initiate the connection
 		/// </summary>
-		[Required]
-		public string Token { get; set; } = null!;
+		public string? Token { get; set; }
 
 		/// <summary>
 		/// Thumbprint of a certificate to trust. Allows using self-signed certs for the server.
diff --git a/Engine/Source/Programs/Horde/HordeAgent/Services/GrpcService.cs b/Engine/Source/Programs/Horde/HordeAgent/Services/GrpcService.cs
--- a/Engine/Source/Programs/Horde/HordeAgent/Services/GrpcService.cs
+++ b/Engine/Source/Programs/Horde/HordeAgent/Services/GrpcService.cs
@@ -60,9 +60,16 @@ namespace HordeAgent.Services
 		/// Create a GRPC channel with the given bearer token
 		/// </summary>
 		/// <returns>New grpc channel</returns>
-		public GrpcChannel CreateGrpcChannel(string BearerToken)
+		public GrpcChannel CreateGrpcChannel(string? BearerToken)
 		{
-			return CreateGrpcChannel(ServerProfile.Url, new AuthenticationHeaderValue("Bearer", BearerToken));
+			if (BearerToken == null)
+			{
+				return CreateGrpcChannel(ServerProfile.Url, null);
+			}
+			else
+			{
+				return CreateGrpcChannel(ServerProfile.Url, new AuthenticationHeaderValue("Bearer", BearerToken));
+			}
 		}
 
 		/// <summary>
diff --git a/Engine/Source/Programs/Horde/HordeAgent/appsettings.json b/Engine/Source/Programs/Horde/HordeAgent/appsettings.json
new file mode 100644
--- /dev/null
+++ b/Engine/Source/Programs/Horde/HordeAgent/appsettings.json
@@ -0,0 +1,14 @@
+{
+	"Horde":
+	{
+		"Server": "Local",
+		"ServerProfiles": 
+		[
+			{
+				"Name": "Local",
+				"Environment": "dev",
+				"Url": "https://localhost:5001"
+			}
+		]
+	}
+}
