Horde: Fix working directory for launching applications, and add icon for Horde Agent.

diff --git a/Engine/Source/Programs/Horde/HordeAgent/HordeAgent.csproj b/Engine/Source/Programs/Horde/HordeAgent/HordeAgent.csproj
--- a/Engine/Source/Programs/Horde/HordeAgent/HordeAgent.csproj
+++ b/Engine/Source/Programs/Horde/HordeAgent/HordeAgent.csproj
@@ -4,6 +4,7 @@
     <OutputType>Exe</OutputType>
     <TargetFramework>netcoreapp3.1</TargetFramework>
     <Nullable>enable</Nullable>
+	<ApplicationIcon>HordeAgent.ico</ApplicationIcon>
     <ApplicationManifest>app.manifest</ApplicationManifest>
     <PublishRepositoryUrl>true</PublishRepositoryUrl>
     <EmbedUntrackedSources>false</EmbedUntrackedSources>
