Horde: Fix default project/stream config file paths.

diff --git a/Engine/Source/Programs/Horde/HordeServer/Defaults/globals.json b/Engine/Source/Programs/Horde/HordeServer/Defaults/globals.json
--- a/Engine/Source/Programs/Horde/HordeServer/Defaults/globals.json
+++ b/Engine/Source/Programs/Horde/HordeServer/Defaults/globals.json
@@ -2,7 +2,7 @@
 	"projects": [
 		{
 			"id": "ue5",
-			"path": "project-ue5.json"
+			"path": "ue5.project.json"
 		}
 	]
 }
diff --git a/Engine/Source/Programs/Horde/HordeServer/Defaults/ue5.project.json b/Engine/Source/Programs/Horde/HordeServer/Defaults/ue5.project.json
--- a/Engine/Source/Programs/Horde/HordeServer/Defaults/ue5.project.json
+++ b/Engine/Source/Programs/Horde/HordeServer/Defaults/ue5.project.json
@@ -1,10 +1,10 @@
 {
 	"name": "UE5",
-	"logo": "project-ue5.png",
+	"logo": "ue5.project.png",
 	"streams": [
 		{
 		  "id": "ue5-main",
-		  "path": "stream-ue5-main.json"
+		  "path": "ue5-main.stream.json"
 		}
 	],
 	"categories": [
