Horde: Fix authentication of agents when using anonymous auth handler.

diff --git a/Engine/Source/Programs/Horde/HordeServer/Authentication/AnonymousHandler.cs b/Engine/Source/Programs/Horde/HordeServer/Authentication/AnonymousHandler.cs
--- a/Engine/Source/Programs/Horde/HordeServer/Authentication/AnonymousHandler.cs
+++ b/Engine/Source/Programs/Horde/HordeServer/Authentication/AnonymousHandler.cs
@@ -54,4 +54,12 @@ namespace HordeServer.Authentication
 			return AuthenticateResult.Success(Ticket);
 		}
 	}
+
+	static class AnonymousExtensions
+	{
+		public static AuthenticationBuilder AddAnonymous(this AuthenticationBuilder Builder, Action<AnonymousAuthenticationOptions> Configure)
+		{
+			return Builder.AddScheme<AnonymousAuthenticationOptions, AnonymousAuthenticationHandler>(AnonymousAuthenticationHandler.AuthenticationScheme, Configure);
+		}
+	}
 }
diff --git a/Engine/Source/Programs/Horde/HordeServer/Authentication/ServiceAccountAuthHandler.cs b/Engine/Source/Programs/Horde/HordeServer/Authentication/ServiceAccountAuthHandler.cs
--- a/Engine/Source/Programs/Horde/HordeServer/Authentication/ServiceAccountAuthHandler.cs
+++ b/Engine/Source/Programs/Horde/HordeServer/Authentication/ServiceAccountAuthHandler.cs
@@ -69,4 +69,12 @@ namespace HordeServer.Authentication
 			return AuthenticateResult.Success(Ticket);
 		}
 	}
+
+	static class ServiceAccountExtensions
+	{
+		public static AuthenticationBuilder AddServiceAccount(this AuthenticationBuilder Builder, Action<ServiceAccountAuthOptions> Config)
+		{
+			return Builder.AddScheme<ServiceAccountAuthOptions, ServiceAccountAuthHandler>(ServiceAccountAuthHandler.AuthenticationScheme, Config);
+		}
+	}
 }
diff --git a/Engine/Source/Programs/Horde/HordeServer/Startup.cs b/Engine/Source/Programs/Horde/HordeServer/Startup.cs
--- a/Engine/Source/Programs/Horde/HordeServer/Startup.cs
+++ b/Engine/Source/Programs/Horde/HordeServer/Startup.cs
@@ -410,85 +410,88 @@ namespace HordeServer
 			ConfigureLogStorage(Services);
 			//			ConfigureLogFileWriteCache(Services, Settings);
 
+			AuthenticationBuilder AuthBuilder = Services.AddAuthentication(Options =>
+				{
+					// If an authentication cookie is present, use it to get authentication information
+					Options.DefaultAuthenticateScheme = CookieAuthenticationDefaults.AuthenticationScheme;
+					Options.DefaultSignInScheme = CookieAuthenticationDefaults.AuthenticationScheme;
+
+					// If authentication is required, and no cookie is present, use OIDC to sign in
+					if (Settings.DisableAuth)
+					{
+						Options.DefaultChallengeScheme = AnonymousAuthenticationHandler.AuthenticationScheme;
+					}
+					else
+					{
+						Options.DefaultChallengeScheme = OktaDefaults.AuthenticationScheme;
+					}
+				});
+
+			List<string> Schemes = new List<string>();
+
+			AuthBuilder.AddCookie(Options =>
+				 {
+					 Options.Events.OnValidatePrincipal = Context =>
+					 {
+						 if (Context.Principal.FindFirst(HordeClaimTypes.UserId) == null)
+						 {
+							 Context.RejectPrincipal();
+						 }
+						 return Task.CompletedTask;
+					 };
+
+					 Options.Events.OnRedirectToAccessDenied = Context =>
+					 {
+						 Context.Response.StatusCode = StatusCodes.Status403Forbidden;
+						 return Context.Response.CompleteAsync();
+					 };
+				 });
+			Schemes.Add(CookieAuthenticationDefaults.AuthenticationScheme);
+
+			AuthBuilder.AddServiceAccount(Options => { });
+			Schemes.Add(ServiceAccountAuthHandler.AuthenticationScheme);
+
 			if (Settings.DisableAuth)
 			{
-				var Scheme = AnonymousAuthenticationHandler.AuthenticationScheme;
-				Services.AddAuthentication(Scheme)
-					.AddScheme<AnonymousAuthenticationOptions, AnonymousAuthenticationHandler>(Scheme, Options =>
-					{
-						Options.AdminClaimType = Settings.AdminClaimType;
-						Options.AdminClaimValue = Settings.AdminClaimValue;
-					});
+				AuthBuilder.AddAnonymous(Options =>
+				{
+					Options.AdminClaimType = Settings.AdminClaimType;
+					Options.AdminClaimValue = Settings.AdminClaimValue;
+				});
+				Schemes.Add(AnonymousAuthenticationHandler.AuthenticationScheme);
 			}
-			else
+			else if (Settings.OidcClientId != null && Settings.OidcAuthority != null)
 			{
-				AuthenticationBuilder AuthBuilder = Services.AddAuthentication(Options =>
+				AuthBuilder.AddOkta(OktaDefaults.AuthenticationScheme, OpenIdConnectDefaults.DisplayName, 
+					Options =>
 					{
-						// If an authentication cookie is present, use it to get authentication information
-						Options.DefaultAuthenticateScheme = CookieAuthenticationDefaults.AuthenticationScheme;
-						Options.DefaultSignInScheme = CookieAuthenticationDefaults.AuthenticationScheme;
+						Options.Authority = Settings.OidcAuthority;
+						Options.ClientId = Settings.OidcClientId;
 
-						// If authentication is required, and no cookie is present, use OIDC to sign in
-						Options.DefaultChallengeScheme = OktaDefaults.AuthenticationScheme;
-					})
-					.AddCookie(Options =>
-					{
-						Options.Events.OnValidatePrincipal = Context =>
+						if (!String.IsNullOrEmpty(Settings.OidcSigninRedirect))
 						{
-							if (Context.Principal.FindFirst(HordeClaimTypes.UserId) == null)
+							Options.Events = new OpenIdConnectEvents
 							{
-								Context.RejectPrincipal();
-							}
-							return Task.CompletedTask;
-						};
-
-						Options.Events.OnRedirectToAccessDenied = Context => 
-						{ 
-							Context.Response.StatusCode = StatusCodes.Status403Forbidden; 
-							return Context.Response.CompleteAsync(); 
-						};
+								OnRedirectToIdentityProvider = async RedirectContext =>
+								{
+									RedirectContext.ProtocolMessage.RedirectUri = Settings.OidcSigninRedirect;
+									await Task.CompletedTask;
+								}
+							};
+						}
 					});
+				Schemes.Add(OktaDefaults.AuthenticationScheme);
+			}
 
-				List<string> AuthSchemes = new List<string>();
-				AuthSchemes.Add(CookieAuthenticationDefaults.AuthenticationScheme);
-
-				AuthBuilder.AddScheme<ServiceAccountAuthOptions, ServiceAccountAuthHandler>(
-					ServiceAccountAuthHandler.AuthenticationScheme, Options => { });
-				AuthSchemes.Add(ServiceAccountAuthHandler.AuthenticationScheme);
+			AuthBuilder.AddScheme<JwtBearerOptions, HordeJwtBearerHandler>(HordeJwtBearerHandler.AuthenticationScheme, Options => { });
+			Schemes.Add(HordeJwtBearerHandler.AuthenticationScheme);
 
-				if (Settings.OidcClientId != null && Settings.OidcAuthority != null)
+			Services.AddAuthorization(Options =>
 				{
-					AuthBuilder.AddOkta(OktaDefaults.AuthenticationScheme, OpenIdConnectDefaults.DisplayName, 
-						Options =>
-						{
-							Options.Authority = Settings.OidcAuthority;
-							Options.ClientId = Settings.OidcClientId;
-
-							if (!String.IsNullOrEmpty(Settings.OidcSigninRedirect))
-							{
-								Options.Events = new OpenIdConnectEvents
-								{
-									OnRedirectToIdentityProvider = async RedirectContext =>
-									{
-										RedirectContext.ProtocolMessage.RedirectUri = Settings.OidcSigninRedirect;
-										await Task.CompletedTask;
-									}
-								};
-							}
-						});
-					AuthSchemes.Add(OktaDefaults.AuthenticationScheme);
-				}
-
-				AuthBuilder.AddScheme<JwtBearerOptions, HordeJwtBearerHandler>(HordeJwtBearerHandler.AuthenticationScheme, Options => { });
-				AuthSchemes.Add(HordeJwtBearerHandler.AuthenticationScheme);
-
-				Services.AddAuthorization(Options =>
-					{
-						Options.DefaultPolicy = new AuthorizationPolicyBuilder(AuthSchemes.ToArray())
-							.RequireAuthenticatedUser()							
-							.Build();
-					});
-			}
+					Options.DefaultPolicy = new AuthorizationPolicyBuilder(Schemes.ToArray())
+						.RequireAuthenticatedUser()							
+						.Build();
+				});
 
 			Services.AddHostedService(Provider => Provider.GetRequiredService<AgentService>());
 			Services.AddHostedService(Provider => Provider.GetRequiredService<AutoscaleService>());
