Horde: Adjust P4 bridge router
#rnx

diff --git a/Engine/Source/Programs/Horde/HordeServer/Services/PerforceServiceRouter.cs b/Engine/Source/Programs/Horde/HordeServer/Services/PerforceServiceRouter.cs
--- a/Engine/Source/Programs/Horde/HordeServer/Services/PerforceServiceRouter.cs
+++ b/Engine/Source/Programs/Horde/HordeServer/Services/PerforceServiceRouter.cs
@@ -42,40 +42,40 @@ namespace HordeServer.Services
 			return P4APIService.GetUserInfoAsync(UserName);
 		}
 
-		/// <inheritdoc/>
-		override public Task<int> CreateNewChangeAsync(string StreamName, string Path)
+		/// <inheritdoc />
+		override public Task<string> CreateTicket(string ImpersonateUser)
 		{
-			return PythonBridgeService.CreateNewChangeAsync(StreamName, Path);
+			return P4APIService.CreateTicket(ImpersonateUser);
 		}
 
 		/// <inheritdoc/>
-		override public Task<List<ChangeSummary>> GetChangesAsync(string StreamName, int? MinChange, int? MaxChange, int Results, string? ImpersonateUser)
+		override public Task<int> GetCodeChangeAsync(string StreamName, int Change)
 		{
-			return P4APIService.GetChangesAsync(StreamName, MinChange, MaxChange, Results, ImpersonateUser);
+			return P4APIService.GetCodeChangeAsync(StreamName, Change);
 		}
 
 		/// <inheritdoc/>
-		override public Task<List<ChangeDetails>> GetChangeDetailsAsync(string StreamName, IReadOnlyList<int> ChangeNumbers, string? ImpersonateUser)
+		override public Task<List<ChangeSummary>> GetChangesAsync(string StreamName, int? MinChange, int? MaxChange, int Results, string? ImpersonateUser)
 		{
-			return P4APIService.GetChangeDetailsAsync(StreamName, ChangeNumbers, ImpersonateUser);
+			return P4APIService.GetChangesAsync(StreamName, MinChange, MaxChange, Results, ImpersonateUser);
 		}
 
-		/// <inheritdoc />
-		override public Task<string> CreateTicket(string ImpersonateUser)
+		/// <inheritdoc/>
+		override public Task<int> CreateNewChangeAsync(string StreamName, string Path)
 		{
-			return P4APIService.CreateTicket(ImpersonateUser);
+			return PythonBridgeService.CreateNewChangeAsync(StreamName, Path);
 		}
 
 		/// <inheritdoc/>
-		override public Task<int> GetCodeChangeAsync(string StreamName, int Change)
+		override public Task<List<ChangeDetails>> GetChangeDetailsAsync(string StreamName, IReadOnlyList<int> ChangeNumbers, string? ImpersonateUser)
 		{
-			return P4APIService.GetCodeChangeAsync(StreamName, Change);
+			return PythonBridgeService.GetChangeDetailsAsync(StreamName, ChangeNumbers, ImpersonateUser);
 		}
 
 		/// <inheritdoc/>
 		override public Task<List<FileSummary>> FindFilesAsync(IEnumerable<string> Paths)
 		{
-			return P4APIService.FindFilesAsync(Paths);
+			return PythonBridgeService.FindFilesAsync(Paths);
 		}
 
 		/// <inheritdoc/>
@@ -105,7 +105,7 @@ namespace HordeServer.Services
 		/// <inheritdoc/>
 		override public Task UpdateChangelistDescription(int Change, string Description)
 		{
-			return P4APIService.UpdateChangelistDescription(Change, Description );
+			return PythonBridgeService.UpdateChangelistDescription(Change, Description );
 		}
 
 		/// <inheritdoc/>
