Horde: Fix anonymous authentication in tests.

diff --git a/Engine/Source/Programs/Horde/HordeServer/Startup.cs b/Engine/Source/Programs/Horde/HordeServer/Startup.cs
--- a/Engine/Source/Programs/Horde/HordeServer/Startup.cs
+++ b/Engine/Source/Programs/Horde/HordeServer/Startup.cs
@@ -412,17 +412,19 @@ namespace HordeServer
 
 			AuthenticationBuilder AuthBuilder = Services.AddAuthentication(Options =>
 				{
-					// If an authentication cookie is present, use it to get authentication information
-					Options.DefaultAuthenticateScheme = CookieAuthenticationDefaults.AuthenticationScheme;
-					Options.DefaultSignInScheme = CookieAuthenticationDefaults.AuthenticationScheme;
-
-					// If authentication is required, and no cookie is present, use OIDC to sign in
 					if (Settings.DisableAuth)
 					{
+						Options.DefaultAuthenticateScheme = AnonymousAuthenticationHandler.AuthenticationScheme;
+						Options.DefaultSignInScheme = AnonymousAuthenticationHandler.AuthenticationScheme;
 						Options.DefaultChallengeScheme = AnonymousAuthenticationHandler.AuthenticationScheme;
 					}
 					else
 					{
+						// If an authentication cookie is present, use it to get authentication information
+						Options.DefaultAuthenticateScheme = CookieAuthenticationDefaults.AuthenticationScheme;
+						Options.DefaultSignInScheme = CookieAuthenticationDefaults.AuthenticationScheme;
+
+						// If authentication is required, and no cookie is present, use OIDC to sign in
 						Options.DefaultChallengeScheme = OktaDefaults.AuthenticationScheme;
 					}
 				});
diff --git a/Engine/Source/Programs/Horde/HordeServerTests/ControllerIntegrationTest.cs b/Engine/Source/Programs/Horde/HordeServerTests/ControllerIntegrationTest.cs
--- a/Engine/Source/Programs/Horde/HordeServerTests/ControllerIntegrationTest.cs
+++ b/Engine/Source/Programs/Horde/HordeServerTests/ControllerIntegrationTest.cs
@@ -71,7 +71,6 @@ namespace HordeServerTests
 
             _factory = GetFactory();
             _client = _factory.CreateClient();
-            _client.DefaultRequestHeaders.Authorization = new AuthenticationHeaderValue("Test");
             return _client;
         }
 
diff --git a/Engine/Source/Programs/Horde/HordeServerTests/Fixture.cs b/Engine/Source/Programs/Horde/HordeServerTests/Fixture.cs
--- a/Engine/Source/Programs/Horde/HordeServerTests/Fixture.cs
+++ b/Engine/Source/Programs/Horde/HordeServerTests/Fixture.cs
@@ -71,9 +71,10 @@ namespace HordeServerTests
 			List<CreateStreamTabRequest> Tabs = new List<CreateStreamTabRequest>();
 			Tabs.Add(new CreateJobsTabRequest { Title = "foo", Templates = new List<string> { TemplateRefId1.ToString(), TemplateRefId2.ToString() } });
 
+			Stream = await StreamService.StreamCollection.GetAsync(new StreamId("ue5-main"));
 			Stream = await StreamService.StreamCollection.TryCreateOrReplaceAsync(
 				new StreamId("ue5-main"),
-				null,
+				Stream,
 				"",
 				"",
 				new ProjectId("does-not-exist"),
