Horde: Fixing issue with null results on changes query
#rnx

diff --git a/Engine/Source/Programs/Horde/HordeServer/Services/PerforceServiceApi.cs b/Engine/Source/Programs/Horde/HordeServer/Services/PerforceServiceApi.cs
--- a/Engine/Source/Programs/Horde/HordeServer/Services/PerforceServiceApi.cs
+++ b/Engine/Source/Programs/Horde/HordeServer/Services/PerforceServiceApi.cs
@@ -444,9 +444,13 @@ namespace HordeServer.Services
 					IList<P4.Changelist> Changelists = Repository.GetChangelists(Options, FileSpec);
 
 					List<ChangeSummary> Changes = new List<ChangeSummary>();
-					foreach (P4.Changelist Changelist in Changelists)
+
+					if (Changelists != null)
 					{
-						Changes.Add(new ChangeSummary(Changelist.Id, Changelist.OwnerName, Changelist.Description));
+						foreach (P4.Changelist Changelist in Changelists)
+						{
+							Changes.Add(new ChangeSummary(Changelist.Id, Changelist.OwnerName, Changelist.Description));
+						}
 					}
 
 					return Changes;
