Horde: Add the -TokenSignature= argument to build jobs.

diff --git a/Engine/Source/Programs/Horde/HordeAgent/Execution/BuildGraphExecutor.cs b/Engine/Source/Programs/Horde/HordeAgent/Execution/BuildGraphExecutor.cs
--- a/Engine/Source/Programs/Horde/HordeAgent/Execution/BuildGraphExecutor.cs
+++ b/Engine/Source/Programs/Horde/HordeAgent/Execution/BuildGraphExecutor.cs
@@ -352,6 +352,7 @@ namespace HordeAgent.Execution
 			}
 			Arguments.AppendArgument("-HordeExport=", DefinitionFile.FullName);
 			Arguments.AppendArgument("-ListOnly");
+			Arguments.AppendArgument("-TokenSignature=", JobId.ToString());
 			foreach (string AdditionalArgument in AdditionalArguments)
 			{
 				Arguments.AppendArgument(AdditionalArgument);
@@ -584,7 +585,8 @@ namespace HordeAgent.Execution
 			{
 				Arguments.AppendArgument("-SharedStorageDir=", SharedStorageDir.FullName);
 			}
-			foreach(string AdditionalArgument in AdditionalArguments)
+			Arguments.AppendArgument("-TokenSignature=", JobId.ToString());
+			foreach (string AdditionalArgument in AdditionalArguments)
 			{
 				if (!bPreprocessScript || !AdditionalArgument.StartsWith("-set:", StringComparison.OrdinalIgnoreCase))
 				{
