Horde: Allow overriding of agent name

diff --git a/Engine/Source/Programs/Horde/HordeAgent/AgentSettings.cs b/Engine/Source/Programs/Horde/HordeAgent/AgentSettings.cs
--- a/Engine/Source/Programs/Horde/HordeAgent/AgentSettings.cs
+++ b/Engine/Source/Programs/Horde/HordeAgent/AgentSettings.cs
@@ -154,6 +154,12 @@ namespace HordeAgent
 		/// The default server, unless overridden from the command line
 		/// </summary>
 		public string? Server { get; set; }
+		
+		/// <summary>
+		/// Name of agent to report as when connecting to server.
+		/// By default, the computer's hostname will be used.
+		/// </summary>
+		public string? Name { get; set; }
 
 		/// <summary>
 		/// The executor to use for jobs. Defaults to the Perforce executor.
@@ -224,5 +230,10 @@ namespace HordeAgent
 
 			return GetServerProfile(Server);
 		}
+		
+		internal string GetAgentName()
+		{
+			return Name ?? Environment.MachineName;
+		}
 	}
 }
diff --git a/Engine/Source/Programs/Horde/HordeAgent/Services/WorkerService.cs b/Engine/Source/Programs/Horde/HordeAgent/Services/WorkerService.cs
--- a/Engine/Source/Programs/Horde/HordeAgent/Services/WorkerService.cs
+++ b/Engine/Source/Programs/Horde/HordeAgent/Services/WorkerService.cs
@@ -369,7 +369,7 @@ namespace HordeAgent.Services
 
 				// Create the session information
 				CreateSessionRequest SessionRequest = new CreateSessionRequest();
-				SessionRequest.Name = Environment.MachineName;
+				SessionRequest.Name = Settings.GetAgentName();
 				SessionRequest.Status = AgentStatus.Ok;
 				SessionRequest.Capabilities = Capabilities;
 				SessionRequest.Version = Program.Version;
