Horde: Use semibold to avoid fuzzy font on device view buttons
#rnx

diff --git a/Engine/Source/Programs/Horde/HordeServer/Models/Device.cs b/Engine/Source/Programs/Horde/HordeServer/Models/Device.cs
--- a/Engine/Source/Programs/Horde/HordeServer/Models/Device.cs
+++ b/Engine/Source/Programs/Horde/HordeServer/Models/Device.cs
@@ -169,7 +169,7 @@ namespace HordeServer.Models
 		public DateTime? ProblemTimeUtc { get; }
 
 		/// <summary>
-		/// The last time a problem was reported
+		/// The time device was marked for maintenance
 		/// </summary>
 		public DateTime? MaintenanceTimeUtc { get; }
 
