Horde: Allow specifying path to dotnet binary when installing agent as a service

diff --git a/Engine/Source/Programs/Horde/HordeAgent/Commands/Service/InstallCommand.cs b/Engine/Source/Programs/Horde/HordeAgent/Commands/Service/InstallCommand.cs
--- a/Engine/Source/Programs/Horde/HordeAgent/Commands/Service/InstallCommand.cs
+++ b/Engine/Source/Programs/Horde/HordeAgent/Commands/Service/InstallCommand.cs
@@ -39,6 +39,13 @@ namespace HordeAgent.Commands.Service
 		/// </summary>
 		[CommandLine("-Server=")]
 		public string? Server = null;
+		
+		/// <summary>
+		/// Path to dotnet executable (dotnet.exe on Windows)
+		/// When left empty, the value of "dotnet" will be used.
+		/// </summary>
+		[CommandLine("-DotNetExecutable=")]
+		public string DotNetExecutable = "dotnet";
 
 		/// <summary>
 		/// Runs the service indefinitely
@@ -71,7 +78,7 @@ namespace HordeAgent.Commands.Service
 				Logger.LogInformation("Registering {ServiceName} service", ServiceName);
 
 				StringBuilder CommandLine = new StringBuilder();
-				CommandLine.AppendFormat("dotnet \"{0}\" service run", Assembly.GetEntryAssembly()!.Location);
+				CommandLine.AppendFormat("{0} \"{1}\" service run", DotNetExecutable, Assembly.GetEntryAssembly()!.Location);
 				if(Server != null)
 				{
 					CommandLine.Append($" -server={Server}");
