Horde: Revert change to specify TokenSignature; does not work because job setup expands entire graph.

diff --git a/Engine/Source/Programs/Horde/HordeAgent/Execution/BuildGraphExecutor.cs b/Engine/Source/Programs/Horde/HordeAgent/Execution/BuildGraphExecutor.cs
--- a/Engine/Source/Programs/Horde/HordeAgent/Execution/BuildGraphExecutor.cs
+++ b/Engine/Source/Programs/Horde/HordeAgent/Execution/BuildGraphExecutor.cs
@@ -352,7 +352,7 @@ namespace HordeAgent.Execution
 			}
 			Arguments.AppendArgument("-HordeExport=", DefinitionFile.FullName);
 			Arguments.AppendArgument("-ListOnly");
-			Arguments.AppendArgument("-TokenSignature=", JobId.ToString());
+			//Arguments.AppendArgument("-TokenSignature=", JobId.ToString());
 			foreach (string AdditionalArgument in AdditionalArguments)
 			{
 				Arguments.AppendArgument(AdditionalArgument);
@@ -585,7 +585,7 @@ namespace HordeAgent.Execution
 			{
 				Arguments.AppendArgument("-SharedStorageDir=", SharedStorageDir.FullName);
 			}
-			Arguments.AppendArgument("-TokenSignature=", JobId.ToString());
+//			Arguments.AppendArgument("-TokenSignature=", JobId.ToString());
 			foreach (string AdditionalArgument in AdditionalArguments)
 			{
 				if (!bPreprocessScript || !AdditionalArgument.StartsWith("-set:", StringComparison.OrdinalIgnoreCase))
