Horde: Use IoHash for all blob hashes in CAS.

diff --git a/Engine/Source/Programs/Horde/HordeServer/HordeServer.csproj b/Engine/Source/Programs/Horde/HordeServer/HordeServer.csproj
--- a/Engine/Source/Programs/Horde/HordeServer/HordeServer.csproj
+++ b/Engine/Source/Programs/Horde/HordeServer/HordeServer.csproj
@@ -74,6 +74,7 @@
 
   <ItemGroup>
     <ProjectReference Include="..\..\Shared\EpicGames.Core\EpicGames.Core.csproj" />
+    <ProjectReference Include="..\..\Shared\EpicGames.IoHash\EpicGames.IoHash.csproj" />
     <ProjectReference Include="..\..\Shared\EpicGames.Perforce\EpicGames.Perforce.csproj" />
     <ProjectReference Include="..\HordeCommon\HordeCommon.csproj" />
     <ProjectReference Include="..\..\..\ThirdParty\Perforce\P4Api.Net\p4api.net\p4api.net.csproj" />
diff --git a/Engine/Source/Programs/Horde/HordeServer/Rpc/ActionCacheService.cs b/Engine/Source/Programs/Horde/HordeServer/Rpc/ActionCacheService.cs
--- a/Engine/Source/Programs/Horde/HordeServer/Rpc/ActionCacheService.cs
+++ b/Engine/Source/Programs/Horde/HordeServer/Rpc/ActionCacheService.cs
@@ -45,8 +45,8 @@ namespace HordeServer.Rpc
 				Request.ActionResult.WriteTo(Stream);
 			}
 
-			BlobHash ActionHash = BlobHash.Parse(Request.ActionDigest.Hash);
-			BlobHash ResultHash = await StorageService.PutBlobAsync(Data);
+			IoHash ActionHash = IoHash.Parse(Request.ActionDigest.Hash);
+			IoHash ResultHash = await StorageService.PutBlobAsync(Data);
 			await StorageService.SetRefAsync(ActionHash, ResultHash);
 
 			return Request.ActionResult;
@@ -70,9 +70,9 @@ namespace HordeServer.Rpc
 		/// <returns></returns>
 		public async Task<ActionResult?> TryGetResult(Digest ActionDigest)
 		{
-			BlobHash ActionHash = ActionDigest.ToHashValue();
+			IoHash ActionHash = ActionDigest.ToHashValue();
 
-			BlobHash? ResultHash = await StorageService.GetRefAsync(ActionHash, TimeSpan.FromDays(1.0));
+			IoHash? ResultHash = await StorageService.GetRefAsync(ActionHash, TimeSpan.FromDays(1.0));
 			if (ResultHash == null)
 			{
 				return null;
diff --git a/Engine/Source/Programs/Horde/HordeServer/Rpc/ByteStreamService.cs b/Engine/Source/Programs/Horde/HordeServer/Rpc/ByteStreamService.cs
--- a/Engine/Source/Programs/Horde/HordeServer/Rpc/ByteStreamService.cs
+++ b/Engine/Source/Programs/Horde/HordeServer/Rpc/ByteStreamService.cs
@@ -26,10 +26,10 @@ namespace HordeServer.Rpc
 		class UploadInfo
 		{
 			public long Position;
-			public BlobHash Hash { get; }
+			public IoHash Hash { get; }
 			public byte[] Data { get; }
 
-			public UploadInfo(BlobHash Hash, long Length)
+			public UploadInfo(IoHash Hash, long Length)
 			{
 				this.Hash = Hash;
 				this.Data = new byte[Length];
@@ -60,7 +60,7 @@ namespace HordeServer.Rpc
 		/// </summary>
 		/// <param name="ResourceName">The resource name</param>
 		/// <returns>THe hash and size</returns>
-		static (BlobHash, long) ParseResourceName(string ResourceName)
+		static (IoHash, long) ParseResourceName(string ResourceName)
 		{
 			// blobs/{hash}/{size}/foo/bar/baz.cc
 			Match Match = Regex.Match(ResourceName, @"^(?:[^/]+/)?(?:uploads/[^/]+/)?blobs/([0-9a-zA-Z]+)/([0-9]+)(?:/|$)");
@@ -69,7 +69,7 @@ namespace HordeServer.Rpc
 				throw new RpcException(new Status(StatusCode.InvalidArgument, "Invalid resource name"));
 			}
 
-			BlobHash Hash = BlobHash.Parse(Match.Groups[1].Value);
+			IoHash Hash = IoHash.Parse(Match.Groups[1].Value);
 			long Size = long.Parse(Match.Groups[2].Value, CultureInfo.InvariantCulture);
 			return (Hash, Size);
 		}
@@ -123,7 +123,7 @@ namespace HordeServer.Rpc
 			UploadInfo? UploadInfo;
 			if (!Uploads.TryGetValue(ResourceName, out UploadInfo))
 			{
-				(BlobHash Hash, long Size) = ParseResourceName(ResourceName);
+				(IoHash Hash, long Size) = ParseResourceName(ResourceName);
 
 				UploadInfo = new UploadInfo(Hash, Size);
 				while (!Uploads.TryAdd(ResourceName, UploadInfo!))
@@ -199,7 +199,7 @@ namespace HordeServer.Rpc
 				}
 			}
 
-			(BlobHash Hash, long Size) = ParseResourceName(Request.ResourceName);
+			(IoHash Hash, long Size) = ParseResourceName(Request.ResourceName);
 			if (await StorageService.ShouldPutBlobAsync(Hash))
 			{
 				throw new RpcException(new Status(StatusCode.NotFound, "Invalid resource name"));
diff --git a/Engine/Source/Programs/Horde/HordeServer/Storage/IStorageService.cs b/Engine/Source/Programs/Horde/HordeServer/Storage/IStorageService.cs
--- a/Engine/Source/Programs/Horde/HordeServer/Storage/IStorageService.cs
+++ b/Engine/Source/Programs/Horde/HordeServer/Storage/IStorageService.cs
@@ -1,5 +1,6 @@
 // Copyright Epic Games, Inc. All Rights Reserved.
 
+using EpicGames.Core;
 using HordeServer.Storage.Primitives;
 using System;
 using System.Collections.Generic;
@@ -17,7 +18,7 @@ namespace HordeServer.Storage
 		/// Attempts to retrieve a blob from the store. 
 		/// 
 		/// Note that this call may succeed even though the object is not in persistent storage; it may be retrieved
-		/// from a volatile cache. Call <see cref="ShouldPutBlobAsync(BlobHash)"/> to determine if a blob needs to
+		/// from a volatile cache. Call <see cref="ShouldPutBlobAsync(IoHash)"/> to determine if a blob needs to
 		/// be put to the store.
 		/// 
 		/// This method takes a timestamp indicating the required replication consistency. The method may return null
@@ -26,7 +27,7 @@ namespace HordeServer.Storage
 		/// <param name="Hash">Hash of the blob</param>
 		/// <param name="Deadline">Deadline for returning a result. If zero, the storage provider should use a default value.</param>
 		/// <returns>The blob data. May return ReadOnlyMemory<byte>.Empty if it does not exist.</byte></returns>
-		Task<ReadOnlyMemory<byte>?> TryGetBlobAsync(BlobHash Hash, DateTime Deadline = default(DateTime));
+		Task<ReadOnlyMemory<byte>?> TryGetBlobAsync(IoHash Hash, DateTime Deadline = default(DateTime));
 
 		/// <summary>
 		/// Adds a blob to the store.
@@ -34,7 +35,7 @@ namespace HordeServer.Storage
 		/// <param name="Hash">Hash of the value. Used as a hint for possible sharding; the server will verify this hash and fail if it does not match.</param>
 		/// <param name="Value">The blob data</param>
 		/// <returns>Async task</returns>
-		Task PutBlobAsync(BlobHash Hash, ReadOnlyMemory<byte> Value);
+		Task PutBlobAsync(IoHash Hash, ReadOnlyMemory<byte> Value);
 
 		/// <summary>
 		/// Determines if the blob should be written to the store. The server may return true even if it exists in a store
@@ -43,7 +44,7 @@ namespace HordeServer.Storage
 		/// </summary>
 		/// <param name="Hash">Hash of the blob to write</param>
 		/// <returns>True if the blob should be written to the store</returns>
-		Task<bool> ShouldPutBlobAsync(BlobHash Hash);
+		Task<bool> ShouldPutBlobAsync(IoHash Hash);
 
 		/// <summary>
 		/// Gets a reference from the storage layer
@@ -51,7 +52,7 @@ namespace HordeServer.Storage
 		/// <param name="Key">Key for the reference</param>
 		/// <param name="MaxDrift">Required replication consistency</param>
 		/// <returns>Hash of the result</returns>
-		Task<BlobHash?> GetRefAsync(BlobHash Key, TimeSpan MaxDrift = default(TimeSpan));
+		Task<IoHash?> GetRefAsync(IoHash Key, TimeSpan MaxDrift = default(TimeSpan));
 
 		/// <summary>
 		/// Sets the value associated with a key. Writes are propagated as a "last writer wins" policy.
@@ -59,7 +60,7 @@ namespace HordeServer.Storage
 		/// <param name="Key">The key to set</param>
 		/// <param name="Value">Value for the key. If null, the key is deleted.</param>
 		/// <returns>Async task</returns>
-		Task SetRefAsync(BlobHash Key, BlobHash? Value);
+		Task SetRefAsync(IoHash Key, IoHash? Value);
 
 		/// <summary>
 		/// Updates the last modified timestamp of the given key. Possibly equivalent to doing a SET() on the same key/value pair, but does not 
@@ -67,7 +68,7 @@ namespace HordeServer.Storage
 		/// </summary>
 		/// <param name="Key">The key to touch</param>
 		/// <returns>Async task</returns>
-		Task TouchRefAsync(BlobHash Key);
+		Task TouchRefAsync(IoHash Key);
 	}
 
 	/// <summary>
@@ -82,7 +83,7 @@ namespace HordeServer.Storage
 		/// <param name="Hash">Hash of the blob to retrieve</param>
 		/// <param name="Deadline">Deadline for retreving the data</param>
 		/// <returns>Memory for the blob</returns>
-		public static async Task<ReadOnlyMemory<byte>> GetBlobAsync(this IStorageService StorageService, BlobHash Hash, DateTime Deadline = default(DateTime))
+		public static async Task<ReadOnlyMemory<byte>> GetBlobAsync(this IStorageService StorageService, IoHash Hash, DateTime Deadline = default(DateTime))
 		{
 			ReadOnlyMemory<byte>? Blob = await StorageService.TryGetBlobAsync(Hash, Deadline);
 			if (Blob == null)
@@ -93,14 +94,14 @@ namespace HordeServer.Storage
 		}
 
 		/// <summary>
-		/// Hashes a blob, then adds it to the storage service. If the hash is already known, prefer <see cref="IStorageService.PutBlobAsync(BlobHash, ReadOnlyMemory{byte})"/>
+		/// Hashes a blob, then adds it to the storage service. If the hash is already known, prefer <see cref="IStorageService.PutBlobAsync(IoHash, ReadOnlyMemory{byte})"/>
 		/// </summary>
 		/// <param name="StorageService">The storage service</param>
 		/// <param name="Data">The data to add</param>
 		/// <returns>Hash of the data that was added</returns>
-		public static async Task<BlobHash> PutBlobAsync(this IStorageService StorageService, ReadOnlyMemory<byte> Data)
+		public static async Task<IoHash> PutBlobAsync(this IStorageService StorageService, ReadOnlyMemory<byte> Data)
 		{
-			BlobHash Hash = BlobHash.Compute(Data);
+			IoHash Hash = IoHash.Compute(Data.Span);
 			await StorageService.PutBlobAsync(Hash, Data);
 			return Hash;
 		}
diff --git a/Engine/Source/Programs/Horde/HordeServer/Storage/Primitives/BlobHash.cs b/Engine/Source/Programs/Horde/HordeServer/Storage/Primitives/BlobHash.cs
deleted file mode 100644
diff --git a/Engine/Source/Programs/Horde/HordeServer/Storage/Primitives/BlobRef.cs b/Engine/Source/Programs/Horde/HordeServer/Storage/Primitives/BlobRef.cs
--- a/Engine/Source/Programs/Horde/HordeServer/Storage/Primitives/BlobRef.cs
+++ b/Engine/Source/Programs/Horde/HordeServer/Storage/Primitives/BlobRef.cs
@@ -1,5 +1,6 @@
 // Copyright Epic Games, Inc. All Rights Reserved.
 
+using EpicGames.Core;
 using System;
 using System.Collections.Generic;
 using System.Linq;
@@ -20,14 +21,14 @@ namespace HordeServer.Storage.Primitives
 		/// <summary>
 		/// The blob hash
 		/// </summary>
-		public BlobHash Hash { get; }
+		public IoHash Hash { get; }
 
 		/// <summary>
 		/// Constructor
 		/// </summary>
 		/// <param name="Type"></param>
 		/// <param name="Hash"></param>
-		public BlobRef(IBlobType Type, BlobHash Hash)
+		public BlobRef(IBlobType Type, IoHash Hash)
 		{
 			this.Type = Type;
 			this.Hash = Hash;
@@ -54,13 +55,13 @@ namespace HordeServer.Storage.Primitives
 		/// <summary>
 		/// Hash of the referenced blob
 		/// </summary>
-		public BlobHash Hash { get; }
+		public IoHash Hash { get; }
 
 		/// <summary>
 		/// Constructor
 		/// </summary>
 		/// <param name="Hash">Hash of the blob</param>
-		public BlobRef(BlobHash Hash)
+		public BlobRef(IoHash Hash)
 		{
 			this.Hash = Hash;
 		}
diff --git a/Engine/Source/Programs/Horde/HordeServer/Storage/Primitives/KeyValueItem.cs b/Engine/Source/Programs/Horde/HordeServer/Storage/Primitives/KeyValueItem.cs
--- a/Engine/Source/Programs/Horde/HordeServer/Storage/Primitives/KeyValueItem.cs
+++ b/Engine/Source/Programs/Horde/HordeServer/Storage/Primitives/KeyValueItem.cs
@@ -1,5 +1,6 @@
 // Copyright Epic Games, Inc. All Rights Reserved.
 
+using EpicGames.Core;
 using System;
 using System.Collections.Generic;
 using System.Linq;
@@ -15,7 +16,7 @@ namespace HordeServer.Storage.Primitives
 		/// <summary>
 		/// The key
 		/// </summary>
-		public BlobHash Key { get; }
+		public IoHash Key { get; }
 
 		/// <summary>
 		/// Block of memory containing the value for this item
@@ -27,7 +28,7 @@ namespace HordeServer.Storage.Primitives
 		/// </summary>
 		/// <param name="Key"></param>
 		/// <param name="Value"></param>
-		public KeyValueItem(BlobHash Key, ReadOnlyMemory<byte> Value)
+		public KeyValueItem(IoHash Key, ReadOnlyMemory<byte> Value)
 		{
 			this.Key = Key;
 			this.Value = Value;
@@ -44,14 +45,14 @@ namespace HordeServer.Storage.Primitives
 		{
 			this.Keys = Node.Keys;
 			this.Values = Node.Values;
-			this.ValueSize = Node.IsLeafNode ? Tree.LeafValueSize : BlobHash.NumBytes;
+			this.ValueSize = Node.IsLeafNode ? Tree.LeafValueSize : IoHash.NumBytes;
 		}
 
 		public bool IsEmpty => Keys.Length == 0;
 
 		public void MoveNext()
 		{
-			Keys = Keys.Slice(BlobHash.NumBytes);
+			Keys = Keys.Slice(IoHash.NumBytes);
 			Values = Values.Slice(ValueSize);
 		}
 
@@ -65,7 +66,7 @@ namespace HordeServer.Storage.Primitives
 
 		public void CopyTo(KeyValueItemWriter Writer, int Count)
 		{
-			ReadOnlyMemory<byte> CopyKeys = Keys.Slice(0, Count * BlobHash.NumBytes);
+			ReadOnlyMemory<byte> CopyKeys = Keys.Slice(0, Count * IoHash.NumBytes);
 			ReadOnlyMemory<byte> CopyValues = Values.Slice(0, Count * ValueSize);
 
 			Writer.WriteItems(CopyKeys, CopyValues);
@@ -86,7 +87,7 @@ namespace HordeServer.Storage.Primitives
 			this.Values = Node.Values;
 		}
 
-		public void WriteItem(BlobHash Key, ReadOnlyMemory<byte> Value)
+		public void WriteItem(IoHash Key, ReadOnlyMemory<byte> Value)
 		{
 			WriteItems(Key.Memory, Value);
 		}
diff --git a/Engine/Source/Programs/Horde/HordeServer/Storage/Primitives/KeyValueNode.cs b/Engine/Source/Programs/Horde/HordeServer/Storage/Primitives/KeyValueNode.cs
--- a/Engine/Source/Programs/Horde/HordeServer/Storage/Primitives/KeyValueNode.cs
+++ b/Engine/Source/Programs/Horde/HordeServer/Storage/Primitives/KeyValueNode.cs
@@ -1,5 +1,6 @@
 // Copyright Epic Games, Inc. All Rights Reserved.
 
+using EpicGames.Core;
 using System;
 using System.Buffers.Binary;
 using System.Collections;
@@ -65,7 +66,7 @@ namespace HordeServer.Storage.Primitives
 			this.NumKeyBits = NumKeyBits;
 			this.NumKeyBitsIfMerged = NumKeyBitsIfMerged;
 
-			int KeyDataLength = BlobHash.NumBytes * NumItems;
+			int KeyDataLength = IoHash.NumBytes * NumItems;
 			int ValueDataLength = ValueSize * NumItems;
 			Memory = new byte[(sizeof(int) * 4) + KeyDataLength + ValueDataLength];
 
@@ -119,7 +120,7 @@ namespace HordeServer.Storage.Primitives
 		/// <summary>
 		/// Whether this node is a leaf node
 		/// </summary>
-		public bool IsLeafNode => NumKeyBits == BlobHash.NumBits;
+		public bool IsLeafNode => NumKeyBits == IoHash.NumBits;
 
 		/// <summary>
 		/// Constructor
@@ -143,7 +144,7 @@ namespace HordeServer.Storage.Primitives
 			NumKeyBitsIfMerged = BinaryPrimitives.ReadInt32LittleEndian(Input.Span);
 			Input = Input.Slice(sizeof(int));
 
-			Keys = Input.Slice(0, NumItems * BlobHash.NumBytes);
+			Keys = Input.Slice(0, NumItems * IoHash.NumBytes);
 			Input = Input.Slice(Keys.Length);
 
 			Values = Input;
diff --git a/Engine/Source/Programs/Horde/HordeServer/Storage/Primitives/KeyValueTree.cs b/Engine/Source/Programs/Horde/HordeServer/Storage/Primitives/KeyValueTree.cs
--- a/Engine/Source/Programs/Horde/HordeServer/Storage/Primitives/KeyValueTree.cs
+++ b/Engine/Source/Programs/Horde/HordeServer/Storage/Primitives/KeyValueTree.cs
@@ -55,7 +55,7 @@ namespace HordeServer.Storage.Primitives
 		/// <param name="LeafValueSize">Length of each value</param>
 		/// <param name="LeafValueRefTypes">References from each value</param>
 		public KeyValueTree(IBlobType KeyType, int LeafValueSize, params (int, IBlobType)[] LeafValueRefTypes)
-			: this(KeyType, LeafValueSize, LeafValueRefTypes, DefaultMaxNodeSize / (BlobHash.NumBytes * 2), DefaultMaxNodeSize / (BlobHash.NumBytes + LeafValueSize))
+			: this(KeyType, LeafValueSize, LeafValueRefTypes, DefaultMaxNodeSize / (IoHash.NumBytes * 2), DefaultMaxNodeSize / (IoHash.NumBytes + LeafValueSize))
 		{
 		}
 
@@ -84,14 +84,14 @@ namespace HordeServer.Storage.Primitives
 		public IEnumerable<BlobRef> GetRefs(ReadOnlyMemory<byte> Data)
 		{
 			ReadOnlyKeyValueNode Node = new ReadOnlyKeyValueNode(Data);
-			if (Node.NumKeyBits == BlobHash.NumBits)
+			if (Node.NumKeyBits == IoHash.NumBits)
 			{
 				// Leaf node: keys are references to other blobs, values follow type definition
 				ReadOnlyMemory<byte> Keys = Node.Keys;
 				while (!Keys.IsEmpty)
 				{
-					yield return new BlobRef(KeyType, new BlobHash(Keys.Slice(0, BlobHash.NumBytes)));
-					Keys = Keys.Slice(BlobHash.NumBytes);
+					yield return new BlobRef(KeyType, new IoHash(Keys.Slice(0, IoHash.NumBytes)));
+					Keys = Keys.Slice(IoHash.NumBytes);
 				}
 
 				ReadOnlyMemory<byte> Values = Node.Values;
@@ -99,7 +99,7 @@ namespace HordeServer.Storage.Primitives
 				{
 					foreach ((int Offset, IBlobType Type) in LeafValueRefTypes)
 					{
-						yield return new BlobRef(Type, new BlobHash(Values.Slice(Offset, BlobHash.NumBytes)));
+						yield return new BlobRef(Type, new IoHash(Values.Slice(Offset, IoHash.NumBytes)));
 					}
 					Values = Values.Slice(LeafValueSize);
 				}
@@ -107,8 +107,8 @@ namespace HordeServer.Storage.Primitives
 			else
 			{
 				// Internal node: keys are just prefixes, and not concrete references. Values are always just references to other trees.
-				ReadOnlyBlobHashArray Values = new ReadOnlyBlobHashArray(Node.Values);
-				foreach(BlobHash Value in Values)
+				ReadOnlyHashArray Values = new ReadOnlyHashArray(Node.Values);
+				foreach(IoHash Value in Values)
 				{
 					yield return new BlobRef(this, Value);
 				}
@@ -141,7 +141,7 @@ namespace HordeServer.Storage.Primitives
 		/// Constructor
 		/// </summary>
 		public KeyValueTree()
-			: base(BlobRef<TKey>.Type, BlobHash.NumBytes, (0, BlobRef<TValue>.Type))
+			: base(BlobRef<TKey>.Type, IoHash.NumBytes, (0, BlobRef<TValue>.Type))
 		{
 		}
 	}
@@ -159,48 +159,48 @@ namespace HordeServer.Storage.Primitives
 		/// <returns>Reference to the new tree</returns>
 		public static async Task<BlobRef<TTree>> CreateKeyValueTreeAsync<TTree>(this IStorageService Storage) where TTree : KeyValueTree, new()
 		{
-			KeyValueNode Node = new KeyValueNode(0, 0, BlobHash.NumBits, BlobHash.NumBits, 0);
+			KeyValueNode Node = new KeyValueNode(0, 0, IoHash.NumBits, IoHash.NumBits, 0);
 			return new BlobRef<TTree>(await Storage.PutKeyValueNodeAsync(Node));
 		}
 
-		public static async Task<BlobRef<TValue>?> SearchKeyValueTreeAsync<TKey, TValue>(this IStorageService Storage, BlobRef<KeyValueTree<TKey, TValue>> Tree, BlobHash FindHash) where TKey : IBlobType, new() where TValue : IBlobType, new()
+		public static async Task<BlobRef<TValue>?> SearchKeyValueTreeAsync<TKey, TValue>(this IStorageService Storage, BlobRef<KeyValueTree<TKey, TValue>> Tree, IoHash FindHash) where TKey : IBlobType, new() where TValue : IBlobType, new()
 		{
 			ReadOnlyMemory<byte>? Result = await SearchKeyValueTreeInternalAsync(Storage, Tree, FindHash);
 			if(Result == null)
 			{
 				return null;
 			}
-			return new BlobRef<TValue>(new BlobHash(Result.Value));
+			return new BlobRef<TValue>(new IoHash(Result.Value));
 		}
 
-		public static Task<ReadOnlyMemory<byte>?> SearchKeyValueTreeAsync<TTree>(this IStorageService Storage, BlobRef<TTree> Tree, BlobHash FindHash) where TTree : KeyValueTree, new()
+		public static Task<ReadOnlyMemory<byte>?> SearchKeyValueTreeAsync<TTree>(this IStorageService Storage, BlobRef<TTree> Tree, IoHash FindHash) where TTree : KeyValueTree, new()
 		{
 			return SearchKeyValueTreeInternalAsync(Storage, Tree, FindHash);
 		}
 
-		static async Task<ReadOnlyMemory<byte>?> SearchKeyValueTreeInternalAsync<TTree>(this IStorageService Storage, BlobRef<TTree> Tree, BlobHash FindHash) where TTree : KeyValueTree, new()
+		static async Task<ReadOnlyMemory<byte>?> SearchKeyValueTreeInternalAsync<TTree>(this IStorageService Storage, BlobRef<TTree> Tree, IoHash FindHash) where TTree : KeyValueTree, new()
 		{
-			BlobHash Hash = Tree.Hash;
+			IoHash Hash = Tree.Hash;
 			for (; ; )
 			{
 				ReadOnlyKeyValueNode Node = await Storage.GetKeyValueNodeAsync(Hash);
 
-				BlobHash MaskHash = FindHash.MaskLeft(Node.NumKeyBits);
+				IoHash MaskHash = MaskLeft(FindHash, Node.NumKeyBits);
 
-				ReadOnlyBlobHashArray Keys = new ReadOnlyBlobHashArray(Node.Keys);
+				ReadOnlyHashArray Keys = new ReadOnlyHashArray(Node.Keys);
 
 				int KeyIdx = Keys.BinarySearch(MaskHash);
 				if (KeyIdx < 0)
 				{
 					return null;
 				}
-				else if (Node.NumKeyBits == BlobHash.NumBits)
+				else if (Node.NumKeyBits == IoHash.NumBits)
 				{
 					return Node.Values.Slice(KeyIdx * BlobRef<TTree>.Type.LeafValueSize, BlobRef<TTree>.Type.LeafValueSize);
 				}
 				else
 				{
-					Hash = new ReadOnlyBlobHashArray(Node.Values)[KeyIdx];
+					Hash = new ReadOnlyHashArray(Node.Values)[KeyIdx];
 				}
 			}
 		}
@@ -215,7 +215,7 @@ namespace HordeServer.Storage.Primitives
 		/// <returns></returns>
 		public static async Task<BlobRef<TTree>> UpdateKeyValueTreeAsync<TTree>(this IStorageService Storage, BlobRef<TTree> Tree, IEnumerable<KeyValueItem> NewItems) where TTree : KeyValueTree, new()
 		{
-			BlobHash NewRootHash = await UpdateKeyValueTreeAsync(Storage, BlobRef<TTree>.Type, Tree.Hash, NewItems);
+			IoHash NewRootHash = await UpdateKeyValueTreeAsync(Storage, BlobRef<TTree>.Type, Tree.Hash, NewItems);
 			return new BlobRef<TTree>(NewRootHash);
 		}
 
@@ -227,7 +227,7 @@ namespace HordeServer.Storage.Primitives
 		/// <param name="RootHash">Hash of the root node of the tree</param>
 		/// <param name="NewItems">New items to add</param>
 		/// <returns>Hash for the new tree root</returns>
-		static async Task<BlobHash> UpdateKeyValueTreeAsync(this IStorageService Storage, KeyValueTree Tree, BlobHash RootHash, IEnumerable<KeyValueItem> NewItems)
+		static async Task<IoHash> UpdateKeyValueTreeAsync(this IStorageService Storage, KeyValueTree Tree, IoHash RootHash, IEnumerable<KeyValueItem> NewItems)
 		{
 			// Sort all the items and remove duplicates
 			ArraySegment<KeyValueItem> NewSortedItems = NewItems.OrderBy(x => x.Key).ToArray();
@@ -246,9 +246,9 @@ namespace HordeServer.Storage.Primitives
 
 			// Get the root node
 			ReadOnlyKeyValueNode RootNode;
-			if (RootHash == BlobHash.Zero)
+			if (RootHash == IoHash.Zero)
 			{
-				RootNode = new KeyValueNode(0, 0, BlobHash.NumBits, BlobHash.NumBits, Tree.LeafValueSize);
+				RootNode = new KeyValueNode(0, 0, IoHash.NumBits, IoHash.NumBits, Tree.LeafValueSize);
 			}
 			else
 			{
@@ -272,14 +272,14 @@ namespace HordeServer.Storage.Primitives
 		static async Task<KeyValueNode> UpdateKeyValueNodeAsync(IStorageService Storage, KeyValueTree Tree, int ParentKeyBits, ReadOnlyKeyValueNode Node, ArraySegment<KeyValueItem> NewItems)
 		{
 			// Create a list wrapper for the keys, so we can perform binary searches on it
-			ReadOnlyBlobHashArray Keys = new ReadOnlyBlobHashArray(Node.Keys);
+			ReadOnlyHashArray Keys = new ReadOnlyHashArray(Node.Keys);
 
 			// Keep track of the total number of child items underneath this node in order to update the merge count
 			int NewNumChildItems = Node.NumItemsIfMerged - Node.NumItems;
 
 			// Find the necessary updates to the node
 			List<(int KeyIdx, KeyValueItem? NewItem)> Updates = new List<(int, KeyValueItem?)>(NewItems.Count);
-			if (Node.NumKeyBits == BlobHash.NumBits)
+			if (Node.NumKeyBits == IoHash.NumBits)
 			{
 				foreach (KeyValueItem NewItem in NewItems)
 				{
@@ -292,14 +292,14 @@ namespace HordeServer.Storage.Primitives
 			}
 			else
 			{
-				ReadOnlyBlobHashArray Children = new ReadOnlyBlobHashArray(Node.Values);
+				ReadOnlyHashArray Children = new ReadOnlyHashArray(Node.Values);
 				for (int Idx = 0; Idx < NewItems.Count;)
 				{
-					BlobHash KeyPrefix = NewItems[Idx].Key.MaskLeft(Node.NumKeyBits);
+					IoHash KeyPrefix = MaskLeft(NewItems[Idx].Key, Node.NumKeyBits);
 
 					// Find all the operations which start with the same prefix
 					int MinIdx = Idx++;
-					while (Idx < NewItems.Count && NewItems[Idx].Key.StartsWith(KeyPrefix, Node.NumKeyBits))
+					while (Idx < NewItems.Count && StartsWith(NewItems[Idx].Key, KeyPrefix, Node.NumKeyBits))
 					{
 						Idx++;
 					}
@@ -331,7 +331,7 @@ namespace HordeServer.Storage.Primitives
 					else
 					{
 						// Add the new item to the update list
-						BlobHash NewChildHash = await Storage.PutKeyValueNodeAsync(NewChildNode);
+						IoHash NewChildHash = await Storage.PutKeyValueNodeAsync(NewChildNode);
 						KeyValueItem NewItem = new KeyValueItem(KeyPrefix, NewChildHash.Memory);
 						Updates.Add((KeyIdx, NewItem));
 
@@ -381,7 +381,7 @@ namespace HordeServer.Storage.Primitives
 			}
 
 			// If it's larger than a threshold, we need to split this node into two
-			while (NewNode.NumItems > ((NewNode.NumKeyBits == BlobHash.NumBits)? Tree.MaxItemsPerLeafNode : Tree.MaxItemsPerInternalNode))
+			while (NewNode.NumItems > ((NewNode.NumKeyBits == IoHash.NumBits)? Tree.MaxItemsPerLeafNode : Tree.MaxItemsPerInternalNode))
 			{
 				NewNode = await SplitAsync(Storage, Tree, NewNode, ParentKeyBits);
 			}
@@ -402,7 +402,7 @@ namespace HordeServer.Storage.Primitives
 			KeyValueNode NewNode = CreateNode(Tree, Node.NumItemsIfMerged, Node.NumItemsIfMerged, Node.NumKeyBitsIfMerged, Node.NumKeyBitsIfMerged);
 			KeyValueItemWriter Writer = new KeyValueItemWriter(NewNode);
 
-			ReadOnlyBlobHashArray Values = new ReadOnlyBlobHashArray(Node.Values);
+			ReadOnlyHashArray Values = new ReadOnlyHashArray(Node.Values);
 			for (int Idx = 0; Idx < Values.Count; Idx++)
 			{
 				ReadOnlyKeyValueNode ChildNode = await Storage.GetKeyValueNodeAsync(Values[Idx]);
@@ -439,15 +439,15 @@ namespace HordeServer.Storage.Primitives
 
 			KeyValueItemReader Reader = new KeyValueItemReader(Tree, Node);
 
-			ReadOnlyBlobHashArray Keys = new ReadOnlyBlobHashArray(Node.Keys);
+			ReadOnlyHashArray Keys = new ReadOnlyHashArray(Node.Keys);
 			for (int KeyIdx = 0; KeyIdx < Keys.Count; )
 			{
 				// Mask off the first key prefix
-				BlobHash KeyPrefix = Keys[KeyIdx].MaskLeft(NewMaxKeyBits);
+				IoHash KeyPrefix = MaskLeft(Keys[KeyIdx], NewMaxKeyBits);
 
 				// Count the number of keys with the same prefix
 				int Count = 1;
-				while (KeyIdx + Count < Keys.Count && Keys[KeyIdx + Count].StartsWith(KeyPrefix, NewMaxKeyBits))
+				while (KeyIdx + Count < Keys.Count && StartsWith(Keys[KeyIdx + Count], KeyPrefix, NewMaxKeyBits))
 				{
 					Count++;
 				}
@@ -464,11 +464,11 @@ namespace HordeServer.Storage.Primitives
 				}
 
 				// Finally, store the child node and add it to its new parent
-				BlobHash NewChildHash = await Storage.PutKeyValueNodeAsync(NewChildNode);
+				IoHash NewChildHash = await Storage.PutKeyValueNodeAsync(NewChildNode);
 				NewItems.Add(new KeyValueItem(KeyPrefix, NewChildHash.Memory));
 			}
 
-			KeyValueNode NewNode = new KeyValueNode(NewItems.Count, Node.NumItems, NewMaxKeyBits, Node.NumKeyBits, BlobHash.NumBytes);
+			KeyValueNode NewNode = new KeyValueNode(NewItems.Count, Node.NumItems, NewMaxKeyBits, Node.NumKeyBits, IoHash.NumBytes);
 
 			KeyValueItemWriter Writer = new KeyValueItemWriter(NewNode);
 			foreach (KeyValueItem NewItem in NewItems)
@@ -481,7 +481,7 @@ namespace HordeServer.Storage.Primitives
 
 		public static KeyValueNode CreateNode(KeyValueTree Tree, int NumItems, int NumItemsIfMerged, int NumKeyBits, int NumKeyBitsIfMerged)
 		{
-			int NewValueSize = (NumKeyBits == BlobHash.NumBits) ? Tree.LeafValueSize : BlobHash.NumBytes;
+			int NewValueSize = (NumKeyBits == IoHash.NumBits) ? Tree.LeafValueSize : IoHash.NumBytes;
 			return new KeyValueNode(NumItems, NumItemsIfMerged, NumKeyBits, NumKeyBitsIfMerged, NewValueSize);
 		}
 
@@ -495,7 +495,7 @@ namespace HordeServer.Storage.Primitives
 		{
 			// Find the largest power of two boundary that the current key length is on
 			int KeyStep = 1;
-			while (KeyLength + KeyStep < BlobHash.NumBits)
+			while (KeyLength + KeyStep < IoHash.NumBits)
 			{
 				int NextKeyStep = KeyStep << 1;
 				if ((KeyLength & (NextKeyStep - 1)) != 0)
@@ -507,13 +507,62 @@ namespace HordeServer.Storage.Primitives
 			return KeyLength + KeyStep;
 		}
 
+		/// <summary>
+		/// Check if one hash starts with the same bits as another
+		/// </summary>
+		/// <param name="This"></param>
+		/// <param name="Other">The hash to test against</param>
+		/// <param name="Length">Number of bits to compare</param>
+		/// <returns>True if the hashes match, false otherwise</returns>
+		static bool StartsWith(IoHash This, IoHash Other, int Length)
+		{
+			ReadOnlySpan<byte> ThisSpan = This.Span;
+			ReadOnlySpan<byte> OtherSpan = Other.Span;
+
+			int Idx = 0;
+			int RemainingLength = Length;
+
+			for (; RemainingLength > 8; RemainingLength -= 8, Idx++)
+			{
+				if (ThisSpan[Idx] != OtherSpan[Idx])
+				{
+					return false;
+				}
+			}
+
+			int Mask = ~((1 << (8 - RemainingLength)) - 1);
+			return (ThisSpan[Idx] & Mask) == (OtherSpan[Idx] & Mask);
+		}
+
+		/// <summary>
+		/// Mask off the left number of bits from the hash
+		/// </summary>
+		/// <param name="Hash"></param>
+		/// <param name="NumBits">Number of bits to include</param>
+		/// <returns>Masked hash value</returns>
+		static IoHash MaskLeft(IoHash Hash, int NumBits)
+		{
+			ReadOnlySpan<byte> OldData = Hash.Span;
+			byte[] NewData = new byte[IoHash.NumBytes];
+
+			int Idx = 0;
+			int RemainingBits = NumBits;
+			for (; RemainingBits > 8; RemainingBits -= 8, Idx++)
+			{
+				NewData[Idx] = OldData[Idx];
+			}
+
+			NewData[Idx] = (byte)(OldData[Idx] & ~((1 << (8 - RemainingBits)) - 1));
+			return new IoHash(NewData);
+		}
+
 		/// <summary>
 		/// Adds a key/value tree node to a storage provider
 		/// </summary>
 		/// <param name="Provider"></param>
 		/// <param name="Node"></param>
 		/// <returns></returns>
-		public static async Task<BlobHash> PutKeyValueNodeAsync(this IStorageService Provider, ReadOnlyKeyValueNode Node)
+		public static async Task<IoHash> PutKeyValueNodeAsync(this IStorageService Provider, ReadOnlyKeyValueNode Node)
 		{
 			return await Provider.PutBlobAsync(Node.Memory);
 		}
@@ -524,7 +573,7 @@ namespace HordeServer.Storage.Primitives
 		/// <param name="Provider"></param>
 		/// <param name="Hash"></param>
 		/// <returns></returns>
-		public static async Task<ReadOnlyKeyValueNode> GetKeyValueNodeAsync(this IStorageService Provider, BlobHash Hash)
+		public static async Task<ReadOnlyKeyValueNode> GetKeyValueNodeAsync(this IStorageService Provider, IoHash Hash)
 		{
 			return new ReadOnlyKeyValueNode(await Provider.GetBlobAsync(Hash));
 		}
diff --git a/Engine/Source/Programs/Horde/HordeServer/Storage/Primitives/ReadOnlyHashArray.cs b/Engine/Source/Programs/Horde/HordeServer/Storage/Primitives/ReadOnlyHashArray.cs
new file mode 100644
--- /dev/null
+++ b/Engine/Source/Programs/Horde/HordeServer/Storage/Primitives/ReadOnlyHashArray.cs
@@ -0,0 +1,58 @@
+// Copyright Epic Games, Inc. All Rights Reserved.
+
+using EpicGames.Core;
+using System;
+using System.Collections;
+using System.Collections.Generic;
+using System.Diagnostics.CodeAnalysis;
+using System.Linq;
+using System.Security.Cryptography;
+using System.Threading.Tasks;
+
+namespace HordeServer.Storage.Primitives
+{
+	/// <summary>
+	/// Read-only array of blob hashes
+	/// </summary>
+	struct ReadOnlyHashArray : IReadOnlyList<IoHash>
+	{
+		/// <summary>
+		/// Underlying storage for the array
+		/// </summary>
+		public ReadOnlyMemory<byte> Memory { get; }
+
+		/// <summary>
+		/// Constructor
+		/// </summary>
+		/// <param name="Memory">Storage for the array</param>
+		public ReadOnlyHashArray(ReadOnlyMemory<byte> Memory)
+		{
+			this.Memory = Memory;
+		}
+
+		/// <summary>
+		/// Accessor for elements in the array
+		/// </summary>
+		/// <param name="Index">Index of the element to retrieve</param>
+		/// <returns>Hash at the given index</returns>
+		public IoHash this[int Index] => new IoHash(Memory.Slice(Index * IoHash.NumBytes, IoHash.NumBytes));
+
+		/// <inheritdoc/>
+		public int Count => Memory.Length / IoHash.NumBytes;
+
+		/// <inheritdoc/>
+		public IEnumerator<IoHash> GetEnumerator()
+		{
+			ReadOnlyMemory<byte> Remaining = Memory;
+			while (!Remaining.IsEmpty)
+			{
+				yield return new IoHash(Remaining.Slice(0, IoHash.NumBytes));
+				Remaining = Remaining.Slice(IoHash.NumBytes);
+			}
+		}
+
+		/// <inheritdoc/>
+		IEnumerator IEnumerable.GetEnumerator() => GetEnumerator();
+	}
+}
+
diff --git a/Engine/Source/Programs/Horde/HordeServer/Storage/Services/SimpleStorageService.cs b/Engine/Source/Programs/Horde/HordeServer/Storage/Services/SimpleStorageService.cs
--- a/Engine/Source/Programs/Horde/HordeServer/Storage/Services/SimpleStorageService.cs
+++ b/Engine/Source/Programs/Horde/HordeServer/Storage/Services/SimpleStorageService.cs
@@ -48,22 +48,22 @@ namespace HordeServer.Storage.Services
 		/// <param name="BaseDir">Base directory for storage</param>
 		/// <param name="Hash">Hash of the file</param>
 		/// <returns>Path to the file for this blob</returns>
-		static string GetItemPath(string BaseDir, BlobHash Hash)
+		static string GetItemPath(string BaseDir, IoHash Hash)
 		{
 			return $"{BaseDir}/{StringUtils.FormatHexString(Hash.Span.Slice(0, 1))}/{StringUtils.FormatHexString(Hash.Span.Slice(1, 1))}/{Hash}.dat";
 		}
 
 		/// <inheritdoc/>
-		public async Task<ReadOnlyMemory<byte>?> TryGetBlobAsync(BlobHash Hash, DateTime Deadline = default)
+		public async Task<ReadOnlyMemory<byte>?> TryGetBlobAsync(IoHash Hash, DateTime Deadline = default)
 		{
 			string Path = GetItemPath(BlobsPath, Hash);
 			return await Backend.ReadAsync(Path);
 		}
 
 		/// <inheritdoc/>
-		public async Task PutBlobAsync(BlobHash Hash, ReadOnlyMemory<byte> Value)
+		public async Task PutBlobAsync(IoHash Hash, ReadOnlyMemory<byte> Value)
 		{
-			BlobHash RealHash = BlobHash.Compute(Value);
+			IoHash RealHash = IoHash.Compute(Value.Span);
 			if(Hash != RealHash)
 			{
 				throw new InvalidDataException($"Hash for blob does not match. Expected {RealHash}, got {Hash}");
@@ -77,14 +77,14 @@ namespace HordeServer.Storage.Services
 		}
 
 		/// <inheritdoc/>
-		public Task<bool> ShouldPutBlobAsync(BlobHash Hash)
+		public Task<bool> ShouldPutBlobAsync(IoHash Hash)
 		{
 			// Always return true to force timestamps to be updated
 			return Task.FromResult(true);
 		}
 
 		/// <inheritdoc/>
-		public async Task SetRefAsync(BlobHash Key, BlobHash? Value)
+		public async Task SetRefAsync(IoHash Key, IoHash? Value)
 		{
 			string Path = GetItemPath(RefsPath, Key);
 			if (Value == null)
@@ -98,14 +98,14 @@ namespace HordeServer.Storage.Services
 		}
 
 		/// <inheritdoc/>
-		public async Task TouchRefAsync(BlobHash Key)
+		public async Task TouchRefAsync(IoHash Key)
 		{
 			string Path = GetItemPath(RefsPath, Key);
 			await Backend.TouchAsync(Path);
 		}
 
 		/// <inheritdoc/>
-		public async Task<BlobHash?> GetRefAsync(BlobHash Key, TimeSpan MaxDrift = default)
+		public async Task<IoHash?> GetRefAsync(IoHash Key, TimeSpan MaxDrift = default)
 		{
 			ReadOnlyMemory<byte>? Memory = await Backend.ReadAsync(GetItemPath(RefsPath, Key));
 			if (Memory == null)
@@ -114,7 +114,7 @@ namespace HordeServer.Storage.Services
 			}
 			else
 			{
-				return new BlobHash(Memory.Value);
+				return new IoHash(Memory.Value);
 			}
 		}
 	}
diff --git a/Engine/Source/Programs/Horde/HordeServer/Utilities/RemoteExtensions.cs b/Engine/Source/Programs/Horde/HordeServer/Utilities/RemoteExtensions.cs
--- a/Engine/Source/Programs/Horde/HordeServer/Utilities/RemoteExtensions.cs
+++ b/Engine/Source/Programs/Horde/HordeServer/Utilities/RemoteExtensions.cs
@@ -9,14 +9,17 @@ using System;
 using System.Collections.Generic;
 using System.Linq;
 using System.Threading.Tasks;
+using EpicGames.Core;
+
+using Digest = Build.Bazel.Remote.Execution.V2.Digest;
 
 namespace HordeServer.Utility
 {
 	static class RemoteExtensions
 	{
-		public static BlobHash ToHashValue(this Digest Digest)
+		public static IoHash ToHashValue(this Digest Digest)
 		{
-			return BlobHash.Parse(Digest.Hash);
+			return IoHash.Parse(Digest.Hash);
 		}
 
 		public static void MergeFrom<T>(this IMessage<T> Message, ReadOnlyMemory<byte> Memory) where T : IMessage<T>
@@ -31,7 +34,7 @@ namespace HordeServer.Utility
 			return Parser.ParseFrom(Memory.ToArray());
 		}
 
-		public static async Task<T> GetProtoMessageAsync<T>(this IStorageService StorageProvider, BlobHash HashValue, DateTime Deadline = default) where T : class, IMessage<T>, new()
+		public static async Task<T> GetProtoMessageAsync<T>(this IStorageService StorageProvider, IoHash HashValue, DateTime Deadline = default) where T : class, IMessage<T>, new()
 		{
 			T? Result = await TryGetProtoMessageAsync<T>(StorageProvider, HashValue, Deadline);
 			if (Result == null)
@@ -41,7 +44,7 @@ namespace HordeServer.Utility
 			return Result;
 		}
 
-		public static async Task<T?> TryGetProtoMessageAsync<T>(this IStorageService StorageProvider, BlobHash HashValue, DateTime Deadline = default) where T : class, IMessage<T>, new()
+		public static async Task<T?> TryGetProtoMessageAsync<T>(this IStorageService StorageProvider, IoHash HashValue, DateTime Deadline = default) where T : class, IMessage<T>, new()
 		{
 			ReadOnlyMemory<byte>? Data = await StorageProvider.TryGetBlobAsync(HashValue, Deadline);
 			if (Data == null)
@@ -54,32 +57,32 @@ namespace HordeServer.Utility
 			return NewItem;
 		}
 
-		public static async Task<BlobHash> PutProtoMessageAsync<T>(this IStorageService StorageProvider, IMessage<T> Message) where T : class, IMessage<T>
+		public static async Task<IoHash> PutProtoMessageAsync<T>(this IStorageService StorageProvider, IMessage<T> Message) where T : class, IMessage<T>
 		{
-			HashSet<BlobHash> References = new HashSet<BlobHash>();
+			HashSet<IoHash> References = new HashSet<IoHash>();
 			FindReferences(Message, References);
 
 			byte[] MessageData = Message.ToByteArray(); // TODO: Could generate this directly into the buffer
 
-			BlobHash HashValue = BlobHash.Compute(MessageData);
+			IoHash HashValue = IoHash.Compute(MessageData);
 			await StorageProvider.PutBlobAsync(HashValue, MessageData);
 
 			return HashValue;
 		}
 
-		public static Task<ReadOnlyMemory<byte>> GetBulkDataTreeAsync(this IStorageService StorageProvider, BlobHash Hash)
+		public static Task<ReadOnlyMemory<byte>> GetBulkDataTreeAsync(this IStorageService StorageProvider, IoHash Hash)
 		{
 			return StorageProvider.GetBlobAsync(Hash);
 		}
 
-		public static async Task<BlobHash> PutBulkDataTreeAsync(this IStorageService StorageProvider, ReadOnlyMemory<byte> BulkData)
+		public static async Task<IoHash> PutBulkDataTreeAsync(this IStorageService StorageProvider, ReadOnlyMemory<byte> BulkData)
 		{
-			BlobHash HashValue = BlobHash.Compute(BulkData);
+			IoHash HashValue = IoHash.Compute(BulkData.Span);
 			await StorageProvider.PutBlobAsync(HashValue, BulkData);
 			return HashValue;
 		}
 
-		static void FindReferences(this IMessage Message, HashSet<BlobHash> References)
+		static void FindReferences(this IMessage Message, HashSet<IoHash> References)
 		{
 			foreach (FieldDescriptor FieldDescriptor in Message.Descriptor.Fields.InFieldNumberOrder())
 			{
diff --git a/Engine/Source/Programs/Horde/HordeServerTests/KeyValueTreeTests.cs b/Engine/Source/Programs/Horde/HordeServerTests/KeyValueTreeTests.cs
--- a/Engine/Source/Programs/Horde/HordeServerTests/KeyValueTreeTests.cs
+++ b/Engine/Source/Programs/Horde/HordeServerTests/KeyValueTreeTests.cs
@@ -1,5 +1,6 @@
 // Copyright Epic Games, Inc. All Rights Reserved.
 
+using EpicGames.Core;
 using HordeServer.Storage;
 using HordeServer.Storage.Impl;
 using HordeServer.Storage.Primitives;
@@ -20,7 +21,7 @@ namespace HordeServerTests
 		class TestTreeMax2 : KeyValueTree
 		{
 			public TestTreeMax2()
-				: base(BlobRef<LeafBlob>.Type, BlobHash.NumBytes, new[] { (0, (IBlobType)BlobRef<LeafBlob>.Type) }, 2, 2)
+				: base(BlobRef<LeafBlob>.Type, IoHash.NumBytes, new[] { (0, (IBlobType)BlobRef<LeafBlob>.Type) }, 2, 2)
 			{
 			}
 		}
@@ -28,7 +29,7 @@ namespace HordeServerTests
 		class TestTreeMax256 : KeyValueTree
 		{
 			public TestTreeMax256()
-				: base(BlobRef<LeafBlob>.Type, BlobHash.NumBytes, new[] { (0, (IBlobType)BlobRef<LeafBlob>.Type) }, 256, 256)
+				: base(BlobRef<LeafBlob>.Type, IoHash.NumBytes, new[] { (0, (IBlobType)BlobRef<LeafBlob>.Type) }, 256, 256)
 			{
 			}
 		}
@@ -36,35 +37,30 @@ namespace HordeServerTests
 		[TestMethod]
 		public void TestChildKeys()
 		{
-			Assert.AreEqual(256, KeyValueTreeExtensions.GetChildKeyLength(128));
+			Assert.AreEqual(160, KeyValueTreeExtensions.GetChildKeyLength(128));
 			Assert.AreEqual(128, KeyValueTreeExtensions.GetChildKeyLength(64));
 			Assert.AreEqual(64, KeyValueTreeExtensions.GetChildKeyLength(32));
 
-			Assert.AreEqual(256, KeyValueTreeExtensions.GetChildKeyLength(128));
+			Assert.AreEqual(160, KeyValueTreeExtensions.GetChildKeyLength(128));
 
 			Assert.AreEqual(128, KeyValueTreeExtensions.GetChildKeyLength(64));
 			Assert.AreEqual(64, KeyValueTreeExtensions.GetChildKeyLength(56));
 			Assert.AreEqual(64, KeyValueTreeExtensions.GetChildKeyLength(48));
 			Assert.AreEqual(48, KeyValueTreeExtensions.GetChildKeyLength(40));
 
-			Assert.AreEqual(128 + 128, KeyValueTreeExtensions.GetChildKeyLength(128 + 64));
-			Assert.AreEqual(128 + 64, KeyValueTreeExtensions.GetChildKeyLength(128 + 56));
-			Assert.AreEqual(128 + 64, KeyValueTreeExtensions.GetChildKeyLength(128 + 48));
-			Assert.AreEqual(128 + 48, KeyValueTreeExtensions.GetChildKeyLength(128 + 40));
-
 			Assert.AreEqual(64, KeyValueTreeExtensions.GetChildKeyLength(32));
 		}
 
-		static KeyValueItem MakeUpdate(BlobHash HashValue)
+		static KeyValueItem MakeUpdate(IoHash HashValue)
 		{
-			byte[] Buffer = new byte[BlobHash.NumBytes];
+			byte[] Buffer = new byte[IoHash.NumBytes];
 			new Random().NextBytes(Buffer);
 			return new KeyValueItem(HashValue, Buffer);
 		}
 
 		static KeyValueItem MakeItem(string HashString)
 		{
-			return MakeUpdate(BlobHash.Parse(HashString));
+			return MakeUpdate(IoHash.Parse(HashString));
 		}
 
 		[TestMethod]
@@ -76,30 +72,30 @@ namespace HordeServerTests
 			BlobRef<TestTreeMax2> Tree = await Storage.CreateKeyValueTreeAsync<TestTreeMax2>();
 			{
 				ReadOnlyKeyValueNode RootNode = await Storage.GetKeyValueNodeAsync(Tree.Hash);
-				Assert.AreEqual(BlobHash.NumBits, RootNode.NumKeyBits);
+				Assert.AreEqual(IoHash.NumBits, RootNode.NumKeyBits);
 				Assert.AreEqual(0, RootNode.NumItems);
 				Assert.AreEqual(0, RootNode.NumItemsIfMerged);
 			}
 
-			KeyValueItem Item1 = MakeItem("aaaaaaaaaaaaaaaaaaaaaaaaaaaaaaaaaaaaaaaaaaaaaaaaaaaaaaaaaaaaaaaa");
+			KeyValueItem Item1 = MakeItem("aaaaaaaaaaaaaaaaaaaaaaaaaaaaaaaaaaaaaaaa");
 			Tree = await Storage.UpdateKeyValueTreeAsync(Tree, new[] { Item1 });
 			{
 				ReadOnlyKeyValueNode RootNode = await Storage.GetKeyValueNodeAsync(Tree.Hash);
-				Assert.AreEqual(BlobHash.NumBits, RootNode.NumKeyBits);
+				Assert.AreEqual(IoHash.NumBits, RootNode.NumKeyBits);
 				Assert.AreEqual(1, RootNode.NumItems);
 				Assert.AreEqual(1, RootNode.NumItemsIfMerged);
 			}
 
-			KeyValueItem Item2 = MakeItem("aaaaaaaaaaaaaaaaaaaaaaaaaaaaaaaaaaaaaaaaaaaaaaaaaaaaaaaaaaaaaaac");
+			KeyValueItem Item2 = MakeItem("aaaaaaaaaaaaaaaaaaaaaaaaaaaaaaaaaaaaaaac");
 			Tree = await Storage.UpdateKeyValueTreeAsync(Tree, new[] { Item2 });
 			{
 				ReadOnlyKeyValueNode RootNode = await Storage.GetKeyValueNodeAsync(Tree.Hash);
-				Assert.AreEqual(BlobHash.NumBits, RootNode.NumKeyBits);
+				Assert.AreEqual(IoHash.NumBits, RootNode.NumKeyBits);
 				Assert.AreEqual(2, RootNode.NumItems);
 				Assert.AreEqual(2, RootNode.NumItemsIfMerged);
 			}
 
-			KeyValueItem Item3 = MakeItem("aaaaaaaaaaaaaaaaaaaaaaaaaaaaaaaaaaaaaaaaaaaaaaaaaaaaaaaaaaaaaaab");
+			KeyValueItem Item3 = MakeItem("aaaaaaaaaaaaaaaaaaaaaaaaaaaaaaaaaaaaaaab");
 			Tree = await Storage.UpdateKeyValueTreeAsync(Tree, new[] { Item3 });
 			{
 				ReadOnlyKeyValueNode Node = await Storage.GetKeyValueNodeAsync(Tree.Hash);
@@ -130,19 +126,19 @@ namespace HordeServerTests
 
 			BlobRef<TestTreeMax256> Tree = await Storage.CreateKeyValueTreeAsync<TestTreeMax256>();
 
-			Dictionary<BlobHash, BlobHash> Items = new Dictionary<BlobHash, BlobHash>();
+			Dictionary<IoHash, IoHash> Items = new Dictionary<IoHash, IoHash>();
 			for (int OuterIdx = 0; OuterIdx < 128; OuterIdx++)
 			{
 				List<KeyValueItem> NewItems = new List<KeyValueItem>();
 				for (int Idx = 0; Idx < 128; Idx++)
 				{
-					byte[] KeyBytes = new byte[BlobHash.NumBytes];
+					byte[] KeyBytes = new byte[IoHash.NumBytes];
 					Random.NextBytes(KeyBytes);
-					BlobHash Key = new BlobHash(KeyBytes);
+					IoHash Key = new IoHash(KeyBytes);
 
-					byte[] ValueBytes = new byte[BlobHash.NumBytes];
+					byte[] ValueBytes = new byte[IoHash.NumBytes];
 					Random.NextBytes(ValueBytes);
-					BlobHash Value = new BlobHash(ValueBytes);
+					IoHash Value = new IoHash(ValueBytes);
 
 					NewItems.Add(new KeyValueItem(Key, ValueBytes));
 					Items[Key] = Value;
@@ -153,7 +149,7 @@ namespace HordeServerTests
 			int NumItems = await CountItemsAsync(Storage, Tree.Hash);
 			Assert.AreEqual(Items.Count, NumItems);
 
-			foreach (KeyValuePair<BlobHash, BlobHash> Pair in Items)
+			foreach (KeyValuePair<IoHash, IoHash> Pair in Items)
 			{
 				ReadOnlyMemory<byte>? Value = await Storage.SearchKeyValueTreeAsync(Tree, Pair.Key);
 				Assert.IsNotNull(Value);
@@ -161,10 +157,10 @@ namespace HordeServerTests
 			}
 		}
 
-		static async Task<int> CountItemsAsync(IStorageService Storage, BlobHash Hash)
+		static async Task<int> CountItemsAsync(IStorageService Storage, IoHash Hash)
 		{
 			ReadOnlyKeyValueNode Node = await Storage.GetKeyValueNodeAsync(Hash);
-			if (Node.NumKeyBits == BlobHash.NumBits)
+			if (Node.NumKeyBits == IoHash.NumBits)
 			{
 				return Node.NumItems;
 			}
@@ -172,8 +168,8 @@ namespace HordeServerTests
 			{
 				int NumItems = 0;
 
-				ReadOnlyBlobHashArray Children = new ReadOnlyBlobHashArray(Node.Values);
-				foreach (BlobHash Child in Children)
+				ReadOnlyHashArray Children = new ReadOnlyHashArray(Node.Values);
+				foreach (IoHash Child in Children)
 				{
 					NumItems += await CountItemsAsync(Storage, Child);
 				}
