Horde: Set proper worker name during remote execution

The agent recently got the ability of overriding the agent name instead of sourcing from the current hostname.
This will ensure that override is respected.

diff --git a/Engine/Source/Programs/Horde/HordeAgent/Execution/ActionExecutor.cs b/Engine/Source/Programs/Horde/HordeAgent/Execution/ActionExecutor.cs
--- a/Engine/Source/Programs/Horde/HordeAgent/Execution/ActionExecutor.cs
+++ b/Engine/Source/Programs/Horde/HordeAgent/Execution/ActionExecutor.cs
@@ -40,6 +40,11 @@ namespace HordeAgent
 		/// </summary>
 		string InstanceName;
 		
+		/// <summary>
+		/// The agent name name
+		/// </summary>
+		string AgentName;
+		
 		/// <summary>
 		/// The storage wrapper
 		/// </summary>
@@ -63,13 +68,15 @@ namespace HordeAgent
 		/// <summary>
 		/// Constructor
 		/// </summary>
+		/// <param name="AgentName"></param>
 		/// <param name="InstanceName"></param>
 		/// <param name="CasChannel"></param>
 		/// <param name="ActionCacheChannel"></param>
 		/// <param name="ActionRpcChannel"></param>
 		/// <param name="Logger"></param>
-		public ActionExecutor(string InstanceName, GrpcChannel CasChannel, GrpcChannel ActionCacheChannel, GrpcChannel ActionRpcChannel, ILogger Logger)
+		public ActionExecutor(string AgentName, string InstanceName, GrpcChannel CasChannel, GrpcChannel ActionCacheChannel, GrpcChannel ActionRpcChannel, ILogger Logger)
 		{
+			this.AgentName = AgentName;
 			this.InstanceName = InstanceName;
 			this.Storage = new ContentAddressableStorage.ContentAddressableStorageClient(CasChannel);
 			this.Cache = new ActionCache.ActionCacheClient(ActionCacheChannel);
@@ -147,7 +154,7 @@ namespace HordeAgent
 					Result.ExitCode = Process.ExitCode;
 					Result.ExecutionMetadata = new ExecutedActionMetadata
 					{
-						Worker = System.Net.Dns.GetHostName()
+						Worker = AgentName
 					};
 
 					foreach (string Line in Encoding.UTF8.GetString(StdOutData).Split('\n'))
diff --git a/Engine/Source/Programs/Horde/HordeAgent/Services/WorkerService.cs b/Engine/Source/Programs/Horde/HordeAgent/Services/WorkerService.cs
--- a/Engine/Source/Programs/Horde/HordeAgent/Services/WorkerService.cs
+++ b/Engine/Source/Programs/Horde/HordeAgent/Services/WorkerService.cs
@@ -795,7 +795,7 @@ namespace HordeAgent.Services
 
 				try
 				{
-					ActionExecutor Executor = new ActionExecutor(ActionTask.InstanceName, CasChannel, ActionCacheChannel, ActionRpcChannel, Logger);
+					ActionExecutor Executor = new ActionExecutor(Settings.GetAgentName(), ActionTask.InstanceName, CasChannel, ActionCacheChannel, ActionRpcChannel, Logger);
 					await Executor.ExecuteActionAsync(LeaseId, ActionTask, LeaseDir);
 				}
 				finally
