Horde: Option to disable network share mounting on agent

Certain agents cannot always mount shares (like the remote exec ones with no AD access). This boolean will provide an easy override to skip that.

diff --git a/Engine/Source/Programs/Horde/HordeAgent/AgentSettings.cs b/Engine/Source/Programs/Horde/HordeAgent/AgentSettings.cs
--- a/Engine/Source/Programs/Horde/HordeAgent/AgentSettings.cs
+++ b/Engine/Source/Programs/Horde/HordeAgent/AgentSettings.cs
@@ -181,6 +181,11 @@ namespace HordeAgent
 		/// </summary>
 		public string? WorkingDir { get; set; } = Program.DataDir.FullName;
 
+		/// <summary>
+		/// Whether to mount the specified list of network shares
+		/// </summary>
+		public bool ShareMountingEnabled { get; set; } = true;
+		
 		/// <summary>
 		/// List of network shares to mount
 		/// </summary>
diff --git a/Engine/Source/Programs/Horde/HordeAgent/Services/WorkerService.cs b/Engine/Source/Programs/Horde/HordeAgent/Services/WorkerService.cs
--- a/Engine/Source/Programs/Horde/HordeAgent/Services/WorkerService.cs
+++ b/Engine/Source/Programs/Horde/HordeAgent/Services/WorkerService.cs
@@ -349,7 +349,7 @@ namespace HordeAgent.Services
 			}
 
 			// Mount all the necessary network shares. Currently only supported on Windows.
-			if (RuntimeInformation.IsOSPlatform(OSPlatform.Windows))
+			if (Settings.ShareMountingEnabled && RuntimeInformation.IsOSPlatform(OSPlatform.Windows))
 			{
 				foreach (MountNetworkShare Share in Settings.Shares)
 				{
