Horde: Include avatar URLs in the user response. Also deprecate the /api/v1/user endpoint in favor of /api/v1/users/current.

diff --git a/Engine/Source/Programs/Horde/HordeServer/Api/Users.cs b/Engine/Source/Programs/Horde/HordeServer/Api/Users.cs
--- a/Engine/Source/Programs/Horde/HordeServer/Api/Users.cs
+++ b/Engine/Source/Programs/Horde/HordeServer/Api/Users.cs
@@ -2,6 +2,7 @@
 
 using EpicGames.Core;
 using HordeServer.Models;
+using HordeServer.Services;
 using MongoDB.Bson;
 using System;
 using System.Collections.Generic;
@@ -21,6 +22,26 @@ namespace HordeServer.Api
 		/// </summary>
 		public string Id { get; set; }
 
+		/// <summary>
+		/// Avatar image URL (24px)
+		/// </summary>
+		public string? Image24 { get; set; }
+
+		/// <summary>
+		/// Avatar image URL (32px)
+		/// </summary>
+		public string? Image32 { get; set; }
+
+		/// <summary>
+		/// Avatar image URL (48px)
+		/// </summary>
+		public string? Image48 { get; set; }
+
+		/// <summary>
+		/// Avatar image URL (72px)
+		/// </summary>
+		public string? Image72 { get; set; }
+
 		/// <summary>
 		/// Claims for the user
 		/// </summary>
@@ -49,12 +70,13 @@ namespace HordeServer.Api
 		/// <summary>
 		/// Constructor
 		/// </summary>
-		/// <param name="User"></param>
-		/// <param name="Claims"></param>
-		/// <param name="Settings"></param>
-		public GetUserResponse(IUser User, IUserClaims Claims, IUserSettings Settings)
+		public GetUserResponse(IUser User, IAvatar? Avatar, IUserClaims Claims, IUserSettings Settings)
 		{
 			this.Id = User.Id.ToString();
+			this.Image24 = Avatar?.Image24;
+			this.Image32 = Avatar?.Image32;
+			this.Image48 = Avatar?.Image48;
+			this.Image72 = Avatar?.Image72;
 			this.Claims = Claims.Claims.Select(x => new UserClaim(x)).ToList();
 			this.EnableExperimentalFeatures = Settings.EnableExperimentalFeatures;
 			this.EnableIssueNotifications = Settings.EnableIssueNotifications;
diff --git a/Engine/Source/Programs/Horde/HordeServer/Controllers/UserController.cs b/Engine/Source/Programs/Horde/HordeServer/Controllers/UserController.cs
--- a/Engine/Source/Programs/Horde/HordeServer/Controllers/UserController.cs
+++ b/Engine/Source/Programs/Horde/HordeServer/Controllers/UserController.cs
@@ -3,6 +3,7 @@
 using HordeServer.Api;
 using HordeServer.Collections;
 using HordeServer.Models;
+using HordeServer.Services;
 using HordeServer.Utilities;
 using Microsoft.AspNetCore.Authorization;
 using Microsoft.AspNetCore.Mvc;
@@ -27,13 +28,20 @@ namespace HordeServer.Controllers
 		/// </summary>
 		IUserCollection UserCollection { get; set; }
 
+		/// <summary>
+		/// The avatar service
+		/// </summary>
+		IAvatarService? AvatarService { get; set; }
+
 		/// <summary>
 		/// Constructor
 		/// </summary>
 		/// <param name="UserCollection"></param>
-		public UserController(IUserCollection UserCollection)
+		/// <param name="AvatarService"></param>
+		public UserController(IUserCollection UserCollection, IAvatarService? AvatarService)
 		{
 			this.UserCollection = UserCollection;
+			this.AvatarService = AvatarService;
 		}
 
 		/// <summary>
@@ -51,9 +59,10 @@ namespace HordeServer.Controllers
 				return NotFound();
 			}
 
-			IUserClaims? Claims = await UserCollection.GetClaimsAsync(InternalUser.Id);
-			IUserSettings? Settings = await UserCollection.GetSettingsAsync(InternalUser.Id);
-			return PropertyFilter.Apply(new GetUserResponse(InternalUser, Claims, Settings), Filter);
+			IAvatar? Avatar = (AvatarService == null)? (IAvatar?)null : await AvatarService.GetAvatarAsync(InternalUser);
+			IUserClaims Claims = await UserCollection.GetClaimsAsync(InternalUser.Id);
+			IUserSettings Settings = await UserCollection.GetSettingsAsync(InternalUser.Id);
+			return PropertyFilter.Apply(new GetUserResponse(InternalUser, Avatar, Claims, Settings), Filter);
 		}
 
 		/// <summary>
diff --git a/Engine/Source/Programs/Horde/HordeServer/Controllers/UsersController.cs b/Engine/Source/Programs/Horde/HordeServer/Controllers/UsersController.cs
new file mode 100644
--- /dev/null
+++ b/Engine/Source/Programs/Horde/HordeServer/Controllers/UsersController.cs
@@ -0,0 +1,93 @@
+// Copyright Epic Games, Inc. All Rights Reserved.
+
+using HordeServer.Api;
+using HordeServer.Collections;
+using HordeServer.Models;
+using HordeServer.Notifications.Impl;
+using HordeServer.Services;
+using HordeServer.Utilities;
+using Microsoft.AspNetCore.Authorization;
+using Microsoft.AspNetCore.Mvc;
+using Microsoft.Extensions.Options;
+using MongoDB.Bson;
+using System;
+using System.Collections.Generic;
+using System.Linq;
+using System.Threading.Tasks;
+
+namespace HordeServer.Controllers
+{
+	/// <summary>
+	/// Controller for the /api/v1/users endpoint
+	/// </summary>
+	[ApiController]
+	[Authorize]
+	[Route("[controller]")]
+	public class UsersController : ControllerBase
+	{
+		/// <summary>
+		/// The user collection instance
+		/// </summary>
+		IUserCollection UserCollection { get; set; }
+
+		/// <summary>
+		/// The avatar service
+		/// </summary>
+		IAvatarService? AvatarService { get; set; }
+
+		/// <summary>
+		/// Constructor
+		/// </summary>
+		/// <param name="UserCollection"></param>
+		/// <param name="AvatarService"></param>
+		public UsersController(IUserCollection UserCollection, IAvatarService? AvatarService)
+		{
+			this.UserCollection = UserCollection;
+			this.AvatarService = AvatarService;
+		}
+
+		/// <summary>
+		/// Gets information about the logged in user
+		/// </summary>
+		/// <returns>Http result code</returns>
+		[HttpGet]
+		[Route("/api/v1/users/{Id}")]
+		[ProducesResponseType(typeof(List<GetUserResponse>), 200)]
+		public async Task<ActionResult<object>> GetUserAsync(string Id, [FromQuery] PropertyFilter? Filter = null)
+		{
+			IUser? User = await GetUserInternalAsync(Id);
+			if (User == null)
+			{
+				return NotFound();
+			}
+
+			IAvatar? Avatar = (AvatarService == null) ? (IAvatar?)null : await AvatarService.GetAvatarAsync(User);
+			IUserClaims? Claims = await UserCollection.GetClaimsAsync(User.Id);
+			IUserSettings? Settings = await UserCollection.GetSettingsAsync(User.Id);
+			return PropertyFilter.Apply(new GetUserResponse(User, Avatar, Claims, Settings), Filter);
+		}
+
+		async Task<IUser?> GetUserInternalAsync(string Id)
+		{
+			ObjectId? UserId = ParseUserId(Id);
+			if(UserId == null)
+			{
+				return null;
+			}
+			return await UserCollection.GetUserAsync(UserId.Value);
+		}
+
+		ObjectId? ParseUserId(string Id)
+		{
+			if (Id.Equals("current", StringComparison.OrdinalIgnoreCase))
+			{
+				return User.GetUserId();
+			}
+			else if(ObjectId.TryParse(Id, out ObjectId Result))
+			{
+				return Result;
+			}
+			return null;
+		}
+	}
+}
diff --git a/Engine/Source/Programs/Horde/HordeServer/Models/IAvatar.cs b/Engine/Source/Programs/Horde/HordeServer/Models/IAvatar.cs
new file mode 100644
--- /dev/null
+++ b/Engine/Source/Programs/Horde/HordeServer/Models/IAvatar.cs
@@ -0,0 +1,35 @@
+// Copyright Epic Games, Inc. All Rights Reserved.
+
+using System;
+using System.Collections.Generic;
+using System.Linq;
+using System.Threading.Tasks;
+
+namespace HordeServer.Models
+{
+	/// <summary>
+	/// Avatar for the user
+	/// </summary>
+	public interface IAvatar
+	{
+		/// <summary>
+		/// Url for a 24x24 pixel image
+		/// </summary>
+		public string? Image24 { get; set; }
+
+		/// <summary>
+		/// Url for a 32x32 pixel image
+		/// </summary>
+		public string? Image32 { get; set; }
+
+		/// <summary>
+		/// Url for a 48x48 pixel image
+		/// </summary>
+		public string? Image48 { get; set; }
+
+		/// <summary>
+		/// Url for a 72x72 pixel image
+		/// </summary>
+		public string? Image72 { get; set; }
+	}
+}
diff --git a/Engine/Source/Programs/Horde/HordeServer/Notifications/Impl/SlackNotificationSink.cs b/Engine/Source/Programs/Horde/HordeServer/Notifications/Impl/SlackNotificationSink.cs
--- a/Engine/Source/Programs/Horde/HordeServer/Notifications/Impl/SlackNotificationSink.cs
+++ b/Engine/Source/Programs/Horde/HordeServer/Notifications/Impl/SlackNotificationSink.cs
@@ -41,7 +41,7 @@ namespace HordeServer.Notifications.Impl
 	/// <summary>
 	/// Maintains a connection to Slack, in order to receive socket-mode notifications of user interactions
 	/// </summary>
-	public class SlackNotificationSink : BackgroundService, INotificationSink
+	public class SlackNotificationSink : BackgroundService, INotificationSink, IAvatarService
 	{
 		const string PostMessageUrl = "https://slack.com/api/chat.postMessage";
 		const string UpdateMessageUrl = "https://slack.com/api/chat.update";
@@ -82,10 +82,40 @@ namespace HordeServer.Notifications.Impl
 			public List<ActionInfo> Actions { get; set; } = new List<ActionInfo>();
 		}
 
+		class UserResponse
+		{
+			[JsonPropertyName("ok")]
+			public bool Ok { get; set; }
+
+			[JsonPropertyName("user")]
+			public UserInfo? User { get; set; }
+		}
+
 		class UserInfo
 		{
 			[JsonPropertyName("username")]
 			public string? UserName { get; set; }
+
+			[JsonPropertyName("profile")]
+			public UserProfile? Profile { get; set; }
+		}
+
+		class UserProfile
+		{
+			[JsonPropertyName("is_custom_image")]
+			public bool IsCustomImage { get; set; }
+
+			[JsonPropertyName("image_24")]
+			public string? Image24 { get; set; }
+
+			[JsonPropertyName("image_32")]
+			public string? Image32 { get; set; }
+
+			[JsonPropertyName("image_48")]
+			public string? Image48 { get; set; }
+
+			[JsonPropertyName("image_72")]
+			public string? Image72 { get; set; }
 		}
 
 		class ActionInfo
@@ -118,12 +148,52 @@ namespace HordeServer.Notifications.Impl
 			public string Digest { get; set; } = String.Empty;
 		}
 
-		class UserDocument
+		class SlackUser : IAvatar
 		{
+			public const int CurrentVersion = 2;
+
 			[BsonId]
 			public string Email { get; set; } = String.Empty;
+
+			[BsonElement("u")]
 			public string? SlackUserId { get; set; }
+
+			[BsonElement("i24")]
+			public string? Image24 { get; set; }
+
+			[BsonElement("i32")]
+			public string? Image32 { get; set; }
+			
+			[BsonElement("i48")]
+			public string? Image48 { get; set; }
+
+			[BsonElement("i72")]
+			public string? Image72 { get; set; }
+
+			[BsonElement("t")]
 			public DateTime Time { get; set; }
+
+			[BsonElement("v")]
+			public int Version { get; set; }
+
+			private SlackUser()
+			{
+			}
+
+			public SlackUser(string Email, UserInfo? Info)
+			{
+				this.Email = Email;
+				this.SlackUserId = Info?.UserName;
+				if (Info != null && Info.Profile != null && Info.Profile.IsCustomImage)
+				{
+					this.Image24 = Info.Profile.Image24;
+					this.Image32 = Info.Profile.Image32;
+					this.Image48 = Info.Profile.Image48;
+					this.Image72 = Info.Profile.Image72;
+				}
+				this.Time = DateTime.UtcNow;
+				this.Version = CurrentVersion;
+			}
 		}
 
 		IIssueService IssueService;
@@ -132,7 +202,7 @@ namespace HordeServer.Notifications.Impl
 		StreamService StreamService;
 		ServerSettings Settings;
 		IMongoCollection<MessageStateDocument> MessageStates;
-		IMongoCollection<UserDocument> Users;
+		IMongoCollection<SlackUser> SlackUsers;
 		ILogger Logger;
 
 		/// <summary>
@@ -158,7 +228,7 @@ namespace HordeServer.Notifications.Impl
 			this.StreamService = StreamService;
 			this.Settings = Settings.Value;
 			this.MessageStates = DatabaseService.Database.GetCollection<MessageStateDocument>("Slack");
-			this.Users = DatabaseService.Database.GetCollection<UserDocument>("Slack.Users");
+			this.SlackUsers = DatabaseService.Database.GetCollection<SlackUser>("Slack.Users");
 			this.Logger = Logger;
 		}
 
@@ -172,6 +242,16 @@ namespace HordeServer.Notifications.Impl
 			GC.SuppressFinalize(this);
 		}
 
+		#region Avatars
+
+		/// <inheritdoc/>
+		public async Task<IAvatar?> GetAvatarAsync(IUser User)
+		{
+			return await GetSlackUser(User);
+		}
+
+		#endregion
+
 		#region Message state 
 
 		async Task<MessageStateDocument?> GetMessageStateAsync(string Recipient, string EventId)
@@ -897,9 +977,9 @@ namespace HordeServer.Notifications.Impl
 			}
 		}
 
-		static bool ShouldUpdateUser(UserDocument? Document)
+		static bool ShouldUpdateUser(SlackUser? Document)
 		{
-			if(Document == null)
+			if(Document == null || Document.Version < SlackUser.CurrentVersion)
 			{
 				return true;
 			}
@@ -917,6 +997,11 @@ namespace HordeServer.Notifications.Impl
 		}
 
 		private async Task<string?> GetSlackUserId(IUser User)
+		{
+			return (await GetSlackUser(User))?.SlackUserId;
+		}
+
+		private async Task<SlackUser?> GetSlackUser(IUser User)
 		{
 			string? Email = User.Email;
 			if (Email == null)
@@ -925,17 +1010,17 @@ namespace HordeServer.Notifications.Impl
 				return null;
 			}
 
-			UserDocument? UserDocument;
+			SlackUser? UserDocument;
 			if (!UserCache.TryGetValue(Email, out UserDocument))
 			{
-				UserDocument = await Users.Find(x => x.Email == Email).FirstOrDefaultAsync();
+				UserDocument = await SlackUsers.Find(x => x.Email == Email).FirstOrDefaultAsync();
 				if (UserDocument == null || ShouldUpdateUser(UserDocument))
 				{
-					string? UserId = await GetSlackUserIdByEmail(Email);
-					if (UserDocument == null || UserId != null)
+					UserInfo? UserInfo = await GetSlackUserInfoByEmail(Email);
+					if (UserDocument == null || UserInfo != null)
 					{
-						UserDocument = new UserDocument { Email = Email, SlackUserId = UserId, Time = DateTime.UtcNow };
-						await Users.ReplaceOneAsync(x => x.Email == Email, UserDocument, new ReplaceOptions { IsUpsert = true });
+						UserDocument = new SlackUser(Email, UserInfo);
+						await SlackUsers.ReplaceOneAsync(x => x.Email == Email, UserDocument, new ReplaceOptions { IsUpsert = true });
 					}
 				}
 				using (ICacheEntry Entry = UserCache.CreateEntry(Email))
@@ -945,34 +1030,27 @@ namespace HordeServer.Notifications.Impl
 				}
 			}
 
-			return UserDocument?.SlackUserId;
+			return UserDocument;
 		}
 
-		private async Task<string?> GetSlackUserIdByEmail(string Email)
+		private async Task<UserInfo?> GetSlackUserInfoByEmail(string Email)
 		{
 			using HttpClient Client = new HttpClient();
 
 			using HttpRequestMessage GetUserIdRequest = new HttpRequestMessage(HttpMethod.Post, $"https://slack.com/api/users.lookupByEmail?email={Email}");
 			GetUserIdRequest.Headers.Add("Authorization", $"Bearer {Settings.SlackToken ?? ""}");
 
-			HttpResponseMessage Response = await Client.SendAsync(GetUserIdRequest);
-			string? SlackResponse = await Response.Content.ReadAsStringAsync();
-
-			JObject ResponseObject = JObject.Parse(SlackResponse);
-			if (!ResponseObject["ok"].ToObject<bool>())
-			{
-				Logger.LogWarning("Unable to find Slack user id for {Email}: {Response}", Email, SlackResponse);
-				return null;
-			}
+			HttpResponseMessage ResponseMessage = await Client.SendAsync(GetUserIdRequest);
+			byte[] ResponseData = await ResponseMessage.Content.ReadAsByteArrayAsync();
 
-			string? UserId = ResponseObject["user"]["id"].ToString();
-			if (String.IsNullOrEmpty(UserId))
+			UserResponse UserResponse = JsonSerializer.Deserialize<UserResponse>(ResponseData);
+			if(!UserResponse.Ok || UserResponse.User == null)
 			{
-				Logger.LogWarning("Unexpected response from Slack user query on {Email}: {Response}", Email, SlackResponse);
+				Logger.LogWarning("Unable to find Slack user id for {Email}: {Response}", Email, Encoding.UTF8.GetString(ResponseData));
 				return null;
 			}
 
-			return UserId;
+			return UserResponse.User;
 		}
 
 		class SlackResponse
diff --git a/Engine/Source/Programs/Horde/HordeServer/Startup.cs b/Engine/Source/Programs/Horde/HordeServer/Startup.cs
--- a/Engine/Source/Programs/Horde/HordeServer/Startup.cs
+++ b/Engine/Source/Programs/Horde/HordeServer/Startup.cs
@@ -389,6 +389,7 @@ namespace HordeServer
 			Services.AddSingleton<ProjectService>();
 			Services.AddSingleton<ScheduleService>();
 			Services.AddSingleton<SlackNotificationSink>();
+			Services.AddSingleton<IAvatarService, SlackNotificationSink>(SP => SP.GetRequiredService<SlackNotificationSink>());
 			Services.AddSingleton<INotificationSink, SlackNotificationSink>(SP => SP.GetRequiredService<SlackNotificationSink>());
 			Services.AddSingleton<SoftwareService>();
 			Services.AddSingleton<StreamService>();
