Horde: Missing file.

diff --git a/Engine/Source/Programs/Horde/HordeServer/Services/IAvatarService.cs b/Engine/Source/Programs/Horde/HordeServer/Services/IAvatarService.cs
new file mode 100644
--- /dev/null
+++ b/Engine/Source/Programs/Horde/HordeServer/Services/IAvatarService.cs
@@ -0,0 +1,28 @@
+// Copyright Epic Games, Inc. All Rights Reserved.
+
+using HordeServer.Models;
+using Microsoft.Extensions.Logging;
+using Microsoft.Extensions.Options;
+using Newtonsoft.Json.Linq;
+using System;
+using System.Collections.Generic;
+using System.Linq;
+using System.Net.Http;
+using System.Threading.Tasks;
+
+namespace HordeServer.Services
+{
+	/// <summary>
+	/// Provides avatar images
+	/// </summary>
+	public interface IAvatarService
+	{
+		/// <summary>
+		/// Gets a users's avatar
+		/// </summary>
+		/// <param name="User"></param>
+		/// <returns></returns>
+		Task<IAvatar?> GetAvatarAsync(IUser User);
+	}
+}
+
