Horde: Use user id as primary key when looking up Slack users.

diff --git a/Engine/Source/Programs/Horde/HordeServer/Notifications/Impl/SlackNotificationSink.cs b/Engine/Source/Programs/Horde/HordeServer/Notifications/Impl/SlackNotificationSink.cs
--- a/Engine/Source/Programs/Horde/HordeServer/Notifications/Impl/SlackNotificationSink.cs
+++ b/Engine/Source/Programs/Horde/HordeServer/Notifications/Impl/SlackNotificationSink.cs
@@ -151,9 +151,8 @@ namespace HordeServer.Notifications.Impl
 		class SlackUser : IAvatar
 		{
 			public const int CurrentVersion = 2;
-
-			[BsonId]
-			public string Email { get; set; } = String.Empty;
+			
+			public ObjectId Id { get; set; }
 
 			[BsonElement("u")]
 			public string? SlackUserId { get; set; }
@@ -180,9 +179,9 @@ namespace HordeServer.Notifications.Impl
 			{
 			}
 
-			public SlackUser(string Email, UserInfo? Info)
+			public SlackUser(ObjectId Id, UserInfo? Info)
 			{
-				this.Email = Email;
+				this.Id = Id;
 				this.SlackUserId = Info?.UserName;
 				if (Info != null && Info.Profile != null && Info.Profile.IsCustomImage)
 				{
@@ -228,7 +227,7 @@ namespace HordeServer.Notifications.Impl
 			this.StreamService = StreamService;
 			this.Settings = Settings.Value;
 			this.MessageStates = DatabaseService.Database.GetCollection<MessageStateDocument>("Slack");
-			this.SlackUsers = DatabaseService.Database.GetCollection<SlackUser>("Slack.Users");
+			this.SlackUsers = DatabaseService.Database.GetCollection<SlackUser>("Slack.UsersV2");
 			this.Logger = Logger;
 		}
 
@@ -1013,14 +1012,14 @@ namespace HordeServer.Notifications.Impl
 			SlackUser? UserDocument;
 			if (!UserCache.TryGetValue(Email, out UserDocument))
 			{
-				UserDocument = await SlackUsers.Find(x => x.Email == Email).FirstOrDefaultAsync();
+				UserDocument = await SlackUsers.Find(x => x.Id == User.Id).FirstOrDefaultAsync();
 				if (UserDocument == null || ShouldUpdateUser(UserDocument))
 				{
 					UserInfo? UserInfo = await GetSlackUserInfoByEmail(Email);
 					if (UserDocument == null || UserInfo != null)
 					{
-						UserDocument = new SlackUser(Email, UserInfo);
-						await SlackUsers.ReplaceOneAsync(x => x.Email == Email, UserDocument, new ReplaceOptions { IsUpsert = true });
+						UserDocument = new SlackUser(User.Id, UserInfo);
+						await SlackUsers.ReplaceOneAsync(x => x.Id == User.Id, UserDocument, new ReplaceOptions { IsUpsert = true });
 					}
 				}
 				using (ICacheEntry Entry = UserCache.CreateEntry(Email))
