Horde: Fill in metadata on remote execution timings

REAPI supports these fields and will help troubleshoot for how long actions gets queued on the server.

diff --git a/Engine/Source/Programs/Horde/HordeAgent/Execution/ActionExecutor.cs b/Engine/Source/Programs/Horde/HordeAgent/Execution/ActionExecutor.cs
--- a/Engine/Source/Programs/Horde/HordeAgent/Execution/ActionExecutor.cs
+++ b/Engine/Source/Programs/Horde/HordeAgent/Execution/ActionExecutor.cs
@@ -90,8 +90,9 @@ namespace HordeAgent
 		/// <param name="LeaseId">The lease id</param>
 		/// <param name="ActionTask">Task to execute</param>
 		/// <param name="SandboxDir">Directory to use as a sandbox for execution</param>
+		/// <param name="ActionTaskStartTime">When the agent received the action task. Used for reporting execution stats in REAPI.</param>
 		/// <returns>The action result</returns>
-		public async Task<ActionResult> ExecuteActionAsync(string LeaseId, ActionTask ActionTask, DirectoryReference SandboxDir)
+		public async Task<ActionResult> ExecuteActionAsync(string LeaseId, ActionTask ActionTask, DirectoryReference SandboxDir, DateTimeOffset ActionTaskStartTime)
 		{
 			Action? Action = await Storage.GetProtoMessageAsync<Action>(InstanceName, ActionTask.Digest);
 			if (Action == null)
@@ -102,8 +103,11 @@ namespace HordeAgent
 			DirectoryReference.CreateDirectory(SandboxDir);
 			FileUtils.ForceDeleteDirectoryContents(SandboxDir);
 
+
+			DateTimeOffset InputFetchStart = DateTimeOffset.UtcNow;
 			Directory InputDirectory = await Storage.GetProtoMessageAsync<Directory>(InstanceName, Action.InputRootDigest);
 			await SetupSandboxAsync(InputDirectory, SandboxDir);
+			DateTimeOffset InputFetchCompleted = DateTimeOffset.UtcNow;
 
 			RpcCommand Command = await Storage.GetProtoMessageAsync<RpcCommand>(InstanceName, Action.CommandDigest);
 
@@ -136,6 +140,7 @@ namespace HordeAgent
 				}
 
 				Logger.LogInformation("Executing {FileName} with arguments {Arguments}", FileName, Arguments);
+				DateTimeOffset ExecutionStartTime = DateTimeOffset.UtcNow;
 				using (ManagedProcess Process = new ManagedProcess(ProcessGroup, FileName, Arguments, WorkingDirectory, NewEnvironment, ProcessPriorityClass.Normal, ManagedProcessFlags.None))
 				{
 					byte[] StdOutData;
@@ -147,6 +152,8 @@ namespace HordeAgent
 						StdOutData = StdOutStream.ToArray();
 						StdErrData = StdErrStream.ToArray();
 					}
+					DateTimeOffset ExecutionCompletedTime = DateTimeOffset.UtcNow;
+						
 
 					Result = new ActionResult();
 					Result.StdoutDigest = await Storage.PutBulkDataAsync(InstanceName, StdOutData);
@@ -154,7 +161,13 @@ namespace HordeAgent
 					Result.ExitCode = Process.ExitCode;
 					Result.ExecutionMetadata = new ExecutedActionMetadata
 					{
-						Worker = AgentName
+						Worker = AgentName,
+						QueuedTimestamp = ActionTask.QueueTime,
+						ExecutionStartTimestamp = Timestamp.FromDateTimeOffset(ExecutionStartTime),
+						ExecutionCompletedTimestamp = Timestamp.FromDateTimeOffset(ExecutionCompletedTime),
+						WorkerStartTimestamp = Timestamp.FromDateTimeOffset(ActionTaskStartTime),
+						InputFetchStartTimestamp = Timestamp.FromDateTimeOffset(InputFetchStart),
+						InputFetchCompletedTimestamp = Timestamp.FromDateTimeOffset(InputFetchCompleted)
 					};
 
 					foreach (string Line in Encoding.UTF8.GetString(StdOutData).Split('\n'))
@@ -168,6 +181,7 @@ namespace HordeAgent
 
 					Logger.LogInformation("exit: {ExitCode}", Process.ExitCode);
 
+					Result.ExecutionMetadata.OutputUploadStartTimestamp = Timestamp.FromDateTimeOffset(DateTimeOffset.UtcNow);
 					foreach (var (FileRef, OutputPath) in ResolveOutputPaths(SandboxDir, Command.OutputPaths))
 					{
 						byte[] Bytes = await FileReference.ReadAllBytesAsync(FileRef);
@@ -180,6 +194,8 @@ namespace HordeAgent
 				
 						Logger.LogInformation("Uploaded {File} (digest: {Digest}, size: {Size})", FileRef, Digest.Hash, Digest.SizeBytes);
 					}
+					Result.ExecutionMetadata.OutputUploadCompletedTimestamp = Timestamp.FromDateTimeOffset(DateTimeOffset.UtcNow);
+					Result.ExecutionMetadata.WorkerCompletedTimestamp = Timestamp.FromDateTimeOffset(DateTimeOffset.UtcNow);
 
 					PostActionResultRequest PostResultRequest = new PostActionResultRequest();
 					PostResultRequest.LeaseId = LeaseId;
diff --git a/Engine/Source/Programs/Horde/HordeAgent/Services/WorkerService.cs b/Engine/Source/Programs/Horde/HordeAgent/Services/WorkerService.cs
--- a/Engine/Source/Programs/Horde/HordeAgent/Services/WorkerService.cs
+++ b/Engine/Source/Programs/Horde/HordeAgent/Services/WorkerService.cs
@@ -766,6 +766,7 @@ namespace HordeAgent.Services
 		/// <returns></returns>
 		async Task<LeaseOutcome> ActionAsync(IRpcConnection RpcConnection, string LeaseId, ActionTask ActionTask, ILogger Logger, CancellationToken CancellationToken)
 		{
+			DateTimeOffset ActionTaskStartTime = DateTimeOffset.UtcNow;
 			using (IRpcClientRef Client = await RpcConnection.GetClientRef(new RpcContext(), CancellationToken))
 			{
 				DirectoryReference LeaseDir = DirectoryReference.Combine(WorkingDir, "Remote", LeaseId);
@@ -796,7 +797,7 @@ namespace HordeAgent.Services
 				try
 				{
 					ActionExecutor Executor = new ActionExecutor(Settings.GetAgentName(), ActionTask.InstanceName, CasChannel, ActionCacheChannel, ActionRpcChannel, Logger);
-					await Executor.ExecuteActionAsync(LeaseId, ActionTask, LeaseDir);
+					await Executor.ExecuteActionAsync(LeaseId, ActionTask, LeaseDir, ActionTaskStartTime);
 				}
 				finally
 				{
diff --git a/Engine/Source/Programs/Horde/HordeServer/Tasks/Impl/ActionTaskSource.cs b/Engine/Source/Programs/Horde/HordeServer/Tasks/Impl/ActionTaskSource.cs
--- a/Engine/Source/Programs/Horde/HordeServer/Tasks/Impl/ActionTaskSource.cs
+++ b/Engine/Source/Programs/Horde/HordeServer/Tasks/Impl/ActionTaskSource.cs
@@ -393,6 +393,7 @@ namespace HordeServer.Tasks.Impl
 				}
 			}
 
+			DateTimeOffset QueueTime = DateTimeOffset.UtcNow;
 			for(; ;)
 			{
 				Operation.SetStatus(ExecutionStage.Types.Value.Queued, null);
@@ -415,6 +416,7 @@ namespace HordeServer.Tasks.Impl
 				ActionTask.InstanceName = Request.InstanceName;
 				ActionTask.Digest = Request.ActionDigest;
 				ActionTask.LogId = LogFile.Id.ToString();
+				ActionTask.QueueTime = Timestamp.FromDateTimeOffset(QueueTime);
 
 				if (Request.InstanceName != null && RemoteExecSettings.Instances.TryGetValue(Request.InstanceName, out var InstanceSettings))
 				{
