Horde: Fix matching of Docker link warnings.

diff --git a/Engine/Source/Programs/Horde/HordeAgent/Parser/Matchers/LinkEventMatcher.cs b/Engine/Source/Programs/Horde/HordeAgent/Parser/Matchers/LinkEventMatcher.cs
--- a/Engine/Source/Programs/Horde/HordeAgent/Parser/Matchers/LinkEventMatcher.cs
+++ b/Engine/Source/Programs/Horde/HordeAgent/Parser/Matchers/LinkEventMatcher.cs
@@ -21,28 +21,45 @@ namespace HordeAgent.Parser.Matchers
 		/// <inheritdoc/>
 		public LogEvent? Match(ILogCursor Cursor, ILogContext Context)
 		{
-			Match? Match;
-
 			int LineCount = 0;
-			bool bIsWarning = true;
-			while (Cursor.IsMatch(LineCount, @": undefined reference to |undefined symbol|^\s*(ld|ld.lld):|[^a-zA-Z]ld: |^\s*>>>"))
-			{
-				bIsWarning &= Cursor.IsMatch("warning:");
-				LineCount++;
-			}
-			if (Cursor.IsMatch(LineCount, @"error: linker command failed with exit code "))
+			bool bIsError = false;
+			bool bIsWarning = false;
+			while (Cursor.IsMatch(LineCount, @": undefined reference to |undefined symbol|^\s*(ld|ld.lld):|^[^:]*[^a-zA-Z][a-z]?ld: |^\s*>>>"))
 			{
-				bIsWarning = false;
+				bIsError |= Cursor.IsMatch("error:");
+				bIsWarning |= Cursor.IsMatch("warning:");
 				LineCount++;
 			}
 			if (LineCount > 0)
 			{
+				for (; ; )
+				{
+					if (Cursor.IsMatch(LineCount, "ld:"))
+					{
+						break;
+					}
+					else if (Cursor.IsMatch(LineCount, "error:"))
+					{
+						bIsError = true;
+					}
+					else if (Cursor.IsMatch(LineCount, "warning:"))
+					{
+						bIsWarning = true;
+					}
+					else
+					{
+						break;
+					}
+					LineCount++;
+				}
+
 				LogEventBuilder Builder = new LogEventBuilder(Cursor);
 				Builder.LineCount = LineCount;
 				EventId EventId = AddSymbolMarkup(Builder) ? KnownLogEvents.Linker_UndefinedSymbol : KnownLogEvents.Linker;
-				return Builder.ToLogEvent(LogEventPriority.Normal, bIsWarning? LogLevel.Warning : LogLevel.Error, EventId);
+				return Builder.ToLogEvent(LogEventPriority.Normal, (bIsError || !bIsWarning)? LogLevel.Error : LogLevel.Warning, EventId);
 			}
 
+			Match? Match;
 			if (Cursor.TryMatch(@"^(\s*)Undefined symbols for architecture", out Match))
 			{
 				LogEventBuilder Builder = new LogEventBuilder(Cursor);
diff --git a/Engine/Source/Programs/Horde/HordeAgentTests/LogParserTests.cs b/Engine/Source/Programs/Horde/HordeAgentTests/LogParserTests.cs
--- a/Engine/Source/Programs/Horde/HordeAgentTests/LogParserTests.cs
+++ b/Engine/Source/Programs/Horde/HordeAgentTests/LogParserTests.cs
@@ -843,6 +843,24 @@ namespace HordeAgentTests
 			Assert.AreEqual(1, Events[0].LineCount);
 		}
 
+		[TestMethod]
+		public void DockerErrorMatcher()
+		{
+			string[] Lines =
+			{
+				@"  #14 9.301 cc -O2 -Wall -DLUA_ANSI -DENABLE_CJSON_GLOBAL -DREDIS_STATIC=''    -c -o lua.o lua.c",
+				@"  #14 9.419 cc -o lua  lua.o liblua.a -lm",
+				@"  #14 9.447 /usr/bin/ld: liblua.a(loslib.o): in function `os_tmpname':",
+				@"  #14 9.447 loslib.c:(.text+0x280): warning: the use of `tmpnam' is dangerous, better use `mkstemp'",
+				@"  #14 9.448 cc -O2 -Wall -DLUA_ANSI -DENABLE_CJSON_GLOBAL -DREDIS_STATIC=''    -c -o luac.o luac.c"
+			};
+			List<CapturedEvent> Events = Parse(Lines);
+			Assert.AreEqual(1, Events.Count);
+
+			Assert.AreEqual(LogLevel.Warning, Events[0].Level);
+			Assert.AreEqual(KnownLogEvents.Linker, Events[0].Id);
+		}
+
 		[TestMethod]
 		public void GauntletErrorMatcher()
 		{
