Horde: Return the operation ID for remote exec actions

Helps track multiple outstanding operations from the client-side.

diff --git a/Engine/Source/Programs/Horde/HordeServer/Tasks/Impl/ActionTaskSource.cs b/Engine/Source/Programs/Horde/HordeServer/Tasks/Impl/ActionTaskSource.cs
--- a/Engine/Source/Programs/Horde/HordeServer/Tasks/Impl/ActionTaskSource.cs
+++ b/Engine/Source/Programs/Horde/HordeServer/Tasks/Impl/ActionTaskSource.cs
@@ -149,6 +149,7 @@ namespace HordeServer.Tasks.Impl
 					Operation.Done = true;
 				}
 				Operation.Metadata = Any.Pack(Metadata);
+				Operation.Name = Id.ToString();
 
 				return TrySetStatus(Operation);
 			}
