Horde: Increase logging for completed REAPI operations

To help identify an outstanding bug in live environment.

diff --git a/Engine/Source/Programs/Horde/HordeServer/Tasks/Impl/ActionTaskSource.cs b/Engine/Source/Programs/Horde/HordeServer/Tasks/Impl/ActionTaskSource.cs
--- a/Engine/Source/Programs/Horde/HordeServer/Tasks/Impl/ActionTaskSource.cs
+++ b/Engine/Source/Programs/Horde/HordeServer/Tasks/Impl/ActionTaskSource.cs
@@ -326,6 +326,7 @@ namespace HordeServer.Tasks.Impl
 			ExecuteOperation? Operation;
 			if (TryGetOperation(Lease.Id, out Operation))
 			{
+				Logger.LogDebug("Lease cancelled. LeaseId={LeaseId} OperationId={OperationId}", Lease.Id, Operation.Id);
 				Operation.TrySetResult(null);
 			}
 			return Task.CompletedTask;
@@ -442,7 +443,14 @@ namespace HordeServer.Tasks.Impl
 
 				// Try to set the lease on the subscription. If it fails, the subscriber has already been allocated
 				Operation.ResultTaskSource = new TaskCompletionSource<ActionResult?>();
-				if (Subscription.TrySetLease(Lease, () => Operation.ResultTaskSource.TrySetResult(null)))
+
+				void OnConnectionLost()
+				{
+					Logger.LogDebug("Connection lost. LeaseId={LeaseId} OperationId={OperationId}", Lease.Id, Operation!.Id);
+					Operation.TrySetResult(null);
+				}
+				
+				if (Subscription.TrySetLease(Lease, OnConnectionLost))
 				{
 					Operation.SetStatus(ExecutionStage.Types.Value.Executing, null);
 
@@ -562,6 +570,7 @@ namespace HordeServer.Tasks.Impl
 			{
 				lock (Bucket!)
 				{
+					Logger.LogDebug("Trimming lease from bucket. LeaseId={LeaseId} OperationId={OperationId}", LeaseId, Operation!.Id);
 					Operation!.TrySetResult(null);
 					Bucket.Remove(LeaseId.Value);
 				}
