Horde: Use user ids for build health issues.  Adding a UserSelect widget for selecting from existing users
#rnx

diff --git a/Engine/Source/Programs/Horde/HordeDashboard/src/backend/Api.ts b/Engine/Source/Programs/Horde/HordeDashboard/src/backend/Api.ts
--- a/Engine/Source/Programs/Horde/HordeDashboard/src/backend/Api.ts
+++ b/Engine/Source/Programs/Horde/HordeDashboard/src/backend/Api.ts
@@ -285,6 +285,16 @@ export type IssueQuery = {
 	resolved?: boolean;
 }
 
+export type UsersQuery = {
+	ids?: string[];
+	nameRegex? : string;
+	index?: number;
+	count?: number;
+	includeAvatar?: boolean;
+	includeClaims?: boolean;
+}
+
+
 export type ScheduleQuery = {
 	streamId?: string;
 	filter?: string;
@@ -2519,6 +2529,9 @@ export type GetIssueResponse = {
 	/**Owner of the issue */
 	owner?: string;
 
+	/**Owner id of the issue */
+	ownerId?: string;
+
 	/**User that nominated the current owner */
 	nominatedBy?: string;
 
@@ -2542,8 +2555,8 @@ export type GetIssueResponse = {
 
 	affectedStreams: GetIssueAffectedStreamResponse[];
 
-	/** Most likely suspects for causing this issue */
-	primarySuspects: string[];
+	/** User id's of the Most likely suspects for causing this issue */
+	primarySuspectIds: string[];
 
 	/** Whether to show alerts for this issue */
 	showDesktopAlerts: boolean;
@@ -2557,10 +2570,10 @@ export type UpdateIssueRequest = {
 	summary?: string;
 
 	/**New owner of the issue */
-	owner?: string | null;
+	ownerId?: string | null;
 
 	/**User than nominates the new owner */
-	nominatedBy?: string | null;
+	nominatedById?: string | null;
 
 	/**Whether the issue has been acknowledged */
 	acknowledged?: boolean;
@@ -2643,20 +2656,38 @@ export type GetUserResponse = {
 	/** Id of the user */
 	id: string;
 
-	/** Claims for the user */
-	claims: UserClaim[];
+	/** Name of the user */
+	name: string;
 
-	/** List of pinned job ids */
-	pinnedJobIds: string[];
+	/** Email of the user */
+	email?: string;
+
+	/** Avatar image URL (24px) */
+	image24?: string;
+
+	/** Avatar image URL (32px) */
+	image32?: string;
+
+	/** Avatar image URL (48px) */
+	image48?: string;
+
+	/** Avatar image URL (72px) */
+	image72?: string;
+
+	/** Claims for the user */
+	claims?: UserClaim[];
 
 	/** Whether to enable slack notifications for this user */
-	enableIssueNotifications: boolean;
+	enableIssueNotifications?: boolean;
 
 	/** Whether to enable experimental features for this user */
-	enableExperimentalFeatures: boolean;
+	enableExperimentalFeatures?: boolean;
 
 	/**  Settings for the dashboard */
-	dashboardSettings: DashboardSettings;
+	dashboardSettings?: DashboardSettings;
+
+	/** List of pinned job ids */
+	pinnedJobIds?: string[];
 
 }
 
diff --git a/Engine/Source/Programs/Horde/HordeDashboard/src/backend/Backend.ts b/Engine/Source/Programs/Horde/HordeDashboard/src/backend/Backend.ts
--- a/Engine/Source/Programs/Horde/HordeDashboard/src/backend/Backend.ts
+++ b/Engine/Source/Programs/Horde/HordeDashboard/src/backend/Backend.ts
@@ -3,7 +3,7 @@
 import { configure } from "mobx";
 import { isNumber } from 'util';
 import templateCache from '../backend/TemplateCache';
-import { AgentData, ArtifactData, BatchUpdatePoolRequest, ChangeSummaryData, CreateJobRequest, CreateJobResponse, CreatePoolRequest, CreateSoftwareResponse, CreateSubscriptionRequest, CreateSubscriptionResponse, DashboardPreference, EventData, GetArtifactZipRequest, GetUserResponse, GetGraphResponse, GetIssueStreamResponse, GetJobStepRefResponse, GetJobStepTraceResponse, GetJobTimingResponse, GetLogEventResponse, GetNotificationResponse, GetSoftwareResponse, GetSubscriptionResponse, IssueData, IssueQuery, JobData, JobQuery, LeaseData, LogData, LogLineData, PoolData, ProjectData, ScheduleData, ScheduleQuery, SearchLogFileResponse, SessionData, StreamData, TemplateData, TestData, UpdateAgentRequest, UpdateIssueRequest, UpdateJobRequest, UpdateNotificationsRequest, UpdatePoolRequest, UpdateStepRequest, UpdateStepResponse, UpdateUserRequest, GetUtilizationTelemetryResponse, TabType, GetJobsTabResponse, JobsTabColumnType, GetPerforceServerStatusResponse, GetDeviceResponse, GetDevicePlatformResponse, CreateDeviceResponse, CreateDeviceRequest, UpdateDeviceRequest, GetDevicePoolResponse, GetDeviceReservationResponse } from './Api';
+import { AgentData, ArtifactData, BatchUpdatePoolRequest, ChangeSummaryData, CreateJobRequest, CreateJobResponse, CreatePoolRequest, CreateSoftwareResponse, CreateSubscriptionRequest, CreateSubscriptionResponse, DashboardPreference, EventData, GetArtifactZipRequest, GetUserResponse, GetGraphResponse, GetIssueStreamResponse, GetJobStepRefResponse, GetJobStepTraceResponse, GetJobTimingResponse, GetLogEventResponse, GetNotificationResponse, GetSoftwareResponse, GetSubscriptionResponse, IssueData, IssueQuery, JobData, JobQuery, LeaseData, LogData, LogLineData, PoolData, ProjectData, ScheduleData, ScheduleQuery, SearchLogFileResponse, SessionData, StreamData, TemplateData, TestData, UpdateAgentRequest, UpdateIssueRequest, UpdateJobRequest, UpdateNotificationsRequest, UpdatePoolRequest, UpdateStepRequest, UpdateStepResponse, UpdateUserRequest, GetUtilizationTelemetryResponse, TabType, GetJobsTabResponse, JobsTabColumnType, GetPerforceServerStatusResponse, GetDeviceResponse, GetDevicePlatformResponse, CreateDeviceResponse, CreateDeviceRequest, UpdateDeviceRequest, GetDevicePoolResponse, GetDeviceReservationResponse, UsersQuery } from './Api';
 import dashboard from './Dashboard';
 import { ChallengeStatus, Fetch } from './Fetch';
 import graphCache, { GraphQuery } from './GraphCache';
@@ -766,11 +766,25 @@ export class Backend {
 
     }
 
-    getUser(): Promise<GetUserResponse> {
+    getUsers(query?: UsersQuery): Promise<GetUserResponse[]> {
+
+        query = query ?? {};
+
+        return new Promise<GetUserResponse[]>((resolve, reject) => {
+            this.backend.get("/api/v1/users",{
+                params: query as any
+            }).then((response) => {
+                resolve(response.data as GetUserResponse[] );
+            }).catch(reason => reject(reason));
+        });
+
+    }
+
+    getCurrentUser(): Promise<GetUserResponse> {
 
         return new Promise<GetUserResponse>((resolve, reject) => {
 
-            this.backend.get("/api/v1/user", { suppress404: true }).then((response) => {
+            this.backend.get("/api/v1/users/current", { suppress404: true }).then((response) => {
 
                 let data = response.data as GetUserResponse;
 
diff --git a/Engine/Source/Programs/Horde/HordeDashboard/src/backend/Dashboard.ts b/Engine/Source/Programs/Horde/HordeDashboard/src/backend/Dashboard.ts
--- a/Engine/Source/Programs/Horde/HordeDashboard/src/backend/Dashboard.ts
+++ b/Engine/Source/Programs/Horde/HordeDashboard/src/backend/Dashboard.ts
@@ -20,11 +20,11 @@ export class Dashboard {
 
     pinJob(id: string) {
 
-        if (this.data.pinnedJobIds.find(j => j === id)) {
+        if (this.data.pinnedJobIds!.find(j => j === id)) {
             return;
         }
 
-        this.data.pinnedJobIds.push(id);
+        this.data.pinnedJobIds!.push(id);
 
         backend.updateUser({ addPinnedJobIds: [id] });
 
@@ -33,11 +33,11 @@ export class Dashboard {
 
     unpinJob(id: string) {
 
-        if (!this.data.pinnedJobIds.find(j => j === id)) {
+        if (!this.data.pinnedJobIds!.find(j => j === id)) {
             return;
         }
 
-        this.data.pinnedJobIds = this.data.pinnedJobIds.filter(j => j !== id);
+        this.data.pinnedJobIds = this.data.pinnedJobIds!.filter(j => j !== id);
 
         backend.updateUser({ removePinnedJobIds: [id] });
 
@@ -45,14 +45,15 @@ export class Dashboard {
     }
 
     get username(): string {
+        return this.data.name;
+    }
 
-        let email = this.claims.find(c => c.type.endsWith("/ue/horde/user"));
-
-        if (!email) {
-            email = this.claims.find(c => c.type.endsWith("/identity/claims/name"));
-        }
+    get userImage48(): string | undefined {
+        return this.data.image48;
+    }
 
-        return email ? email.value : "???";
+    get userImage32(): string | undefined {
+        return this.data.image32;
     }
 
 
@@ -84,6 +85,11 @@ export class Dashboard {
 
     }
 
+    get userId(): string {
+        return this.data.id;
+    }
+
+
     get pinnedJobsIds(): string[] {
 
         return this.data.pinnedJobIds ?? [];
@@ -91,7 +97,7 @@ export class Dashboard {
 
     get issueNotifications(): boolean {
 
-        return this.data.enableIssueNotifications;
+        return this.data.enableIssueNotifications!;
     }
 
     set issueNotifications(value: boolean) {
@@ -107,7 +113,7 @@ export class Dashboard {
 
     get experimentalFeatures(): boolean {
 
-        return this.data.enableExperimentalFeatures;
+        return this.data.enableExperimentalFeatures!;
     }
 
     set experimentalFeatures(value: boolean) {
@@ -187,7 +193,7 @@ export class Dashboard {
     }
 
     private get preferences() {
-        return this.data.dashboardSettings.preferences;
+        return this.data.dashboardSettings!.preferences;
     }
 
     async update() {
@@ -206,7 +212,7 @@ export class Dashboard {
 
                 const cancelId = this.cancelId++;
 
-                const response = await backend.getUser();
+                const response = await backend.getCurrentUser();
 
                 // check for canceled during graph request
                 if (!this.canceled.has(cancelId)) {
@@ -293,7 +299,7 @@ export class Dashboard {
         const data: any = {};
 
         for (const key of Object.keys(DashboardPreference)) {
-            data[key] = this.data.dashboardSettings.preferences?.get(key as DashboardPreference);
+            data[key] = this.data.dashboardSettings!.preferences?.get(key as DashboardPreference);
         }
 
         let success = true;
@@ -310,7 +316,7 @@ export class Dashboard {
 
     requestLogout = false;
 
-    private data: GetUserResponse = { id: "", enableIssueNotifications: false, enableExperimentalFeatures: false, claims: [], pinnedJobIds: [], dashboardSettings: { preferences: new Map() } };
+    private data: GetUserResponse = { id: "", name: "", enableIssueNotifications: false, enableExperimentalFeatures: false, claims: [], pinnedJobIds: [], dashboardSettings: { preferences: new Map() } };
 
     private updateTimeoutId: any = undefined;
 
diff --git a/Engine/Source/Programs/Horde/HordeDashboard/src/backend/UserCache.ts b/Engine/Source/Programs/Horde/HordeDashboard/src/backend/UserCache.ts
new file mode 100644
--- /dev/null
+++ b/Engine/Source/Programs/Horde/HordeDashboard/src/backend/UserCache.ts
@@ -0,0 +1,54 @@
+
+import { action, observable } from 'mobx';
+import backend from '.';
+import { GetUserResponse } from './Api';
+
+
+class UserCache {
+
+    @action
+    setUpdated() {
+        this.updated++;
+    }
+
+    @observable
+    updated: number = 0;
+
+    async getUsers(startsWith:string) {
+
+        startsWith = startsWith.toLowerCase();
+        
+        let users = this.userMap.get(startsWith);
+
+        if (users !== undefined) {
+            return users;
+        }
+
+        users = [];
+
+        try {
+            users = await backend.getUsers({count: 128, includeAvatar: true, nameRegex: `^${startsWith}`});
+            console.log(`Found ${users.length} users starting with ${startsWith}`);
+
+            users = users.sort((a, b) => {
+                const aname = a.name.toLowerCase();
+                const bname = b.name.toLowerCase();
+
+                if (aname < bname) return -1;
+                if (aname > bname) return 1;
+                return 0;
+            })
+            this.userMap.set(startsWith, users);
+            this.setUpdated();
+        } catch (reason) {
+            console.error("Unable to get users", reason);
+        }
+        
+        return users;
+    }
+
+    private userMap:Map<string, GetUserResponse[]> = new Map();
+
+}
+
+export default new UserCache();
\ No newline at end of file
diff --git a/Engine/Source/Programs/Horde/HordeServer/Api/Issues.cs b/Engine/Source/Programs/Horde/HordeServer/Api/Issues.cs
--- a/Engine/Source/Programs/Horde/HordeServer/Api/Issues.cs
+++ b/Engine/Source/Programs/Horde/HordeServer/Api/Issues.cs
@@ -488,10 +488,15 @@ namespace HordeServer.Api
 		public List<GetIssueAffectedStreamResponse> AffectedStreams { get; set; }
 
 		/// <summary>
-		/// Most likely suspects for causing this issue
+		/// Most likely suspects for causing this issue, deprecated
 		/// </summary>
 		public List<string> PrimarySuspects { get; set; }
 
+		/// <summary>
+		/// User ids of the most likely suspects
+		/// </summary>
+		public List<string> PrimarySuspectIds { get; set; }
+
 		/// <summary>
 		/// Whether to show alerts for this issue
 		/// </summary>
@@ -537,6 +542,7 @@ namespace HordeServer.Api
 				}
 			}
 			this.PrimarySuspects = Details.SuspectUsers.Where(x => x.Login != null).Select(x => x.Login).ToList();
+			this.PrimarySuspectIds= Details.SuspectUsers.Select(x => x.Id.ToString()).ToList();
 			this.ShowDesktopAlerts = ShowDesktopAlerts;
 		}
 	}
@@ -552,14 +558,14 @@ namespace HordeServer.Api
 		public string? Summary { get; set; }
 
 		/// <summary>
-		/// New owner of the issue
+		/// New user id for owner of the issue
 		/// </summary>
-		public string? Owner { get; set; }
+		public string? OwnerId { get; set; }
 
 		/// <summary>
-		/// User than nominates the new owner
+		/// User id that nominated the new owner
 		/// </summary>
-		public string? NominatedBy { get; set; }
+		public string? NominatedById { get; set; }
 
 		/// <summary>
 		/// Whether the issue has been acknowledged
diff --git a/Engine/Source/Programs/Horde/HordeServer/Api/Users.cs b/Engine/Source/Programs/Horde/HordeServer/Api/Users.cs
--- a/Engine/Source/Programs/Horde/HordeServer/Api/Users.cs
+++ b/Engine/Source/Programs/Horde/HordeServer/Api/Users.cs
@@ -22,6 +22,11 @@ namespace HordeServer.Api
 		/// </summary>
 		public string Id { get; set; }
 
+		/// <summary>
+		/// Name of the user
+		/// </summary>
+		public string Name { get; set; }
+
 		/// <summary>
 		/// Avatar image URL (24px)
 		/// </summary>
@@ -42,46 +47,59 @@ namespace HordeServer.Api
 		/// </summary>
 		public string? Image72 { get; set; }
 
+		/// <summary>
+		/// Email of the user
+		/// </summary>
+		public string? Email { get; set; }
+
 		/// <summary>
 		/// Claims for the user
 		/// </summary>
-		public List<UserClaim> Claims { get; set; }
+		public List<UserClaim>? Claims { get; set; }
 
 		/// <summary>
 		/// Whether to enable experimental features for this user
 		/// </summary>
-		public bool EnableExperimentalFeatures { get; set; }
+		public bool? EnableExperimentalFeatures { get; set; }
 
 		/// <summary>
 		/// Whether to enable slack notifications for this user
 		/// </summary>
-		public bool EnableIssueNotifications { get; set; }
+		public bool? EnableIssueNotifications { get; set; }
 
 		/// <summary>
 		/// Settings for the dashboard
 		/// </summary>
-		public object DashboardSettings { get; set; }
+		public object? DashboardSettings { get; set; }
 
 		/// <summary>
 		/// List of pinned job ids
 		/// </summary>
-		public List<string> PinnedJobIds { get; set; }
+		public List<string>? PinnedJobIds { get; set; }
 
 		/// <summary>
 		/// Constructor
 		/// </summary>
-		public GetUserResponse(IUser User, IAvatar? Avatar, IUserClaims Claims, IUserSettings Settings)
+		public GetUserResponse(IUser User, IAvatar? Avatar, IUserClaims? Claims, IUserSettings? Settings)
 		{
 			this.Id = User.Id.ToString();
+			this.Name = User.Name;
+			this.Email = User.Email;
+
 			this.Image24 = Avatar?.Image24;
 			this.Image32 = Avatar?.Image32;
 			this.Image48 = Avatar?.Image48;
 			this.Image72 = Avatar?.Image72;
-			this.Claims = Claims.Claims.Select(x => new UserClaim(x)).ToList();
-			this.EnableExperimentalFeatures = Settings.EnableExperimentalFeatures;
-			this.EnableIssueNotifications = Settings.EnableIssueNotifications;
-			this.DashboardSettings = BsonTypeMapper.MapToDotNetValue(Settings.DashboardSettings);
-			this.PinnedJobIds = Settings.PinnedJobIds.ConvertAll(x => x.ToString());
+						
+			this.Claims = Claims == null ? null : Claims.Claims.Select(x => new UserClaim(x)).ToList();
+
+			if (Settings != null)
+			{
+				this.EnableExperimentalFeatures = Settings.EnableExperimentalFeatures;
+				this.EnableIssueNotifications = Settings.EnableIssueNotifications;
+				this.DashboardSettings = BsonTypeMapper.MapToDotNetValue(Settings.DashboardSettings);
+				this.PinnedJobIds = Settings.PinnedJobIds.ConvertAll(x => x.ToString());
+			}
 		}
 	}
 
@@ -115,4 +133,5 @@ namespace HordeServer.Api
 		/// </summary>
 		public List<string>? RemovePinnedJobIds { get; set; }
 	}
+
 }
diff --git a/Engine/Source/Programs/Horde/HordeServer/Collections/IUserCollection.cs b/Engine/Source/Programs/Horde/HordeServer/Collections/IUserCollection.cs
--- a/Engine/Source/Programs/Horde/HordeServer/Collections/IUserCollection.cs
+++ b/Engine/Source/Programs/Horde/HordeServer/Collections/IUserCollection.cs
@@ -35,10 +35,11 @@ namespace HordeServer.Collections
 		/// Gets a user by unique id
 		/// </summary>
 		/// <param name="Ids">Ids of the users</param>
+		/// <param name="NameRegex">Name regex to match for the users</param>
 		/// <param name="Index">Maximum number of results</param>
 		/// <param name="Count">Number of results to return</param>
 		/// <returns>The user information</returns>
-		Task<List<IUser>> FindUsersAsync(IEnumerable<ObjectId>? Ids = null, int? Index = null, int? Count = null);
+		Task<List<IUser>> FindUsersAsync(IEnumerable<ObjectId>? Ids = null, string? NameRegex = null, int? Index = null, int? Count = null);
 
 		/// <summary>
 		/// Gets a user by login
diff --git a/Engine/Source/Programs/Horde/HordeServer/Collections/Impl/UserCollection.cs b/Engine/Source/Programs/Horde/HordeServer/Collections/Impl/UserCollection.cs
--- a/Engine/Source/Programs/Horde/HordeServer/Collections/Impl/UserCollection.cs
+++ b/Engine/Source/Programs/Horde/HordeServer/Collections/Impl/UserCollection.cs
@@ -109,7 +109,7 @@ namespace HordeServer.Collections.Impl
 		}
 
 		/// <inheritdoc/>
-		public async Task<List<IUser>> FindUsersAsync(IEnumerable<ObjectId>? Ids, int? Index, int? Count)
+		public async Task<List<IUser>> FindUsersAsync(IEnumerable<ObjectId>? Ids, string? NameRegex, int? Index, int? Count)
 		{
 			FilterDefinition<UserDocument> Filter = Builders<UserDocument>.Filter.In(x => x.Id, Ids);
 			return await Users.Find(Filter).Range(Index, Count).ToListAsync<UserDocument, IUser>();
diff --git a/Engine/Source/Programs/Horde/HordeServer/Collections/Impl/UserCollectionV2.cs b/Engine/Source/Programs/Horde/HordeServer/Collections/Impl/UserCollectionV2.cs
--- a/Engine/Source/Programs/Horde/HordeServer/Collections/Impl/UserCollectionV2.cs
+++ b/Engine/Source/Programs/Horde/HordeServer/Collections/Impl/UserCollectionV2.cs
@@ -217,13 +217,20 @@ namespace HordeServer.Collections.Impl
 		}
 
 		/// <inheritdoc/>
-		public async Task<List<IUser>> FindUsersAsync(IEnumerable<ObjectId>? Ids, int? Index = null, int? Count = null)
+		public async Task<List<IUser>> FindUsersAsync(IEnumerable<ObjectId>? Ids, string? NameRegex = null, int? Index = null, int? Count = null)
 		{
 			FilterDefinition<UserDocument> Filter = FilterDefinition<UserDocument>.Empty;
 			if (Ids != null)
 			{
 				Filter &= Builders<UserDocument>.Filter.In(x => x.Id, Ids);
 			}
+			
+			if (NameRegex != null)
+			{
+				BsonRegularExpression Regex = new BsonRegularExpression(NameRegex, "i");
+				Filter &= Builders<UserDocument>.Filter.Regex(x => x.Name, Regex);
+			}
+
 			return await Users.Find(Filter).Range(Index, Count ?? 100).ToListAsync<UserDocument, IUser>();
 		}
 
diff --git a/Engine/Source/Programs/Horde/HordeServer/Controllers/IssuesController.cs b/Engine/Source/Programs/Horde/HordeServer/Controllers/IssuesController.cs
--- a/Engine/Source/Programs/Horde/HordeServer/Controllers/IssuesController.cs
+++ b/Engine/Source/Programs/Horde/HordeServer/Controllers/IssuesController.cs
@@ -401,15 +401,15 @@ namespace HordeServer.Controllers
 		public async Task<ActionResult> UpdateIssueAsync(int IssueId, [FromBody] UpdateIssueRequest Request)
 		{
 			ObjectId? NewOwnerId = null;
-			if (Request.Owner != null)
+			if (Request.OwnerId != null)
 			{
-				NewOwnerId = (await UserCollection.FindOrAddUserByLoginAsync(Request.Owner))?.Id;
+				NewOwnerId = new ObjectId(Request.OwnerId);
 			}
 
 			ObjectId? NewNominatedById = null;
-			if (Request.NominatedBy != null)
+			if (Request.NominatedById != null)
 			{
-				NewNominatedById = (await UserCollection.FindOrAddUserByLoginAsync(Request.NominatedBy))?.Id;
+				NewNominatedById = new ObjectId(Request.NominatedById);
 			}
 
 			ObjectId? NewDeclinedById = null;
diff --git a/Engine/Source/Programs/Horde/HordeServer/Controllers/UsersController.cs b/Engine/Source/Programs/Horde/HordeServer/Controllers/UsersController.cs
--- a/Engine/Source/Programs/Horde/HordeServer/Controllers/UsersController.cs
+++ b/Engine/Source/Programs/Horde/HordeServer/Controllers/UsersController.cs
@@ -47,7 +47,7 @@ namespace HordeServer.Controllers
 		}
 
 		/// <summary>
-		/// Gets information about the logged in user
+		/// Gets information about a user by id, specify "current" for id to get the currently logged in user
 		/// </summary>
 		/// <returns>Http result code</returns>
 		[HttpGet]
@@ -67,6 +67,42 @@ namespace HordeServer.Controllers
 			return PropertyFilter.Apply(new GetUserResponse(User, Avatar, Claims, Settings), Filter);
 		}
 
+		/// <summary>
+		/// Gets a list of users
+		/// </summary>
+		/// <returns>List of user responses</returns>
+		[HttpGet]
+		[Route("/api/v1/users")]
+		[ProducesResponseType(typeof(List<GetUserResponse>), 200)]
+		public async Task<ActionResult<List<GetUserResponse>>> FindUsersAsync(
+			[FromQuery] string[]? Ids = null,
+			[FromQuery] string? NameRegex = null,
+			[FromQuery] int Index = 0,
+			[FromQuery] int Count = 100,
+			[FromQuery] bool IncludeClaims = false,
+			[FromQuery] bool IncludeAvatar = false)
+		{
+
+			ObjectId[]? UserIds = null;
+			if (Ids != null && Ids.Length > 0)
+			{
+				UserIds = Ids.Select(x => new ObjectId(x)).ToArray();
+			}
+
+			List<IUser> Users = await UserCollection.FindUsersAsync(UserIds, NameRegex, Index, Count);
+
+			List<GetUserResponse> Response = new List<GetUserResponse>();
+			foreach (IUser User in Users)
+			{
+				IAvatar? Avatar = (AvatarService == null || !IncludeAvatar) ? (IAvatar?)null : await AvatarService.GetAvatarAsync(User);
+				IUserClaims? Claims = (!IncludeClaims) ? null : await UserCollection.GetClaimsAsync(User.Id);
+				Response.Add(new GetUserResponse(User, Avatar, Claims, null));
+			}
+
+			return Response;
+		}
+
+
 		async Task<IUser?> GetUserInternalAsync(string Id)
 		{
 			ObjectId? UserId = ParseUserId(Id);
