Horde: More logging to identify REAPI retry issues

diff --git a/Engine/Source/Programs/Horde/HordeServer/Rpc/ActionRpcService.cs b/Engine/Source/Programs/Horde/HordeServer/Rpc/ActionRpcService.cs
--- a/Engine/Source/Programs/Horde/HordeServer/Rpc/ActionRpcService.cs
+++ b/Engine/Source/Programs/Horde/HordeServer/Rpc/ActionRpcService.cs
@@ -9,16 +9,19 @@ using System;
 using System.Collections.Generic;
 using System.Linq;
 using System.Threading.Tasks;
+using Microsoft.Extensions.Logging;
 
 namespace HordeServer.Rpc
 {
 	class ActionRpcService : ActionRpc.ActionRpcBase
 	{
 		ActionTaskSource ActionTaskSource;
+		ILogger<ActionRpcService> Logger;
 
-		public ActionRpcService(ActionTaskSource ActionTaskSource)
+		public ActionRpcService(ActionTaskSource ActionTaskSource, ILogger<ActionRpcService> Logger)
 		{
 			this.ActionTaskSource = ActionTaskSource;
+			this.Logger = Logger;
 		}
 
 		/// <inheritdoc/>
@@ -27,6 +30,7 @@ namespace HordeServer.Rpc
 			IActionExecuteOperation? Operation;
 			if (ActionTaskSource.TryGetOperation(Request.LeaseId.ToObjectId(), out Operation))
 			{
+				Logger.LogInformation("Setting operation result to {Result}. LeaseId={LeaseId} OperationId={OperationId}", Request.Result, Request.LeaseId, Operation.Id);
 				Operation.TrySetResult(Request.Result);
 			}
 			return Task.FromResult<Empty>(new Empty());
diff --git a/Engine/Source/Programs/Horde/HordeServer/Tasks/Impl/ActionTaskSource.cs b/Engine/Source/Programs/Horde/HordeServer/Tasks/Impl/ActionTaskSource.cs
--- a/Engine/Source/Programs/Horde/HordeServer/Tasks/Impl/ActionTaskSource.cs
+++ b/Engine/Source/Programs/Horde/HordeServer/Tasks/Impl/ActionTaskSource.cs
@@ -326,7 +326,7 @@ namespace HordeServer.Tasks.Impl
 			ExecuteOperation? Operation;
 			if (TryGetOperation(Lease.Id, out Operation))
 			{
-				Logger.LogDebug("Lease cancelled. LeaseId={LeaseId} OperationId={OperationId}", Lease.Id, Operation.Id);
+				Logger.LogInformation("Lease cancelled. LeaseId={LeaseId} OperationId={OperationId}", Lease.Id, Operation.Id);
 				Operation.TrySetResult(null);
 			}
 			return Task.CompletedTask;
@@ -446,7 +446,7 @@ namespace HordeServer.Tasks.Impl
 
 				void OnConnectionLost()
 				{
-					Logger.LogDebug("Connection lost. LeaseId={LeaseId} OperationId={OperationId}", Lease.Id, Operation!.Id);
+					Logger.LogInformation("Connection lost. LeaseId={LeaseId} OperationId={OperationId}", Lease.Id, Operation!.Id);
 					Operation.TrySetResult(null);
 				}
 				
@@ -570,7 +570,7 @@ namespace HordeServer.Tasks.Impl
 			{
 				lock (Bucket!)
 				{
-					Logger.LogDebug("Trimming lease from bucket. LeaseId={LeaseId} OperationId={OperationId}", LeaseId, Operation!.Id);
+					Logger.LogInformation("Trimming lease from bucket. LeaseId={LeaseId} OperationId={OperationId}", LeaseId, Operation!.Id);
 					Operation!.TrySetResult(null);
 					Bucket.Remove(LeaseId.Value);
 				}
