Horde: Fix file-based config using hash of filename rather than hash of file contents.

diff --git a/Engine/Source/Programs/Horde/HordeServer/Services/ConfigService.cs b/Engine/Source/Programs/Horde/HordeServer/Services/ConfigService.cs
--- a/Engine/Source/Programs/Horde/HordeServer/Services/ConfigService.cs
+++ b/Engine/Source/Programs/Horde/HordeServer/Services/ConfigService.cs
@@ -319,7 +319,7 @@ namespace HordeServer.Services
 				}
 				else
 				{
-					Revisions[Path] = $"ver={Version},md5={ContentHash.MD5(Path)}";
+					Revisions[Path] = $"ver={Version},md5={ContentHash.MD5(new FileReference(Path))}";
 				}
 			}
 
