Horde: Fix warnings in dev environment due to empty string for websocket url.

diff --git a/Engine/Source/Programs/Horde/HordeServer/Notifications/Impl/SlackNotificationSink.cs b/Engine/Source/Programs/Horde/HordeServer/Notifications/Impl/SlackNotificationSink.cs
--- a/Engine/Source/Programs/Horde/HordeServer/Notifications/Impl/SlackNotificationSink.cs
+++ b/Engine/Source/Programs/Horde/HordeServer/Notifications/Impl/SlackNotificationSink.cs
@@ -1160,7 +1160,7 @@ namespace HordeServer.Notifications.Impl
 		/// <inheritdoc/>
 		protected async override Task ExecuteAsync(CancellationToken StoppingToken)
 		{
-			if (Settings.SlackSocketToken != null)
+			if (!String.IsNullOrEmpty(Settings.SlackSocketToken))
 			{
 				while (!StoppingToken.IsCancellationRequested)
 				{
