Horde: Some more UserId cleanup
#rnx

diff --git a/Engine/Source/Programs/Horde/HordeDashboard/src/backend/Api.ts b/Engine/Source/Programs/Horde/HordeDashboard/src/backend/Api.ts
--- a/Engine/Source/Programs/Horde/HordeDashboard/src/backend/Api.ts
+++ b/Engine/Source/Programs/Horde/HordeDashboard/src/backend/Api.ts
@@ -262,7 +262,7 @@ export type JobQuery = {
 	maxChange?: number;
 	preflightChange?: number;
 	includePreflight?: boolean,
-	preflightStartedByUser?: string,
+	preflightStartedByUserId?: string,
 	minCreateTime?: string;
 	maxCreateTime?: string;
 	modifiedAfter?: string;
@@ -281,7 +281,7 @@ export type IssueQuery = {
 	maxChange?: number;
 	index?: number;
 	count?: number;
-	user?: string;
+	userId?: string;
 	resolved?: boolean;
 }
 
diff --git a/Engine/Source/Programs/Horde/HordeDashboard/src/backend/Backend.ts b/Engine/Source/Programs/Horde/HordeDashboard/src/backend/Backend.ts
--- a/Engine/Source/Programs/Horde/HordeDashboard/src/backend/Backend.ts
+++ b/Engine/Source/Programs/Horde/HordeDashboard/src/backend/Backend.ts
@@ -1034,7 +1034,7 @@ export class Backend {
 
     createSubscription(requests: CreateSubscriptionRequest[]): Promise<CreateSubscriptionResponse[]> {
         requests.forEach(request => {
-            request.userId = dashboard.username ?? "Anonymous";
+            request.userId = dashboard.userId;
         });
         let url = `api/v1/subscriptions`;
         return new Promise<CreateSubscriptionResponse[]>((resolve, reject) => {
@@ -1048,7 +1048,7 @@ export class Backend {
     }
 
     getSubscriptions(): Promise<GetSubscriptionResponse[]> {
-        let url = `api/v1/subscriptions?userId=${dashboard.username ?? "Anonymous"}`;
+        let url = `api/v1/subscriptions?userId=${dashboard.userId}`;
         return new Promise<GetSubscriptionResponse[]>((resolve, reject) => {
             this.backend.get(`${url}`, { suppress404: true }).then((value) => {
                 resolve(value.data as GetSubscriptionResponse[]);
diff --git a/Engine/Source/Programs/Horde/HordeDashboard/src/backend/JobHandler.ts b/Engine/Source/Programs/Horde/HordeDashboard/src/backend/JobHandler.ts
--- a/Engine/Source/Programs/Horde/HordeDashboard/src/backend/JobHandler.ts
+++ b/Engine/Source/Programs/Horde/HordeDashboard/src/backend/JobHandler.ts
@@ -20,19 +20,19 @@ export class JobHandler {
     }
 
 
-    filter(stream: StreamData, templateNamesIn?: string[], preflightStartedByUser?: string): boolean {
+    filter(stream: StreamData, templateNamesIn?: string[], preflightStartedByUserId?: string): boolean {
 
         const names = templateNamesIn ?? [];
         const sameNames = (this.templateNames.length === names.length && this.templateNames.every((value, index) => value === names[index]));
 
-        if (stream.id === this.stream?.id && sameNames && this.preflightStartedByUser === preflightStartedByUser) {
+        if (stream.id === this.stream?.id && sameNames && this.preflightStartedByUserId === preflightStartedByUserId) {
             return false;
         }
 
         this.clear();
         this.stream = stream;
         this.templateNames = names;
-        this.preflightStartedByUser = preflightStartedByUser;
+        this.preflightStartedByUserId = preflightStartedByUserId;
         this.setUpdated();
         this.update();
         return true;
@@ -49,7 +49,7 @@ export class JobHandler {
         this.bumpCount = false;
         this.count = jobPageSize * 2;
         this.includePreflights = true;
-        this.preflightStartedByUser = undefined;
+        this.preflightStartedByUserId = undefined;
 
         // cancel any pending        
         for (let i = 0; i < this.cancelId; i++) {
@@ -91,7 +91,7 @@ export class JobHandler {
                 streamId: this.stream?.id,
                 template: this.templateNames,
                 includePreflight: this.includePreflights,
-                preflightStartedByUser: this.preflightStartedByUser
+                preflightStartedByUserId: this.preflightStartedByUserId
             };
 
             // instance property can mutate externally as we're async
@@ -322,7 +322,7 @@ export class JobHandler {
     includeBatches = false;
 
     includePreflights = true;
-    private preflightStartedByUser?: string;
+    private preflightStartedByUserId?: string;
 
     private jobLimit?:number;
 
diff --git a/Engine/Source/Programs/Horde/HordeDashboard/src/base/utilities/issueUtils.ts b/Engine/Source/Programs/Horde/HordeDashboard/src/base/utilities/issueUtils.ts
--- a/Engine/Source/Programs/Horde/HordeDashboard/src/base/utilities/issueUtils.ts
+++ b/Engine/Source/Programs/Horde/HordeDashboard/src/base/utilities/issueUtils.ts
@@ -33,10 +33,10 @@ export function getIssueStatus(issue: GetIssueResponse, showResolveTime?:boolean
 
 	}
 
-	if (!issue.owner) {
+	if (!issue.ownerId) {
 		text = "Currently unassigned.";
 	} else {
-		if (issue.owner.toLowerCase() === dashboard.username?.toLowerCase()) {
+		if (issue.ownerId === dashboard.userId) {
 			if (issue.nominatedBy) {
 				text = `You have been nominated to fix this issue by ${issue.nominatedBy}.`;
 			} else {
diff --git a/Engine/Source/Programs/Horde/HordeServer/Collections/IJobCollection.cs b/Engine/Source/Programs/Horde/HordeServer/Collections/IJobCollection.cs
--- a/Engine/Source/Programs/Horde/HordeServer/Collections/IJobCollection.cs
+++ b/Engine/Source/Programs/Horde/HordeServer/Collections/IJobCollection.cs
@@ -103,7 +103,7 @@ namespace HordeServer.Collections
 		/// <param name="Index">Index of the first result to return</param>
 		/// <param name="Count">Number of results to return</param>
 		/// <returns>List of jobs matching the given criteria</returns>
-		Task<List<IJob>> FindAsync(ObjectId[]? JobIds = null, StreamId? StreamId = null, string? Name = null, TemplateRefId[]? Templates = null, int? MinChange = null, int? MaxChange = null, int? PreflightChange = null, string? PreflightStartedByUser = null, DateTimeOffset? MinCreateTime = null, DateTimeOffset? MaxCreateTime = null, DateTimeOffset? ModifiedBefore = null, DateTimeOffset? ModifiedAfter = null, int? Index = null, int? Count = null);
+		Task<List<IJob>> FindAsync(ObjectId[]? JobIds = null, StreamId? StreamId = null, string? Name = null, TemplateRefId[]? Templates = null, int? MinChange = null, int? MaxChange = null, int? PreflightChange = null, ObjectId? PreflightStartedByUser = null, DateTimeOffset? MinCreateTime = null, DateTimeOffset? MaxCreateTime = null, DateTimeOffset? ModifiedBefore = null, DateTimeOffset? ModifiedAfter = null, int? Index = null, int? Count = null);
 
 		/// <summary>
 		/// Updates a new job
diff --git a/Engine/Source/Programs/Horde/HordeServer/Collections/Impl/JobCollection.cs b/Engine/Source/Programs/Horde/HordeServer/Collections/Impl/JobCollection.cs
--- a/Engine/Source/Programs/Horde/HordeServer/Collections/Impl/JobCollection.cs
+++ b/Engine/Source/Programs/Horde/HordeServer/Collections/Impl/JobCollection.cs
@@ -408,7 +408,7 @@ namespace HordeServer.Collections.Impl
 		}
 
 		/// <inheritdoc/>
-		public async Task<List<IJob>> FindAsync(ObjectId[]? JobIds, StreamId? StreamId, string? Name, TemplateRefId[]? Templates, int? MinChange, int? MaxChange, int? PreflightChange, string? PreflightStartedByUser, DateTimeOffset? MinCreateTime, DateTimeOffset? MaxCreateTime, DateTimeOffset? ModifiedBefore, DateTimeOffset? ModifiedAfter, int? Index, int? Count)
+		public async Task<List<IJob>> FindAsync(ObjectId[]? JobIds, StreamId? StreamId, string? Name, TemplateRefId[]? Templates, int? MinChange, int? MaxChange, int? PreflightChange, ObjectId? PreflightStartedByUser, DateTimeOffset? MinCreateTime, DateTimeOffset? MaxCreateTime, DateTimeOffset? ModifiedBefore, DateTimeOffset? ModifiedAfter, int? Index, int? Count)
 		{
 			FilterDefinitionBuilder<JobDocument> FilterBuilder = Builders<JobDocument>.Filter;
 
@@ -443,7 +443,7 @@ namespace HordeServer.Collections.Impl
 			}
 			if (PreflightStartedByUser != null)
 			{
-				Filter &= FilterBuilder.Or(FilterBuilder.Eq(x => x.PreflightChange, 0), FilterBuilder.Eq(x => x.StartedByUser, PreflightStartedByUser));
+				Filter &= FilterBuilder.Or(FilterBuilder.Eq(x => x.PreflightChange, 0), FilterBuilder.Eq(x => x.StartedByUserId, PreflightStartedByUser));
 			}
 			if (MinCreateTime != null)
 			{
diff --git a/Engine/Source/Programs/Horde/HordeServer/Controllers/IssuesController.cs b/Engine/Source/Programs/Horde/HordeServer/Controllers/IssuesController.cs
--- a/Engine/Source/Programs/Horde/HordeServer/Controllers/IssuesController.cs
+++ b/Engine/Source/Programs/Horde/HordeServer/Controllers/IssuesController.cs
@@ -92,7 +92,7 @@ namespace HordeServer.Controllers
 		/// <param name="BatchId">The batch to filter by</param>
 		/// <param name="StepId">The step to filter by</param>
 		/// <param name="LabelIdx">The label within the job to filter by</param>
-		/// <param name="User">User to filter issues for</param>
+		/// <param name="UserId">User to filter issues for</param>
 		/// <param name="Resolved">Whether to include resolved issues</param>
 		/// <param name="Index">Starting offset of the window of results to return</param>
 		/// <param name="Count">Number of results to return</param>
@@ -101,7 +101,7 @@ namespace HordeServer.Controllers
 		[HttpGet]
 		[Route("/api/v1/issues")]
 		[ProducesResponseType(typeof(List<GetIssueResponse>), 200)]
-		public async Task<ActionResult<object>> FindIssuesAsync([FromQuery(Name = "Id")] int[]? Ids = null, [FromQuery] string? StreamId = null, [FromQuery] int? Change = null, [FromQuery] int? MinChange = null, [FromQuery] int? MaxChange = null, [FromQuery] string? JobId = null, [FromQuery] string? BatchId = null, [FromQuery] string? StepId = null, [FromQuery(Name = "label")] int? LabelIdx = null, [FromQuery] string? User = null, [FromQuery] bool? Resolved = null, [FromQuery] int Index = 0, [FromQuery] int Count = 10, [FromQuery] PropertyFilter? Filter = null)
+		public async Task<ActionResult<object>> FindIssuesAsync([FromQuery(Name = "Id")] int[]? Ids = null, [FromQuery] string? StreamId = null, [FromQuery] int? Change = null, [FromQuery] int? MinChange = null, [FromQuery] int? MaxChange = null, [FromQuery] string? JobId = null, [FromQuery] string? BatchId = null, [FromQuery] string? StepId = null, [FromQuery(Name = "label")] int? LabelIdx = null, [FromQuery] string? UserId = null, [FromQuery] bool? Resolved = null, [FromQuery] int Index = 0, [FromQuery] int Count = 10, [FromQuery] PropertyFilter? Filter = null)
 		{
 			if(Ids != null && Ids.Length == 0)
 			{
@@ -109,13 +109,9 @@ namespace HordeServer.Controllers
 			}
 
 			ObjectId? UserIdValue = null;
-			if (User != null)
+			if (UserId != null)
 			{
-				UserIdValue = (await UserCollection.FindUserByLoginAsync(User))?.Id;
-				if (UserIdValue == null)
-				{
-					return new List<IIssue>();
-				}
+				UserIdValue = new ObjectId(UserId);
 			}
 
 			List<IIssue> Issues;
diff --git a/Engine/Source/Programs/Horde/HordeServer/Controllers/JobsController.cs b/Engine/Source/Programs/Horde/HordeServer/Controllers/JobsController.cs
--- a/Engine/Source/Programs/Horde/HordeServer/Controllers/JobsController.cs
+++ b/Engine/Source/Programs/Horde/HordeServer/Controllers/JobsController.cs
@@ -512,7 +512,7 @@ namespace HordeServer.Controllers
 		/// <param name="MaxChange">The maximum changelist number</param>
 		/// <param name="IncludePreflight">Whether to include preflight jobs</param>
 		/// <param name="PreflightChange">The preflighted changelist</param>
-		/// <param name="PreflightStartedByUser">User for which to include preflight jobs</param>
+		/// <param name="PreflightStartedByUserId">User id for which to include preflight jobs</param>
 		/// <param name="MinCreateTime">Minimum creation time</param>
 		/// <param name="MaxCreateTime">Maximum creation time</param>
 		/// <param name="ModifiedBefore">If specified, only jobs updated before the give time will be returned</param>
@@ -536,7 +536,7 @@ namespace HordeServer.Controllers
 			[FromQuery] int? MaxChange = null,
 			[FromQuery] bool IncludePreflight = true,
 			[FromQuery] int? PreflightChange = null,
-			[FromQuery] string? PreflightStartedByUser = null,
+			[FromQuery] string? PreflightStartedByUserId = null,
 			[FromQuery] DateTimeOffset? MinCreateTime = null,
 			[FromQuery] DateTimeOffset? MaxCreateTime = null,
 			[FromQuery] DateTimeOffset? ModifiedBefore = null,
@@ -558,11 +558,18 @@ namespace HordeServer.Controllers
 				PreflightChange = 0;
 			}
 
+			ObjectId? PreflightStartedByUserIdValue = null;
+
+			if (PreflightStartedByUserId != null)
+			{
+				PreflightStartedByUserIdValue = new ObjectId(PreflightStartedByUserId);
+			}
+
 			List<IJob> Jobs;
 			using (Scope _ = Tracer.Instance.StartActive("FindJobs"))
 			{
 				Jobs = await JobService.FindJobsAsync(JobIdValues, StreamIdValue, Name, TemplateRefIds, MinChange,
-					MaxChange, PreflightChange, PreflightStartedByUser, MinCreateTime?.UtcDateTime, MaxCreateTime?.UtcDateTime, Target, State, Outcome,
+					MaxChange, PreflightChange, PreflightStartedByUserIdValue, MinCreateTime?.UtcDateTime, MaxCreateTime?.UtcDateTime, Target, State, Outcome,
 					ModifiedBefore, ModifiedAfter, Index, Count);
 			}
 
diff --git a/Engine/Source/Programs/Horde/HordeServer/Services/JobService.cs b/Engine/Source/Programs/Horde/HordeServer/Services/JobService.cs
--- a/Engine/Source/Programs/Horde/HordeServer/Services/JobService.cs
+++ b/Engine/Source/Programs/Horde/HordeServer/Services/JobService.cs
@@ -497,7 +497,7 @@ namespace HordeServer.Services
 		/// <param name="Index">Index of the first result to return</param>
 		/// <param name="Count">Number of results to return</param>
 		/// <returns>List of jobs matching the given criteria</returns>
-		public async Task<List<IJob>> FindJobsAsync(ObjectId[]? JobIds = null, StreamId? StreamId = null, string? Name = null, TemplateRefId[]? Templates = null, int? MinChange = null, int? MaxChange = null, int? PreflightChange = null, string? PreflightStartedByUser = null, DateTimeOffset? MinCreateTime = null, DateTimeOffset? MaxCreateTime = null, string? Target = null, JobStepState[]? State = null, JobStepOutcome[]? Outcome = null, DateTimeOffset? ModifiedBefore = null, DateTimeOffset? ModifiedAfter = null, int? Index = null, int? Count = null)
+		public async Task<List<IJob>> FindJobsAsync(ObjectId[]? JobIds = null, StreamId? StreamId = null, string? Name = null, TemplateRefId[]? Templates = null, int? MinChange = null, int? MaxChange = null, int? PreflightChange = null, ObjectId? PreflightStartedByUser = null, DateTimeOffset? MinCreateTime = null, DateTimeOffset? MaxCreateTime = null, string? Target = null, JobStepState[]? State = null, JobStepOutcome[]? Outcome = null, DateTimeOffset? ModifiedBefore = null, DateTimeOffset? ModifiedAfter = null, int? Index = null, int? Count = null)
 		{
 			if (Target == null && (State == null || State.Length == 0) && (Outcome == null || Outcome.Length == 0))
 			{
