Horde: Catch and log exceptions in background run in REAPI scheduling

diff --git a/Engine/Source/Programs/Horde/HordeServer/Tasks/Impl/ActionTaskSource.cs b/Engine/Source/Programs/Horde/HordeServer/Tasks/Impl/ActionTaskSource.cs
--- a/Engine/Source/Programs/Horde/HordeServer/Tasks/Impl/ActionTaskSource.cs
+++ b/Engine/Source/Programs/Horde/HordeServer/Tasks/Impl/ActionTaskSource.cs
@@ -369,8 +369,17 @@ namespace HordeServer.Tasks.Impl
 		/// <returns></returns>
 		async Task RunAsync(ExecuteOperation Operation)
 		{
-			ExecuteResponse Response = await RunInternalAsync(Operation);
-			Operation.SetStatus(ExecutionStage.Types.Value.Completed, Any.Pack(Response));
+			try
+			{
+				ExecuteResponse Response = await RunInternalAsync(Operation);
+				Operation.SetStatus(ExecutionStage.Types.Value.Completed, Any.Pack(Response));
+			}
+			catch (Exception e)
+			{
+				Logger.LogError(e, "Exception in background run");
+				throw;
+			}
+			
 		}
 
 		/// <summary>
