Horde: Add basic test for ActionTaskSource

diff --git a/Engine/Source/Programs/Horde/HordeServerTests/RemoteExec/ActionTaskSourceTest.cs b/Engine/Source/Programs/Horde/HordeServerTests/RemoteExec/ActionTaskSourceTest.cs
new file mode 100644
--- /dev/null
+++ b/Engine/Source/Programs/Horde/HordeServerTests/RemoteExec/ActionTaskSourceTest.cs
@@ -0,0 +1,97 @@
+﻿// Copyright Epic Games, Inc. All Rights Reserved.
+
+extern alias HordeAgent;
+using System;
+using System.Threading.Tasks;
+using Build.Bazel.Remote.Execution.V2;
+using Google.LongRunning;
+using HordeServer.Models;
+using HordeServer.Services;
+using HordeServer.Tasks.Impl;
+using Microsoft.Extensions.Logging;
+using Microsoft.Extensions.Logging.Console;
+using Microsoft.VisualStudio.TestTools.UnitTesting;
+
+namespace HordeServerTests.RemoteExec
+{
+	[TestClass]
+	public class ActionTaskSourceTest : DatabaseIntegrationTest
+	{
+		private readonly TestSetup TestSetup;
+
+		public ActionTaskSourceTest()
+		{
+			TestSetup = GetTestSetup().Result;
+		}
+
+		[TestMethod]
+		public async Task BasicExecution()
+		{
+			ActionTaskSource TaskSource = CreateTaskSource();
+
+			FakeAgent Agent = new FakeAgent();
+			await TaskSource.SubscribeAsync(Agent);
+			
+			ExecuteRequest Req1 = new ExecuteRequest();
+			Req1.SkipCacheLookup = true;
+			Req1.ActionDigest = new Digest { Hash = "bogusTestingHash1", SizeBytes = 17 };
+			
+			ExecuteRequest Req2 = new ExecuteRequest();
+			Req2.SkipCacheLookup = true;
+			Req2.ActionDigest = new Digest { Hash = "bogusTestingHash2", SizeBytes = 17 };
+
+			IActionExecuteOperation ExecuteOp1 = TaskSource.Execute(Req1);
+			IActionExecuteOperation ExecuteOp2 = TaskSource.Execute(Req2);
+
+			Task Pump1 = PumpExecuteOperation(TaskSource, ExecuteOp1, 10, Agent);
+			Task Pump2 = PumpExecuteOperation(TaskSource, ExecuteOp2, 20, Agent);
+
+			await Task.WhenAll(Pump1, Pump2);
+		}
+		
+		private async Task<Operation> PumpExecuteOperation(ActionTaskSource TaskSource, IActionExecuteOperation ExecuteOp, int ExitCode, IAgent? AgentToReschedule)
+		{
+			Operation? LastOperation = null;
+			await foreach (Operation Op in ExecuteOp.ReadStatusUpdatesAsync())
+			{
+				LastOperation = Op;
+				Assert.IsTrue(Op.Metadata.TryUnpack(out ExecuteOperationMetadata Metadata));
+
+				string Msg = $"id={Op.Name} digest={Metadata.ActionDigest.Hash} stage={Metadata.Stage}";
+
+				if (Metadata.Stage == ExecutionStage.Types.Value.Executing)
+				{
+					ExecuteOp.TrySetResult(new ActionResult {ExitCode = ExitCode});
+				}
+				
+				if (Op.Done)
+				{
+					Assert.IsTrue(Op.Response.TryUnpack(out ExecuteResponse Response));
+					Assert.AreEqual(ExitCode, Response.Result.ExitCode);
+					Msg += " result=" + Response;
+				}
+				
+				Console.WriteLine(Msg);
+			}
+
+			if (AgentToReschedule != null)
+			{
+				// Schedule the agent again so the task source can execute actions on it
+				await TaskSource.SubscribeAsync(AgentToReschedule);
+			}
+
+			return LastOperation!;
+		}
+
+		private ActionTaskSource CreateTaskSource()
+		{
+			ILogFileService LogFileService = (TestSetup.ServiceProvider.GetService(typeof(ILogFileService)) as ILogFileService)!;
+			ILoggerFactory LoggerFactory = new LoggerFactory();
+			ConsoleLoggerOptions LoggerOptions = new ConsoleLoggerOptions();
+			TestOptionsMonitor<ConsoleLoggerOptions> LoggerOptionsMon = new TestOptionsMonitor<ConsoleLoggerOptions>(LoggerOptions);
+			LoggerFactory.AddProvider(new ConsoleLoggerProvider(LoggerOptionsMon));
+			ILogger<ActionTaskSource> Logger = LoggerFactory.CreateLogger<ActionTaskSource>();
+			return new ActionTaskSource(null!, LogFileService, Logger, TestSetup.ServerSettingsMon);
+		}
+	}
+}
\ No newline at end of file
