Horde: Add setting for specifying which agent pool to use per REAPI instance

diff --git a/Engine/Source/Programs/Horde/HordeServer/ServerSettings.cs b/Engine/Source/Programs/Horde/HordeServer/ServerSettings.cs
--- a/Engine/Source/Programs/Horde/HordeServer/ServerSettings.cs
+++ b/Engine/Source/Programs/Horde/HordeServer/ServerSettings.cs
@@ -81,6 +81,12 @@ namespace HordeServer
 		/// </summary>
 		[Required]
 		public string? ServiceAccountToken { get; set; }
+		
+		/// <summary>
+		/// List of pool IDs this instance can be scheduled in
+		/// </summary>
+		[Required]
+		public List<string>? PoolFilter { get; set; }
 	}
 	
 	/// <summary>
diff --git a/Engine/Source/Programs/Horde/HordeServer/Tasks/Impl/ActionTaskSource.cs b/Engine/Source/Programs/Horde/HordeServer/Tasks/Impl/ActionTaskSource.cs
--- a/Engine/Source/Programs/Horde/HordeServer/Tasks/Impl/ActionTaskSource.cs
+++ b/Engine/Source/Programs/Horde/HordeServer/Tasks/Impl/ActionTaskSource.cs
@@ -25,6 +25,8 @@ using System.Threading;
 using System.Threading.Tasks;
 using Microsoft.Extensions.Options;
 
+using PoolId = HordeServer.Utilities.StringId<HordeServer.Models.IPool>;
+
 namespace HordeServer.Tasks.Impl
 {
 	/// <summary>
@@ -409,28 +411,17 @@ namespace HordeServer.Tasks.Impl
 			{
 				Operation.SetStatus(ExecutionStage.Types.Value.Queued, null);
 
-				// Get a new subscription
-				Subscription? Subscription = await WaitForSubscriberAsync(TimeSpan.FromSeconds(20.0));
-				if (Subscription == null)
-				{
-					ExecuteResponse FailedResponse = new ExecuteResponse();
-					FailedResponse.Status = new Status(StatusCode.ResourceExhausted, "Unable to find agent to execute request");
-					return FailedResponse;
-				}
-
-				// Create a lease for it
-				string LeaseName = $"Remote action ({Request.ActionDigest.Hash.Substring(0, 16)}...)";
-
-				ILogFile LogFile = await LogFileService.CreateLogFileAsync(ObjectId.Empty, Subscription.Agent.SessionId, LogType.Json);
-
+				List<string>? PoolFilter = null;
+				
 				ActionTask ActionTask = new ActionTask();
 				ActionTask.InstanceName = Request.InstanceName;
 				ActionTask.Digest = Request.ActionDigest;
-				ActionTask.LogId = LogFile.Id.ToString();
 				ActionTask.QueueTime = Timestamp.FromDateTimeOffset(QueueTime);
-
+				
 				if (Request.InstanceName != null && RemoteExecSettings.Instances.TryGetValue(Request.InstanceName, out var InstanceSettings))
 				{
+					PoolFilter = InstanceSettings.PoolFilter;
+					
 					if (InstanceSettings.CasUrl != null)
 					{
 						ActionTask.CasUrl = InstanceSettings.CasUrl.ToString();
@@ -446,6 +437,21 @@ namespace HordeServer.Tasks.Impl
 						ActionTask.ServiceAccountToken = InstanceSettings.ServiceAccountToken;
 					}
 				}
+				
+				// Get a new subscription
+				Subscription? Subscription = await WaitForSubscriberAsync(TimeSpan.FromSeconds(20.0), PoolFilter);
+				if (Subscription == null)
+				{
+					ExecuteResponse FailedResponse = new ExecuteResponse();
+					FailedResponse.Status = new Status(StatusCode.ResourceExhausted, "Unable to find agent to execute request");
+					return FailedResponse;
+				}
+
+				// Create a lease for it
+				string LeaseName = $"Remote action ({Request.ActionDigest.Hash.Substring(0, 16)}...)";
+
+				ILogFile LogFile = await LogFileService.CreateLogFileAsync(ObjectId.Empty, Subscription.Agent.SessionId, LogType.Json);
+				ActionTask.LogId = LogFile.Id.ToString();
 
 				byte[] Payload = Any.Pack(ActionTask).ToByteArray();
 				AgentLease Lease = new AgentLease(Operation.Id, LeaseName, null, null, LogFile.Id, LeaseState.Pending, Payload, new AgentRequirements(), null);
@@ -479,8 +485,9 @@ namespace HordeServer.Tasks.Impl
 		/// Waits until there is an available agent for executing work
 		/// </summary>
 		/// <param name="TimeLimit"></param>
+		/// <param name="PoolFilter"></param>
 		/// <returns></returns>
-		async Task<Subscription?> WaitForSubscriberAsync(TimeSpan TimeLimit)
+		async Task<Subscription?> WaitForSubscriberAsync(TimeSpan TimeLimit, List<string>? PoolFilter)
 		{
 			Stopwatch Timer = Stopwatch.StartNew();
 			for (; ; )
@@ -489,7 +496,7 @@ namespace HordeServer.Tasks.Impl
 				lock (Subscriptions)
 				{
 					Subscription? Subscription = Subscriptions.FirstOrDefault();
-					if (Subscription != null)
+					if (Subscription != null && CanBeScheduledOnAgent(Subscription.Agent, PoolFilter))
 					{
 						Subscriptions.Remove(Subscription);
 						return Subscription;
@@ -597,5 +604,10 @@ namespace HordeServer.Tasks.Impl
 		{
 			return Task.CompletedTask;
 		}
+		
+		internal static bool CanBeScheduledOnAgent(IAgent Agent, List<string>? PoolFilter)
+		{
+			return PoolFilter == null || Agent.ExplicitPools.Any(PoolId => PoolFilter.Contains(PoolId.ToString()));
+		}
 	}
 }
diff --git a/Engine/Source/Programs/Horde/HordeServerTests/RemoteExec/ActionTaskSourceTest.cs b/Engine/Source/Programs/Horde/HordeServerTests/RemoteExec/ActionTaskSourceTest.cs
--- a/Engine/Source/Programs/Horde/HordeServerTests/RemoteExec/ActionTaskSourceTest.cs
+++ b/Engine/Source/Programs/Horde/HordeServerTests/RemoteExec/ActionTaskSourceTest.cs
@@ -2,16 +2,22 @@
 
 extern alias HordeAgent;
 using System;
+using System.Collections.Generic;
 using System.Threading.Tasks;
+using Amazon.EC2.Model;
 using Build.Bazel.Remote.Execution.V2;
 using Google.LongRunning;
+using HordeServer;
 using HordeServer.Models;
 using HordeServer.Services;
 using HordeServer.Tasks.Impl;
+using HordeServer.Utilities;
 using Microsoft.Extensions.Logging;
 using Microsoft.Extensions.Logging.Console;
 using Microsoft.VisualStudio.TestTools.UnitTesting;
 
+using PoolId = HordeServer.Utilities.StringId<HordeServer.Models.IPool>;
+
 namespace HordeServerTests.RemoteExec
 {
 	[TestClass]
@@ -49,6 +55,48 @@ namespace HordeServerTests.RemoteExec
 			await Task.WhenAll(Pump1, Pump2);
 		}
 		
+		[TestMethod]
+		public async Task BasicExecutionWithPoolFilter()
+		{
+			RemoteExecInstanceSettings InstanceSettings = new RemoteExecInstanceSettings();
+			InstanceSettings.PoolFilter = new List<string> {"pool1"};
+			TestSetup.ServerSettingsMon.CurrentValue.RemoteExecSettings.Instances["foobar"] = InstanceSettings; 
+			ActionTaskSource TaskSource = CreateTaskSource();
+
+			FakeAgent Agent = new FakeAgent();
+			Agent.ExplicitPools = new List<PoolId> {new PoolId("pool1")};
+			await TaskSource.SubscribeAsync(Agent);
+			
+			ExecuteRequest Req1 = new ExecuteRequest();
+			Req1.InstanceName = "foobar";
+			Req1.SkipCacheLookup = true;
+			Req1.ActionDigest = new Digest { Hash = "bogusTestingHash1", SizeBytes = 17 };
+
+			IActionExecuteOperation ExecuteOp1 = TaskSource.Execute(Req1);
+
+			await PumpExecuteOperation(TaskSource, ExecuteOp1, 10, Agent);
+		}
+
+		[TestMethod]
+		public void CanBeScheduledOnAgent()
+		{
+			FakeAgent AgentWithNoPools = new FakeAgent();
+			FakeAgent Agent = new FakeAgent();
+			Agent.ExplicitPools = new List<StringId<IPool>> {new PoolId("pool1"), new PoolId("pool2")};
+
+			Assert.IsTrue(ActionTaskSource.CanBeScheduledOnAgent(Agent, new List<string> {"pool1"}));
+			Assert.IsTrue(ActionTaskSource.CanBeScheduledOnAgent(Agent, new List<string> {"pool2"}));
+			Assert.IsTrue(ActionTaskSource.CanBeScheduledOnAgent(Agent, new List<string> {"pool1", "pool2"}));
+			Assert.IsTrue(ActionTaskSource.CanBeScheduledOnAgent(Agent, new List<string> {"bogus", "pool2"}));
+			Assert.IsFalse(ActionTaskSource.CanBeScheduledOnAgent(AgentWithNoPools, new List<string> {"bogus", "pool2"}));
+			Assert.IsFalse(ActionTaskSource.CanBeScheduledOnAgent(Agent, new List<string> {"bogus", "foobar"}));
+			Assert.IsFalse(ActionTaskSource.CanBeScheduledOnAgent(Agent, new List<string>()));
+			Assert.IsFalse(ActionTaskSource.CanBeScheduledOnAgent(AgentWithNoPools, new List<string>()));
+			
+			Assert.IsTrue(ActionTaskSource.CanBeScheduledOnAgent(Agent, null));
+			Assert.IsTrue(ActionTaskSource.CanBeScheduledOnAgent(AgentWithNoPools, null));
+		}
+		
 		private async Task<Operation> PumpExecuteOperation(ActionTaskSource TaskSource, IActionExecuteOperation ExecuteOp, int ExitCode, IAgent? AgentToReschedule)
 		{
 			Operation? LastOperation = null;
