Horde: Lower timeout waiting for agent to become available in REAPI actions

Should be moved to the instance config, but this will do for now.

diff --git a/Engine/Source/Programs/Horde/HordeServer/Tasks/Impl/ActionTaskSource.cs b/Engine/Source/Programs/Horde/HordeServer/Tasks/Impl/ActionTaskSource.cs
--- a/Engine/Source/Programs/Horde/HordeServer/Tasks/Impl/ActionTaskSource.cs
+++ b/Engine/Source/Programs/Horde/HordeServer/Tasks/Impl/ActionTaskSource.cs
@@ -439,7 +439,7 @@ namespace HordeServer.Tasks.Impl
 				}
 				
 				// Get a new subscription
-				Subscription? Subscription = await WaitForSubscriberAsync(TimeSpan.FromSeconds(20.0), PoolFilter);
+				Subscription? Subscription = await WaitForSubscriberAsync(TimeSpan.FromSeconds(5.0), PoolFilter);
 				if (Subscription == null)
 				{
 					ExecuteResponse FailedResponse = new ExecuteResponse();
