Horde: Allow empty pool filter for REAPI

diff --git a/Engine/Source/Programs/Horde/HordeServer/ServerSettings.cs b/Engine/Source/Programs/Horde/HordeServer/ServerSettings.cs
--- a/Engine/Source/Programs/Horde/HordeServer/ServerSettings.cs
+++ b/Engine/Source/Programs/Horde/HordeServer/ServerSettings.cs
@@ -85,7 +85,6 @@ namespace HordeServer
 		/// <summary>
 		/// List of pool IDs this instance can be scheduled in
 		/// </summary>
-		[Required]
 		public List<string>? PoolFilter { get; set; }
 	}
 	
diff --git a/Engine/Source/Programs/Horde/HordeServer/Tasks/Impl/ActionTaskSource.cs b/Engine/Source/Programs/Horde/HordeServer/Tasks/Impl/ActionTaskSource.cs
--- a/Engine/Source/Programs/Horde/HordeServer/Tasks/Impl/ActionTaskSource.cs
+++ b/Engine/Source/Programs/Horde/HordeServer/Tasks/Impl/ActionTaskSource.cs
@@ -607,7 +607,7 @@ namespace HordeServer.Tasks.Impl
 		
 		internal static bool CanBeScheduledOnAgent(IAgent Agent, List<string>? PoolFilter)
 		{
-			return PoolFilter == null || Agent.ExplicitPools.Any(PoolId => PoolFilter.Contains(PoolId.ToString()));
+			return PoolFilter == null || PoolFilter.Count == 0 || Agent.ExplicitPools.Any(PoolId => PoolFilter.Contains(PoolId.ToString()));
 		}
 	}
 }
diff --git a/Engine/Source/Programs/Horde/HordeServerTests/RemoteExec/ActionTaskSourceTest.cs b/Engine/Source/Programs/Horde/HordeServerTests/RemoteExec/ActionTaskSourceTest.cs
--- a/Engine/Source/Programs/Horde/HordeServerTests/RemoteExec/ActionTaskSourceTest.cs
+++ b/Engine/Source/Programs/Horde/HordeServerTests/RemoteExec/ActionTaskSourceTest.cs
@@ -90,8 +90,8 @@ namespace HordeServerTests.RemoteExec
 			Assert.IsTrue(ActionTaskSource.CanBeScheduledOnAgent(Agent, new List<string> {"bogus", "pool2"}));
 			Assert.IsFalse(ActionTaskSource.CanBeScheduledOnAgent(AgentWithNoPools, new List<string> {"bogus", "pool2"}));
 			Assert.IsFalse(ActionTaskSource.CanBeScheduledOnAgent(Agent, new List<string> {"bogus", "foobar"}));
-			Assert.IsFalse(ActionTaskSource.CanBeScheduledOnAgent(Agent, new List<string>()));
-			Assert.IsFalse(ActionTaskSource.CanBeScheduledOnAgent(AgentWithNoPools, new List<string>()));
+			Assert.IsTrue(ActionTaskSource.CanBeScheduledOnAgent(Agent, new List<string>()));
+			Assert.IsTrue(ActionTaskSource.CanBeScheduledOnAgent(AgentWithNoPools, new List<string>()));
 			
 			Assert.IsTrue(ActionTaskSource.CanBeScheduledOnAgent(Agent, null));
 			Assert.IsTrue(ActionTaskSource.CanBeScheduledOnAgent(AgentWithNoPools, null));
