Horde: Fix property name on Slack user response.

diff --git a/Engine/Source/Programs/Horde/HordeServer/Notifications/Impl/SlackNotificationSink.cs b/Engine/Source/Programs/Horde/HordeServer/Notifications/Impl/SlackNotificationSink.cs
--- a/Engine/Source/Programs/Horde/HordeServer/Notifications/Impl/SlackNotificationSink.cs
+++ b/Engine/Source/Programs/Horde/HordeServer/Notifications/Impl/SlackNotificationSink.cs
@@ -93,7 +93,7 @@ namespace HordeServer.Notifications.Impl
 
 		class UserInfo
 		{
-			[JsonPropertyName("username")]
+			[JsonPropertyName("id")]
 			public string? UserName { get; set; }
 
 			[JsonPropertyName("profile")]
