Missing file.

diff --git a/Engine/Source/Programs/Horde/HordeServer/Program.cs b/Engine/Source/Programs/Horde/HordeServer/Program.cs
--- a/Engine/Source/Programs/Horde/HordeServer/Program.cs
+++ b/Engine/Source/Programs/Horde/HordeServer/Program.cs
@@ -14,9 +14,6 @@ using System.Threading.Tasks;
 using EpicGames.Core;
 using HordeServer.Models;
 using HordeServer.Utilities;
-using Json.Schema;
-using Json.Schema.Generation;
-using Json.Schema.Generation.Intents;
 using Microsoft.AspNetCore.Hosting;
 using Microsoft.AspNetCore.Server.Kestrel.Core;
 using Microsoft.Extensions.Configuration;
@@ -125,7 +122,7 @@ namespace HordeServer
 				foreach (Type SchemaType in ConfigSchemas)
 				{
 					FileReference OutputFile = FileReference.Combine(SchemaDir, $"{SchemaType.Name}.json");
-					Schemas.WriteSchema(SchemaType, OutputFile);
+					Schemas.CreateSchema(SchemaType).Write(OutputFile);
 				}
 				return;
 			}
