Horde: Hide users with the Hidden flag set when searching. These users appear to be invalid, but - due to a lack of validation for user names - may still be referenced by issues.

diff --git a/Engine/Source/Programs/Horde/HordeServer/Collections/Impl/UserCollectionV2.cs b/Engine/Source/Programs/Horde/HordeServer/Collections/Impl/UserCollectionV2.cs
--- a/Engine/Source/Programs/Horde/HordeServer/Collections/Impl/UserCollectionV2.cs
+++ b/Engine/Source/Programs/Horde/HordeServer/Collections/Impl/UserCollectionV2.cs
@@ -32,6 +32,9 @@ namespace HordeServer.Collections.Impl
 			public string? Email { get; set; }
 			public string? EmailUpper { get; set; }
 
+			[BsonIgnoreIfDefault]
+			public bool? Hidden { get; set; }
+
 			[BsonConstructor]
 			private UserDocument()
 			{
@@ -231,6 +234,8 @@ namespace HordeServer.Collections.Impl
 				Filter &= Builders<UserDocument>.Filter.Regex(x => x.Name, Regex);
 			}
 
+			Filter &= Builders<UserDocument>.Filter.Ne(x => x.Hidden, true);
+
 			return await Users.Find(Filter).Range(Index, Count ?? 100).ToListAsync<UserDocument, IUser>();
 		}
 
