Horde: Enable local graph caching for Chromium based browsers
#rnx

diff --git a/Engine/Source/Programs/Horde/HordeDashboard/src/backend/Dashboard.ts b/Engine/Source/Programs/Horde/HordeDashboard/src/backend/Dashboard.ts
--- a/Engine/Source/Programs/Horde/HordeDashboard/src/backend/Dashboard.ts
+++ b/Engine/Source/Programs/Horde/HordeDashboard/src/backend/Dashboard.ts
@@ -4,6 +4,12 @@ import { action, observable } from 'mobx';
 import backend from '.';
 import { UserClaim, DashboardPreference, GetUserResponse } from './Api';
 
+export enum WebBrowser {
+    Chromium = "Chromium",
+    Safari = "Safari", 
+    Other = "Other"
+}
+
 export class Dashboard {
 
     startPolling() {
@@ -85,6 +91,23 @@ export class Dashboard {
 
     }
 
+    get browser() : WebBrowser {
+
+        const agent = window.navigator.userAgent.toLowerCase();
+        
+        switch (true) {
+            case agent.indexOf("edge") > -1: return WebBrowser.Other;
+            case agent.indexOf("edg") > -1: return WebBrowser.Chromium;
+            case agent.indexOf("opr") > -1: return WebBrowser.Other ;
+            case agent.indexOf("chrome") > -1 && !!(window as any).chrome: return WebBrowser.Chromium;
+            case agent.indexOf("trident") > -1: return WebBrowser.Other;
+            case agent.indexOf("firefox") > -1: return WebBrowser.Other;
+            case agent.indexOf("safari") > -1: return WebBrowser.Safari;
+            default: return WebBrowser.Other;
+        }
+    
+    }
+
     get userId(): string {
         return this.data.id;
     }
@@ -147,7 +170,18 @@ export class Dashboard {
         this.setPreference(DashboardPreference.DisplayUTC, value ? "true" : "false");
     }
 
+    private hasLoggedLocalCache = false;
+
     get localCache(): boolean {
+        
+        if (this.browser === WebBrowser.Chromium) {
+            if (!this.hasLoggedLocalCache) {
+                this.hasLoggedLocalCache = true;
+                console.log("Chromium browser detected, local caching is enabled")
+            }
+            
+            return true;
+        }
 
         return this.preferences.get(DashboardPreference.LocalCache) === 'true';
 
@@ -313,6 +347,7 @@ export class Dashboard {
         return success;
 
     }
+    
 
     requestLogout = false;
 
