Horde: Add support for Perforce binary prints
#rnx

diff --git a/Engine/Source/Programs/Horde/HordeServer/Controllers/DebugController.cs b/Engine/Source/Programs/Horde/HordeServer/Controllers/DebugController.cs
--- a/Engine/Source/Programs/Horde/HordeServer/Controllers/DebugController.cs
+++ b/Engine/Source/Programs/Horde/HordeServer/Controllers/DebugController.cs
@@ -687,6 +687,8 @@ namespace HordeServer.Controllers
 				return Forbid();
 			}
 
+			StringBuilder Content = new StringBuilder();
+
 			const string DevBuild = "//UE4/Dev-Build";
 			// a test CL I have setup in Dev-Build
 			const int TestCL = 16242697;
@@ -771,7 +773,6 @@ namespace HordeServer.Controllers
 			Bytes = await Perforce.PrintAsync($"{DevBuild}/RunUAT.bat@13916380");
 			Results.Add($"PrintAsync: {DevBuild}/RunUAT.bat@13916380", $"{System.Text.Encoding.Default.GetString(Bytes)}");
 
-
 			// DuplicateShelvedChangeAsync is disabled due to issue with p4 librarian
 			/*
 			// need to be running as service account for these
@@ -817,13 +818,20 @@ namespace HordeServer.Controllers
 
 			}
 			*/
-			StringBuilder Content = new StringBuilder();
+
+			byte[] ImageBytes = await Perforce.PrintAsync($"{DevBuild}/Samples/Games/ShooterGame/ShooterGame.png");
+
+
 			Content.AppendLine("<html><body><pre>");
 			foreach (KeyValuePair<string, string> Pair in Results)
 			{
 				Content.AppendLine(HttpUtility.HtmlEncode($"{Pair.Key}={Pair.Value}"));
 			}
-			Content.Append("</pre></body></html>");
+
+			Content.Append($"PrintAsync Binary: {DevBuild}/Samples/Games/ShooterGame/ShooterGame.png\n</pre>");
+			Content.Append($"<img src=\"data:image/png;base64,{Convert.ToBase64String(ImageBytes, Base64FormattingOptions.None)}\"/>");
+
+			Content.Append("</body></html>");
 			return new ContentResult { ContentType = "text/html", StatusCode = (int)HttpStatusCode.OK, Content = Content.ToString() };
 
 		}
diff --git a/Engine/Source/Programs/Horde/HordeServer/Services/PerforceServiceApi.cs b/Engine/Source/Programs/Horde/HordeServer/Services/PerforceServiceApi.cs
--- a/Engine/Source/Programs/Horde/HordeServer/Services/PerforceServiceApi.cs
+++ b/Engine/Source/Programs/Horde/HordeServer/Services/PerforceServiceApi.cs
@@ -608,6 +608,11 @@ namespace HordeServer.Services
 					{
 						P4.P4CommandResult Result = Command.Run();
 
+						if (Result.BinaryOutput != null)
+						{
+							return Result.BinaryOutput;
+						}
+					
 						return Encoding.Default.GetBytes(Result.TextOutput);
 
 					}
