Horde: Moving more P4 api off python bridge
#rnx

diff --git a/Engine/Source/Programs/Horde/HordeServer/Services/PerforceServiceApi.cs b/Engine/Source/Programs/Horde/HordeServer/Services/PerforceServiceApi.cs
--- a/Engine/Source/Programs/Horde/HordeServer/Services/PerforceServiceApi.cs
+++ b/Engine/Source/Programs/Horde/HordeServer/Services/PerforceServiceApi.cs
@@ -682,19 +682,26 @@ namespace HordeServer.Services
 		{
 			await Task.Run(() =>
 			{
-				using (P4.Repository Repository = GetConnection(NoClient: true))
+				try
 				{
-					P4.Changelist Changelist = Repository.GetChangelist(Change);
+					using (P4.Repository Repository = GetConnection(NoClient: true))
+					{
+						P4.Changelist Changelist = Repository.GetChangelist(Change);
 
-					Repository.Connection.Disconnect();
+						Repository.Connection.Disconnect();
 
-					// the client must exist for the change list, otherwise will fail (for example, CreateNewChangeAsync deletes the client before returning)
-					using (P4.Repository UpdateRepository = GetConnection(ClientFromChange: Change, UseClientFromChange: true, Username: Changelist.OwnerName))
-					{
-						P4.Changelist UpdatedChangelist = UpdateRepository.GetChangelist(Change);
-						UpdatedChangelist.Description = Description;
-						UpdateRepository.UpdateChangelist(UpdatedChangelist);
+						// the client must exist for the change list, otherwise will fail (for example, CreateNewChangeAsync deletes the client before returning)
+						using (P4.Repository UpdateRepository = GetConnection(ClientFromChange: Change, UseClientFromChange: true, Username: Changelist.OwnerName))
+						{
+							P4.Changelist UpdatedChangelist = UpdateRepository.GetChangelist(Change);
+							UpdatedChangelist.Description = Description;
+							UpdateRepository.UpdateChangelist(UpdatedChangelist);
+						}
 					}
+				} 
+				catch (Exception Ex)
+				{
+					LogPerforce(1, "", $"Unable to update Changelist for CL {Change} to ${Description}, {Ex.Message}");
 				}
 			});
 		}
diff --git a/Engine/Source/Programs/Horde/HordeServer/Services/PerforceServiceRouter.cs b/Engine/Source/Programs/Horde/HordeServer/Services/PerforceServiceRouter.cs
--- a/Engine/Source/Programs/Horde/HordeServer/Services/PerforceServiceRouter.cs
+++ b/Engine/Source/Programs/Horde/HordeServer/Services/PerforceServiceRouter.cs
@@ -61,33 +61,27 @@ namespace HordeServer.Services
 		}
 
 		/// <inheritdoc/>
-		override public Task<int> CreateNewChangeAsync(string StreamName, string Path)
+		override public Task<byte[]> PrintAsync(string Path)
 		{
-			return PythonBridgeService.CreateNewChangeAsync(StreamName, Path);
+			return P4APIService.PrintAsync(Path);
 		}
 
 		/// <inheritdoc/>
 		override public Task<List<ChangeDetails>> GetChangeDetailsAsync(string StreamName, IReadOnlyList<int> ChangeNumbers, string? ImpersonateUser)
 		{
-			return PythonBridgeService.GetChangeDetailsAsync(StreamName, ChangeNumbers, ImpersonateUser);
+			return P4APIService.GetChangeDetailsAsync(StreamName, ChangeNumbers, ImpersonateUser);
 		}
 
 		/// <inheritdoc/>
 		override public Task<List<FileSummary>> FindFilesAsync(IEnumerable<string> Paths)
 		{
-			return PythonBridgeService.FindFilesAsync(Paths);
-		}
-
-		/// <inheritdoc/>
-		override public Task<byte[]> PrintAsync(string Path)
-		{
-			return PythonBridgeService.PrintAsync(Path);
+			return P4APIService.FindFilesAsync(Paths);
 		}
 
 		/// <inheritdoc/>
-		override public Task<int> DuplicateShelvedChangeAsync(int ShelvedChange)
+		override public Task<int> CreateNewChangeAsync(string StreamName, string Path)
 		{
-			throw new Exception("DuplicateShelvedChangeAsync is disabled due to p4 issue: https://jira.it.epicgames.com/servicedesk/customer/portal/1/ITH-144069");	
+			return PythonBridgeService.CreateNewChangeAsync(StreamName, Path);
 		}
 
 		/// <inheritdoc/>
@@ -108,6 +102,12 @@ namespace HordeServer.Services
 			return PythonBridgeService.UpdateChangelistDescription(Change, Description );
 		}
 
+		/// <inheritdoc/>
+		override public Task<int> DuplicateShelvedChangeAsync(int ShelvedChange)
+		{
+			throw new Exception("DuplicateShelvedChangeAsync is disabled due to p4 issue: https://jira.it.epicgames.com/servicedesk/customer/portal/1/ITH-144069");
+		}
+
 		/// <inheritdoc/>
 		public void Dispose()
 		{
