Horde: When submitting a preflight, return the error message instead of throwing
#rnx

diff --git a/Engine/Source/Programs/Horde/HordeServer/Services/PerforceServiceApi.cs b/Engine/Source/Programs/Horde/HordeServer/Services/PerforceServiceApi.cs
--- a/Engine/Source/Programs/Horde/HordeServer/Services/PerforceServiceApi.cs
+++ b/Engine/Source/Programs/Horde/HordeServer/Services/PerforceServiceApi.cs
@@ -888,13 +888,14 @@ namespace HordeServer.Services
 
 					   if (!Result.Success)
 					   {
-						   throw new Exception("Unable to get changes");
+							return (null, $"Unable to get change {Change}");
+						
 					   }
 
 					   if (Result.TaggedOutput == null || Result.TaggedOutput.Count != 2)
 					   {
 
-						   throw new Exception($"Unable to get tagged output for change: {Change} and original change: {OriginalChange}");
+							return (null, $"Unable to get tagged output for change: {Change} and original change: {OriginalChange}");
 					   }
 
 					   bool DSTMismatch = false;
@@ -914,12 +915,12 @@ namespace HordeServer.Services
 
 					   if (OriginalChangelist.ShelvedFiles.Count != Changelist.ShelvedFiles.Count)
 					   {
-						   throw new Exception($"Mismatched number of shelved files for change: {Change} and original change: {OriginalChange}");
+							return (null, $"Mismatched number of shelved files for change: {Change} and original change: {OriginalChange}");
 					   }
 
 					   if (OriginalChangelist.ShelvedFiles.Count == 0)
 					   {
-						   throw new Exception($"No shelved file for change: {Change} and original change: {OriginalChange}");
+							return (null, $"No shelved file for change: {Change} and original change: {OriginalChange}");
 					   }
 
 						// todo, test straight up delete
@@ -929,7 +930,7 @@ namespace HordeServer.Services
 
 						   if (Found == null)
 						   {
-							   throw new Exception($"Mismatch in shelved file digest or action for {ShelvedFile.Path}");
+								return (null, $"Mismatch in shelved file digest or action for {ShelvedFile.Path}");
 						   }
 					   }
 
@@ -944,7 +945,8 @@ namespace HordeServer.Services
 
 							   if (!Result.Success)
 							   {
-								   throw new Exception($"Unable to submit {Change}");
+									string Message = (Result.ErrorList != null && Result.ErrorList.Count > 0) ? Result.ErrorList[0].ErrorMessage : "Unknown error, no errors in list" ;
+									return (null, $"Unable to submit {Change}, {Message}");
 							   }
 
 							   int? SubmittedChangeId = null;
@@ -960,7 +962,7 @@ namespace HordeServer.Services
 
 							   if (SubmittedChangeId == null)
 							   {
-								   throw new Exception($"Unable to get submitted changelist");
+								   return (null, $"Submit command succeeded, though unable to parse submitted change number");
 							   }
 
 							   return (SubmittedChangeId, Changelist.Description);
