Horde: Fix dashboard warnings and remove SoftwareHome
#rnx

diff --git a/Engine/Source/Programs/Horde/HordeDashboard/src/backend/Config.ts b/Engine/Source/Programs/Horde/HordeDashboard/src/backend/Config.ts
--- a/Engine/Source/Programs/Horde/HordeDashboard/src/backend/Config.ts
+++ b/Engine/Source/Programs/Horde/HordeDashboard/src/backend/Config.ts
@@ -1,6 +1,6 @@
 // Copyright Epic Games, Inc. All Rights Reserved.
 
-import { initDatadog} from './Datadog';
+// import { initDatadog} from './Datadog';
 
 export type SiteConfig = {
     environment: "dev" | "production" | "dev-local" | "";
diff --git a/Engine/Source/Programs/Horde/HordeDashboard/src/backend/Fetch.ts b/Engine/Source/Programs/Horde/HordeDashboard/src/backend/Fetch.ts
--- a/Engine/Source/Programs/Horde/HordeDashboard/src/backend/Fetch.ts
+++ b/Engine/Source/Programs/Horde/HordeDashboard/src/backend/Fetch.ts
@@ -1,7 +1,7 @@
 // Copyright Epic Games, Inc. All Rights Reserved.
 
 import { ErrorHandler, ErrorInfo } from "../components/ErrorHandler";
-import { setDatadogUser } from './Datadog';
+//import { setDatadogUser } from './Datadog';
 
 export enum ChallengeStatus {
     Ok,
diff --git a/Engine/Source/Programs/Horde/HordeDashboard/src/backend/Issues.ts b/Engine/Source/Programs/Horde/HordeDashboard/src/backend/Issues.ts
--- a/Engine/Source/Programs/Horde/HordeDashboard/src/backend/Issues.ts
+++ b/Engine/Source/Programs/Horde/HordeDashboard/src/backend/Issues.ts
@@ -29,7 +29,7 @@ const getCompilerIssueSummary = (issue: GetIssueResponse, events: GetLogEventRes
 
                     if (relativePath) {
 
-                        const sourceFile = relativePath.replace(/^.*[\\\/]/, '');
+                        const sourceFile = relativePath.replace(/^.*[\\]/, '');
                         if (!sourceFiles.find(s => s === sourceFile)) {
                             sourceFiles.push(sourceFile);
                         }
diff --git a/Engine/Source/Programs/Horde/HordeDashboard/src/backend/LocalCache.ts b/Engine/Source/Programs/Horde/HordeDashboard/src/backend/LocalCache.ts
--- a/Engine/Source/Programs/Horde/HordeDashboard/src/backend/LocalCache.ts
+++ b/Engine/Source/Programs/Horde/HordeDashboard/src/backend/LocalCache.ts
@@ -3,11 +3,6 @@
 import * as CryptoJS from "crypto-js";
 import { openDB, DBSchema, IDBPDatabase } from 'idb';
 
-type CacheItem = {
-    storeTime: number;
-    encodedData: string;
-}
-
 interface CacheSchema extends DBSchema {
     'cache-item': {
         key: string;
