Horde: P4 submit command throws an exception, which we need to catch
#rnx

diff --git a/Engine/Source/Programs/Horde/HordeServer/Services/PerforceServiceApi.cs b/Engine/Source/Programs/Horde/HordeServer/Services/PerforceServiceApi.cs
--- a/Engine/Source/Programs/Horde/HordeServer/Services/PerforceServiceApi.cs
+++ b/Engine/Source/Programs/Horde/HordeServer/Services/PerforceServiceApi.cs
@@ -941,7 +941,14 @@ namespace HordeServer.Services
 							// we might not need a client here, possibly -e below facilitates this, check!
 							using (P4.P4Command SubmitCommand = new P4.P4Command(SubmitRepository, "submit", null, true, new string[] { "-e", Change.ToString(CultureInfo.InvariantCulture) }))
 						   {
-							   Result = SubmitCommand.Run();
+								try
+								{
+									Result = SubmitCommand.Run();
+								}
+								catch (Exception Ex)
+								{
+									return (null, $"Submit command failed: {Ex.Message}");
+								}
 
 							   if (!Result.Success)
 							   {
