Horde: Update file for new change p4 test
#rnx

diff --git a/Engine/Source/Programs/Horde/HordeServer/Controllers/DebugController.cs b/Engine/Source/Programs/Horde/HordeServer/Controllers/DebugController.cs
--- a/Engine/Source/Programs/Horde/HordeServer/Controllers/DebugController.cs
+++ b/Engine/Source/Programs/Horde/HordeServer/Controllers/DebugController.cs
@@ -702,7 +702,7 @@ namespace HordeServer.Controllers
 
 			Dictionary<string, string> Results = new Dictionary<string, string>();
 
-			int NewChangeId = await Perforce.CreateNewChangeAsync("//UE4/Private-Build", "Counter.txt");
+			int NewChangeId = await Perforce.CreateNewChangeAsync("//UE4/Private-Build", "HordeTestFile.txt");
 
 			List<ChangeDetails> ChangeDetails = await Perforce.GetChangeDetailsAsync("//UE4/Private-Build", new int[] { NewChangeId }, PerforceUser);
 
