Horde: Fix path for debugging new p4 changes, and add comment to avoid missing it
#rnx

diff --git a/Engine/Source/Programs/Horde/HordeServer/Controllers/DebugController.cs b/Engine/Source/Programs/Horde/HordeServer/Controllers/DebugController.cs
--- a/Engine/Source/Programs/Horde/HordeServer/Controllers/DebugController.cs
+++ b/Engine/Source/Programs/Horde/HordeServer/Controllers/DebugController.cs
@@ -702,7 +702,8 @@ namespace HordeServer.Controllers
 
 			Dictionary<string, string> Results = new Dictionary<string, string>();
 
-			int NewChangeId = await Perforce.CreateNewChangeAsync("//UE4/Private-Build", "HordeTestFile.txt");
+			// NOTE: This is in Private-Build not DevBuild
+			int NewChangeId = await Perforce.CreateNewChangeAsync("//UE4/Private-Build", "Counter.txt");
 
 			List<ChangeDetails> ChangeDetails = await Perforce.GetChangeDetailsAsync("//UE4/Private-Build", new int[] { NewChangeId }, PerforceUser);
 
