Horde: Limit number of concurrent leases per agent in ActionTaskSource

Currently hardcoded to just allow one lease per agent. Probably good to make this configurable later on.

diff --git a/Engine/Source/Programs/Horde/HordeServer/Tasks/Impl/ActionTaskSource.cs b/Engine/Source/Programs/Horde/HordeServer/Tasks/Impl/ActionTaskSource.cs
--- a/Engine/Source/Programs/Horde/HordeServer/Tasks/Impl/ActionTaskSource.cs
+++ b/Engine/Source/Programs/Horde/HordeServer/Tasks/Impl/ActionTaskSource.cs
@@ -596,6 +596,12 @@ namespace HordeServer.Tasks.Impl
 		/// <inheritdoc/>
 		public Task<ITaskListener?> SubscribeAsync(IAgent Agent)
 		{
+			if (Agent.Leases.Count > 0)
+			{
+				Logger.LogDebug("Cannot subscribe agent {AgentId} as it already has one or more leases. Num leases {Leases}", Agent.Id, Agent.Leases.Count);
+				return Task.FromResult<ITaskListener?>(null);
+			}
+			
 			return Task.FromResult<ITaskListener?>(new Subscription(this, Agent));
 		}
 
