Horde: Account for pool filter when finding suitable subscription in REAPI

diff --git a/Engine/Source/Programs/Horde/HordeServer/Tasks/Impl/ActionTaskSource.cs b/Engine/Source/Programs/Horde/HordeServer/Tasks/Impl/ActionTaskSource.cs
--- a/Engine/Source/Programs/Horde/HordeServer/Tasks/Impl/ActionTaskSource.cs
+++ b/Engine/Source/Programs/Horde/HordeServer/Tasks/Impl/ActionTaskSource.cs
@@ -495,8 +495,8 @@ namespace HordeServer.Tasks.Impl
 				// Try to allocate a subscription
 				lock (Subscriptions)
 				{
-					Subscription? Subscription = Subscriptions.FirstOrDefault();
-					if (Subscription != null && CanBeScheduledOnAgent(Subscription.Agent, PoolFilter))
+					Subscription? Subscription = Subscriptions.FirstOrDefault(s => CanBeScheduledOnAgent(s.Agent, PoolFilter));
+					if (Subscription != null)
 					{
 						Subscriptions.Remove(Subscription);
 						return Subscription;
