Horde: Support viewing step history and step artifacts in log view
#rnx

diff --git a/Engine/Source/Programs/Horde/HordeDashboard/src/backend/JobDetails.ts b/Engine/Source/Programs/Horde/HordeDashboard/src/backend/JobDetails.ts
--- a/Engine/Source/Programs/Horde/HordeDashboard/src/backend/JobDetails.ts
+++ b/Engine/Source/Programs/Horde/HordeDashboard/src/backend/JobDetails.ts
@@ -777,9 +777,9 @@ export class JobDetails {
                 stepName = groups[batch.groupIdx]?.nodes[stepNode.nodeIdx]?.name;
             }
 
+            requests.push(backend.getJobArtifacts(this.id!));
 
-            if (!this.isLogView) {
-                requests.push(backend.getJobArtifacts(this.id!));
+            if (!this.isLogView) {                
                 requests.push(backend.getJobTestData(this.id!));
                 requests.push(this.getIssues());
                 requests.push(this.queryReports());
