Horde: Make agent queue timeout and max concurrent leases configurable for REAPI

Longer term, max leases shouldn't be used but rather have actions specify required agent parameters. Such as available CPU.

diff --git a/Engine/Source/Programs/Horde/HordeServer/ServerSettings.cs b/Engine/Source/Programs/Horde/HordeServer/ServerSettings.cs
--- a/Engine/Source/Programs/Horde/HordeServer/ServerSettings.cs
+++ b/Engine/Source/Programs/Horde/HordeServer/ServerSettings.cs
@@ -86,6 +86,16 @@ namespace HordeServer
 		/// List of pool IDs this instance can be scheduled in
 		/// </summary>
 		public List<string>? PoolFilter { get; set; }
+
+		/// <summary>
+		/// Max time to wait for an agent to become available for serving a remote execution request (in milliseconds)
+		/// </summary>
+		public int AgentQueueTimeout { get; set; } = 5000;
+
+		public TimeSpan AgentQueueTimeoutAsTimeSpan()
+		{
+			return TimeSpan.FromMilliseconds(AgentQueueTimeout);
+		}
 	}
 	
 	/// <summary>
@@ -97,6 +107,11 @@ namespace HordeServer
 		/// Mapping instance names to instance settings
 		/// </summary>
 		public Dictionary<string, RemoteExecInstanceSettings> Instances { get; set; } = new Dictionary<string, RemoteExecInstanceSettings>();
+
+		/// <summary>
+		/// Max number of concurrent leases per agent
+		/// </summary>
+		public int MaxConcurrentLeasesPerAgent = 2;
 	}
 
 	/// <summary>
diff --git a/Engine/Source/Programs/Horde/HordeServer/Tasks/Impl/ActionTaskSource.cs b/Engine/Source/Programs/Horde/HordeServer/Tasks/Impl/ActionTaskSource.cs
--- a/Engine/Source/Programs/Horde/HordeServer/Tasks/Impl/ActionTaskSource.cs
+++ b/Engine/Source/Programs/Horde/HordeServer/Tasks/Impl/ActionTaskSource.cs
@@ -411,35 +411,21 @@ namespace HordeServer.Tasks.Impl
 			{
 				Operation.SetStatus(ExecutionStage.Types.Value.Queued, null);
 
-				List<string>? PoolFilter = null;
-				
 				ActionTask ActionTask = new ActionTask();
 				ActionTask.InstanceName = Request.InstanceName;
 				ActionTask.Digest = Request.ActionDigest;
 				ActionTask.QueueTime = Timestamp.FromDateTimeOffset(QueueTime);
-				
-				if (Request.InstanceName != null && RemoteExecSettings.Instances.TryGetValue(Request.InstanceName, out var InstanceSettings))
-				{
-					PoolFilter = InstanceSettings.PoolFilter;
-					
-					if (InstanceSettings.CasUrl != null)
-					{
-						ActionTask.CasUrl = InstanceSettings.CasUrl.ToString();
-					}
-					
-					if (InstanceSettings.ActionCacheUrl != null)
-					{
-						ActionTask.ActionCacheUrl = InstanceSettings.ActionCacheUrl.ToString();
-					}
+
+				RemoteExecInstanceSettings InstanceSettings =
+					RemoteExecSettings.Instances.GetValueOrDefault(Request.InstanceName ?? "") ??
+					new RemoteExecInstanceSettings();
 					
-					if (InstanceSettings.ServiceAccountToken != null)
-					{
-						ActionTask.ServiceAccountToken = InstanceSettings.ServiceAccountToken;
-					}
-				}
+				if (InstanceSettings.CasUrl != null) ActionTask.CasUrl = InstanceSettings.CasUrl.ToString();
+				if (InstanceSettings.ActionCacheUrl != null) ActionTask.ActionCacheUrl = InstanceSettings.ActionCacheUrl.ToString();
+				if (InstanceSettings.ServiceAccountToken != null) ActionTask.ServiceAccountToken = InstanceSettings.ServiceAccountToken;
 				
 				// Get a new subscription
-				Subscription? Subscription = await WaitForSubscriberAsync(TimeSpan.FromSeconds(20.0), PoolFilter);
+				Subscription? Subscription = await WaitForSubscriberAsync(InstanceSettings.AgentQueueTimeoutAsTimeSpan(), InstanceSettings.PoolFilter);
 				if (Subscription == null)
 				{
 					ExecuteResponse FailedResponse = new ExecuteResponse();
@@ -596,7 +582,7 @@ namespace HordeServer.Tasks.Impl
 		/// <inheritdoc/>
 		public Task<ITaskListener?> SubscribeAsync(IAgent Agent)
 		{
-			if (Agent.Leases.Count > 0)
+			if (Agent.Leases.Count >= RemoteExecSettings.MaxConcurrentLeasesPerAgent)
 			{
 				Logger.LogDebug("Cannot subscribe agent {AgentId} as it already has one or more leases. Num leases {Leases}", Agent.Id, Agent.Leases.Count);
 				return Task.FromResult<ITaskListener?>(null);
