Horde: Fix warnings

diff --git a/Engine/Source/Programs/Horde/HordeServer/ServerSettings.cs b/Engine/Source/Programs/Horde/HordeServer/ServerSettings.cs
--- a/Engine/Source/Programs/Horde/HordeServer/ServerSettings.cs
+++ b/Engine/Source/Programs/Horde/HordeServer/ServerSettings.cs
@@ -92,7 +92,7 @@ namespace HordeServer
 		/// </summary>
 		public int AgentQueueTimeout { get; set; } = 5000;
 
-		public TimeSpan AgentQueueTimeoutAsTimeSpan()
+		internal TimeSpan AgentQueueTimeoutAsTimeSpan()
 		{
 			return TimeSpan.FromMilliseconds(AgentQueueTimeout);
 		}
@@ -111,7 +111,7 @@ namespace HordeServer
 		/// <summary>
 		/// Max number of concurrent leases per agent
 		/// </summary>
-		public int MaxConcurrentLeasesPerAgent = 2;
+		public int MaxConcurrentLeasesPerAgent { get; set; } = 2;
 	}
 
 	/// <summary>
