Horde: Switch router to use new native perforce api exclusively
#rnx

diff --git a/Engine/Source/Programs/Horde/HordeServer/ServerSettings.cs b/Engine/Source/Programs/Horde/HordeServer/ServerSettings.cs
--- a/Engine/Source/Programs/Horde/HordeServer/ServerSettings.cs
+++ b/Engine/Source/Programs/Horde/HordeServer/ServerSettings.cs
@@ -405,6 +405,11 @@ namespace HordeServer
 		/// </summary>
 		public string? P4BridgeServicePassword { get; set; }
 
+		/// <summary>
+		/// Whether the p4 bridge service account can impersonate other users
+		/// </summary>
+		public bool P4BridgeCanImpersonate { get; set; } = false;
+
 		/// <summary>
 		/// Whether to use the (temporary) P4 router while transitioning to p4 api 
 		/// </summary>
diff --git a/Engine/Source/Programs/Horde/HordeServer/Services/PerforceServiceApi.cs b/Engine/Source/Programs/Horde/HordeServer/Services/PerforceServiceApi.cs
--- a/Engine/Source/Programs/Horde/HordeServer/Services/PerforceServiceApi.cs
+++ b/Engine/Source/Programs/Horde/HordeServer/Services/PerforceServiceApi.cs
@@ -923,7 +923,6 @@ namespace HordeServer.Services
 							return (null, $"No shelved file for change: {Change} and original change: {OriginalChange}");
 					   }
 
-						// todo, test straight up delete
 						foreach (P4.ShelvedFile ShelvedFile in Changelist.ShelvedFiles)
 					   {
 						   P4.ShelvedFile? Found = OriginalChangelist.ShelvedFiles.FirstOrDefault(Original => Original.Digest == ShelvedFile.Digest && Original.Action == ShelvedFile.Action);
@@ -1026,8 +1025,7 @@ namespace HordeServer.Services
 			get
 			{
 				ServerSettings Settings = this.Settings.CurrentValue;
-				// Todo: we need a better way of identifying we're using a service account
-				return Settings.P4BridgeServiceUsername == "SVC-P4-HordeProxy-P";
+				return Settings.P4BridgeCanImpersonate;
 			}
 		}
 
diff --git a/Engine/Source/Programs/Horde/HordeServer/Services/PerforceServiceRouter.cs b/Engine/Source/Programs/Horde/HordeServer/Services/PerforceServiceRouter.cs
--- a/Engine/Source/Programs/Horde/HordeServer/Services/PerforceServiceRouter.cs
+++ b/Engine/Source/Programs/Horde/HordeServer/Services/PerforceServiceRouter.cs
@@ -79,27 +79,27 @@ namespace HordeServer.Services
 		}
 
 		/// <inheritdoc/>
-		override public Task<int> CreateNewChangeAsync(string StreamName, string Path)
+		override public Task<(int? Change, string Message)> SubmitShelvedChangeAsync(int ShelvedChange, int OriginalChange)
 		{
-			return PythonBridgeService.CreateNewChangeAsync(StreamName, Path);
+			return P4APIService.SubmitShelvedChangeAsync(ShelvedChange, OriginalChange);
 		}
 
 		/// <inheritdoc/>
-		override public Task<(int? Change, string Message)> SubmitShelvedChangeAsync(int ShelvedChange, int OriginalChange)
+		override public Task UpdateChangelistDescription(int Change, string Description)
 		{
-			return PythonBridgeService.SubmitShelvedChangeAsync(ShelvedChange, OriginalChange);
+			return P4APIService.UpdateChangelistDescription(Change, Description);
 		}
 
 		/// <inheritdoc/>
 		override public Task DeleteShelvedChangeAsync(int ShelvedChange)
 		{
-			return PythonBridgeService.DeleteShelvedChangeAsync(ShelvedChange);
+			return P4APIService.DeleteShelvedChangeAsync(ShelvedChange);
 		}
 
 		/// <inheritdoc/>
-		override public Task UpdateChangelistDescription(int Change, string Description)
+		override public Task<int> CreateNewChangeAsync(string StreamName, string Path)
 		{
-			return PythonBridgeService.UpdateChangelistDescription(Change, Description );
+			return P4APIService.CreateNewChangeAsync(StreamName, Path);
 		}
 
 		/// <inheritdoc/>
