Horde: Fix debug controller change list
#rnx

diff --git a/Engine/Source/Programs/Horde/HordeServer/Controllers/DebugController.cs b/Engine/Source/Programs/Horde/HordeServer/Controllers/DebugController.cs
--- a/Engine/Source/Programs/Horde/HordeServer/Controllers/DebugController.cs
+++ b/Engine/Source/Programs/Horde/HordeServer/Controllers/DebugController.cs
@@ -691,7 +691,7 @@ namespace HordeServer.Controllers
 
 			const string DevBuild = "//UE4/Dev-Build";
 			// a test CL I have setup in Dev-Build
-			const int TestCL = 16242697;
+			const int TestCL = 16687685;
 
 			string? PerforceUser = User.GetPerforceUser();
 
