Horde: Add ability to change log level at runtime in agent

To be used from REAPI test client.

diff --git a/Engine/Source/Programs/Horde/HordeAgent/Logging.cs b/Engine/Source/Programs/Horde/HordeAgent/Logging.cs
--- a/Engine/Source/Programs/Horde/HordeAgent/Logging.cs
+++ b/Engine/Source/Programs/Horde/HordeAgent/Logging.cs
@@ -22,6 +22,8 @@ namespace HordeAgent
 		/// Current version of the HordeAgent
 		/// </summary>
 		public static string Version = "";
+		
+		public static LoggingLevelSwitch LogLevelSwitch =  new LoggingLevelSwitch();
 
 		private class VersionLogEnricher : ILogEventEnricher
 		{
@@ -71,6 +73,7 @@ namespace HordeAgent
 				.MinimumLevel.Debug()
 				.MinimumLevel.Override("Microsoft", LogEventLevel.Information)
 				.MinimumLevel.Override("System.Net.Http.HttpClient", LogEventLevel.Warning)
+				.MinimumLevel.ControlledBy(LogLevelSwitch)
 				.Enrich.FromLogContext()
 				.Enrich.With<VersionLogEnricher>()
 				.WriteTo.Console(outputTemplate: "[{Timestamp:HH:mm:ss} {Level:w3}] {Indent}{Message:l}{NewLine}{Exception}", theme: Theme)
