Horde: Improve sending of concurrent REAPI actions for load testing

diff --git a/Engine/Source/Programs/Horde/HordeAgent/Commands/Execution/ExecuteCommand.cs b/Engine/Source/Programs/Horde/HordeAgent/Commands/Execution/ExecuteCommand.cs
--- a/Engine/Source/Programs/Horde/HordeAgent/Commands/Execution/ExecuteCommand.cs
+++ b/Engine/Source/Programs/Horde/HordeAgent/Commands/Execution/ExecuteCommand.cs
@@ -5,6 +5,7 @@ using Google.Protobuf;
 using HordeAgent.Utility;
 using Microsoft.Extensions.Logging;
 using System;
+using System.Collections.Concurrent;
 using System.Collections.Generic;
 using System.IO;
 using System.Text;
@@ -21,13 +22,12 @@ using HordeAgent.Services;
 using Grpc.Net.Client;
 using System.Linq;
 using System.Net.Http.Headers;
-using System.Threading;
 using Grpc.Core;
 using Google.LongRunning;
 using Microsoft.Extensions.Hosting;
 using Microsoft.Extensions.DependencyInjection;
 using Microsoft.Extensions.Configuration;
-
+using Serilog.Events;
 using Digest = Build.Bazel.Remote.Execution.V2.Digest;
 
 namespace HordeAgent.Commands
@@ -223,10 +223,26 @@ namespace HordeAgent.Commands
 		/// </summary>
 		[CommandLine("-NumDuplicatedExecutions")]
 		public int NumDuplicatedExecutions = 1;
+		
+		/// <summary>
+		/// Log verbosity level (use normal Serilog levels such as debug, warning or info)
+		/// </summary>
+		[CommandLine("-LogLevel")]
+		public string LogLevelStr = "debug";
 
 		/// <inheritdoc/>
 		public override async Task<int> ExecuteAsync(ILogger Logger)
 		{
+			if (Enum.TryParse(LogLevelStr, true, out LogEventLevel LogEventLevel))
+			{
+				Logging.LogLevelSwitch.MinimumLevel = LogEventLevel;
+			}
+			else
+			{
+				Console.WriteLine($"Unable to parse log level: {this.LogLevelStr}");
+				return 0;
+			}
+			
 			IConfiguration Configuration = new ConfigurationBuilder()
 				.AddEnvironmentVariables()
 				.AddJsonFile("appsettings.json")
@@ -234,7 +250,10 @@ namespace HordeAgent.Commands
 
 			IConfigurationSection ConfigSection = Configuration.GetSection(AgentSettings.SectionName);
 			IHostBuilder HostBuilder = Host.CreateDefaultBuilder()
-				.ConfigureLogging(Builder => { })
+				.ConfigureLogging(Builder =>
+				{
+					Builder.SetMinimumLevel(LogLevel.Warning);
+				})
 				.ConfigureServices(Services =>
 				{
 					Services.AddOptions<AgentSettings>().Configure(Options => ConfigSection.Bind(Options)).ValidateDataAnnotations();
@@ -245,7 +264,7 @@ namespace HordeAgent.Commands
 			using (IHost Host = HostBuilder.Build())
 			{
 				GrpcService GrpcService = Host.Services.GetRequiredService<GrpcService>();
-
+				
 				byte[] Json = await FileReference.ReadAllBytesAsync(Input);
 				JsonAction Action = JsonSerializer.Deserialize<JsonAction>(Json, new JsonSerializerOptions { PropertyNameCaseInsensitive = true });
 
@@ -280,28 +299,86 @@ namespace HordeAgent.Commands
 			return 0;
 		}
 
+		private class ConcurrentExecuteOperation
+		{
+			internal DateTimeOffset CreatedTime;
+			internal DateTimeOffset ExecutionSendTime;
+			internal DateTimeOffset CompletedTime;
+			internal ActionResult? Result;
+			internal int NumRetries = 0;
+		
+			public ConcurrentExecuteOperation()
+			{
+				CreatedTime = DateTimeOffset.Now;
+			}
+		}
+
 		private async Task ExecuteActionsConcurrently(ILogger Logger, GrpcService GrpcService, UploadList UploadList, Digest ActionDigest)
 		{
 			Logger.LogInformation($"Running {NumDuplicatedExecutions} concurrently using {NumConcurrentExecutions} tasks...");
-			int Count = NumDuplicatedExecutions;
-			List<Task> TaskList = new List<Task>();
+			ConcurrentQueue<ConcurrentExecuteOperation> PendingOps = new ConcurrentQueue<ConcurrentExecuteOperation>();
+			ConcurrentQueue<ConcurrentExecuteOperation> CompletedOps = new ConcurrentQueue<ConcurrentExecuteOperation>();
+
+			for (int i = 0; i < NumDuplicatedExecutions; i++)
+			{
+				PendingOps.Enqueue(new ConcurrentExecuteOperation());
+			}
+			
+			List<Task> RunnerTasks = new List<Task>();
+			DateTime StartTime = DateTime.UtcNow;
+			
 			for (int i = 0; i < NumConcurrentExecutions; i++)
 			{
-				TaskList.Add(Task.Run(async () =>
+				RunnerTasks.Add(Task.Run(async () =>
 				{
-					while (Interlocked.Decrement(ref Count) > 0)
+					while (PendingOps.TryDequeue(out var Op))
 					{
-						await ExecuteAction(Logger, GrpcService, UploadList, ActionDigest);
+						Op.ExecutionSendTime = DateTimeOffset.Now;
+						Op.Result = await ExecuteAction(Logger, GrpcService, UploadList, ActionDigest);
+						if (Op.Result == null)
+						{
+							Logger.LogInformation("Retrying...");
+							Op.NumRetries++;
+							PendingOps.Enqueue(Op);
+							continue;
+						}
+						
+						Op.CompletedTime = DateTimeOffset.Now;
+						CompletedOps.Enqueue(Op);
+						Logger.LogInformation($"{CompletedOps.Count}/{NumDuplicatedExecutions} done");
 					}
 				}));
 			}
 
-			await Task.WhenAll(TaskList.ToArray());
+			await Task.WhenAll(RunnerTasks.ToArray());
+			DateTime EndTime = DateTime.UtcNow;
+			TimeSpan TimeTaken = EndTime - StartTime;
+			Logger.LogInformation("Finished executing all tasks");
+
+			Logger.LogInformation("{0,35} {1,16} {2,15} {3,7}", "Worker", "ServerQueueTime", "WorkerTime", "Retries");
+			Logger.LogInformation(new string('-', 100));
+			foreach (ConcurrentExecuteOperation Op in CompletedOps.ToArray())
+			{
+				if (Op.Result == null)
+				{
+					Logger.LogError("Received a null result");
+					continue;
+				}
+				ExecutedActionMetadata Md = Op.Result.ExecutionMetadata;
+				TimeSpan QueueTime = Md.WorkerStartTimestamp.ToDateTimeOffset() - Md.QueuedTimestamp.ToDateTimeOffset();
+				TimeSpan TotalWorkerTime = Md.WorkerCompletedTimestamp.ToDateTimeOffset() - Md.WorkerStartTimestamp.ToDateTimeOffset();
+				Logger.LogInformation("{Worker,35} {QueueTime,16:N2} {TotalWorkerTime,15:N2} {Retries,7}", Md.Worker, QueueTime.TotalSeconds, TotalWorkerTime.TotalSeconds, Op.NumRetries);
+			}
+			
+			Logger.LogInformation(" Time taken: {TimeTaken:N1} secs", TimeTaken.TotalSeconds);
+			Logger.LogInformation("  Num execs: {NumExecutions}", CompletedOps.Count);
+			Logger.LogInformation("Num retries: {NumRetries}", CompletedOps.Sum(o => o.NumRetries));
 		}
 
-		private async Task ExecuteAction(ILogger Logger, GrpcService GrpcService, UploadList UploadList, Digest ActionDigest)
+		private async Task<ActionResult?> ExecuteAction(ILogger Logger, GrpcService GrpcService, UploadList UploadList, Digest ActionDigest)
 		{
 			using GrpcChannel Channel = GrpcService.CreateGrpcChannel();
+			ActionResult? ActionResult = null;
 
 			GrpcChannel GetChannel()
 			{
@@ -333,7 +410,7 @@ namespace HordeAgent.Commands
 					{
 						if (Response.Status.Code != (int) Google.Rpc.Code.Ok)
 						{
-							Console.WriteLine($"Upload failed for {Response.Digest}. Code: {Response.Status.Code} Message: {Response.Status.Message}");
+							Logger.LogError($"Upload failed for {Response.Digest}. Code: {Response.Status.Code} Message: {Response.Status.Message}");
 							UpdateFailed = true;
 						}
 					}
@@ -356,13 +433,18 @@ namespace HordeAgent.Commands
 				ExecuteRequest.InstanceName = InstanceName;
 			}
 
+			ExecutionStage.Types.Value CurrentStage = ExecutionStage.Types.Value.Unknown;
 			using AsyncServerStreamingCall<Operation> Call = ExecutionClient.Execute(ExecuteRequest);
 			while (await Call.ResponseStream.MoveNext())
 			{
 				Operation Operation = Call.ResponseStream.Current;
 
 				ExecuteOperationMetadata Metadata = Operation.Metadata.Unpack<ExecuteOperationMetadata>();
-				Logger.LogInformation("Execution state: {State}", Metadata.Stage);
+				Logger.LogDebug("{OperationName}: Execution state: {State}", Operation.Name, Metadata.Stage);
+				if (Metadata.Stage < CurrentStage)
+				{
+					Logger.LogError("Trying to set stage {NewStage} when current stage is {CurrentStage}", Metadata.Stage, CurrentStage);
+				}
 
 				if (Operation.Done)
 				{
@@ -379,25 +461,26 @@ namespace HordeAgent.Commands
 						break;
 					}
 
-					ActionResult? ActionResult = ExecuteResponse.Result;
+					ActionResult = ExecuteResponse.Result;
 					if (ExecuteResponse.Result == null)
 					{
 						Logger.LogError("Empty result");
 						break;
 					}
 
-					Logger.LogInformation("Cached: {CachedResult}", ExecuteResponse.CachedResult);
+					Logger.LogDebug("Worker: {Worker}", ActionResult.ExecutionMetadata.Worker);
+					Logger.LogDebug("Cached: {CachedResult}", ExecuteResponse.CachedResult);
 							
-					Logger.LogInformation("Output directories:");
+					Logger.LogDebug("Output directories:");
 					foreach (OutputDirectory Dir in ActionResult.OutputDirectories)
 					{
-						Logger.LogInformation("Path: {Path} Digest: {Digest}", Dir.Path, Dir.TreeDigest);
+						Logger.LogDebug("Path: {Path} Digest: {Digest}", Dir.Path, Dir.TreeDigest);
 					}
 							
-					Logger.LogInformation("Output files:");
+					Logger.LogDebug("Output files:");
 					foreach (OutputFile File in ActionResult.OutputFiles)
 					{
-						Logger.LogInformation("Path: {Path} Digest: {Digest}", File.Path, File.Digest);
+						Logger.LogDebug("Path: {Path} Digest: {Digest}", File.Path, File.Digest);
 					}
 
 					// Print the result
@@ -424,7 +507,7 @@ namespace HordeAgent.Commands
 						{
 							foreach (string Line in Encoding.UTF8.GetString(StdoutResponse.Data.ToByteArray()).Split('\n'))
 							{
-								Logger.LogInformation("stdout: {Line}", Line);
+								Logger.LogDebug("stdout: {Line}", Line);
 							}
 						}
 					}
@@ -435,7 +518,8 @@ namespace HordeAgent.Commands
 						{
 							foreach (string Line in Encoding.UTF8.GetString(StderrResponse.Data.ToByteArray()).Split('\n'))
 							{
-								Logger.LogError("stderr: {Line}", Line);
+								if (!string.IsNullOrWhiteSpace(Line))
+									Logger.LogError("stderr: {Line}", Line);
 							}
 						}
 					}
@@ -451,14 +535,16 @@ namespace HordeAgent.Commands
 							string OutputPath = Path.Join(OutputDir, OutputFile.Path);
 							System.IO.Directory.CreateDirectory(Path.GetDirectoryName(OutputPath));
 							await File.WriteAllBytesAsync(OutputPath, FileResponse.Data.ToByteArray());
-							Logger.LogInformation("Wrote {OutputFilePath}", OutputFile.Path);
+							Logger.LogDebug("Wrote {OutputFilePath}", OutputFile.Path);
 						}
 					}
 
-					Logger.LogInformation("exit: {ExitCode}", ActionResult.ExitCode);
+					Logger.LogDebug("exit: {ExitCode}", ActionResult.ExitCode);
 					break;
 				}
 			}
+
+			return ActionResult;
 		}
 	}
 }
