Horde: Move Horde protobuf schemas to Shared/EpicGames.Horde.Common

#rnx

diff --git a/Engine/Source/Programs/Horde/HordeAgent/HordeAgent.csproj b/Engine/Source/Programs/Horde/HordeAgent/HordeAgent.csproj
--- a/Engine/Source/Programs/Horde/HordeAgent/HordeAgent.csproj
+++ b/Engine/Source/Programs/Horde/HordeAgent/HordeAgent.csproj
@@ -70,7 +70,7 @@
   </ItemGroup>
 
   <ItemGroup>
-    <Content Include="..\HordeCommon\Protos\**\*.*">
+    <Content Include="..\..\Shared\EpicGames.Horde.Common\Protos\**\*.*">
       <Link>Protos\%(RecursiveDir)%(FileName)%(Extension)</Link>
     </Content>
     <Compile Include="../../../../Platforms/*/Source/Programs/Horde/HordeAgent/**/*.cs">
@@ -83,27 +83,28 @@
 
   <ItemGroup>
     <ProjectReference Include="..\..\Shared\EpicGames.Core\EpicGames.Core.csproj" />
+    <ProjectReference Include="..\..\Shared\EpicGames.Horde.Common\EpicGames.Horde.Common.csproj" />
     <ProjectReference Include="..\..\Shared\EpicGames.Perforce.Managed\EpicGames.Perforce.Managed.csproj" />
     <ProjectReference Include="..\..\Shared\EpicGames.Perforce\EpicGames.Perforce.csproj" />
     <ProjectReference Include="..\HordeCommon\HordeCommon.csproj" />
   </ItemGroup>
 
   <ItemGroup>
-    <Protobuf Include="..\HordeCommon\Protos\build\bazel\remote\execution\v2\remote_execution.proto" ProtoRoot="..\HordeCommon\Protos" GrpcServices="Client" />
-    <Protobuf Include="..\HordeCommon\Protos\build\bazel\semver\semver.proto" ProtoRoot="..\HordeCommon\Protos" GrpcServices="Client" />
-    <Protobuf Include="..\HordeCommon\Protos\google\api\annotations.proto" ProtoRoot="..\HordeCommon\Protos" GrpcServices="Client" />
-    <Protobuf Include="..\HordeCommon\Protos\google\api\client.proto" ProtoRoot="..\HordeCommon\Protos" GrpcServices="Client" />
-    <Protobuf Include="..\HordeCommon\Protos\google\api\field_behavior.proto" ProtoRoot="..\HordeCommon\Protos" GrpcServices="Client" />
-    <Protobuf Include="..\HordeCommon\Protos\google\api\http.proto" ProtoRoot="..\HordeCommon\Protos" GrpcServices="Client" />
-    <Protobuf Include="..\HordeCommon\Protos\google\api\resource.proto" ProtoRoot="..\HordeCommon\Protos" GrpcServices="Client" />
-    <Protobuf Include="..\HordeCommon\Protos\google\bytestream\bytestream.proto" ProtoRoot="..\HordeCommon\Protos" GrpcServices="Client" />
-    <Protobuf Include="..\HordeCommon\Protos\google\devtools\remoteworkers\v1test2\bots.proto" ProtoRoot="..\HordeCommon\Protos" GrpcServices="Client" />
-    <Protobuf Include="..\HordeCommon\Protos\google\devtools\remoteworkers\v1test2\worker.proto" ProtoRoot="..\HordeCommon\Protos" GrpcServices="Client" />
-    <Protobuf Include="..\HordeCommon\Protos\google\longrunning\operations.proto" ProtoRoot="..\HordeCommon\Protos" GrpcServices="Client" />
-    <Protobuf Include="..\HordeCommon\Protos\google\rpc\code.proto" ProtoRoot="..\HordeCommon\Protos" GrpcServices="Client" />
-    <Protobuf Include="..\HordeCommon\Protos\horde\action_rpc.proto" ProtoRoot="..\HordeCommon\Protos" GrpcServices="Client" />
-    <Protobuf Include="..\HordeCommon\Protos\horde\horde_rpc.proto" ProtoRoot="..\HordeCommon\Protos" GrpcServices="Client" />
-    <Protobuf Include="..\HordeCommon\Protos\horde\tasks\action_task.proto" ProtoRoot="..\HordeCommon\Protos" GrpcServices="Client" />
+    <Protobuf Include="..\..\Shared\EpicGames.Horde.Common\Protos\build\bazel\remote\execution\v2\remote_execution.proto" ProtoRoot="..\..\Shared\EpicGames.Horde.Common\Protos" GrpcServices="Client" />
+    <Protobuf Include="..\..\Shared\EpicGames.Horde.Common\Protos\build\bazel\semver\semver.proto" ProtoRoot="..\..\Shared\EpicGames.Horde.Common\Protos" GrpcServices="Client" />
+    <Protobuf Include="..\..\Shared\EpicGames.Horde.Common\Protos\google\api\annotations.proto" ProtoRoot="..\..\Shared\EpicGames.Horde.Common\Protos" GrpcServices="Client" />
+    <Protobuf Include="..\..\Shared\EpicGames.Horde.Common\Protos\google\api\client.proto" ProtoRoot="..\..\Shared\EpicGames.Horde.Common\Protos" GrpcServices="Client" />
+    <Protobuf Include="..\..\Shared\EpicGames.Horde.Common\Protos\google\api\field_behavior.proto" ProtoRoot="..\..\Shared\EpicGames.Horde.Common\Protos" GrpcServices="Client" />
+    <Protobuf Include="..\..\Shared\EpicGames.Horde.Common\Protos\google\api\http.proto" ProtoRoot="..\..\Shared\EpicGames.Horde.Common\Protos" GrpcServices="Client" />
+    <Protobuf Include="..\..\Shared\EpicGames.Horde.Common\Protos\google\api\resource.proto" ProtoRoot="..\..\Shared\EpicGames.Horde.Common\Protos" GrpcServices="Client" />
+    <Protobuf Include="..\..\Shared\EpicGames.Horde.Common\Protos\google\bytestream\bytestream.proto" ProtoRoot="..\..\Shared\EpicGames.Horde.Common\Protos" GrpcServices="Client" />
+    <Protobuf Include="..\..\Shared\EpicGames.Horde.Common\Protos\google\devtools\remoteworkers\v1test2\bots.proto" ProtoRoot="..\..\Shared\EpicGames.Horde.Common\Protos" GrpcServices="Client" />
+    <Protobuf Include="..\..\Shared\EpicGames.Horde.Common\Protos\google\devtools\remoteworkers\v1test2\worker.proto" ProtoRoot="..\..\Shared\EpicGames.Horde.Common\Protos" GrpcServices="Client" />
+    <Protobuf Include="..\..\Shared\EpicGames.Horde.Common\Protos\google\longrunning\operations.proto" ProtoRoot="..\..\Shared\EpicGames.Horde.Common\Protos" GrpcServices="Client" />
+    <Protobuf Include="..\..\Shared\EpicGames.Horde.Common\Protos\google\rpc\code.proto" ProtoRoot="..\..\Shared\EpicGames.Horde.Common\Protos" GrpcServices="Client" />
+    <Protobuf Include="..\..\Shared\EpicGames.Horde.Common\Protos\horde\action_rpc.proto" ProtoRoot="..\..\Shared\EpicGames.Horde.Common\Protos" GrpcServices="Client" />
+    <Protobuf Include="..\..\Shared\EpicGames.Horde.Common\Protos\horde\horde_rpc.proto" ProtoRoot="..\..\Shared\EpicGames.Horde.Common\Protos" GrpcServices="Client" />
+    <Protobuf Include="..\..\Shared\EpicGames.Horde.Common\Protos\horde\tasks\action_task.proto" ProtoRoot="..\..\Shared\EpicGames.Horde.Common\Protos" GrpcServices="Client" />
   </ItemGroup>
 
   <ItemGroup>
diff --git a/Engine/Source/Programs/Horde/HordeCommon/HordeCommon.csproj b/Engine/Source/Programs/Horde/HordeCommon/HordeCommon.csproj
--- a/Engine/Source/Programs/Horde/HordeCommon/HordeCommon.csproj
+++ b/Engine/Source/Programs/Horde/HordeCommon/HordeCommon.csproj
@@ -16,24 +16,27 @@
   </ItemGroup>
 
   <ItemGroup>
-    <Protobuf Include="Protos\horde\horde_rpc_messages.proto" ProtoRoot="Protos" />
-    <Protobuf Include="Protos\horde\common\agent_status.proto" ProtoRoot="Protos" />
-    <Protobuf Include="Protos\horde\common\event_severity.proto" ProtoRoot="Protos" />
-    <Protobuf Include="Protos\horde\common\jobstep_outcome.proto" ProtoRoot="Protos" />
-    <Protobuf Include="Protos\horde\common\jobstep_state.proto" ProtoRoot="Protos" />
-    <Protobuf Include="Protos\horde\common\label_change.proto" ProtoRoot="Protos" />
-    <Protobuf Include="Protos\horde\common\lease_outcome.proto" ProtoRoot="Protos" />
-    <Protobuf Include="Protos\horde\common\lease_state.proto" ProtoRoot="Protos" />
-    <Protobuf Include="Protos\horde\common\priority.proto" ProtoRoot="Protos" />
-    <Protobuf Include="Protos\horde\messages\agent_capabilities.proto" ProtoRoot="Protos" />
-    <Protobuf Include="Protos\horde\messages\agent_requirements.proto" ProtoRoot="Protos" />
-    <Protobuf Include="Protos\horde\messages\agent_workspace.proto" ProtoRoot="Protos" />
-    <Protobuf Include="Protos\horde\messages\lease.proto" ProtoRoot="Protos" />
-    <Protobuf Include="Protos\horde\tasks\conform_task.proto" ProtoRoot="Protos" />
-    <Protobuf Include="Protos\horde\tasks\job_task.proto" ProtoRoot="Protos" />
-    <Protobuf Include="Protos\horde\tasks\restart_task.proto" ProtoRoot="Protos" />
-    <Protobuf Include="Protos\horde\tasks\shutdown_task.proto" ProtoRoot="Protos" />
-    <Protobuf Include="Protos\horde\tasks\upgrade_task.proto" ProtoRoot="Protos" />
-    <Protobuf Include="Protos\google\rpc\status.proto" ProtoRoot="Protos" />
+    <ProjectReference Include="..\..\Shared\EpicGames.Horde.Common\EpicGames.Horde.Common.csproj" />
+  </ItemGroup>
+
+  <ItemGroup>
+    <Protobuf Include="..\..\Shared\EpicGames.Horde.Common\Protos\horde\horde_rpc_messages.proto" ProtoRoot="..\..\Shared\EpicGames.Horde.Common\Protos" />
+    <Protobuf Include="..\..\Shared\EpicGames.Horde.Common\Protos\horde\common\agent_status.proto" ProtoRoot="..\..\Shared\EpicGames.Horde.Common\Protos" />
+    <Protobuf Include="..\..\Shared\EpicGames.Horde.Common\Protos\horde\common\event_severity.proto" ProtoRoot="..\..\Shared\EpicGames.Horde.Common\Protos" />
+    <Protobuf Include="..\..\Shared\EpicGames.Horde.Common\Protos\horde\common\jobstep_outcome.proto" ProtoRoot="..\..\Shared\EpicGames.Horde.Common\Protos" />
+    <Protobuf Include="..\..\Shared\EpicGames.Horde.Common\Protos\horde\common\jobstep_state.proto" ProtoRoot="..\..\Shared\EpicGames.Horde.Common\Protos" />
+    <Protobuf Include="..\..\Shared\EpicGames.Horde.Common\Protos\horde\common\label_change.proto" ProtoRoot="..\..\Shared\EpicGames.Horde.Common\Protos" />
+    <Protobuf Include="..\..\Shared\EpicGames.Horde.Common\Protos\horde\common\lease_outcome.proto" ProtoRoot="..\..\Shared\EpicGames.Horde.Common\Protos" />
+    <Protobuf Include="..\..\Shared\EpicGames.Horde.Common\Protos\horde\common\lease_state.proto" ProtoRoot="..\..\Shared\EpicGames.Horde.Common\Protos" />
+    <Protobuf Include="..\..\Shared\EpicGames.Horde.Common\Protos\horde\common\priority.proto" ProtoRoot="..\..\Shared\EpicGames.Horde.Common\Protos" />
+    <Protobuf Include="..\..\Shared\EpicGames.Horde.Common\Protos\horde\messages\agent_capabilities.proto" ProtoRoot="..\..\Shared\EpicGames.Horde.Common\Protos" />
+    <Protobuf Include="..\..\Shared\EpicGames.Horde.Common\Protos\horde\messages\agent_requirements.proto" ProtoRoot="..\..\Shared\EpicGames.Horde.Common\Protos" />
+    <Protobuf Include="..\..\Shared\EpicGames.Horde.Common\Protos\horde\messages\agent_workspace.proto" ProtoRoot="..\..\Shared\EpicGames.Horde.Common\Protos" />
+    <Protobuf Include="..\..\Shared\EpicGames.Horde.Common\Protos\horde\messages\lease.proto" ProtoRoot="..\..\Shared\EpicGames.Horde.Common\Protos" />
+    <Protobuf Include="..\..\Shared\EpicGames.Horde.Common\Protos\horde\tasks\conform_task.proto" ProtoRoot="..\..\Shared\EpicGames.Horde.Common\Protos" />
+    <Protobuf Include="..\..\Shared\EpicGames.Horde.Common\Protos\horde\tasks\job_task.proto" ProtoRoot="..\..\Shared\EpicGames.Horde.Common\Protos" />
+    <Protobuf Include="..\..\Shared\EpicGames.Horde.Common\Protos\horde\tasks\restart_task.proto" ProtoRoot="..\..\Shared\EpicGames.Horde.Common\Protos" />
+    <Protobuf Include="..\..\Shared\EpicGames.Horde.Common\Protos\horde\tasks\shutdown_task.proto" ProtoRoot="..\..\Shared\EpicGames.Horde.Common\Protos" />
+    <Protobuf Include="..\..\Shared\EpicGames.Horde.Common\Protos\horde\tasks\upgrade_task.proto" ProtoRoot="..\..\Shared\EpicGames.Horde.Common\Protos" />
   </ItemGroup>
 </Project>
diff --git a/Engine/Source/Programs/Horde/HordeCommon/Status.cs b/Engine/Source/Programs/Horde/HordeCommon/Status.cs
deleted file mode 100644
diff --git a/Engine/Source/Programs/Horde/HordeServer/HordeServer.csproj b/Engine/Source/Programs/Horde/HordeServer/HordeServer.csproj
--- a/Engine/Source/Programs/Horde/HordeServer/HordeServer.csproj
+++ b/Engine/Source/Programs/Horde/HordeServer/HordeServer.csproj
@@ -47,7 +47,7 @@
   </ItemGroup>
 
   <ItemGroup>
-    <Content Include="..\HordeCommon\Protos\**\*.*">
+    <Content Include="..\..\Shared\EpicGames.Horde.Common\Protos\**\*.*">
       <Link>Protos\%(RecursiveDir)%(FileName)%(Extension)</Link>
     </Content>
   </ItemGroup>
@@ -57,21 +57,22 @@
   </ItemGroup>
 
   <ItemGroup>
-    <Protobuf Include="..\HordeCommon\Protos\build\bazel\remote\execution\v2\remote_execution.proto" ProtoRoot="..\HordeCommon\Protos" GrpcServices="Server" />
-    <Protobuf Include="..\HordeCommon\Protos\build\bazel\semver\semver.proto" ProtoRoot="..\HordeCommon\Protos" GrpcServices="Server" />
-    <Protobuf Include="..\HordeCommon\Protos\google\api\annotations.proto" ProtoRoot="..\HordeCommon\Protos" GrpcServices="Server" />
-    <Protobuf Include="..\HordeCommon\Protos\google\api\client.proto" ProtoRoot="..\HordeCommon\Protos" GrpcServices="Server" />
-    <Protobuf Include="..\HordeCommon\Protos\google\api\http.proto" ProtoRoot="..\HordeCommon\Protos" GrpcServices="Server" />
-    <Protobuf Include="..\HordeCommon\Protos\google\bytestream\bytestream.proto" ProtoRoot="..\HordeCommon\Protos" GrpcServices="Server" />
-    <Protobuf Include="..\HordeCommon\Protos\google\longrunning\operations.proto" ProtoRoot="..\HordeCommon\Protos" GrpcServices="Server" />
-    <Protobuf Include="..\HordeCommon\Protos\google\rpc\code.proto" ProtoRoot="..\HordeCommon\Protos" GrpcServices="Server" />
-    <Protobuf Include="..\HordeCommon\Protos\horde\action_rpc.proto" ProtoRoot="..\HordeCommon\Protos" GrpcServices="Server" />
-    <Protobuf Include="..\HordeCommon\Protos\horde\horde_rpc.proto" ProtoRoot="..\HordeCommon\Protos" GrpcServices="Server" />
-    <Protobuf Include="..\HordeCommon\Protos\horde\tasks\action_task.proto" ProtoRoot="..\HordeCommon\Protos" GrpcServices="Server" />
+    <Protobuf Include="..\..\Shared\EpicGames.Horde.Common\Protos\build\bazel\remote\execution\v2\remote_execution.proto" ProtoRoot="..\..\Shared\EpicGames.Horde.Common\Protos" GrpcServices="Server" />
+    <Protobuf Include="..\..\Shared\EpicGames.Horde.Common\Protos\build\bazel\semver\semver.proto" ProtoRoot="..\..\Shared\EpicGames.Horde.Common\Protos" GrpcServices="Server" />
+    <Protobuf Include="..\..\Shared\EpicGames.Horde.Common\Protos\google\api\annotations.proto" ProtoRoot="..\..\Shared\EpicGames.Horde.Common\Protos" GrpcServices="Server" />
+    <Protobuf Include="..\..\Shared\EpicGames.Horde.Common\Protos\google\api\client.proto" ProtoRoot="..\..\Shared\EpicGames.Horde.Common\Protos" GrpcServices="Server" />
+    <Protobuf Include="..\..\Shared\EpicGames.Horde.Common\Protos\google\api\http.proto" ProtoRoot="..\..\Shared\EpicGames.Horde.Common\Protos" GrpcServices="Server" />
+    <Protobuf Include="..\..\Shared\EpicGames.Horde.Common\Protos\google\bytestream\bytestream.proto" ProtoRoot="..\..\Shared\EpicGames.Horde.Common\Protos" GrpcServices="Server" />
+    <Protobuf Include="..\..\Shared\EpicGames.Horde.Common\Protos\google\longrunning\operations.proto" ProtoRoot="..\..\Shared\EpicGames.Horde.Common\Protos" GrpcServices="Server" />
+    <Protobuf Include="..\..\Shared\EpicGames.Horde.Common\Protos\google\rpc\code.proto" ProtoRoot="..\..\Shared\EpicGames.Horde.Common\Protos" GrpcServices="Server" />
+    <Protobuf Include="..\..\Shared\EpicGames.Horde.Common\Protos\horde\action_rpc.proto" ProtoRoot="..\..\Shared\EpicGames.Horde.Common\Protos" GrpcServices="Server" />
+    <Protobuf Include="..\..\Shared\EpicGames.Horde.Common\Protos\horde\horde_rpc.proto" ProtoRoot="..\..\Shared\EpicGames.Horde.Common\Protos" GrpcServices="Server" />
+    <Protobuf Include="..\..\Shared\EpicGames.Horde.Common\Protos\horde\tasks\action_task.proto" ProtoRoot="..\..\Shared\EpicGames.Horde.Common\Protos" GrpcServices="Server" />
   </ItemGroup>
 
   <ItemGroup>
     <ProjectReference Include="..\..\Shared\EpicGames.Core\EpicGames.Core.csproj" />
+    <ProjectReference Include="..\..\Shared\EpicGames.Horde.Common\EpicGames.Horde.Common.csproj" />
     <ProjectReference Include="..\..\Shared\EpicGames.IoHash\EpicGames.IoHash.csproj" />
     <ProjectReference Include="..\..\Shared\EpicGames.Perforce\EpicGames.Perforce.csproj" />
     <ProjectReference Include="..\HordeCommon\HordeCommon.csproj" />
diff --git a/Engine/Source/Programs/Shared/EpicGames.Horde.Common/EpicGames.Horde.Common.csproj b/Engine/Source/Programs/Shared/EpicGames.Horde.Common/EpicGames.Horde.Common.csproj
new file mode 100644
--- /dev/null
+++ b/Engine/Source/Programs/Shared/EpicGames.Horde.Common/EpicGames.Horde.Common.csproj
@@ -0,0 +1,42 @@
+﻿<Project Sdk="Microsoft.NET.Sdk">
+  <PropertyGroup>
+    <TargetFramework>netcoreapp3.1</TargetFramework>
+    <Configuration Condition=" '$(Configuration)' == '' ">Development</Configuration>
+    <OutputType>Library</OutputType>
+    <Deterministic>true</Deterministic>
+    <RootNamespace>EpicGames.Horde.Common</RootNamespace>
+    <GenerateAssemblyInfo>false</GenerateAssemblyInfo>
+    <GenerateTargetFrameworkAttribute>false</GenerateTargetFrameworkAttribute>
+    <Configurations>Debug;Release;Development</Configurations>
+    <Nullable>enable</Nullable>
+  </PropertyGroup>
+  <PropertyGroup Condition=" '$(Configuration)|$(Platform)' == 'Debug|AnyCPU' ">
+    <DocumentationFile></DocumentationFile>
+  </PropertyGroup>
+  <PropertyGroup Condition="'$(Configuration)|$(Platform)' == 'Development|AnyCPU'">
+    <DocumentationFile></DocumentationFile>
+    <Optimize>false</Optimize>
+  </PropertyGroup>
+  <PropertyGroup Condition="'$(Configuration)|$(Platform)' == 'Release|AnyCPU'">
+    <DocumentationFile></DocumentationFile>
+    <Optimize>true</Optimize>
+  </PropertyGroup>
+  <ItemGroup>
+    <Compile Include="..\..\DotNETCommon\MetaData.cs">
+      <Link>Properties\MetaData.cs</Link>
+    </Compile>
+  </ItemGroup>
+  <ItemGroup>
+    <PackageReference Include="Google.Protobuf" Version="3.12.3" />
+    <PackageReference Include="Grpc.Core.Api" Version="2.35.0" />
+    <PackageReference Include="Grpc.Tools" Version="2.30.0">
+      <PrivateAssets>all</PrivateAssets>
+      <IncludeAssets>runtime; build; native; contentfiles; analyzers; buildtransitive</IncludeAssets>
+    </PackageReference>
+    <PackageReference Include="Microsoft.Extensions.Logging" Version="3.1.5" />
+  </ItemGroup>
+
+  <ItemGroup>
+    <Protobuf Include="Protos\google\rpc\status.proto" ProtoRoot="Protos" />
+  </ItemGroup>
+</Project>
\ No newline at end of file
diff --git a/Engine/Source/Programs/Shared/EpicGames.Horde.Common/Status.cs b/Engine/Source/Programs/Shared/EpicGames.Horde.Common/Status.cs
new file mode 100644
--- /dev/null
+++ b/Engine/Source/Programs/Shared/EpicGames.Horde.Common/Status.cs
@@ -0,0 +1,27 @@
+// Copyright Epic Games, Inc. All Rights Reserved.
+
+using System;
+using System.Collections.Generic;
+using System.Text;
+
+namespace Google.Rpc
+{
+	public partial class Status
+	{
+		public Status(Grpc.Core.StatusCode Code, string Message)
+		{
+			this.Code = (int)Code;
+			this.Message = Message;
+		}
+
+		public static implicit operator Grpc.Core.Status(Status Other)
+		{
+			return new Grpc.Core.Status((Grpc.Core.StatusCode)Other.Code, Other.Message);
+		}
+
+		public static implicit operator Status(Grpc.Core.Status Other)
+		{
+			return new Status((Grpc.Core.StatusCode)Other.StatusCode, Other.Detail);
+		}
+	}
+}
