Horde: Add support for AutoSDK on Linux.

diff --git a/Engine/Source/Programs/Horde/HordeServer/Models/Agent.cs b/Engine/Source/Programs/Horde/HordeServer/Models/Agent.cs
--- a/Engine/Source/Programs/Horde/HordeServer/Models/Agent.cs
+++ b/Engine/Source/Programs/Horde/HordeServer/Models/Agent.cs
@@ -770,7 +770,11 @@ namespace HordeServer.Models
 					{
 						return new AgentWorkspace(null, null, AgentWorkspace.AutoSdkIdentifier, "//UE4/Private-AutoSDK-Windows", null, true);
 					}
-					else if (PrimaryDevice.Properties.Contains("OSVersion=MacOS"))
+					else if (PrimaryDevice.Properties.Contains("OSFamily=Linux"))
+					{
+						return new AgentWorkspace(null, null, AgentWorkspace.AutoSdkIdentifier, "//UE4/Private-AutoSDK-Linux", null, true);
+					}
+					else if (PrimaryDevice.Properties.Contains("OSFamily=MacOS"))
 					{
 						return new AgentWorkspace(null, null, AgentWorkspace.AutoSdkIdentifier, "//UE4/Private-AutoSDK-Mac", null, true);
 					}
