Horde: Remove hard-coded paths to AutoSDK streams. Streams are now configured through Perforce clusters in config instead.

diff --git a/Engine/Source/Programs/Horde/HordeServer/Models/Agent.cs b/Engine/Source/Programs/Horde/HordeServer/Models/Agent.cs
--- a/Engine/Source/Programs/Horde/HordeServer/Models/Agent.cs
+++ b/Engine/Source/Programs/Horde/HordeServer/Models/Agent.cs
@@ -30,11 +30,6 @@ namespace HordeServer.Models
 	/// </summary>
 	public class AgentWorkspace
 	{
-		/// <summary>
-		/// Identifier for the AutoSDK workspace
-		/// </summary>
-		public const string AutoSdkIdentifier = "AutoSDK";
-
 		/// <summary>
 		/// Name of the Perforce cluster to use
 		/// </summary>
@@ -755,54 +750,53 @@ namespace HordeServer.Models
 		}
 
 		/// <summary>
-		/// Get the AutoSDK workspace required for an agent
+		/// Gets all the autosdk workspaces required for an agent
 		/// </summary>
 		/// <param name="Agent"></param>
+		/// <param name="Globals"></param>
+		/// <param name="Workspaces"></param>
 		/// <returns></returns>
-		public static AgentWorkspace? GetAutoSdkWorkspace(this IAgent Agent)
+		public static HashSet<AgentWorkspace> GetAutoSdkWorkspaces(this IAgent Agent, Globals Globals, List<AgentWorkspace> Workspaces)
 		{
-			if (Agent.Capabilities.Devices.Count > 0)
+			HashSet<AgentWorkspace> AutoSdkWorkspaces = new HashSet<AgentWorkspace>();
+			foreach (string? ClusterName in Workspaces.Select(x => x.Cluster).Distinct())
 			{
-				DeviceCapabilities PrimaryDevice = Agent.Capabilities.Devices[0];
-				if (PrimaryDevice.Properties != null)
+				PerforceCluster? Cluster = Globals.FindPerforceCluster(ClusterName);
+				if (Cluster != null)
 				{
-					if (PrimaryDevice.Properties.Contains("OSFamily=Windows"))
-					{
-						return new AgentWorkspace(null, null, AgentWorkspace.AutoSdkIdentifier, "//UE4/Private-AutoSDK-Windows", null, true);
-					}
-					else if (PrimaryDevice.Properties.Contains("OSFamily=Linux"))
+					AgentWorkspace? AutoSdkWorkspace = GetAutoSdkWorkspace(Agent, Cluster);
+					if (AutoSdkWorkspace != null)
 					{
-						return new AgentWorkspace(null, null, AgentWorkspace.AutoSdkIdentifier, "//UE4/Private-AutoSDK-Linux", null, true);
-					}
-					else if (PrimaryDevice.Properties.Contains("OSFamily=MacOS"))
-					{
-						return new AgentWorkspace(null, null, AgentWorkspace.AutoSdkIdentifier, "//UE4/Private-AutoSDK-Mac", null, true);
+						AutoSdkWorkspaces.Add(AutoSdkWorkspace);
 					}
 				}
 			}
-			return null;
+			return AutoSdkWorkspaces;
 		}
 
 		/// <summary>
-		/// Check that an agent has the required workspaces to execute a task
+		/// Get the AutoSDK workspace required for an agent
 		/// </summary>
 		/// <param name="Agent"></param>
-		/// <param name="Workspace"></param>
+		/// <param name="Cluster">The perforce cluster to get a workspace for</param>
 		/// <returns></returns>
-		public static bool HasRequiredWorkspaces(this IAgent Agent, AgentWorkspace Workspace)
+		public static AgentWorkspace? GetAutoSdkWorkspace(this IAgent Agent, PerforceCluster Cluster)
 		{
-			if (!Agent.Workspaces.Contains(Workspace))
-			{
-				return false;
-			}
-
-			AgentWorkspace? AutoSdkWorkspace = Agent.GetAutoSdkWorkspace();
-			if (AutoSdkWorkspace != null && !Agent.Workspaces.Contains(AutoSdkWorkspace))
+			if (Agent.Capabilities.Devices.Count > 0)
 			{
-				return false;
+				DeviceCapabilities PrimaryDevice = Agent.Capabilities.Devices[0];
+				if (PrimaryDevice.Properties != null)
+				{
+					foreach (AutoSdkWorkspace AutoSdk in Cluster.AutoSdk)
+					{
+						if (AutoSdk.Stream != null && AutoSdk.Properties.All(x => PrimaryDevice.Properties.Contains(x)))
+						{
+							return new AgentWorkspace(Cluster.Name, AutoSdk.UserName, AutoSdk.Name ?? "AutoSDK", AutoSdk.Stream!, null, true);
+						}
+					}
+				}
 			}
-
-			return true;
+			return null;
 		}
 
 		/// <summary>
@@ -810,33 +804,22 @@ namespace HordeServer.Models
 		/// </summary>
 		/// <param name="Agent">The agent to get a workspace for</param>
 		/// <param name="Workspace">The workspace definition</param>
-		/// <param name="Globals">The global state</param>
+		/// <param name="Cluster">The global state</param>
 		/// <param name="LoadBalancer">The Perforce load balancer</param>
 		/// <param name="WorkspaceMessages">List of messages</param>
 		/// <returns>The RPC message</returns>
-		public static async Task<bool> TryAddWorkspaceMessage(this IAgent Agent, AgentWorkspace Workspace, Globals Globals, PerforceLoadBalancer LoadBalancer, IList<HordeCommon.Rpc.Messages.AgentWorkspace> WorkspaceMessages)
+		public static async Task<bool> TryAddWorkspaceMessage(this IAgent Agent, AgentWorkspace Workspace, PerforceCluster Cluster, PerforceLoadBalancer LoadBalancer, IList<HordeCommon.Rpc.Messages.AgentWorkspace> WorkspaceMessages)
 		{
-			// Find the matching Perforce cluster
-			PerforceCluster? ClusterInfo;
-			if (Workspace.Cluster == null)
-			{
-				ClusterInfo = Globals.PerforceClusters.FirstOrDefault();
-			}
-			else
-			{
-				ClusterInfo = Globals.PerforceClusters.FirstOrDefault(x => String.Equals(x.Name, Workspace.Cluster, StringComparison.OrdinalIgnoreCase));
-			}
-
 			// Find a matching server, trying to use a previously selected one if possible
 			string? ServerAndPort = WorkspaceMessages.FirstOrDefault(x => x.ConfiguredCluster == Workspace.Cluster)?.ServerAndPort;
 			if (ServerAndPort == null)
 			{
-				if (ClusterInfo == null)
+				if (Cluster == null)
 				{
 					return false;
 				}
 
-				IPerforceServer? Server = await LoadBalancer.SelectServerAsync(ClusterInfo, Agent);
+				IPerforceServer? Server = await LoadBalancer.SelectServerAsync(Cluster, Agent);
 				if (Server == null)
 				{
 					return false;
@@ -846,15 +829,15 @@ namespace HordeServer.Models
 
 			// Find the matching credentials for the desired user
 			PerforceCredentials? Credentials = null;
-			if (ClusterInfo != null)
+			if (Cluster != null)
 			{
 				if (Workspace.UserName == null)
 				{
-					Credentials = ClusterInfo.Credentials.FirstOrDefault();
+					Credentials = Cluster.Credentials.FirstOrDefault();
 				}
 				else
 				{
-					Credentials = ClusterInfo.Credentials.FirstOrDefault(x => String.Equals(x.UserName, Workspace.UserName, StringComparison.OrdinalIgnoreCase));
+					Credentials = Cluster.Credentials.FirstOrDefault(x => String.Equals(x.UserName, Workspace.UserName, StringComparison.OrdinalIgnoreCase));
 				}
 			}
 
@@ -862,7 +845,7 @@ namespace HordeServer.Models
 			HordeCommon.Rpc.Messages.AgentWorkspace Result = new HordeCommon.Rpc.Messages.AgentWorkspace();
 			Result.ConfiguredCluster = Workspace.Cluster;
 			Result.ConfiguredUserName = Workspace.UserName;
-			Result.Cluster = ClusterInfo?.Name;
+			Result.Cluster = Cluster?.Name;
 			Result.ServerAndPort = ServerAndPort;
 			Result.UserName = Credentials?.UserName ?? Workspace.UserName;
 			Result.Password = Credentials?.Password;
diff --git a/Engine/Source/Programs/Horde/HordeServer/Models/Globals.cs b/Engine/Source/Programs/Horde/HordeServer/Models/Globals.cs
--- a/Engine/Source/Programs/Horde/HordeServer/Models/Globals.cs
+++ b/Engine/Source/Programs/Horde/HordeServer/Models/Globals.cs
@@ -178,6 +178,33 @@ namespace HordeServer.Models
 		public string Message { get; set; } = String.Empty;
 	}
 
+	/// <summary>
+	/// Path to a platform and stream to use for syncing AutoSDK
+	/// </summary>
+	public class AutoSdkWorkspace
+	{
+		/// <summary>
+		/// Name of this workspace
+		/// </summary>
+		public string? Name { get; set; }
+
+		/// <summary>
+		/// The agent properties to check (eg. "OSFamily=Windows")
+		/// </summary>
+		public List<string> Properties { get; set; } = new List<string>();
+
+		/// <summary>
+		/// Username for logging in to the server
+		/// </summary>
+		public string? UserName { get; set; }
+
+		/// <summary>
+		/// Stream to use
+		/// </summary>
+		[Required]
+		public string? Stream { get; set; }
+	}
+
 	/// <summary>
 	/// Information about an individual Perforce server
 	/// </summary>
@@ -245,6 +272,11 @@ namespace HordeServer.Models
 		/// List of server credentials
 		/// </summary>
 		public List<PerforceCredentials> Credentials { get; set; } = new List<PerforceCredentials>();
+
+		/// <summary>
+		/// List of autosdk streams
+		/// </summary>
+		public List<AutoSdkWorkspace> AutoSdk { get; set; } = new List<AutoSdkWorkspace>();
 	}
 
 	/// <summary>
@@ -323,6 +355,23 @@ namespace HordeServer.Models
 			}
 		}
 
+		/// <summary>
+		/// Finds a perforce cluster with the given name
+		/// </summary>
+		/// <param name="Name">Name of the cluster</param>
+		/// <returns></returns>
+		public PerforceCluster? FindPerforceCluster(string? Name)
+		{
+			if (Name == null)
+			{
+				return PerforceClusters.FirstOrDefault();
+			}
+			else
+			{
+				return PerforceClusters.FirstOrDefault(x => String.Equals(x.Name, Name, StringComparison.OrdinalIgnoreCase));
+			}
+		}
+
 		/// <inheritdoc/>
 		public override void PostLoad()
 		{
diff --git a/Engine/Source/Programs/Horde/HordeServer/Services/ConfigService.cs b/Engine/Source/Programs/Horde/HordeServer/Services/ConfigService.cs
--- a/Engine/Source/Programs/Horde/HordeServer/Services/ConfigService.cs
+++ b/Engine/Source/Programs/Horde/HordeServer/Services/ConfigService.cs
@@ -33,7 +33,7 @@ namespace HordeServer.Services
 		/// <summary>
 		/// Config file version number
 		/// </summary>
-		const int Version = 3;
+		const int Version = 4;
 
 		/// <summary>
 		/// Database service
diff --git a/Engine/Source/Programs/Horde/HordeServer/Services/PoolService.cs b/Engine/Source/Programs/Horde/HordeServer/Services/PoolService.cs
--- a/Engine/Source/Programs/Horde/HordeServer/Services/PoolService.cs
+++ b/Engine/Source/Programs/Horde/HordeServer/Services/PoolService.cs
@@ -24,6 +24,11 @@ namespace HordeServer.Services
 	/// </summary>
 	public class PoolService
 	{
+		/// <summary>
+		/// The database service instance
+		/// </summary>
+		DatabaseService DatabaseService;
+
 		/// <summary>
 		/// Collection of pool documents
 		/// </summary>
@@ -42,10 +47,12 @@ namespace HordeServer.Services
 		/// <summary>
 		/// Constructor
 		/// </summary>
+		/// <param name="DatabaseService"></param>
 		/// <param name="Pools">Collection of pool documents</param>
 		/// <param name="Clock"></param>
-		public PoolService(IPoolCollection Pools, IClock Clock)
+		public PoolService(DatabaseService DatabaseService, IPoolCollection Pools, IClock Clock)
 		{
+			this.DatabaseService = DatabaseService;
 			this.Pools = Pools;
 			this.Clock = Clock;
 		}
@@ -128,14 +135,12 @@ namespace HordeServer.Services
 		{
 			HashSet<AgentWorkspace> Workspaces = new HashSet<AgentWorkspace>();
 
-			AgentWorkspace? AutoSdkWorkspace = Agent.GetAutoSdkWorkspace();
-			if(AutoSdkWorkspace != null)
-			{
-				Workspaces.Add(AutoSdkWorkspace);
-			}
-
 			Dictionary<PoolId, IPool> PoolMapping = await GetPoolLookupAsync(ValidAtTime);
 			Workspaces.UnionWith(Agent.GetPools(PoolMapping.Values).SelectMany(x => x.Workspaces));
+
+			Globals Globals = await DatabaseService.GetGlobalsAsync();
+			Workspaces.UnionWith(Agent.GetAutoSdkWorkspaces(Globals, Workspaces.ToList()));
+
 			return Workspaces;
 		}
 
diff --git a/Engine/Source/Programs/Horde/HordeServer/Tasks/Impl/ConformTaskSource.cs b/Engine/Source/Programs/Horde/HordeServer/Tasks/Impl/ConformTaskSource.cs
--- a/Engine/Source/Programs/Horde/HordeServer/Tasks/Impl/ConformTaskSource.cs
+++ b/Engine/Source/Programs/Horde/HordeServer/Tasks/Impl/ConformTaskSource.cs
@@ -192,7 +192,8 @@ namespace HordeServer.Tasks.Impl
 			HashSet<AgentWorkspace> ConformWorkspaces = await PoolService.GetWorkspacesAsync(Agent, DateTime.UtcNow);
 			foreach (AgentWorkspace ConformWorkspace in ConformWorkspaces)
 			{
-				if (!await Agent.TryAddWorkspaceMessage(ConformWorkspace, Globals, PerforceLoadBalancer, Workspaces))
+				PerforceCluster? Cluster = Globals.FindPerforceCluster(ConformWorkspace.Cluster);
+				if (Cluster == null || !await Agent.TryAddWorkspaceMessage(ConformWorkspace, Cluster, PerforceLoadBalancer, Workspaces))
 				{
 					return false;
 				}
diff --git a/Engine/Source/Programs/Horde/HordeServer/Tasks/Impl/JobTaskSource.cs b/Engine/Source/Programs/Horde/HordeServer/Tasks/Impl/JobTaskSource.cs
--- a/Engine/Source/Programs/Horde/HordeServer/Tasks/Impl/JobTaskSource.cs
+++ b/Engine/Source/Programs/Horde/HordeServer/Tasks/Impl/JobTaskSource.cs
@@ -877,17 +877,23 @@ namespace HordeServer.Tasks.Impl
 
 			List<HordeCommon.Rpc.Messages.AgentWorkspace> Workspaces = new List<HordeCommon.Rpc.Messages.AgentWorkspace>();
 
-			AgentWorkspace? AutoSdkWorkspace = Agent.GetAutoSdkWorkspace();
+			PerforceCluster? Cluster = Globals.FindPerforceCluster(Workspace.Cluster);
+			if (Cluster == null)
+			{
+				return null;
+			}
+
+			AgentWorkspace? AutoSdkWorkspace = Agent.GetAutoSdkWorkspace(Cluster);
 			if (AutoSdkWorkspace != null)
 			{
-				if (!await Agent.TryAddWorkspaceMessage(AutoSdkWorkspace, Globals, PerforceLoadBalancer, Workspaces))
+				if (!await Agent.TryAddWorkspaceMessage(AutoSdkWorkspace, Cluster, PerforceLoadBalancer, Workspaces))
 				{
 					return null;
 				}
 				Task.AutoSdkWorkspace = Workspaces.Last();
 			}
 
-			if (!await Agent.TryAddWorkspaceMessage(Workspace, Globals, PerforceLoadBalancer, Workspaces))
+			if (!await Agent.TryAddWorkspaceMessage(Workspace, Cluster, PerforceLoadBalancer, Workspaces))
 			{
 				return null;
 			}
diff --git a/Engine/Source/Programs/Horde/HordeServerTests/PoolServiceTests.cs b/Engine/Source/Programs/Horde/HordeServerTests/PoolServiceTests.cs
--- a/Engine/Source/Programs/Horde/HordeServerTests/PoolServiceTests.cs
+++ b/Engine/Source/Programs/Horde/HordeServerTests/PoolServiceTests.cs
@@ -25,7 +25,7 @@ namespace HordeServerTests
         public PoolServiceTests()
         {
             var PoolCollection = new PoolCollection(GetDatabaseService());
-            PoolService = new PoolService(PoolCollection, new FakeClock());
+            PoolService = new PoolService(GetDatabaseService(), PoolCollection, new FakeClock());
         }
 
         private async Task<IPool> CreatePoolFixture(String Name)
