Horde: Adding dashboard and server version information to settings view
#rnx

diff --git a/Engine/Source/Programs/Horde/HordeDashboard/src/backend/Backend.ts b/Engine/Source/Programs/Horde/HordeDashboard/src/backend/Backend.ts
--- a/Engine/Source/Programs/Horde/HordeDashboard/src/backend/Backend.ts
+++ b/Engine/Source/Programs/Horde/HordeDashboard/src/backend/Backend.ts
@@ -1170,6 +1170,16 @@ export class Backend {
         
     }
 
+    getServerVersion():Promise<string> {
+        return new Promise<string>((resolve, reject) => {
+            this.backend.get(`/api/v1/server/version`).then((value) => {
+                resolve(value.data as string);
+            }).catch(reason => {
+                reject(reason);
+            });
+        });
+    }
+
     private async update() {
 
         if (this.updateID === "updating") {
@@ -1236,6 +1246,8 @@ export class Backend {
                 return;
             }
 
+            this.serverVersion = await this.getServerVersion();
+
             await dashboard.update();
 
             if (dashboard.localCache) {
@@ -1260,6 +1272,8 @@ export class Backend {
     updateID?: any;
     logout: boolean = false;
 
+    serverVersion:string = "";
+
 
     private backend: Fetch;
 
