Horde: Move base RemoteExecution classes to EpicGames.Horde.Common
Fix IoHash Digest size being set to the size of the hash rather than the size of the hashed data

#rnx

diff --git a/Engine/Source/Programs/Horde/HordeAgent/Commands/Execution/ExecuteCommand.cs b/Engine/Source/Programs/Horde/HordeAgent/Commands/Execution/ExecuteCommand.cs
--- a/Engine/Source/Programs/Horde/HordeAgent/Commands/Execution/ExecuteCommand.cs
+++ b/Engine/Source/Programs/Horde/HordeAgent/Commands/Execution/ExecuteCommand.cs
@@ -29,6 +29,7 @@ using Microsoft.Extensions.DependencyInjection;
 using Microsoft.Extensions.Configuration;
 using Serilog.Events;
 using Digest = Build.Bazel.Remote.Execution.V2.Digest;
+using EpicGames.Horde.Common.RemoteExecution;
 
 namespace HordeAgent.Commands
 {
@@ -38,127 +39,6 @@ namespace HordeAgent.Commands
 	[Command("Execute", "Executes a command through the remote execution API")]
 	class ExecuteCommand : Command
 	{
-		/// <summary>
-		/// List of files that need to be uploaded
-		/// </summary>
-		class UploadList
-		{
-			public Dictionary<string, BatchUpdateBlobsRequest.Types.Request> Blobs = new Dictionary<string, BatchUpdateBlobsRequest.Types.Request>();
-
-			public Digest Add<T>(IMessage<T> Message) where T : IMessage<T>
-			{
-				ByteString ByteString = Message.ToByteString();
-
-				Digest Digest = StorageExtensions.GetDigest(ByteString);
-				if (!Blobs.ContainsKey(Digest.Hash))
-				{
-					BatchUpdateBlobsRequest.Types.Request Request = new BatchUpdateBlobsRequest.Types.Request();
-					Request.Data = ByteString;
-					Request.Digest = Digest;
-					Blobs.Add(Digest.Hash, Request);
-				}
-				return Digest;
-			}
-
-			public async Task<Digest> AddFileAsync(string LocalFile)
-			{
-				BatchUpdateBlobsRequest.Types.Request Content = new BatchUpdateBlobsRequest.Types.Request();
-				using (FileStream Stream = File.Open(LocalFile, FileMode.Open, FileAccess.Read, FileShare.Read))
-				{
-					Content.Data = await ByteString.FromStreamAsync(Stream);
-					Content.Digest = StorageExtensions.GetDigest(Content.Data);
-				}
-				Blobs[Content.Digest.Hash] = Content;
-				return Content.Digest;
-			}
-		}
-
-		class JsonFileNode
-		{
-			public string Name { get; set; } = String.Empty;
-			public string LocalFile { get; set; } = String.Empty;
-
-			public async Task<FileNode> BuildAsync(UploadList UploadList)
-			{
-				FileNode Node = new FileNode();
-				Node.Name = String.IsNullOrEmpty(Name) ? Path.GetFileName(LocalFile) : Name;
-				Node.Digest = await UploadList.AddFileAsync(LocalFile);
-				return Node;
-			}
-		}
-
-		class JsonDirectory
-		{
-			public List<JsonFileNode> Files { get; set; } = new List<JsonFileNode>();
-			public List<JsonDirectoryNode> Directories { get; set; } = new List<JsonDirectoryNode>();
-
-			public async Task<Digest> BuildAsync(UploadList UploadList)
-			{
-				Directory Directory = new Directory();
-				foreach (JsonFileNode File in Files)
-				{
-					Directory.Files.Add(await File.BuildAsync(UploadList));
-				}
-				foreach (JsonDirectoryNode DirectoryNode in Directories)
-				{
-					Directory.Directories.Add(await DirectoryNode.BuildAsync(UploadList));
-				}
-				return UploadList.Add(Directory);
-			}
-		}
-
-		class JsonDirectoryNode : JsonDirectory
-		{
-			public string Name { get; set; } = String.Empty;
-
-			public new async Task<DirectoryNode> BuildAsync(UploadList UploadList)
-			{
-				DirectoryNode Node = new DirectoryNode();
-				Node.Name = Name;
-				Node.Digest = await base.BuildAsync(UploadList);
-				return Node;
-			}
-		}
-
-		class JsonCommand
-		{
-			public List<string> OutputPaths { get; set; } = new List<string>();
-			public List<string> Arguments { get; set; } = new List<string>();
-			public Dictionary<string, string> EnvVars { get; set; } = new Dictionary<string, string>();
-			public string WorkingDirectory { get; set; } = String.Empty;
-
-			public Digest Build(UploadList UploadList)
-			{
-				RpcCommand Command = new RpcCommand();
-				Command.Arguments.Add(Arguments);
-				Command.WorkingDirectory = WorkingDirectory;
-				Command.OutputPaths.Add(OutputPaths);
-				
-				foreach (KeyValuePair<string, string> Pair in EnvVars)
-				{
-					Command.EnvironmentVariables.Add(new RpcCommand.Types.EnvironmentVariable { Name = Pair.Key, Value = Pair.Value });
-				}
-				
-				return UploadList.Add(Command);
-			}
-		}
-
-		class JsonAction
-		{
-			public JsonCommand Command { get; set; } = new JsonCommand();
-			public JsonDirectory Workspace { get; set; } = new JsonDirectory();
-
-			public async Task<Digest> BuildAsync(UploadList UploadList, ByteString? Salt, bool DoNotCache)
-			{
-				Action NewAction = new Action();
-				NewAction.CommandDigest = Command.Build(UploadList);
-				NewAction.InputRootDigest = await Workspace.BuildAsync(UploadList);
-				NewAction.Salt = Salt;
-				NewAction.DoNotCache = DoNotCache;
-				return UploadList.Add(NewAction);
-			}
-		}
-
 		/// <summary>
 		/// Input file describing the work to execute
 		/// </summary>
diff --git a/Engine/Source/Programs/Horde/HordeAgent/Execution/ActionExecutor.cs b/Engine/Source/Programs/Horde/HordeAgent/Execution/ActionExecutor.cs
--- a/Engine/Source/Programs/Horde/HordeAgent/Execution/ActionExecutor.cs
+++ b/Engine/Source/Programs/Horde/HordeAgent/Execution/ActionExecutor.cs
@@ -27,6 +27,7 @@ using HordeAgent.Utility;
 using HordeCommon.Rpc.Tasks;
 using Microsoft.Extensions.Logging;
 using HordeCommon.Rpc;
+using EpicGames.Horde.Common.RemoteExecution;
 
 namespace HordeAgent
 {
diff --git a/Engine/Source/Programs/Horde/HordeAgent/HordeAgent.csproj b/Engine/Source/Programs/Horde/HordeAgent/HordeAgent.csproj
--- a/Engine/Source/Programs/Horde/HordeAgent/HordeAgent.csproj
+++ b/Engine/Source/Programs/Horde/HordeAgent/HordeAgent.csproj
@@ -22,7 +22,6 @@
   <ItemGroup>
     <PackageReference Include="AWSSDK.Core" Version="3.5.1.49" />
     <PackageReference Include="AWSSDK.EC2" Version="3.5.1.2" />
-    <PackageReference Include="Blake3" Version="0.3.0" />
     <PackageReference Include="Datadog.Trace" Version="1.18.0" />
     <PackageReference Include="Google.Protobuf" Version="3.12.3" />
     <PackageReference Include="Grpc.Net.Client" Version="2.29.0" />
@@ -90,18 +89,9 @@
   </ItemGroup>
 
   <ItemGroup>
-    <Protobuf Include="..\..\Shared\EpicGames.Horde.Common\Protos\build\bazel\remote\execution\v2\remote_execution.proto" ProtoRoot="..\..\Shared\EpicGames.Horde.Common\Protos" GrpcServices="Client" />
-    <Protobuf Include="..\..\Shared\EpicGames.Horde.Common\Protos\build\bazel\semver\semver.proto" ProtoRoot="..\..\Shared\EpicGames.Horde.Common\Protos" GrpcServices="Client" />
-    <Protobuf Include="..\..\Shared\EpicGames.Horde.Common\Protos\google\api\annotations.proto" ProtoRoot="..\..\Shared\EpicGames.Horde.Common\Protos" GrpcServices="Client" />
-    <Protobuf Include="..\..\Shared\EpicGames.Horde.Common\Protos\google\api\client.proto" ProtoRoot="..\..\Shared\EpicGames.Horde.Common\Protos" GrpcServices="Client" />
-    <Protobuf Include="..\..\Shared\EpicGames.Horde.Common\Protos\google\api\field_behavior.proto" ProtoRoot="..\..\Shared\EpicGames.Horde.Common\Protos" GrpcServices="Client" />
-    <Protobuf Include="..\..\Shared\EpicGames.Horde.Common\Protos\google\api\http.proto" ProtoRoot="..\..\Shared\EpicGames.Horde.Common\Protos" GrpcServices="Client" />
-    <Protobuf Include="..\..\Shared\EpicGames.Horde.Common\Protos\google\api\resource.proto" ProtoRoot="..\..\Shared\EpicGames.Horde.Common\Protos" GrpcServices="Client" />
     <Protobuf Include="..\..\Shared\EpicGames.Horde.Common\Protos\google\bytestream\bytestream.proto" ProtoRoot="..\..\Shared\EpicGames.Horde.Common\Protos" GrpcServices="Client" />
     <Protobuf Include="..\..\Shared\EpicGames.Horde.Common\Protos\google\devtools\remoteworkers\v1test2\bots.proto" ProtoRoot="..\..\Shared\EpicGames.Horde.Common\Protos" GrpcServices="Client" />
     <Protobuf Include="..\..\Shared\EpicGames.Horde.Common\Protos\google\devtools\remoteworkers\v1test2\worker.proto" ProtoRoot="..\..\Shared\EpicGames.Horde.Common\Protos" GrpcServices="Client" />
-    <Protobuf Include="..\..\Shared\EpicGames.Horde.Common\Protos\google\longrunning\operations.proto" ProtoRoot="..\..\Shared\EpicGames.Horde.Common\Protos" GrpcServices="Client" />
-    <Protobuf Include="..\..\Shared\EpicGames.Horde.Common\Protos\google\rpc\code.proto" ProtoRoot="..\..\Shared\EpicGames.Horde.Common\Protos" GrpcServices="Client" />
     <Protobuf Include="..\..\Shared\EpicGames.Horde.Common\Protos\horde\action_rpc.proto" ProtoRoot="..\..\Shared\EpicGames.Horde.Common\Protos" GrpcServices="Client" />
     <Protobuf Include="..\..\Shared\EpicGames.Horde.Common\Protos\horde\horde_rpc.proto" ProtoRoot="..\..\Shared\EpicGames.Horde.Common\Protos" GrpcServices="Client" />
     <Protobuf Include="..\..\Shared\EpicGames.Horde.Common\Protos\horde\tasks\action_task.proto" ProtoRoot="..\..\Shared\EpicGames.Horde.Common\Protos" GrpcServices="Client" />
diff --git a/Engine/Source/Programs/Horde/HordeAgent/Utility/StorageExtensions.cs b/Engine/Source/Programs/Horde/HordeAgent/Utility/StorageExtensions.cs
deleted file mode 100644
diff --git a/Engine/Source/Programs/Horde/HordeCommon/HordeCommon.csproj b/Engine/Source/Programs/Horde/HordeCommon/HordeCommon.csproj
--- a/Engine/Source/Programs/Horde/HordeCommon/HordeCommon.csproj
+++ b/Engine/Source/Programs/Horde/HordeCommon/HordeCommon.csproj
@@ -12,7 +12,7 @@
       <PrivateAssets>all</PrivateAssets>
       <IncludeAssets>runtime; build; native; contentfiles; analyzers; buildtransitive</IncludeAssets>
     </PackageReference>
-    <PackageReference Include="Microsoft.Extensions.Logging" Version="3.1.5" />
+    <PackageReference Include="Microsoft.Extensions.Logging" Version="5.0.0" />
   </ItemGroup>
 
   <ItemGroup>
diff --git a/Engine/Source/Programs/Horde/HordeServer/HordeServer.csproj b/Engine/Source/Programs/Horde/HordeServer/HordeServer.csproj
--- a/Engine/Source/Programs/Horde/HordeServer/HordeServer.csproj
+++ b/Engine/Source/Programs/Horde/HordeServer/HordeServer.csproj
@@ -57,14 +57,7 @@
   </ItemGroup>
 
   <ItemGroup>
-    <Protobuf Include="..\..\Shared\EpicGames.Horde.Common\Protos\build\bazel\remote\execution\v2\remote_execution.proto" ProtoRoot="..\..\Shared\EpicGames.Horde.Common\Protos" GrpcServices="Server" />
-    <Protobuf Include="..\..\Shared\EpicGames.Horde.Common\Protos\build\bazel\semver\semver.proto" ProtoRoot="..\..\Shared\EpicGames.Horde.Common\Protos" GrpcServices="Server" />
-    <Protobuf Include="..\..\Shared\EpicGames.Horde.Common\Protos\google\api\annotations.proto" ProtoRoot="..\..\Shared\EpicGames.Horde.Common\Protos" GrpcServices="Server" />
-    <Protobuf Include="..\..\Shared\EpicGames.Horde.Common\Protos\google\api\client.proto" ProtoRoot="..\..\Shared\EpicGames.Horde.Common\Protos" GrpcServices="Server" />
-    <Protobuf Include="..\..\Shared\EpicGames.Horde.Common\Protos\google\api\http.proto" ProtoRoot="..\..\Shared\EpicGames.Horde.Common\Protos" GrpcServices="Server" />
     <Protobuf Include="..\..\Shared\EpicGames.Horde.Common\Protos\google\bytestream\bytestream.proto" ProtoRoot="..\..\Shared\EpicGames.Horde.Common\Protos" GrpcServices="Server" />
-    <Protobuf Include="..\..\Shared\EpicGames.Horde.Common\Protos\google\longrunning\operations.proto" ProtoRoot="..\..\Shared\EpicGames.Horde.Common\Protos" GrpcServices="Server" />
-    <Protobuf Include="..\..\Shared\EpicGames.Horde.Common\Protos\google\rpc\code.proto" ProtoRoot="..\..\Shared\EpicGames.Horde.Common\Protos" GrpcServices="Server" />
     <Protobuf Include="..\..\Shared\EpicGames.Horde.Common\Protos\horde\action_rpc.proto" ProtoRoot="..\..\Shared\EpicGames.Horde.Common\Protos" GrpcServices="Server" />
     <Protobuf Include="..\..\Shared\EpicGames.Horde.Common\Protos\horde\horde_rpc.proto" ProtoRoot="..\..\Shared\EpicGames.Horde.Common\Protos" GrpcServices="Server" />
     <Protobuf Include="..\..\Shared\EpicGames.Horde.Common\Protos\horde\tasks\action_task.proto" ProtoRoot="..\..\Shared\EpicGames.Horde.Common\Protos" GrpcServices="Server" />
diff --git a/Engine/Source/Programs/Shared/EpicGames.Horde.Common/EpicGames.Horde.Common.csproj b/Engine/Source/Programs/Shared/EpicGames.Horde.Common/EpicGames.Horde.Common.csproj
--- a/Engine/Source/Programs/Shared/EpicGames.Horde.Common/EpicGames.Horde.Common.csproj
+++ b/Engine/Source/Programs/Shared/EpicGames.Horde.Common/EpicGames.Horde.Common.csproj
@@ -33,10 +33,23 @@
       <PrivateAssets>all</PrivateAssets>
       <IncludeAssets>runtime; build; native; contentfiles; analyzers; buildtransitive</IncludeAssets>
     </PackageReference>
-    <PackageReference Include="Microsoft.Extensions.Logging" Version="3.1.5" />
+    <PackageReference Include="Microsoft.Extensions.Logging" Version="5.0.0" />
+  </ItemGroup>
+  <ItemGroup>
+    <ProjectReference Include="..\EpicGames.Core\EpicGames.Core.csproj" />
+    <ProjectReference Include="..\EpicGames.IoHash\EpicGames.IoHash.csproj" />
   </ItemGroup>
 
   <ItemGroup>
+    <Protobuf Include="Protos\build\bazel\remote\execution\v2\remote_execution.proto" ProtoRoot="Protos" />
+    <Protobuf Include="Protos\build\bazel\semver\semver.proto" ProtoRoot="Protos" />
+    <Protobuf Include="Protos\google\api\annotations.proto" ProtoRoot="Protos" GrpcServices="Client" />
+    <Protobuf Include="Protos\google\api\client.proto" ProtoRoot="Protos" GrpcServices="Client" />
+    <Protobuf Include="Protos\google\api\field_behavior.proto" ProtoRoot="Protos" GrpcServices="Client" />
+    <Protobuf Include="Protos\google\api\http.proto" ProtoRoot="Protos" GrpcServices="Client" />
+    <Protobuf Include="Protos\google\api\resource.proto" ProtoRoot="Protos" GrpcServices="Client" />
+    <Protobuf Include="Protos\google\longrunning\operations.proto" ProtoRoot="Protos" GrpcServices="Client" />
+    <Protobuf Include="Protos\google\rpc\code.proto" ProtoRoot="Protos" GrpcServices="Client" />
     <Protobuf Include="Protos\google\rpc\status.proto" ProtoRoot="Protos" />
   </ItemGroup>
 </Project>
\ No newline at end of file
diff --git a/Engine/Source/Programs/Shared/EpicGames.Horde.Common/RemoteExecution/Json.cs b/Engine/Source/Programs/Shared/EpicGames.Horde.Common/RemoteExecution/Json.cs
new file mode 100644
--- /dev/null
+++ b/Engine/Source/Programs/Shared/EpicGames.Horde.Common/RemoteExecution/Json.cs
@@ -0,0 +1,101 @@
+﻿// Copyright Epic Games, Inc. All Rights Reserved.
+
+using Build.Bazel.Remote.Execution.V2;
+using Google.Protobuf;
+using System;
+using System.Collections.Generic;
+using System.IO;
+using System.Threading.Tasks;
+using Action = Build.Bazel.Remote.Execution.V2.Action;
+using Digest = Build.Bazel.Remote.Execution.V2.Digest;
+using Directory = Build.Bazel.Remote.Execution.V2.Directory;
+using RpcCommand = Build.Bazel.Remote.Execution.V2.Command;
+
+namespace EpicGames.Horde.Common.RemoteExecution
+{
+	public class JsonFileNode
+	{
+		public string Name { get; set; } = String.Empty;
+		public string LocalFile { get; set; } = String.Empty;
+
+		public async Task<FileNode> BuildAsync(UploadList UploadList)
+		{
+			FileNode Node = new FileNode();
+			Node.Name = String.IsNullOrEmpty(Name) ? Path.GetFileName(LocalFile) : Name;
+			Node.Digest = await UploadList.AddFileAsync(LocalFile);
+			return Node;
+		}
+	}
+
+	public class JsonDirectory
+	{
+		public List<JsonFileNode> Files { get; set; } = new List<JsonFileNode>();
+		public List<JsonDirectoryNode> Directories { get; set; } = new List<JsonDirectoryNode>();
+
+		public async Task<Digest> BuildAsync(UploadList UploadList)
+		{
+			Directory Directory = new Directory();
+			foreach (JsonFileNode File in Files)
+			{
+				Directory.Files.Add(await File.BuildAsync(UploadList));
+			}
+			foreach (JsonDirectoryNode DirectoryNode in Directories)
+			{
+				Directory.Directories.Add(await DirectoryNode.BuildAsync(UploadList));
+			}
+			return UploadList.Add(Directory);
+		}
+	}
+
+	public class JsonDirectoryNode : JsonDirectory
+	{
+		public string Name { get; set; } = String.Empty;
+
+		public new async Task<DirectoryNode> BuildAsync(UploadList UploadList)
+		{
+			DirectoryNode Node = new DirectoryNode();
+			Node.Name = Name;
+			Node.Digest = await base.BuildAsync(UploadList);
+			return Node;
+		}
+	}
+
+	public class JsonCommand
+	{
+		public List<string> OutputPaths { get; set; } = new List<string>();
+		public List<string> Arguments { get; set; } = new List<string>();
+		public Dictionary<string, string> EnvVars { get; set; } = new Dictionary<string, string>();
+		public string WorkingDirectory { get; set; } = String.Empty;
+
+		public Digest Build(UploadList UploadList)
+		{
+			RpcCommand Command = new RpcCommand();
+			Command.Arguments.Add(Arguments);
+			Command.WorkingDirectory = WorkingDirectory;
+			Command.OutputPaths.Add(OutputPaths);
+
+			foreach (KeyValuePair<string, string> Pair in EnvVars)
+			{
+				Command.EnvironmentVariables.Add(new RpcCommand.Types.EnvironmentVariable { Name = Pair.Key, Value = Pair.Value });
+			}
+
+			return UploadList.Add(Command);
+		}
+	}
+
+	public class JsonAction
+	{
+		public JsonCommand Command { get; set; } = new JsonCommand();
+		public JsonDirectory Workspace { get; set; } = new JsonDirectory();
+
+		public async Task<Digest> BuildAsync(UploadList UploadList, ByteString? Salt, bool DoNotCache)
+		{
+			Action NewAction = new Action();
+			NewAction.CommandDigest = Command.Build(UploadList);
+			NewAction.InputRootDigest = await Workspace.BuildAsync(UploadList);
+			NewAction.Salt = Salt;
+			NewAction.DoNotCache = DoNotCache;
+			return UploadList.Add(NewAction);
+		}
+	}
+}
diff --git a/Engine/Source/Programs/Shared/EpicGames.Horde.Common/RemoteExecution/StorageExtensions.cs b/Engine/Source/Programs/Shared/EpicGames.Horde.Common/RemoteExecution/StorageExtensions.cs
new file mode 100644
--- /dev/null
+++ b/Engine/Source/Programs/Shared/EpicGames.Horde.Common/RemoteExecution/StorageExtensions.cs
@@ -0,0 +1,113 @@
+// Copyright Epic Games, Inc. All Rights Reserved.
+
+using Build.Bazel.Remote.Execution.V2;
+using EpicGames.Core;
+using Google.Protobuf;
+using System;
+using System.Security.Cryptography;
+using System.Threading.Tasks;
+using ContentAddressableStorageClient = Build.Bazel.Remote.Execution.V2.ContentAddressableStorage.ContentAddressableStorageClient;
+using Digest = Build.Bazel.Remote.Execution.V2.Digest;
+
+namespace EpicGames.Horde.Common.RemoteExecution
+{
+	public static class StorageExtensions
+	{
+		public static async Task<byte[]> GetBulkDataAsync(this ContentAddressableStorageClient Storage, string InstanceName, Digest Digest)
+		{
+			BatchReadBlobsRequest Request = new BatchReadBlobsRequest();
+			Request.InstanceName = InstanceName;
+			Request.Digests.Add(Digest);
+
+			BatchReadBlobsResponse Response = await Storage.BatchReadBlobsAsync(Request);
+			return Response.Responses[0].Data.ToByteArray();
+		}
+
+		public static async Task<Digest> PutBulkDataAsync(this ContentAddressableStorageClient Storage, string InstanceName, byte[] Data)
+		{
+			BatchUpdateBlobsRequest.Types.Request BlobRequest = new BatchUpdateBlobsRequest.Types.Request();
+			BlobRequest.Data = ByteString.CopyFrom(Data);
+			BlobRequest.Digest = GetDigest(Data);
+
+			BatchUpdateBlobsRequest Request = new BatchUpdateBlobsRequest();
+			Request.InstanceName = InstanceName;
+			Request.Requests.Add(BlobRequest);
+
+			BatchUpdateBlobsResponse Response = await Storage.BatchUpdateBlobsAsync(Request);
+			//?
+			return BlobRequest.Digest;
+		}
+
+		public static async Task<T> GetProtoMessageAsync<T>(this ContentAddressableStorageClient Storage, string InstanceName, Digest Digest) where T : class, IMessage<T>, new()
+		{
+			BatchReadBlobsRequest Request = new BatchReadBlobsRequest();
+			Request.InstanceName = InstanceName;
+			Request.Digests.Add(Digest);
+
+			BatchReadBlobsResponse Response = await Storage.BatchReadBlobsAsync(Request);
+			if (Response.Responses[0].Status.Code != (int) Google.Rpc.Code.Ok)
+			{
+				throw new Exception($"Failed getting digest {Digest}. Code={Response.Responses[0].Status.Code} Message={Response.Responses[0].Status.Message}");
+			}
+
+			T Item = new T();
+			Item.MergeFrom(Response.Responses[0].Data);
+			return Item;
+		}
+
+		public static async Task<Digest> PutProtoMessageAsync<T>(this ContentAddressableStorageClient Storage, string InstanceName, IMessage<T> Message) where T : class, IMessage<T>
+		{
+			BatchUpdateBlobsRequest.Types.Request BlobRequest = new BatchUpdateBlobsRequest.Types.Request();
+			BlobRequest.Data = Message.ToByteString();
+			BlobRequest.Digest = GetDigest(BlobRequest.Data.ToByteArray());
+
+			BatchUpdateBlobsRequest Request = new BatchUpdateBlobsRequest();
+			Request.InstanceName = InstanceName;
+			Request.Requests.Add(BlobRequest);
+
+			BatchUpdateBlobsResponse Response = await Storage.BatchUpdateBlobsAsync(Request);
+			//??
+			return BlobRequest.Digest;
+		}
+		
+		public static Digest GetDigest(ByteString Data)
+		{
+			return GetDigest(Data.ToByteArray(), DigestFunction.Types.Value.Epiciohash);
+		}
+
+		public static Digest GetDigest(byte[] Data)
+		{
+			return GetDigest(Data, DigestFunction.Types.Value.Epiciohash);
+		}
+
+		public static Digest GetDigest(ByteString Data, DigestFunction.Types.Value DigestFunction)
+		{
+			return GetDigest(Data.ToByteArray(), DigestFunction);
+		}
+
+		public static Digest GetDigest(byte[] Data, DigestFunction.Types.Value DigestFunction)
+		{
+			switch (DigestFunction)
+			{
+				case Build.Bazel.Remote.Execution.V2.DigestFunction.Types.Value.Sha256: return GetDigestSha256(Data);
+				case Build.Bazel.Remote.Execution.V2.DigestFunction.Types.Value.Epiciohash: return GetDigestIoHash(Data);
+				default: throw new Exception("Unknown digest function " + DigestFunction);
+			}
+		}
+		
+		public static Digest GetDigestSha256(byte[] Data)
+		{
+			using (SHA256 Algorithm = SHA256.Create())
+			{
+				byte[] Hash = Algorithm.ComputeHash(Data);
+				return new Digest { Hash = StringUtils.FormatHexString(Hash), SizeBytes = Data.Length };
+			}
+		}
+		
+		public static Digest GetDigestIoHash(byte[] Data)
+		{
+			IoHash Hash = IoHash.Compute(Data.AsSpan());
+			return new Digest { Hash = StringUtils.FormatHexString(Hash.Span).ToUpper(), SizeBytes = Data.Length };
+		}
+	}
+}
diff --git a/Engine/Source/Programs/Shared/EpicGames.Horde.Common/RemoteExecution/UploadList.cs b/Engine/Source/Programs/Shared/EpicGames.Horde.Common/RemoteExecution/UploadList.cs
new file mode 100644
--- /dev/null
+++ b/Engine/Source/Programs/Shared/EpicGames.Horde.Common/RemoteExecution/UploadList.cs
@@ -0,0 +1,45 @@
+﻿// Copyright Epic Games, Inc. All Rights Reserved.
+
+using Build.Bazel.Remote.Execution.V2;
+using Google.Protobuf;
+using System.Collections.Generic;
+using System.IO;
+using System.Threading.Tasks;
+
+namespace EpicGames.Horde.Common.RemoteExecution
+{
+	/// <summary>
+	/// List of files that need to be uploaded
+	/// </summary>
+	public class UploadList
+	{
+		public Dictionary<string, BatchUpdateBlobsRequest.Types.Request> Blobs = new Dictionary<string, BatchUpdateBlobsRequest.Types.Request>();
+
+		public Digest Add<T>(IMessage<T> Message) where T : IMessage<T>
+		{
+			ByteString ByteString = Message.ToByteString();
+
+			Digest Digest = StorageExtensions.GetDigest(ByteString);
+			if (!Blobs.ContainsKey(Digest.Hash))
+			{
+				BatchUpdateBlobsRequest.Types.Request Request = new BatchUpdateBlobsRequest.Types.Request();
+				Request.Data = ByteString;
+				Request.Digest = Digest;
+				Blobs.Add(Digest.Hash, Request);
+			}
+			return Digest;
+		}
+
+		public async Task<Digest> AddFileAsync(string LocalFile)
+		{
+			BatchUpdateBlobsRequest.Types.Request Content = new BatchUpdateBlobsRequest.Types.Request();
+			using (FileStream Stream = File.Open(LocalFile, FileMode.Open, FileAccess.Read, FileShare.Read))
+			{
+				Content.Data = await ByteString.FromStreamAsync(Stream);
+				Content.Digest = StorageExtensions.GetDigest(Content.Data);
+			}
+			Blobs[Content.Digest.Hash] = Content;
+			return Content.Digest;
+		}
+	}
+}
