Horde: Enabling reshelving and adding debug controller test
#rnx

diff --git a/Engine/Source/Programs/Horde/HordeServer/Controllers/DebugController.cs b/Engine/Source/Programs/Horde/HordeServer/Controllers/DebugController.cs
--- a/Engine/Source/Programs/Horde/HordeServer/Controllers/DebugController.cs
+++ b/Engine/Source/Programs/Horde/HordeServer/Controllers/DebugController.cs
@@ -774,17 +774,16 @@ namespace HordeServer.Controllers
 			Bytes = await Perforce.PrintAsync($"{DevBuild}/RunUAT.bat@13916380");
 			Results.Add($"PrintAsync: {DevBuild}/RunUAT.bat@13916380", $"{System.Text.Encoding.Default.GetString(Bytes)}");
 
-			// DuplicateShelvedChangeAsync is disabled due to issue with p4 librarian
-			/*
 			// need to be running as service account for these
 
 			int? DuplicateCL = null;
 
-			int ShelvedChangeOnDevBuild = 16158333;
+			int ShelvedChangeOnDevBuild = 16687685;
 
 			try 
 			{
 				DuplicateCL = await Perforce.DuplicateShelvedChangeAsync(ShelvedChangeOnDevBuild);
+				Results.Add("DuplicateShelvedChangeAsync", $"Duplicated CL {DuplicateCL} from Shelved CL {ShelvedChangeOnDevBuild}");
 			}
 			catch 
 			{
@@ -794,7 +793,7 @@ namespace HordeServer.Controllers
 			if (DuplicateCL.HasValue)
 			{
 				// Note change client much exist for the update
-				ChangeDetails = await Perforce.GetChangeDetailsAsync(DevBuild, new int[]{NewChangeId}, PerforceUser);
+				ChangeDetails = await Perforce.GetChangeDetailsAsync(DevBuild, new int[]{ DuplicateCL.Value}, PerforceUser);
 
 				Results.Add("UpdateChangelistDescription - PreUpdate", $"{DuplicateCL} : {ChangeDetails[0].Description}"); 
 
@@ -802,7 +801,7 @@ namespace HordeServer.Controllers
 				{
 					await Perforce.UpdateChangelistDescription(DuplicateCL.Value, "Updated Description from Horde");
 
-					ChangeDetails = await Perforce.GetChangeDetailsAsync(DevBuild, new int[]{NewChangeId}, PerforceUser);
+					ChangeDetails = await Perforce.GetChangeDetailsAsync(DevBuild, new int[]{ DuplicateCL.Value}, PerforceUser);
 
 					Results.Add("UpdateChangelistDescription - PostUpdate", $"{DuplicateCL} : {ChangeDetails[0].Description}"); 
 				}
@@ -818,11 +817,9 @@ namespace HordeServer.Controllers
 				Results.Add("GetChangeDetailsAsync - After Delete", $"ChangeDetails.Count: {ChangeDetails.Count}"); 
 
 			}
-			*/
 
 			byte[] ImageBytes = await Perforce.PrintAsync($"{DevBuild}/Samples/Games/ShooterGame/ShooterGame.png");
 
-
 			Content.AppendLine("<html><body><pre>");
 			foreach (KeyValuePair<string, string> Pair in Results)
 			{
diff --git a/Engine/Source/Programs/Horde/HordeServer/Services/PerforceServiceApi.cs b/Engine/Source/Programs/Horde/HordeServer/Services/PerforceServiceApi.cs
--- a/Engine/Source/Programs/Horde/HordeServer/Services/PerforceServiceApi.cs
+++ b/Engine/Source/Programs/Horde/HordeServer/Services/PerforceServiceApi.cs
@@ -625,18 +625,16 @@ namespace HordeServer.Services
 		/// <inheritdoc/>
 		public override async Task<int> DuplicateShelvedChangeAsync(int ShelvedChange)
 		{
-			await Task.Yield();
-
-			throw new Exception("DuplicateShelvedChangeAsync is disabled due to p4 issue: https://jira.it.epicgames.com/servicedesk/customer/portal/1/ITH-144069");
-
-			/*
-			using (P4.Repository Repository = GetConnection(NoClient: true))
+			return await Task.Run(() =>
 			{
-				return ReshelveChange(Repository, ShelvedChange);
-			}
+				using (P4.Repository Repository = GetConnection(ReadOnly: false, ClientFromChange: ShelvedChange, UseClientFromChange: false))
+				{
+					return ReshelveChange(Repository, ShelvedChange);				
+				}
 
-			throw new Exception($"Failed to duplicate shelved change {ShelvedChange}");
-			*/
+				throw new Exception($"Unable to duplicate shelve change {ShelvedChange}");
+
+			});
 
 		}
 
@@ -1003,23 +1001,6 @@ namespace HordeServer.Services
 			});
 		}
 
-		void RequireEdgeServer()
-		{
-			string? Value;
-			if (Server.Metadata.RawData.TryGetValue("serverServices", out Value))
-			{
-				if (Value != "edge-server")
-				{
-					Serilog.Log.Logger.Information("Perforce: servierService is not an edge server : {ServerServices}", Value);
-					throw new Exception("Operation requires an edge server");
-				}
-			}
-			else
-			{
-				throw new Exception("unable to get server services");
-			}
-		}
-
 		bool CanImpersonate
 		{
 			get
@@ -1032,11 +1013,27 @@ namespace HordeServer.Services
 		int ReshelveChange(P4.Repository Repository, int Change)
 		{
 
-			// all attempts at using this without edge server have lead to bizarre reshelving behavior, also note that -p requires an edge server
-			// https://jira.it.epicgames.com/servicedesk/customer/portal/1/ITH-144069
-			RequireEdgeServer();
+			bool EdgeServer = false;
+			string? Value;
+			if (Server.Metadata.RawData.TryGetValue("serverServices", out Value))
+			{
+				if (Value == "edge-server")
+				{
+					EdgeServer = true;
+				}
+			}
+
+			List<string> Arguments = new List<string>();
+
+			// promote shelf if we're on an edge server
+			if (EdgeServer)
+			{
+				Arguments.Add("-p");
+			}
+
+			Arguments.AddRange(new string[] { "-s", Change.ToString(CultureInfo.InvariantCulture), "-f" });
 
-			using (P4.P4Command Cmd = new P4.P4Command(Repository, "reshelve", false, new string[] { "-s", Change.ToString(CultureInfo.InvariantCulture), "-f", "-p" }))
+			using (P4.P4Command Cmd = new P4.P4Command(Repository, "reshelve", false, Arguments.ToArray()))
 			{
 				P4.P4CommandResult Results = Cmd.Run();
 
