Horde: Enable shelf cloning for UE5/Main.

diff --git a/Engine/Source/Programs/Horde/HordeServer/Services/JobService.cs b/Engine/Source/Programs/Horde/HordeServer/Services/JobService.cs
--- a/Engine/Source/Programs/Horde/HordeServer/Services/JobService.cs
+++ b/Engine/Source/Programs/Horde/HordeServer/Services/JobService.cs
@@ -201,6 +201,11 @@ namespace HordeServer.Services
 			this.Logger = Logger;
 		}
 
+		static bool ShouldClonePreflightChange(StreamId StreamId)
+		{
+			return StreamId == new StreamId("ue5-main");
+		}
+
 		/// <summary>
 		/// Creates a new job
 		/// </summary>
@@ -233,9 +238,9 @@ namespace HordeServer.Services
 			ObjectId JobIdValue = JobId ?? ObjectId.GenerateNewId();
 			using IDisposable Scope = Logger.BeginScope("CreateJobAsync({JobId})", JobIdValue);
 
-			if (PreflightChange != null)
+			if (PreflightChange != null && ShouldClonePreflightChange(StreamId))
 			{
-				//PF				ClonedPreflightChange = await CloneShelvedChangeAsync(ClonedPreflightChange ?? PreflightChange.Value);
+				ClonedPreflightChange = await CloneShelvedChangeAsync(ClonedPreflightChange ?? PreflightChange.Value);
 			}
 
 			Logger.LogInformation("Creating job at CL {Change}, code CL {CodeChange}, preflight CL {PreflightChange}, cloned CL {ClonedPreflightChange}", Change, CodeChange, PreflightChange, ClonedPreflightChange);
@@ -974,8 +979,14 @@ namespace HordeServer.Services
 				int ClonedPreflightChange = Job.ClonedPreflightChange;
 				if (ClonedPreflightChange == 0)
 				{
-//PF				ClonedPreflightChange = await CloneShelvedChangeAsync(Job.PreflightChange);
-					ClonedPreflightChange = Job.PreflightChange;
+					if (ShouldClonePreflightChange(Job.StreamId))
+					{
+						ClonedPreflightChange = await CloneShelvedChangeAsync(Job.PreflightChange);
+					}
+					else
+					{
+						ClonedPreflightChange = Job.PreflightChange;
+					}
 				}
 
 				Logger.LogInformation("Updating description for {ClonedPreflightChange}", ClonedPreflightChange);
@@ -1004,10 +1015,19 @@ namespace HordeServer.Services
 				Logger.LogError(Ex, "Unable to submit shelved change");
 			}
 
-//PF		if (Change != null && Job.ClonedPreflightChange != 0)
-			if (Change != null && Job.PreflightChange != 0)
+			if (ShouldClonePreflightChange(Job.StreamId))
 			{
-				await DeleteShelvedChangeAsync(Job.PreflightChange);
+				if (Change != null && Job.ClonedPreflightChange != 0)
+				{
+					await DeleteShelvedChangeAsync(Job.PreflightChange);
+				}
+			}
+			else
+			{
+				if (Change != null && Job.PreflightChange != 0)
+				{
+					await DeleteShelvedChangeAsync(Job.PreflightChange);
+				}
 			}
 
 			for (; ; )
