Horde: Add tests for remotely executed actions

diff --git a/Engine/Source/Programs/Horde/HordeAgent/Execution/ActionExecutor.cs b/Engine/Source/Programs/Horde/HordeAgent/Execution/ActionExecutor.cs
--- a/Engine/Source/Programs/Horde/HordeAgent/Execution/ActionExecutor.cs
+++ b/Engine/Source/Programs/Horde/HordeAgent/Execution/ActionExecutor.cs
@@ -59,7 +59,7 @@ namespace HordeAgent
 		/// <summary>
 		/// Horde RPC wrapper
 		/// </summary>
-		ActionRpc.ActionRpcClient ActionRpc;
+		internal ActionRpc.ActionRpcClient ActionRpc { get; }
 
 		/// <summary>
 		/// Logger for internal executor output
@@ -71,17 +71,17 @@ namespace HordeAgent
 		/// </summary>
 		/// <param name="AgentName"></param>
 		/// <param name="InstanceName"></param>
-		/// <param name="CasChannel"></param>
-		/// <param name="ActionCacheChannel"></param>
-		/// <param name="ActionRpcChannel"></param>
+		/// <param name="CasClient"></param>
+		/// <param name="CacheClient"></param>
+		/// <param name="ActionRpcClient"></param>
 		/// <param name="Logger"></param>
-		public ActionExecutor(string AgentName, string InstanceName, GrpcChannel CasChannel, GrpcChannel ActionCacheChannel, GrpcChannel ActionRpcChannel, ILogger Logger)
+		public ActionExecutor(string AgentName, string InstanceName, ContentAddressableStorage.ContentAddressableStorageClient CasClient, ActionCache.ActionCacheClient CacheClient, ActionRpc.ActionRpcClient ActionRpcClient, ILogger Logger)
 		{
 			this.AgentName = AgentName;
 			this.InstanceName = InstanceName;
-			this.Storage = new ContentAddressableStorage.ContentAddressableStorageClient(CasChannel);
-			this.Cache = new ActionCache.ActionCacheClient(ActionCacheChannel);
-			this.ActionRpc = new ActionRpc.ActionRpcClient(ActionRpcChannel);
+			this.Storage = CasClient;
+			this.Cache = CacheClient;
+			this.ActionRpc = ActionRpcClient;
 			this.Logger = Logger;
 		}
 
diff --git a/Engine/Source/Programs/Horde/HordeAgent/Services/GrpcService.cs b/Engine/Source/Programs/Horde/HordeAgent/Services/GrpcService.cs
--- a/Engine/Source/Programs/Horde/HordeAgent/Services/GrpcService.cs
+++ b/Engine/Source/Programs/Horde/HordeAgent/Services/GrpcService.cs
@@ -12,6 +12,9 @@ using System.Collections.Generic;
 using System.Net.Http;
 using System.Net.Http.Headers;
 using System.Text;
+using Build.Bazel.Remote.Execution.V2;
+using Grpc.Core;
+using HordeCommon.Rpc;
 
 namespace HordeAgent.Services
 {
@@ -35,6 +38,10 @@ namespace HordeAgent.Services
 		/// </summary>
 		ILogger Logger;
 
+		internal Func<ChannelBase, ContentAddressableStorage.ContentAddressableStorageClient> CasClientFactory;
+		internal Func<ChannelBase, ActionCache.ActionCacheClient> ActionCacheClientFactory;
+		internal Func<ChannelBase, ActionRpc.ActionRpcClient> ActionRpcClientFactory;
+
 		/// <summary>
 		/// Constructor
 		/// </summary>
@@ -45,6 +52,10 @@ namespace HordeAgent.Services
 			this.Settings = Settings;
 			this.ServerProfile = Settings.Value.GetCurrentServerProfile();
 			this.Logger = Logger;
+			
+			CasClientFactory = Channel => new ContentAddressableStorage.ContentAddressableStorageClient(Channel);
+			ActionCacheClientFactory = Channel => new ActionCache.ActionCacheClient(Channel);
+			ActionRpcClientFactory = Channel => new ActionRpc.ActionRpcClient(Channel);
 		}
 
 		/// <summary>
@@ -106,5 +117,20 @@ namespace HordeAgent.Services
 				DisposeHttpClient = true
 			});
 		}
+
+		public ContentAddressableStorage.ContentAddressableStorageClient CreateCasClient(ChannelBase Channel)
+		{
+			return CasClientFactory(Channel);
+		}
+
+		public ActionCache.ActionCacheClient CreateActionCacheClient(ChannelBase Channel)
+		{
+			return ActionCacheClientFactory(Channel);
+		}
+
+		public ActionRpc.ActionRpcClient CreateActionRpcClient(ChannelBase Channel)
+		{
+			return ActionRpcClientFactory(Channel);
+		}
 	}
 }
diff --git a/Engine/Source/Programs/Horde/HordeAgent/Services/WorkerService.cs b/Engine/Source/Programs/Horde/HordeAgent/Services/WorkerService.cs
--- a/Engine/Source/Programs/Horde/HordeAgent/Services/WorkerService.cs
+++ b/Engine/Source/Programs/Horde/HordeAgent/Services/WorkerService.cs
@@ -51,6 +51,7 @@ using Amazon.Util;
 using Amazon.EC2;
 using Scope = Datadog.Trace.Scope;
 using Amazon.EC2.Model;
+using Status = Grpc.Core.Status;
 
 [assembly: InternalsVisibleTo("HordeAgentTests")]
 
@@ -64,7 +65,7 @@ namespace HordeAgent.Services
 		/// <summary>
 		/// Stores information about an active session
 		/// </summary>
-		class LeaseInfo
+		internal class LeaseInfo
 		{
 			/// <summary>
 			/// The worker lease state
@@ -630,14 +631,14 @@ namespace HordeAgent.Services
 		/// <param name="AgentId">The current agent id</param>
 		/// <param name="LeaseInfo">Information about the lease</param>
 		/// <returns>Outcome from the lease</returns>
-		async Task<LeaseOutcome> HandleLeasePayloadAsync(IRpcConnection RpcConnection, string AgentId, LeaseInfo LeaseInfo)
+		internal async Task<LeaseOutcome> HandleLeasePayloadAsync(IRpcConnection RpcConnection, string AgentId, LeaseInfo LeaseInfo)
 		{
 			Any Payload = LeaseInfo.Lease.Payload;
 
 			ActionTask ActionTask;
 			if (LeaseInfo.Lease.Payload.TryUnpack(out ActionTask))
 			{
-				Tracer.Instance.ActiveScope.Span.SetTag("task", "Action");
+				Tracer.Instance.ActiveScope?.Span.SetTag("task", "Action");
 				Func<ILogger, Task<LeaseOutcome>> Handler = NewLogger => ActionAsync(RpcConnection, LeaseInfo.Lease.Id, ActionTask, NewLogger, LeaseInfo.CancellationTokenSource.Token);
 				return await HandleLeasePayloadWithLogAsync(RpcConnection, ActionTask.LogId, AgentId, null, null, Handler, LeaseInfo.CancellationTokenSource.Token);
 			}
@@ -764,7 +765,7 @@ namespace HordeAgent.Services
 		/// <param name="Logger">Logger for the task</param>
 		/// <param name="CancellationToken">Token used to cancel the operation</param>
 		/// <returns></returns>
-		async Task<LeaseOutcome> ActionAsync(IRpcConnection RpcConnection, string LeaseId, ActionTask ActionTask, ILogger Logger, CancellationToken CancellationToken)
+		internal async Task<LeaseOutcome> ActionAsync(IRpcConnection RpcConnection, string LeaseId, ActionTask ActionTask, ILogger Logger, CancellationToken CancellationToken)
 		{
 			DateTimeOffset ActionTaskStartTime = DateTimeOffset.UtcNow;
 			using (IRpcClientRef Client = await RpcConnection.GetClientRef(new RpcContext(), CancellationToken))
@@ -794,11 +795,27 @@ namespace HordeAgent.Services
 				Logger.LogDebug("ActionCacheChannel={ActionCacheChannel}", ActionCacheChannel.Target);
 				Logger.LogDebug("ActionRpcChannel={ActionRpcChannel}", ActionRpcChannel.Target);
 
+				ActionExecutor Executor = new ActionExecutor(
+					Settings.GetAgentName(), ActionTask.InstanceName,
+					GrpcService.CreateCasClient(CasChannel),
+					GrpcService.CreateActionCacheClient(ActionCacheChannel),
+					GrpcService.CreateActionRpcClient(ActionRpcChannel), Logger);
+
 				try
 				{
-					ActionExecutor Executor = new ActionExecutor(Settings.GetAgentName(), ActionTask.InstanceName, CasChannel, ActionCacheChannel, ActionRpcChannel, Logger);
 					await Executor.ExecuteActionAsync(LeaseId, ActionTask, LeaseDir, ActionTaskStartTime);
 				}
+				catch (Exception re)
+				{
+					Logger.LogError(re, "Unhandled exception during remote execution");
+					PostActionResultRequest PostResultRequest = new PostActionResultRequest();
+					PostResultRequest.LeaseId = LeaseId;
+					PostResultRequest.ActionDigest = ActionTask.Digest;
+					PostResultRequest.Error = new RpcException(new Status(StatusCode.Unknown, "Unhandled exception during remote execution: " + re.Message), re.Message).Status;
+					await Executor.ActionRpc.PostActionResultAsync(PostResultRequest);
+					
+					throw re;
+				}
 				finally
 				{
 					try
diff --git a/Engine/Source/Programs/Horde/HordeAgentTests/ActionExecutorTest.cs b/Engine/Source/Programs/Horde/HordeAgentTests/ActionExecutorTest.cs
--- a/Engine/Source/Programs/Horde/HordeAgentTests/ActionExecutorTest.cs
+++ b/Engine/Source/Programs/Horde/HordeAgentTests/ActionExecutorTest.cs
@@ -1,7 +1,9 @@
 // Copyright Epic Games, Inc. All Rights Reserved.
 
+using System;
 using System.IO;
 using EpicGames.Core;
+using Grpc.Core;
 using HordeAgent;
 using Microsoft.VisualStudio.TestTools.UnitTesting;
 
@@ -60,7 +62,7 @@ namespace HordeAgentTests
 			Directory.Delete(TempDir, true);
 		}
 		
-		private static string GetTemporaryDirectory()
+		public static string GetTemporaryDirectory()
 		{
 			string TempDir = Path.Join(Path.GetTempPath(), "horde-" + Path.GetRandomFileName());
 			if (Directory.Exists(TempDir))
diff --git a/Engine/Source/Programs/Horde/HordeAgentTests/GrpcClients.cs b/Engine/Source/Programs/Horde/HordeAgentTests/GrpcClients.cs
new file mode 100644
--- /dev/null
+++ b/Engine/Source/Programs/Horde/HordeAgentTests/GrpcClients.cs
@@ -0,0 +1,118 @@
+﻿// Copyright Epic Games, Inc. All Rights Reserved.
+
+using System;
+using System.Collections.Concurrent;
+using System.Threading.Tasks;
+using Build.Bazel.Remote.Execution.V2;
+using Google.Protobuf;
+using Google.Protobuf.WellKnownTypes;
+using Grpc.Core;
+using HordeCommon.Rpc;
+using Status = Grpc.Core.Status;
+
+using BlobStore = System.Collections.Concurrent.ConcurrentDictionary<string, Google.Protobuf.ByteString>;
+
+namespace HordeAgentTests
+{
+	public class FakeCasClient : ContentAddressableStorage.ContentAddressableStorageClient
+	{
+		private readonly ConcurrentDictionary<string, BlobStore> BlobStoreInstances = new ConcurrentDictionary<string, BlobStore>();
+
+		private BlobStore GetBlobStore(string InstanceName)
+		{
+			if (!BlobStoreInstances.ContainsKey(InstanceName))
+			{
+				BlobStoreInstances.TryAdd(InstanceName, new BlobStore());
+			}
+
+			return BlobStoreInstances[InstanceName];
+		}
+		
+		public override AsyncUnaryCall<FindMissingBlobsResponse> FindMissingBlobsAsync(FindMissingBlobsRequest Request, CallOptions Options)
+		{
+			FindMissingBlobsResponse Res = new FindMissingBlobsResponse();
+			foreach (Digest Digest in Request.BlobDigests)
+			{
+				if (!GetBlobStore(Request.InstanceName).ContainsKey(Digest.Hash))
+				{
+					Res.MissingBlobDigests.Add(Digest);
+				}
+			}
+
+			return GetResponse(Res);
+		}
+
+		public override AsyncUnaryCall<BatchUpdateBlobsResponse> BatchUpdateBlobsAsync(BatchUpdateBlobsRequest Request, CallOptions Options)
+		{
+			BlobStore BlobStore = GetBlobStore(Request.InstanceName);
+			BatchUpdateBlobsResponse Res = new BatchUpdateBlobsResponse();
+
+			foreach (BatchUpdateBlobsRequest.Types.Request Req in Request.Requests)
+			{
+				BlobStore[Req.Digest.Hash] = Req.Data;
+				Res.Responses.Add(new BatchUpdateBlobsResponse.Types.Response {Digest = Req.Digest, Status = Status.DefaultSuccess});
+			}
+
+			return GetResponse(Res);
+		}
+
+		public override AsyncUnaryCall<BatchReadBlobsResponse> BatchReadBlobsAsync(BatchReadBlobsRequest Request, CallOptions Options)
+		{
+			BatchReadBlobsResponse Response = new BatchReadBlobsResponse();
+			foreach (Digest Digest in Request.Digests)
+			{
+				if (GetBlobStore(Request.InstanceName).TryGetValue(Digest.Hash, out ByteString? Data))
+				{
+					Response.Responses.Add(new BatchReadBlobsResponse.Types.Response
+					{
+						Data = Data,
+						Digest = Digest,
+						Status = Status.DefaultSuccess
+					});
+				}
+				else
+				{
+					Response.Responses.Add(new BatchReadBlobsResponse.Types.Response
+					{
+						Data = ByteString.Empty,
+						Digest = Digest,
+						Status = new Status(StatusCode.NotFound, $"Digest {Digest.Hash} not found")
+					});
+				}
+			}
+
+			return GetResponse(Response);
+		}
+
+		public override AsyncServerStreamingCall<GetTreeResponse> GetTree(GetTreeRequest request, CallOptions options)
+		{
+			throw new NotImplementedException("GetTree is not implemented in FakeCasClient");
+		}
+
+		public void SetBlob(string InstanceName, string DigestHash, ByteString Data)
+		{
+			GetBlobStore(InstanceName)[DigestHash] = Data;
+		}
+		
+		internal static AsyncUnaryCall<T> GetResponse<T>(T Response)
+		{
+			return new AsyncUnaryCall<T>(Task.FromResult(Response), null!, null!, null!, null!, null!);
+		}
+	}
+
+	public class FakeActionRpcClient : ActionRpc.ActionRpcClient
+	{
+		internal PostActionResultRequest? ActionResultRequest;
+		
+		public override AsyncUnaryCall<Empty> PostActionResultAsync(PostActionResultRequest Request, CallOptions Options)
+		{
+			if (ActionResultRequest != null)
+			{
+				throw new Exception("ActionResultRequest has already been set!");
+			}
+			
+			ActionResultRequest = Request;
+			return FakeCasClient.GetResponse(new Empty());
+		}
+	}
+}
\ No newline at end of file
diff --git a/Engine/Source/Programs/Horde/HordeAgentTests/HordeAgentTests.csproj b/Engine/Source/Programs/Horde/HordeAgentTests/HordeAgentTests.csproj
--- a/Engine/Source/Programs/Horde/HordeAgentTests/HordeAgentTests.csproj
+++ b/Engine/Source/Programs/Horde/HordeAgentTests/HordeAgentTests.csproj
@@ -20,4 +20,10 @@
     <ProjectReference Include="..\HordeAgent\HordeAgent.csproj" />
   </ItemGroup>
 
+  <ItemGroup>
+    <Content Include="remote-exec-test-bin\*.exe">
+      <CopyToOutputDirectory>Always</CopyToOutputDirectory>
+    </Content>
+  </ItemGroup>
+
 </Project>
diff --git a/Engine/Source/Programs/Horde/HordeAgentTests/TaskTest.cs b/Engine/Source/Programs/Horde/HordeAgentTests/TaskTest.cs
new file mode 100644
--- /dev/null
+++ b/Engine/Source/Programs/Horde/HordeAgentTests/TaskTest.cs
@@ -0,0 +1,181 @@
+﻿// Copyright Epic Games, Inc. All Rights Reserved.
+
+using System;
+using System.IO;
+using System.Runtime.InteropServices;
+using System.Text;
+using System.Threading.Tasks;
+using Google.Protobuf;
+using Google.Protobuf.WellKnownTypes;
+using Grpc.Net.Client;
+using HordeAgent;
+using HordeAgent.Execution.Interfaces;
+using HordeAgent.Services;
+using HordeAgent.Utility;
+using HordeCommon;
+using HordeCommon.Rpc;
+using HordeCommon.Rpc.Messages;
+using HordeCommon.Rpc.Tasks;
+using Microsoft.Extensions.Logging;
+using Microsoft.Extensions.Logging.Console;
+using Microsoft.Extensions.Options;
+using Microsoft.VisualStudio.TestTools.UnitTesting;
+using Directory = System.IO.Directory;
+
+namespace HordeAgentTests
+{
+	[TestClass]
+	public class TaskTest
+	{
+		private readonly ILogger<WorkerService> WorkerLogger;
+		private readonly ILogger<GrpcService> GrpcServiceLogger;
+		private readonly string TempDir;
+		private readonly AgentSettings Settings = new AgentSettings();
+		private readonly GrpcService GrpcService;
+		private readonly FakeCasClient Cas;
+		private readonly FakeActionRpcClient ActionRpc;
+		private readonly RpcConnectionStub RpcConnection;
+		private readonly WorkerService Ws;
+		private readonly string InstanceName = "testing-instance";
+
+		public TaskTest()
+		{
+			LoggerFactory LoggerFactory = new LoggerFactory();
+			ConsoleLoggerOptions LoggerOptions = new ConsoleLoggerOptions();
+			TestOptionsMonitor<ConsoleLoggerOptions> LoggerOptionsMon = new TestOptionsMonitor<ConsoleLoggerOptions>(LoggerOptions);
+			LoggerFactory.AddProvider(new ConsoleLoggerProvider(LoggerOptionsMon));
+			WorkerLogger = LoggerFactory.CreateLogger<WorkerService>();
+			GrpcServiceLogger = LoggerFactory.CreateLogger<GrpcService>();
+			
+			TempDir = ActionExecutorTest.GetTemporaryDirectory();
+			
+			ServerProfile Profile = new ServerProfile();
+			Profile.Name = "test";
+			Profile.Environment = "test-env";
+			Profile.Token = "bogus-token";
+			Profile.Url = "http://localhost";
+			Settings.ServerProfiles.Add(Profile);
+			Settings.Server = "test";
+			Settings.WorkingDir = Path.GetTempPath();
+			Settings.Executor = ExecutorType.Test; // Not really used since the executor is overridden in the tests
+			
+			GrpcService = new GrpcService(new OptionsWrapper<AgentSettings>(Settings), GrpcServiceLogger);
+			Cas = new FakeCasClient();
+			ActionRpc = new FakeActionRpcClient();
+			GrpcService.CasClientFactory = Channel => Cas;
+			GrpcService.ActionRpcClientFactory = Channel => ActionRpc;
+			
+			HordeRpcClientStub HordeRpcClient = new HordeRpcClientStub(WorkerLogger);
+			RpcConnection = new RpcConnectionStub(GrpcChannel.ForAddress("http://bogus-test-url"), HordeRpcClient);
+			Ws = GetWorkerService();
+		}
+		
+		[TestCleanup]
+		public void Cleanup()
+		{
+			Directory.Delete(TempDir, true);
+		}
+
+		private WorkerService GetWorkerService(
+			Func<IRpcConnection, ExecuteJobTask, BeginBatchResponse, IExecutor>? CreateExecutor = null)
+		{
+			IOptionsMonitor<AgentSettings> SettingsMonitor = new TestOptionsMonitor<AgentSettings>(Settings);
+			return new WorkerService(WorkerLogger, SettingsMonitor, GrpcService, CreateExecutor);
+		}
+	
+		[TestMethod]
+		public async Task ActionTask()
+		{
+			if (!RuntimeInformation.IsOSPlatform(OSPlatform.Windows)) return;
+			
+			ActionTask ActionTask = CreateActionTask();
+			Assert.IsNull(ActionRpc.ActionResultRequest);
+			LeaseOutcome Outcome = await Ws.HandleLeasePayloadAsync(RpcConnection, "my-agent-id", CreateLeaseInfo(ActionTask));
+			Assert.AreEqual(LeaseOutcome.Success, Outcome);
+			Assert.IsNotNull(ActionRpc.ActionResultRequest);
+			Assert.IsNull(ActionRpc.ActionResultRequest!.Error);
+		}
+		
+		[TestMethod]
+		[Ignore]
+		public async Task ActionTaskMissingAction()
+		{
+			if (!RuntimeInformation.IsOSPlatform(OSPlatform.Windows)) return;
+			
+			ActionTask ActionTask = CreateActionTask();
+			LeaseOutcome Outcome = await Ws.HandleLeasePayloadAsync(RpcConnection, "my-agent-id", CreateLeaseInfo(ActionTask));
+			Assert.AreEqual(LeaseOutcome.Failed, Outcome);
+			Assert.IsNotNull(ActionRpc.ActionResultRequest);
+			Assert.IsNull(ActionRpc.ActionResultRequest!.Error);
+		}
+		
+		private string GetHash(byte[] data)
+		{
+			using System.Security.Cryptography.MD5 md5 = System.Security.Cryptography.MD5.Create();
+			byte[] hashBytes = md5.ComputeHash(data);
+			StringBuilder sb = new StringBuilder();
+			foreach (byte b in hashBytes)
+			{
+				sb.Append(b.ToString("X2"));
+			}
+			return sb.ToString();
+		}
+
+		private Build.Bazel.Remote.Execution.V2.Digest GetDigest(byte[] data)
+		{
+			return new Build.Bazel.Remote.Execution.V2.Digest {Hash = GetHash(data), SizeBytes = data.Length};
+		}
+
+		private ActionTask CreateActionTask()
+		{
+			string BinName = "remote-exec-test-bin.exe";
+			string AssemblyPath = Path.GetDirectoryName(System.Reflection.Assembly.GetExecutingAssembly().Location)!;
+			string BinPath = Path.Combine(AssemblyPath, "remote-exec-test-bin", BinName);
+			
+			byte[] ExeFileData = File.ReadAllBytes(BinPath);
+			Cas.SetBlob(InstanceName, GetHash(ExeFileData), ByteString.CopyFrom(ExeFileData));
+				
+			Build.Bazel.Remote.Execution.V2.FileNode ExeFile = new Build.Bazel.Remote.Execution.V2.FileNode();
+			ExeFile.Digest = GetDigest(ExeFileData);
+			ExeFile.Name = BinName;
+			ExeFile.IsExecutable = true;
+			
+			Build.Bazel.Remote.Execution.V2.Directory Dir = new Build.Bazel.Remote.Execution.V2.Directory();
+			Dir.Files.Add(ExeFile);
+			byte[] DirData = Dir.ToByteArray();
+			Cas.SetBlob(InstanceName, GetHash(DirData), ByteString.CopyFrom(DirData));
+			
+			Build.Bazel.Remote.Execution.V2.Command Command = new Build.Bazel.Remote.Execution.V2.Command();
+			Command.Arguments.Add(ExeFile.Name);
+			byte[] CommandData = Command.ToByteArray();
+			Cas.SetBlob(InstanceName, GetHash(CommandData), ByteString.CopyFrom(CommandData));
+			
+			
+			Build.Bazel.Remote.Execution.V2.Action Action = new Build.Bazel.Remote.Execution.V2.Action();
+			Action.CommandDigest = GetDigest(CommandData);
+			Action.InputRootDigest = GetDigest(DirData);
+			Action.DoNotCache = true;
+			
+			byte[] ActionData = Action.ToByteArray();
+			Cas.SetBlob(InstanceName, GetHash(ActionData), ByteString.CopyFrom(ActionData));
+
+			ActionTask ActionTask = new ActionTask();
+			ActionTask.Digest = GetDigest(ActionData);
+			ActionTask.LogId = "bogus-log-id";
+			ActionTask.CasUrl = "http://bogus-cas-url";
+			ActionTask.ActionCacheUrl = "http://bogus-actioncache-url";
+			ActionTask.InstanceName = InstanceName;
+			return ActionTask;
+		}
+		
+		private WorkerService.LeaseInfo CreateLeaseInfo(IMessage Payload)
+		{
+			Lease Lease = new Lease();
+			Lease.Id = "my-lease-id";
+			Lease.Payload = Any.Pack(Payload);
+			Lease.State = LeaseState.Active;
+			Lease.Outcome = LeaseOutcome.Unspecified;
+			return new WorkerService.LeaseInfo(Lease);
+		}
+	}
+}
\ No newline at end of file
diff --git a/Engine/Source/Programs/Horde/HordeAgentTests/remote-exec-test-bin/build.bat b/Engine/Source/Programs/Horde/HordeAgentTests/remote-exec-test-bin/build.bat
new file mode 100755
--- /dev/null
+++ b/Engine/Source/Programs/Horde/HordeAgentTests/remote-exec-test-bin/build.bat
@@ -0,0 +1,2 @@
+@echo off
+cl.exe /O1 /MT remote-exec-test-bin.c
\ No newline at end of file
diff --git a/Engine/Source/Programs/Horde/HordeAgentTests/remote-exec-test-bin/remote-exec-test-bin.c b/Engine/Source/Programs/Horde/HordeAgentTests/remote-exec-test-bin/remote-exec-test-bin.c
new file mode 100644
--- /dev/null
+++ b/Engine/Source/Programs/Horde/HordeAgentTests/remote-exec-test-bin/remote-exec-test-bin.c
@@ -0,0 +1,24 @@
+#include <stdio.h>
+#include <windows.h>
+
+int main(int argc, char *argv[]) {
+    printf("Horde Remote Execution Test Binary\n");
+    int exitCode = 0;
+    int sleepMs = -1;
+    
+    if (argc == 2) {
+        sleepMs = atoi(argv[1]);
+    } else if (argc == 3) {
+        sleepMs = atoi(argv[1]);
+        exitCode = atoi(argv[2]);
+    }
+
+    printf("    Sleep: %d ms\n", sleepMs);
+    printf("Exit code: %d\n", exitCode);
+
+    if (sleepMs > 0) {
+        Sleep(sleepMs);
+    }
+
+    return exitCode;
+}
\ No newline at end of file
