Horde: Fix username on reshelf for autosubmit
#rnx

diff --git a/Engine/Source/Programs/Horde/HordeServer/Services/PerforceServiceApi.cs b/Engine/Source/Programs/Horde/HordeServer/Services/PerforceServiceApi.cs
--- a/Engine/Source/Programs/Horde/HordeServer/Services/PerforceServiceApi.cs
+++ b/Engine/Source/Programs/Horde/HordeServer/Services/PerforceServiceApi.cs
@@ -627,7 +627,41 @@ namespace HordeServer.Services
 		{
 			return await Task.Run(() =>
 			{
-				using (P4.Repository Repository = GetConnection(ReadOnly: false, ClientFromChange: ShelvedChange, UseClientFromChange: false))
+				string? ChangeOwner = null;
+
+				// Get the owner of the shelf
+				using (P4.Repository Repository = GetConnection(NoClient: true))
+				{
+					List<string> Args = new List<string> { "-S", ShelvedChange.ToString(CultureInfo.InvariantCulture)};
+
+					using (P4.P4Command Command = new P4.P4Command(Repository, "describe", true, Args.ToArray()))
+					{
+						P4.P4CommandResult Result = Command.Run();
+
+						if (!Result.Success)
+						{
+							throw new Exception($"Unable to get change {ShelvedChange}");
+						}
+
+						if (Result.TaggedOutput == null || Result.TaggedOutput.Count != 1)
+						{
+
+							throw new Exception($"Unable to get tagged output for change: {ShelvedChange}");
+						}
+
+						P4.Changelist Changelist = new P4.Changelist();
+						Changelist.FromChangeCmdTaggedOutput(Result.TaggedOutput[0], true, string.Empty, false);
+						ChangeOwner = Changelist.OwnerName;
+
+					}
+				}
+
+				if (string.IsNullOrEmpty(ChangeOwner))
+				{
+					throw new Exception($"Unable to get owner for shelved change {ShelvedChange}");
+				}
+
+				using (P4.Repository Repository = GetConnection(ReadOnly: false, ClientFromChange: ShelvedChange, UseClientFromChange: false, Username: ChangeOwner))
 				{
 					return ReshelveChange(Repository, ShelvedChange);				
 				}
diff --git a/Engine/Source/Programs/Horde/HordeServer/Services/PerforceServiceRouter.cs b/Engine/Source/Programs/Horde/HordeServer/Services/PerforceServiceRouter.cs
--- a/Engine/Source/Programs/Horde/HordeServer/Services/PerforceServiceRouter.cs
+++ b/Engine/Source/Programs/Horde/HordeServer/Services/PerforceServiceRouter.cs
@@ -105,7 +105,7 @@ namespace HordeServer.Services
 		/// <inheritdoc/>
 		override public Task<int> DuplicateShelvedChangeAsync(int ShelvedChange)
 		{
-			throw new Exception("DuplicateShelvedChangeAsync is disabled due to p4 issue: https://jira.it.epicgames.com/servicedesk/customer/portal/1/ITH-144069");
+			return P4APIService.DuplicateShelvedChangeAsync(ShelvedChange);
 		}
 
 		/// <inheritdoc/>
