Horde: Surface failed remote executions back to the client

diff --git a/Engine/Source/Programs/Horde/HordeServer/Rpc/ActionRpcService.cs b/Engine/Source/Programs/Horde/HordeServer/Rpc/ActionRpcService.cs
--- a/Engine/Source/Programs/Horde/HordeServer/Rpc/ActionRpcService.cs
+++ b/Engine/Source/Programs/Horde/HordeServer/Rpc/ActionRpcService.cs
@@ -31,7 +31,7 @@ namespace HordeServer.Rpc
 			if (ActionTaskSource.TryGetOperation(Request.LeaseId.ToObjectId(), out Operation))
 			{
 				Logger.LogInformation("Setting operation result to {Result}. LeaseId={LeaseId} OperationId={OperationId}", Request.Result, Request.LeaseId, Operation.Id);
-				Operation.TrySetResult(Request.Result);
+				Operation.TrySetResult(new ActionExecuteResult(Request.Result));
 			}
 			return Task.FromResult<Empty>(new Empty());
 		}
diff --git a/Engine/Source/Programs/Horde/HordeServer/Tasks/Impl/ActionTaskSource.cs b/Engine/Source/Programs/Horde/HordeServer/Tasks/Impl/ActionTaskSource.cs
--- a/Engine/Source/Programs/Horde/HordeServer/Tasks/Impl/ActionTaskSource.cs
+++ b/Engine/Source/Programs/Horde/HordeServer/Tasks/Impl/ActionTaskSource.cs
@@ -29,6 +29,40 @@ using PoolId = HordeServer.Utilities.StringId<HordeServer.Models.IPool>;
 
 namespace HordeServer.Tasks.Impl
 {
+	/// <summary>
+	/// Store either a successful or failed result for an execution
+	/// </summary>
+	public class ActionExecuteResult
+	{
+		/// <summary>
+		/// Successful result
+		/// </summary>
+		public ActionResult? Result { get; }
+		
+		/// <summary>
+		/// Failed result
+		/// </summary>
+		public Google.Rpc.Status? Error { get; }
+
+		/// <summary>
+		/// Construct with a successful result
+		/// </summary>
+		/// <param name="Result"></param>
+		public ActionExecuteResult(ActionResult? Result)
+		{
+			this.Result = Result;
+		}
+
+		/// <summary>
+		/// Construct with a failed result
+		/// </summary>
+		/// <param name="Error"></param>
+		public ActionExecuteResult(Google.Rpc.Status? Error)
+		{
+			this.Error = Error;
+		}
+	}
+	
 	/// <summary>
 	/// Interface for an action operation
 	/// </summary>
@@ -44,7 +78,7 @@ namespace HordeServer.Tasks.Impl
 		/// </summary>
 		/// <param name="Result">The result for this operation. If null, the operation will be rescheduled to operate on a new worker</param>
 		/// <returns></returns>
-		bool TrySetResult(ActionResult? Result);
+		bool TrySetResult(ActionExecuteResult Result);
 
 		/// <summary>
 		/// Read status updates from the operation
@@ -111,7 +145,7 @@ namespace HordeServer.Tasks.Impl
 			/// <summary>
 			/// The result from executing this operation
 			/// </summary>
-			public TaskCompletionSource<ActionResult?> ResultTaskSource = new TaskCompletionSource<ActionResult?>();
+			public TaskCompletionSource<ActionExecuteResult> ResultTaskSource = new TaskCompletionSource<ActionExecuteResult>();
 
 			/// <summary>
 			/// Constructor
@@ -127,7 +161,7 @@ namespace HordeServer.Tasks.Impl
 			/// Sets the operation result
 			/// </summary>
 			/// <param name="Result">The result of the operation</param>
-			public bool TrySetResult(ActionResult? Result)
+			public bool TrySetResult(ActionExecuteResult Result)
 			{
 				return ResultTaskSource.TrySetResult(Result);
 			}
@@ -328,8 +362,9 @@ namespace HordeServer.Tasks.Impl
 			ExecuteOperation? Operation;
 			if (TryGetOperation(Lease.Id, out Operation))
 			{
-				Logger.LogInformation("Lease cancelled. LeaseId={LeaseId} OperationId={OperationId}", Lease.Id, Operation.Id);
-				Operation.TrySetResult(null);
+				string Message = $"Lease cancelled. LeaseId={Lease.Id} OperationId={Operation.Id}";
+				Logger.LogInformation(Message);
+				Operation.TrySetResult(new ActionExecuteResult(new Google.Rpc.Status(StatusCode.Cancelled, Message)));
 			}
 			return Task.CompletedTask;
 		}
@@ -443,26 +478,40 @@ namespace HordeServer.Tasks.Impl
 				AgentLease Lease = new AgentLease(Operation.Id, LeaseName, null, null, LogFile.Id, LeaseState.Pending, Payload, new AgentRequirements(), null);
 
 				// Try to set the lease on the subscription. If it fails, the subscriber has already been allocated
-				Operation.ResultTaskSource = new TaskCompletionSource<ActionResult?>();
+				Operation.ResultTaskSource = new TaskCompletionSource<ActionExecuteResult>();
 
 				void OnConnectionLost()
 				{
-					Logger.LogInformation("Connection lost. LeaseId={LeaseId} OperationId={OperationId}", Lease.Id, Operation!.Id);
-					Operation.TrySetResult(null);
+					string Message = $"Connection lost to agent. LeaseId={Lease.Id} OperationId={Operation!.Id}";
+					Logger.LogInformation(Message);
+					Operation.TrySetResult(new ActionExecuteResult(new Google.Rpc.Status(StatusCode.Internal, Message)));
 				}
 				
 				if (Subscription.TrySetLease(Lease, OnConnectionLost))
 				{
 					Operation.SetStatus(ExecutionStage.Types.Value.Executing, null);
 
-					ActionResult? Result = await Operation.ResultTaskSource.Task;
-					if (Result != null)
+					ActionExecuteResult Result = await Operation.ResultTaskSource.Task;
+					if (Result.Result != null)
 					{
-						ExecuteResponse SuccessResponse = new ExecuteResponse();
-						SuccessResponse.Status = new Status(StatusCode.OK, String.Empty);
-						SuccessResponse.Result = Result;
-						return SuccessResponse;
+						return new ExecuteResponse
+						{
+							Status = new Status(StatusCode.OK, String.Empty),
+							Result = Result.Result
+						};
+					}
+					
+					if (Result.Error != null)
+					{
+						return new ExecuteResponse
+						{
+							Status = new Status((StatusCode) Result.Error.Code, Result.Error.Message),
+							CachedResult = false,
+							Message = "Remote execution failed. See 'status' field for more details."
+						};
 					}
+					
+					throw new Exception("Either ActionResult or Error must be set");
 				}
 			}
 		}
@@ -572,8 +621,9 @@ namespace HordeServer.Tasks.Impl
 			{
 				lock (Bucket!)
 				{
-					Logger.LogInformation("Trimming lease from bucket. LeaseId={LeaseId} OperationId={OperationId}", LeaseId, Operation!.Id);
-					Operation!.TrySetResult(null);
+					string Message = $"Trimming lease from bucket. LeaseId={LeaseId} OperationId={Operation!.Id}"; 
+					Logger.LogInformation(Message);
+					Operation!.TrySetResult(new ActionExecuteResult(new Google.Rpc.Status(StatusCode.Aborted, Message)));
 					Bucket.Remove(LeaseId.Value);
 				}
 			}
diff --git a/Engine/Source/Programs/Horde/HordeServer/Tasks/Impl/ClusteredActionTaskSource.cs b/Engine/Source/Programs/Horde/HordeServer/Tasks/Impl/ClusteredActionTaskSource.cs
--- a/Engine/Source/Programs/Horde/HordeServer/Tasks/Impl/ClusteredActionTaskSource.cs
+++ b/Engine/Source/Programs/Horde/HordeServer/Tasks/Impl/ClusteredActionTaskSource.cs
@@ -107,7 +107,7 @@ namespace HordeServer.Tasks.Impl
 
 		private class RemoteExecOperationInternal : IActionExecuteOperation
 		{
-			public TaskCompletionSource<ActionResult?> ResultTaskSource { get; } = new TaskCompletionSource<ActionResult?>();
+			public TaskCompletionSource<ActionExecuteResult> ResultTaskSource { get; } = new TaskCompletionSource<ActionExecuteResult>();
 			public ObjectId Id { get; }
 			public ObjectId LeaseId { get; }
 
@@ -117,7 +117,7 @@ namespace HordeServer.Tasks.Impl
 				this.LeaseId = LeaseId;
 			}
 
-			public bool TrySetResult(ActionResult? Result)
+			public bool TrySetResult(ActionExecuteResult Result)
 			{
 				return ResultTaskSource.TrySetResult(Result);
 			}
@@ -178,7 +178,8 @@ namespace HordeServer.Tasks.Impl
 				{
 					if (OpInternal.LeaseId == Lease.Id)
 					{
-						OpInternal.TrySetResult(null);
+						string Message = $"Lease cancelled. LeaseId={Lease.Id}";
+						OpInternal.TrySetResult(new ActionExecuteResult(new Google.Rpc.Status(StatusCode.Cancelled, Message)));
 					}
 				}
 			}
@@ -382,25 +383,45 @@ namespace HordeServer.Tasks.Impl
 					ActiveOperations.Remove(OpId);
 				}
 			}
+			
+			void OnConnectionLost()
+			{
+				string Message = $"Connection lost to agent. LeaseId={Lease.Id} OperationId={OpInternal!.Id}";
+				Logger.LogInformation(Message);
+				OpInternal.TrySetResult(new ActionExecuteResult(new Google.Rpc.Status(StatusCode.Internal, Message)));
+			}
 
-			if (Subscription.TrySetLease(Lease, () => OpInternal.ResultTaskSource.TrySetResult(null)))
+			if (Subscription.TrySetLease(Lease, OnConnectionLost))
 			{
 				await SetOperationStatusAsync(Op, ExecutionStage.Types.Value.Executing, null);
 
-				Task<ActionResult?> ResultTask = OpInternal.ResultTaskSource.Task;
+				Task<ActionExecuteResult> ResultTask = OpInternal.ResultTaskSource.Task;
 				if (await Task.WhenAny(ResultTask, Task.Delay(OperationTimeout)) == ResultTask)
 				{
-					ActionResult? Result = ResultTask.Result;
-					if (Result != null)
+					ActionExecuteResult Result = ResultTask.Result;
+					if (Result.Result != null)
 					{
 						ExecuteResponse SuccessResponse = new ExecuteResponse();
 						SuccessResponse.Status = new Status(StatusCode.OK, String.Empty);
-						SuccessResponse.Result = Result;
+						SuccessResponse.Result = Result.Result;
 						
 						RemoveActiveOperation(OpId);
 						return SuccessResponse;
 					}
-					
+
+					if (Result.Error != null)
+					{
+						ExecuteResponse ErrorResponse = new ExecuteResponse
+						{
+							Status = new Status((StatusCode) Result.Error.Code, Result.Error.Message),
+							CachedResult = false,
+							Message = "Remote execution failed. See 'status' field for more details."
+						};
+						
+						RemoveActiveOperation(OpId);
+						return ErrorResponse;
+					}
+
 					Logger.LogError("No result from remote execution under lease {LeaseId}", Lease.Id);
 				}
 				else
@@ -453,7 +474,7 @@ namespace HordeServer.Tasks.Impl
 			}
 		}
 
-		internal bool SetResultForActiveOperation(ObjectId OperationId, ActionResult? Result)
+		internal bool SetResultForActiveOperation(ObjectId OperationId, ActionExecuteResult Result)
 		{
 			lock (ActiveOperations)
 			{
diff --git a/Engine/Source/Programs/Horde/HordeServerTests/RemoteExec/ActionTaskSourceTest.cs b/Engine/Source/Programs/Horde/HordeServerTests/RemoteExec/ActionTaskSourceTest.cs
--- a/Engine/Source/Programs/Horde/HordeServerTests/RemoteExec/ActionTaskSourceTest.cs
+++ b/Engine/Source/Programs/Horde/HordeServerTests/RemoteExec/ActionTaskSourceTest.cs
@@ -109,7 +109,7 @@ namespace HordeServerTests.RemoteExec
 
 				if (Metadata.Stage == ExecutionStage.Types.Value.Executing)
 				{
-					ExecuteOp.TrySetResult(new ActionResult {ExitCode = ExitCode});
+					ExecuteOp.TrySetResult(new ActionExecuteResult(new ActionResult {ExitCode = ExitCode}));
 				}
 				
 				if (Op.Done)
diff --git a/Engine/Source/Programs/Horde/HordeServerTests/RemoteExec/ClusteredActionTaskSourceTest.cs b/Engine/Source/Programs/Horde/HordeServerTests/RemoteExec/ClusteredActionTaskSourceTest.cs
--- a/Engine/Source/Programs/Horde/HordeServerTests/RemoteExec/ClusteredActionTaskSourceTest.cs
+++ b/Engine/Source/Programs/Horde/HordeServerTests/RemoteExec/ClusteredActionTaskSourceTest.cs
@@ -75,7 +75,7 @@ namespace HordeServerTests.RemoteExec
 			Task Server2UpdateTask = Server2.UpdateOperations();
 			Assert.IsTrue(await WaitForOperationToExistAsync(Server2, ObjectId.Parse(Op.Id)));
 			int ExitCode = 11122233;
-			Assert.IsTrue(Server2.SetResultForActiveOperation(ObjectId.Parse(Op.Id), new ActionResult { ExitCode = ExitCode }));
+			Assert.IsTrue(Server2.SetResultForActiveOperation(ObjectId.Parse(Op.Id), new ActionExecuteResult(new ActionResult { ExitCode = ExitCode })));
 			await Server2UpdateTask;
 			
 			// Verify result can fetched from server 1
