Horde: Log REAPI channel setup as info for better visibility

diff --git a/Engine/Source/Programs/Horde/HordeAgent/Services/WorkerService.cs b/Engine/Source/Programs/Horde/HordeAgent/Services/WorkerService.cs
--- a/Engine/Source/Programs/Horde/HordeAgent/Services/WorkerService.cs
+++ b/Engine/Source/Programs/Horde/HordeAgent/Services/WorkerService.cs
@@ -791,9 +791,9 @@ namespace HordeAgent.Services
 				GrpcChannel ActionCacheChannel = GetChannel(ActionTask.ActionCacheUrl, ActionTask.ServiceAccountToken);//ActionTask.ActionCacheUrl == null ? Client.Channel : GrpcChannel.ForAddress(ActionTask.ActionCacheUrl);
 				GrpcChannel ActionRpcChannel = Client.Channel;
 
-				Logger.LogDebug("CasChannel={CasChannel}", CasChannel.Target);
-				Logger.LogDebug("ActionCacheChannel={ActionCacheChannel}", ActionCacheChannel.Target);
-				Logger.LogDebug("ActionRpcChannel={ActionRpcChannel}", ActionRpcChannel.Target);
+				Logger.LogInformation("CasChannel={CasChannel}", CasChannel.Target);
+				Logger.LogInformation("ActionCacheChannel={ActionCacheChannel}", ActionCacheChannel.Target);
+				Logger.LogInformation("ActionRpcChannel={ActionRpcChannel}", ActionRpcChannel.Target);
 
 				ActionExecutor Executor = new ActionExecutor(
 					Settings.GetAgentName(), ActionTask.InstanceName,
