Horde: Add server-side reflection for gRPC

Allows the server to be self-describing, making it easier to use tools like grpcurl to make test calls to the server.

diff --git a/Engine/Source/Programs/Horde/HordeServer/HordeServer.csproj b/Engine/Source/Programs/Horde/HordeServer/HordeServer.csproj
--- a/Engine/Source/Programs/Horde/HordeServer/HordeServer.csproj
+++ b/Engine/Source/Programs/Horde/HordeServer/HordeServer.csproj
@@ -19,6 +19,7 @@
     <PackageReference Include="DogStatsD-CSharp-Client" Version="6.0.0" />
     <PackageReference Include="Google.Protobuf" Version="3.12.3" />
     <PackageReference Include="Grpc.AspNetCore" Version="2.29.0" />
+    <PackageReference Include="Grpc.AspNetCore.Server.Reflection" Version="2.29.0" />
     <PackageReference Include="Grpc.Net.Client" Version="2.29.0" />
     <PackageReference Include="Grpc.Tools" Version="2.30.0">
       <PrivateAssets>all</PrivateAssets>
diff --git a/Engine/Source/Programs/Horde/HordeServer/Startup.cs b/Engine/Source/Programs/Horde/HordeServer/Startup.cs
--- a/Engine/Source/Programs/Horde/HordeServer/Startup.cs
+++ b/Engine/Source/Programs/Horde/HordeServer/Startup.cs
@@ -288,6 +288,7 @@ namespace HordeServer
 				Options.Interceptors.Add(typeof(LifetimeGrpcInterceptor));
 				Options.Interceptors.Add(typeof(GrpcExceptionInterceptor));
 			});
+			Services.AddGrpcReflection();
 
 			Services.AddSingleton<IAgentCollection, AgentCollection>();
 			Services.AddSingleton<IAgentSoftwareCollection, AgentSoftwareCollection>();
@@ -722,6 +723,7 @@ namespace HordeServer
 				Endpoints.MapGrpcService<CapabilitiesService>();
 				Endpoints.MapGrpcService<ContentStorageService>();
 				Endpoints.MapGrpcService<ExecutionService>();
+				Endpoints.MapGrpcReflectionService();
 
 				Endpoints.MapControllers();
 			});
