Horde: Properly handle both result and error from agent REAPI callback

Not considering errors caused one of ActionExecuteResult's fields to be null. Was not cached by the type system since data was received at runtime.

diff --git a/Engine/Source/Programs/Horde/HordeServer/Rpc/ActionRpcService.cs b/Engine/Source/Programs/Horde/HordeServer/Rpc/ActionRpcService.cs
--- a/Engine/Source/Programs/Horde/HordeServer/Rpc/ActionRpcService.cs
+++ b/Engine/Source/Programs/Horde/HordeServer/Rpc/ActionRpcService.cs
@@ -30,8 +30,20 @@ namespace HordeServer.Rpc
 			IActionExecuteOperation? Operation;
 			if (ActionTaskSource.TryGetOperation(Request.LeaseId.ToObjectId(), out Operation))
 			{
-				Logger.LogInformation("Setting operation result to {Result}. LeaseId={LeaseId} OperationId={OperationId}", Request.Result, Request.LeaseId, Operation.Id);
-				Operation.TrySetResult(new ActionExecuteResult(Request.Result));
+				Logger.LogInformation("Setting operation result. Result={Result} Error={Error}. LeaseId={LeaseId} OperationId={OperationId}", Request.Result, Request.Error, Request.LeaseId, Operation.Id);
+				if (Request.Error != null)
+				{
+					Operation.TrySetResult(new ActionExecuteResult(Request.Error));	
+				}
+				else if (Request.Result != null)
+				{
+					Operation.TrySetResult(new ActionExecuteResult(Request.Result));	
+				}
+				else
+				{
+					Logger.LogError("Both Result and Error are null!");
+				}
+				
 			}
 			return Task.FromResult<Empty>(new Empty());
 		}
diff --git a/Engine/Source/Programs/Horde/HordeServer/Tasks/Impl/ActionTaskSource.cs b/Engine/Source/Programs/Horde/HordeServer/Tasks/Impl/ActionTaskSource.cs
--- a/Engine/Source/Programs/Horde/HordeServer/Tasks/Impl/ActionTaskSource.cs
+++ b/Engine/Source/Programs/Horde/HordeServer/Tasks/Impl/ActionTaskSource.cs
@@ -511,7 +511,7 @@ namespace HordeServer.Tasks.Impl
 						};
 					}
 					
-					throw new Exception("Either ActionResult or Error must be set");
+					throw new Exception($"Either Result or Error must be set. Result={Result.Result} Error={Result.Error}");
 				}
 			}
 		}
