Horde: Ignore exceptions from REAPI cache result update

Simply log it as update cannot be sent. PostActionResultAsync has already been called, which probably makes sense.

diff --git a/Engine/Source/Programs/Horde/HordeAgent/Execution/ActionExecutor.cs b/Engine/Source/Programs/Horde/HordeAgent/Execution/ActionExecutor.cs
--- a/Engine/Source/Programs/Horde/HordeAgent/Execution/ActionExecutor.cs
+++ b/Engine/Source/Programs/Horde/HordeAgent/Execution/ActionExecutor.cs
@@ -204,14 +204,7 @@ namespace HordeAgent
 					PostResultRequest.Result = Result;
 					await ActionRpc.PostActionResultAsync(PostResultRequest);
 
-					if (!Action.DoNotCache)
-					{
-						UpdateActionResultRequest Request = new UpdateActionResultRequest();
-						Request.InstanceName = InstanceName;
-						Request.ActionDigest = ActionTask.Digest;
-						Request.ActionResult = Result;
-						await Cache.UpdateActionResultAsync(Request);
-					}
+					await CacheActionResultAsync(Action, ActionTask, Result);
 
 					// Command.OutputPaths
 					// Command.OutputDirectories
@@ -220,6 +213,25 @@ namespace HordeAgent
 			return Result;
 		}
 
+		async Task CacheActionResultAsync(Action Action, ActionTask ActionTask, ActionResult Result)
+		{
+			if (!Action.DoNotCache)
+			{
+				try
+				{
+					UpdateActionResultRequest Request = new UpdateActionResultRequest();
+					Request.InstanceName = InstanceName;
+					Request.ActionDigest = ActionTask.Digest;
+					Request.ActionResult = Result;
+					await Cache.UpdateActionResultAsync(Request);
+				}
+				catch (Exception e)
+				{
+					Logger.LogError(e, "Unable to update action result for digest {DigestHash}", ActionTask.Digest.Hash);
+				}
+			}
+		}
+
 		/// <summary>
 		/// Downloads files to the sandbox
 		/// </summary>
diff --git a/Engine/Source/Programs/Horde/HordeServer/Tasks/Impl/ActionTaskSource.cs b/Engine/Source/Programs/Horde/HordeServer/Tasks/Impl/ActionTaskSource.cs
--- a/Engine/Source/Programs/Horde/HordeServer/Tasks/Impl/ActionTaskSource.cs
+++ b/Engine/Source/Programs/Horde/HordeServer/Tasks/Impl/ActionTaskSource.cs
@@ -497,6 +497,7 @@ namespace HordeServer.Tasks.Impl
 						return new ExecuteResponse
 						{
 							Status = new Status(StatusCode.OK, String.Empty),
+							CachedResult = false,
 							Result = Result.Result
 						};
 					}
