Horde: Log what action digest is used for an REAPI execution

diff --git a/Engine/Source/Programs/Horde/HordeAgent/Execution/ActionExecutor.cs b/Engine/Source/Programs/Horde/HordeAgent/Execution/ActionExecutor.cs
--- a/Engine/Source/Programs/Horde/HordeAgent/Execution/ActionExecutor.cs
+++ b/Engine/Source/Programs/Horde/HordeAgent/Execution/ActionExecutor.cs
@@ -100,6 +100,7 @@ namespace HordeAgent
 			{
 				throw new RpcException(new Grpc.Core.Status(StatusCode.NotFound, "Unable to find action message"));
 			}
+			Logger.LogInformation("Using action {Hash} / {Size} bytes for lease ID {LeaseId}", ActionTask.Digest.Hash, ActionTask.Digest.SizeBytes, LeaseId);
 
 			DirectoryReference.CreateDirectory(SandboxDir);
 			FileUtils.ForceDeleteDirectoryContents(SandboxDir);
