Horde: Add internal retry for REAPI actions

To temporarily hide issues with remote executions and avoiding the client to do extra retries.

diff --git a/Engine/Source/Programs/Horde/HordeServer/Tasks/Impl/ActionTaskSource.cs b/Engine/Source/Programs/Horde/HordeServer/Tasks/Impl/ActionTaskSource.cs
--- a/Engine/Source/Programs/Horde/HordeServer/Tasks/Impl/ActionTaskSource.cs
+++ b/Engine/Source/Programs/Horde/HordeServer/Tasks/Impl/ActionTaskSource.cs
@@ -112,6 +112,11 @@ namespace HordeServer.Tasks.Impl
 		/// </summary>
 		class ExecuteOperation : IActionExecuteOperation
 		{
+			/// <summary>
+			/// Max number of retries for executing an action
+			/// </summary>
+			public const int MaxRetries = 5;
+			
 			/// <summary>
 			/// The operation id
 			/// </summary>
@@ -147,6 +152,11 @@ namespace HordeServer.Tasks.Impl
 			/// </summary>
 			public TaskCompletionSource<ActionExecuteResult> ResultTaskSource = new TaskCompletionSource<ActionExecuteResult>();
 
+			/// <summary>
+			/// Increased each time a result for an action execution is set (a result or an error)
+			/// </summary>
+			public int NumExecResults = 0;
+
 			/// <summary>
 			/// Constructor
 			/// </summary>
@@ -163,6 +173,7 @@ namespace HordeServer.Tasks.Impl
 			/// <param name="Result">The result of the operation</param>
 			public bool TrySetResult(ActionExecuteResult Result)
 			{
+				Interlocked.Increment(ref NumExecResults);
 				return ResultTaskSource.TrySetResult(Result);
 			}
 
@@ -498,21 +509,29 @@ namespace HordeServer.Tasks.Impl
 						{
 							Status = new Status(StatusCode.OK, String.Empty),
 							CachedResult = false,
-							Result = Result.Result
+							Result = Result.Result,
+							Message = $"Remote execution successful. Retries {Operation.NumExecResults} of {ExecuteOperation.MaxRetries}"
 						};
 					}
-					
-					if (Result.Error != null)
+					else if (Result.Error != null)
 					{
-						return new ExecuteResponse
+						if (Operation.NumExecResults > ExecuteOperation.MaxRetries)
 						{
-							Status = new Status((StatusCode) Result.Error.Code, Result.Error.Message),
-							CachedResult = false,
-							Message = "Remote execution failed. See 'status' field for more details."
-						};
+							return new ExecuteResponse
+							{
+								Status = new Status((StatusCode) Result.Error.Code, Result.Error.Message),
+								CachedResult = false,
+								Message = $"Remote execution failed. Retries {Operation.NumExecResults} of {ExecuteOperation.MaxRetries}"
+							};
+						}
+						
+						Logger.LogError("Retrying execution for {OperationId}. Last error: {ErrorStatus} {ErrorMsg}", Operation.Id, Result.Error.Code, Result.Error.Message);
+						// Intended fallthrough that will restart the for-loop, allowing the action to be retried
+					}
+					else
+					{
+						throw new Exception($"Either Result or Error must be set. Result={Result.Result} Error={Result.Error}");
 					}
-					
-					throw new Exception($"Either Result or Error must be set. Result={Result.Result} Error={Result.Error}");
 				}
 			}
 		}
