Horde: Suppress excessive logging in Horde Server

Much too fine-grained at the info level.

diff --git a/Engine/Source/Programs/Horde/HordeAgent/Logging.cs b/Engine/Source/Programs/Horde/HordeAgent/Logging.cs
--- a/Engine/Source/Programs/Horde/HordeAgent/Logging.cs
+++ b/Engine/Source/Programs/Horde/HordeAgent/Logging.cs
@@ -73,6 +73,13 @@ namespace HordeAgent
 				.MinimumLevel.Debug()
 				.MinimumLevel.Override("Microsoft", LogEventLevel.Information)
 				.MinimumLevel.Override("System.Net.Http.HttpClient", LogEventLevel.Warning)
+				.MinimumLevel.Override("Microsoft.AspNetCore.Routing.EndpointMiddleware", LogEventLevel.Warning)
+				.MinimumLevel.Override("Microsoft.AspNetCore.Authorization.DefaultAuthorizationService", LogEventLevel.Warning)
+				.MinimumLevel.Override("HordeServer.Authentication.HordeJwtBearerHandler", LogEventLevel.Warning)
+				.MinimumLevel.Override("HordeServer.Authentication.OktaHandler", LogEventLevel.Warning)
+				.MinimumLevel.Override("Microsoft.AspNetCore.Hosting.Diagnostics", LogEventLevel.Warning)
+				.MinimumLevel.Override("Microsoft.AspNetCore.Mvc.Infrastructure.ControllerActionInvoker", LogEventLevel.Warning)
+				.MinimumLevel.Override("Serilog.AspNetCore.RequestLoggingMiddleware", LogEventLevel.Warning)
 				.MinimumLevel.ControlledBy(LogLevelSwitch)
 				.Enrich.FromLogContext()
 				.Enrich.With<VersionLogEnricher>()
