Horde: Prevent server participating in elections when in read-only mode. Useful for running against prod server without clobbering settings etc...

diff --git a/Engine/Source/Programs/Horde/HordeServer/Utilities/ElectedBackgroundService.cs b/Engine/Source/Programs/Horde/HordeServer/Utilities/ElectedBackgroundService.cs
--- a/Engine/Source/Programs/Horde/HordeServer/Utilities/ElectedBackgroundService.cs
+++ b/Engine/Source/Programs/Horde/HordeServer/Utilities/ElectedBackgroundService.cs
@@ -26,6 +26,7 @@ namespace HordeServer.Utilities
 		SingletonDocument<State> StateAccessor;
 		TimeSpan PollInterval;
 		TimeSpan ElectionTimeout;
+		bool ReadOnlyMode;
 
 		/// <summary>
 		/// Constructor
@@ -52,6 +53,7 @@ namespace HordeServer.Utilities
 			this.StateAccessor = new SingletonDocument<State>(DatabaseService, ServiceId);
 			this.PollInterval = PollInterval;
 			this.ElectionTimeout = ElectionTimeout;
+			this.ReadOnlyMode = DatabaseService.ReadOnlyMode;
 		}
 
 		/// <summary>
@@ -61,6 +63,11 @@ namespace HordeServer.Utilities
 		/// <returns>Async task</returns>
 		protected override sealed async Task TickAsync(CancellationToken StoppingToken)
 		{
+			if (ReadOnlyMode)
+			{
+				return;
+			}
+
 			DateTime UtcNow = DateTime.UtcNow;
 			for (; ; )
 			{
diff --git a/Engine/Source/Programs/Horde/HordeServer/Utilities/ElectedTick.cs b/Engine/Source/Programs/Horde/HordeServer/Utilities/ElectedTick.cs
--- a/Engine/Source/Programs/Horde/HordeServer/Utilities/ElectedTick.cs
+++ b/Engine/Source/Programs/Horde/HordeServer/Utilities/ElectedTick.cs
@@ -29,6 +29,7 @@ namespace HordeServer.Utilities
 		Task BackgroundTask;
 		CancellationTokenSource CancellationTokenSource;
 		Func<CancellationToken, Task<DateTime>> TickAsync;
+		bool ReadOnlyMode;
 		ILogger Logger;
 
 		/// <summary>
@@ -71,6 +72,7 @@ namespace HordeServer.Utilities
 			this.PollInterval = PollInterval;
 			this.ElectionTimeout = ElectionTimeout;
 			this.TickAsync = TickAsync;
+			this.ReadOnlyMode = DatabaseService.ReadOnlyMode;
 			this.Logger = Logger;
 
 			CancellationTokenSource = new CancellationTokenSource();
@@ -122,6 +124,12 @@ namespace HordeServer.Utilities
 					continue;
 				}
 
+				// Don't claim the election unless we're in writeable mode
+				if (ReadOnlyMode)
+				{
+					continue;
+				}
+
 				// Try to claim the election
 				State.Host = System.Net.Dns.GetHostName();
 				State.NextUpdateTime = UtcNow + ElectionTimeout;
