Horde: Enable the hosting log channel output, so the server prints information about which port it's listening on.

diff --git a/Engine/Source/Programs/Horde/HordeServer/appsettings.json b/Engine/Source/Programs/Horde/HordeServer/appsettings.json
--- a/Engine/Source/Programs/Horde/HordeServer/appsettings.json
+++ b/Engine/Source/Programs/Horde/HordeServer/appsettings.json
@@ -19,6 +19,7 @@
 				"Redis" : "Warning", // For bundled Redis output
 				"Microsoft": "Warning",
 				"Microsoft.AspNetCore": "Warning",
+				"Microsoft.Hosting": "Information",
 				"HordeServer.Authentication.OktaHandler": "Warning",
 				"HordeServer.Authentication.OktaViaJwtHandler": "Warning",
 				"HordeServer.Authentication.HordeJwtBearerHandler": "Warning",
