Horde: Fix StringId type converter not correctly handling conversion from strings.

diff --git a/Engine/Source/Programs/Horde/HordeServer/Utilities/StringId.cs b/Engine/Source/Programs/Horde/HordeServer/Utilities/StringId.cs
--- a/Engine/Source/Programs/Horde/HordeServer/Utilities/StringId.cs
+++ b/Engine/Source/Programs/Horde/HordeServer/Utilities/StringId.cs
@@ -276,12 +276,29 @@ namespace HordeServer.Utilities
 	/// <summary>
 	/// Type converter from strings to PropertyFilter objects
 	/// </summary>
-	public sealed class StringIdTypeConverter : TypeConverter
+	sealed class StringIdTypeConverter : TypeConverter
 	{
+		Type Type;
+
+		/// <summary>
+		/// Constructor
+		/// </summary>
+		/// <param name="Type"></param>
+		public StringIdTypeConverter(Type Type)
+		{
+			this.Type = Type;
+		}
+
 		/// <inheritdoc/>
 		public override bool CanConvertFrom(ITypeDescriptorContext Context, Type SourceType)
 		{
-			return SourceType == typeof(string);
+			return SourceType == typeof(string) || base.CanConvertFrom(Context, SourceType);
+		}
+
+		/// <inheritdoc/>
+		public override object ConvertFrom(ITypeDescriptorContext Context, CultureInfo Culture, object Value)
+		{
+			return Activator.CreateInstance(Type, Value)!;
 		}
 
 		/// <inheritdoc/>
