Horde: Track conform count for each server separately.

diff --git a/Engine/Source/Programs/Horde/HordeServer/Models/ConformList.cs b/Engine/Source/Programs/Horde/HordeServer/Models/ConformList.cs
--- a/Engine/Source/Programs/Horde/HordeServer/Models/ConformList.cs
+++ b/Engine/Source/Programs/Horde/HordeServer/Models/ConformList.cs
@@ -34,6 +34,27 @@ namespace HordeServer.Models
 		public DateTime? LastCheckTimeUtc { get; set; }
 	}
 
+	/// <summary>
+	/// List of entries for a particular server
+	/// </summary>
+	public class ConformListServer
+	{
+		/// <summary>
+		/// The perforce cluster
+		/// </summary>
+		public string Cluster { get; set; } = String.Empty;
+
+		/// <summary>
+		/// The server and port
+		/// </summary>
+		public string ServerAndPort { get; set; } = String.Empty;
+
+		/// <summary>
+		/// List of entries
+		/// </summary>
+		public List<ConformListEntry> Entries { get; set; } = new List<ConformListEntry>();
+	}
+
 	/// <summary>
 	/// List of machines that are currently conforming
 	/// </summary>
@@ -41,8 +62,13 @@ namespace HordeServer.Models
 	public class ConformList : SingletonBase
 	{
 		/// <summary>
-		/// List of entries
+		/// List of conforming servers [DEPRECATED]
 		/// </summary>
 		public List<ConformListEntry> Entries { get; set; } = new List<ConformListEntry>();
+
+		/// <summary>
+		/// List of conforming servers
+		/// </summary>
+		public List<ConformListServer> Servers { get; set; } = new List<ConformListServer>();
 	}
 }
diff --git a/Engine/Source/Programs/Horde/HordeServer/Models/Globals.cs b/Engine/Source/Programs/Horde/HordeServer/Models/Globals.cs
--- a/Engine/Source/Programs/Horde/HordeServer/Models/Globals.cs
+++ b/Engine/Source/Programs/Horde/HordeServer/Models/Globals.cs
@@ -225,6 +225,11 @@ namespace HordeServer.Models
 		/// </summary>
 		public bool ResolveDns { get; set; }
 
+		/// <summary>
+		/// Maximum number of simultaneous conforms on this server
+		/// </summary>
+		public int MaxConformCount { get; set; }
+
 		/// <summary>
 		/// List of properties for an agent to be eligable to use this server
 		/// </summary>
diff --git a/Engine/Source/Programs/Horde/HordeServer/Tasks/Impl/ConformTaskSource.cs b/Engine/Source/Programs/Horde/HordeServer/Tasks/Impl/ConformTaskSource.cs
--- a/Engine/Source/Programs/Horde/HordeServer/Tasks/Impl/ConformTaskSource.cs
+++ b/Engine/Source/Programs/Horde/HordeServer/Tasks/Impl/ConformTaskSource.cs
@@ -96,7 +96,7 @@ namespace HordeServer.Tasks.Impl
 
 			// Update any leases that are older than LastCheckTimeUtc
 			Dictionary<ObjectId, bool> RemoveLeases = new Dictionary<ObjectId, bool>();
-			foreach (ConformListEntry Entry in List.Entries)
+			foreach (ConformListEntry Entry in Enumerable.Concat(List.Entries, List.Servers.SelectMany(x => x.Entries)))
 			{
 				if (Entry.LastCheckTimeUtc < LastCheckTimeUtc)
 				{
@@ -128,18 +128,33 @@ namespace HordeServer.Tasks.Impl
 		/// <param name="RemoveLeases">List of leases to update. Entries with values set to true will be removed, entries with values set to false will have their timestamp updated.</param>
 		static void UpdateConformList(ConformList List, DateTime UtcNow, Dictionary<ObjectId, bool> RemoveLeases)
 		{
-			for (int Idx = 0; Idx < List.Entries.Count; Idx++)
+			UpdateConformList(List.Entries, UtcNow, RemoveLeases);
+			foreach (ConformListServer Server in List.Servers)
+			{
+				UpdateConformList(Server.Entries, UtcNow, RemoveLeases);
+			}
+		}
+
+		/// <summary>
+		/// Remove items from the conform list, and update timestamps for them
+		/// </summary>
+		/// <param name="Entries">The list to update</param>
+		/// <param name="UtcNow">Current time</param>
+		/// <param name="RemoveLeases">List of leases to update. Entries with values set to true will be removed, entries with values set to false will have their timestamp updated.</param>
+		static void UpdateConformList(List<ConformListEntry> Entries, DateTime UtcNow, Dictionary<ObjectId, bool> RemoveLeases)
+		{
+			for (int Idx = 0; Idx < Entries.Count; Idx++)
 			{
 				bool Remove;
-				if (RemoveLeases.TryGetValue(List.Entries[Idx].LeaseId, out Remove))
+				if (RemoveLeases.TryGetValue(Entries[Idx].LeaseId, out Remove))
 				{
 					if (Remove)
 					{
-						List.Entries.RemoveAt(Idx--);
+						Entries.RemoveAt(Idx--);
 					}
-					else if (List.Entries[Idx].LastCheckTimeUtc < UtcNow)
+					else if (Entries[Idx].LastCheckTimeUtc < UtcNow)
 					{
-						List.Entries[Idx].LastCheckTimeUtc = UtcNow;
+						Entries[Idx].LastCheckTimeUtc = UtcNow;
 					}
 				}
 			}
@@ -156,13 +171,12 @@ namespace HordeServer.Tasks.Impl
 
 			if (Agent.Leases.Count == 0)
 			{
-				ObjectId LeaseId = ObjectId.GenerateNewId();
-				if (await AllocateConformLeaseAsync(Agent.Id, LeaseId))
+				ConformTask Task = new ConformTask();
+				if (await GetWorkspacesAsync(Agent, Task.Workspaces))
 				{
-					ConformTask Task = new ConformTask();
-					if (!await GetWorkspacesAsync(Agent, Task.Workspaces))
+					ObjectId LeaseId = ObjectId.GenerateNewId();
+					if (!await AllocateConformLeaseAsync(Agent.Id, Task.Workspaces, LeaseId))
 					{
-						await ReleaseConformLeaseAsync(LeaseId);
 						return null;
 					}
 
@@ -212,24 +226,55 @@ namespace HordeServer.Tasks.Impl
 		/// Atempt to allocate a conform resource for the given lease
 		/// </summary>
 		/// <param name="AgentId">The agent id</param>
+		/// <param name="Workspaces">List of workspaces that are required</param>
 		/// <param name="LeaseId">The lease id</param>
 		/// <returns>True if the resource was allocated, false otherwise</returns>
-		private async Task<bool> AllocateConformLeaseAsync(AgentId AgentId, ObjectId LeaseId)
+		private async Task<bool> AllocateConformLeaseAsync(AgentId AgentId, IEnumerable<HordeCommon.Rpc.Messages.AgentWorkspace> Workspaces, ObjectId LeaseId)
 		{
 			Globals Globals = await DatabaseService.GetGlobalsAsync();
 			for (; ; )
 			{
 				ConformList CurrentValue = await ConformList.GetAsync();
-				if (Globals.MaxConformCount != 0 && CurrentValue.Entries.Count >= Globals.MaxConformCount)
+				if (Globals.MaxConformCount != 0 && CurrentValue.Entries.Count + CurrentValue.Servers.Sum(x => x.Entries.Count) >= Globals.MaxConformCount)
 				{
 					return false;
 				}
 
-				ConformListEntry Entry = new ConformListEntry();
-				Entry.AgentId = AgentId;
-				Entry.LeaseId = LeaseId;
-				Entry.LastCheckTimeUtc = DateTime.UtcNow;
-				CurrentValue.Entries.Add(Entry);
+				HashSet<string> Servers = new HashSet<string>(StringComparer.Ordinal);
+				foreach (HordeCommon.Rpc.Messages.AgentWorkspace Workspace in Workspaces)
+				{
+					PerforceCluster? Cluster = Globals.FindPerforceCluster(Workspace.Cluster);
+					if(Cluster == null)
+					{
+						Logger.LogWarning("Unable to find perforce cluster '{Cluster}' for conform", Workspace.Cluster);
+						return false;
+					}
+
+					PerforceServer? Server = Cluster.Servers.FirstOrDefault(x => x.ServerAndPort.Equals(Workspace.ServerAndPort, StringComparison.OrdinalIgnoreCase));
+					if(Server == null)
+					{
+						Logger.LogWarning("Unable to find perforce server '{Server}' for conform", Workspace.Cluster);
+						return false;
+					}
+
+					ConformListServer? ConformServer = CurrentValue.Servers.FirstOrDefault(x => x.Cluster.Equals(Workspace.Cluster, StringComparison.Ordinal) && x.ServerAndPort.Equals(Workspace.ServerAndPort, StringComparison.Ordinal));
+					if (ConformServer == null)
+					{
+						ConformServer = new ConformListServer { Cluster = Workspace.Cluster, ServerAndPort = Workspace.ServerAndPort };
+						CurrentValue.Servers.Add(ConformServer);
+					}
+
+					if (Server.MaxConformCount != 0 && ConformServer.Entries.Count > Server.MaxConformCount)
+					{
+						return false;
+					}
+
+					ConformListEntry Entry = new ConformListEntry();
+					Entry.AgentId = AgentId;
+					Entry.LeaseId = LeaseId;
+					Entry.LastCheckTimeUtc = DateTime.UtcNow;
+					ConformServer.Entries.Add(Entry);
+				}
 
 				if (await ConformList.TryUpdateAsync(CurrentValue))
 				{
@@ -249,7 +294,7 @@ namespace HordeServer.Tasks.Impl
 			for (; ; )
 			{
 				ConformList CurrentValue = await ConformList.GetAsync();
-				if (CurrentValue.Entries.RemoveAll(x => x.LeaseId == LeaseId) == 0)
+				if (CurrentValue.Entries.RemoveAll(x => x.LeaseId == LeaseId) + CurrentValue.Servers.Sum(x => x.Entries.RemoveAll(x => x.LeaseId == LeaseId)) == 0)
 				{
 					Logger.LogInformation("Conform lease {LeaseId} is not in singelton", LeaseId);
 					break;
