Horde: Fix search for default cluster.

diff --git a/Engine/Source/Programs/Horde/HordeServer/Tasks/Impl/ConformTaskSource.cs b/Engine/Source/Programs/Horde/HordeServer/Tasks/Impl/ConformTaskSource.cs
--- a/Engine/Source/Programs/Horde/HordeServer/Tasks/Impl/ConformTaskSource.cs
+++ b/Engine/Source/Programs/Horde/HordeServer/Tasks/Impl/ConformTaskSource.cs
@@ -243,7 +243,7 @@ namespace HordeServer.Tasks.Impl
 				HashSet<string> Servers = new HashSet<string>(StringComparer.Ordinal);
 				foreach (HordeCommon.Rpc.Messages.AgentWorkspace Workspace in Workspaces)
 				{
-					PerforceCluster? Cluster = Globals.FindPerforceCluster(Workspace.Cluster);
+					PerforceCluster? Cluster = Globals.FindPerforceCluster(String.IsNullOrEmpty(Workspace.ConfiguredCluster) ? null : Workspace.ConfiguredCluster);
 					if(Cluster == null)
 					{
 						Logger.LogWarning("Unable to find perforce cluster '{Cluster}' for conform", Workspace.Cluster);
