Horde: Fix reshelf behavior
#rnx

diff --git a/Engine/Source/Programs/Horde/HordeServer/Services/PerforceServiceApi.cs b/Engine/Source/Programs/Horde/HordeServer/Services/PerforceServiceApi.cs
--- a/Engine/Source/Programs/Horde/HordeServer/Services/PerforceServiceApi.cs
+++ b/Engine/Source/Programs/Horde/HordeServer/Services/PerforceServiceApi.cs
@@ -1074,7 +1074,7 @@ namespace HordeServer.Services
 				if (Results.Success)
 				{
 
-					if (Results.InfoOutput.Count != 2)
+					if (Results.InfoOutput.Count == 0)
 					{
 						Serilog.Log.Logger.Information("Perforce: Unexpected info output when reshelving change");
 						throw new Exception("Unexpected info output when reshelving change");
@@ -1082,7 +1082,7 @@ namespace HordeServer.Services
 
 					bool Error = true;
 					int ReshelvedChange = 0;
-					string Message = Results.InfoOutput[1].Message;
+					string Message = Results.InfoOutput[Results.InfoOutput.Count - 1].Message;
 					Match ChangeMatch = Regex.Match(Message, @"Change (\d+) files shelved");
 
 					if (ChangeMatch.Success && ChangeMatch.Groups.Count == 2)
