Horde: Fix error message for invalid P4 cluster.

diff --git a/Engine/Source/Programs/Horde/HordeServer/Tasks/Impl/ConformTaskSource.cs b/Engine/Source/Programs/Horde/HordeServer/Tasks/Impl/ConformTaskSource.cs
--- a/Engine/Source/Programs/Horde/HordeServer/Tasks/Impl/ConformTaskSource.cs
+++ b/Engine/Source/Programs/Horde/HordeServer/Tasks/Impl/ConformTaskSource.cs
@@ -243,17 +243,17 @@ namespace HordeServer.Tasks.Impl
 				HashSet<string> Servers = new HashSet<string>(StringComparer.Ordinal);
 				foreach (HordeCommon.Rpc.Messages.AgentWorkspace Workspace in Workspaces)
 				{
-					PerforceCluster? Cluster = Globals.FindPerforceCluster(String.IsNullOrEmpty(Workspace.ConfiguredCluster) ? null : Workspace.ConfiguredCluster);
+					PerforceCluster? Cluster = Globals.FindPerforceCluster(Workspace.Cluster);
 					if(Cluster == null)
 					{
 						Logger.LogWarning("Unable to find perforce cluster '{Cluster}' for conform", Workspace.Cluster);
 						return false;
 					}
 
-					PerforceServer? Server = Cluster.Servers.FirstOrDefault(x => x.ServerAndPort.Equals(Workspace.ServerAndPort, StringComparison.OrdinalIgnoreCase));
+					PerforceServer? Server = Cluster.Servers.FirstOrDefault(x => x.ServerAndPort.Equals(Workspace.ServerAndPort, StringComparison.Ordinal));
 					if(Server == null)
 					{
-						Logger.LogWarning("Unable to find perforce server '{Server}' for conform", Workspace.Cluster);
+						Logger.LogWarning("Unable to find perforce server '{Server}' in cluster '{Cluster}' for conform", Workspace.ServerAndPort, Workspace.Cluster);
 						return false;
 					}
 
