Horde: Fix matching of servers for conform.

diff --git a/Engine/Source/Programs/Horde/HordeServer/Models/Agent.cs b/Engine/Source/Programs/Horde/HordeServer/Models/Agent.cs
--- a/Engine/Source/Programs/Horde/HordeServer/Models/Agent.cs
+++ b/Engine/Source/Programs/Horde/HordeServer/Models/Agent.cs
@@ -811,8 +811,16 @@ namespace HordeServer.Models
 		public static async Task<bool> TryAddWorkspaceMessage(this IAgent Agent, AgentWorkspace Workspace, PerforceCluster Cluster, PerforceLoadBalancer LoadBalancer, IList<HordeCommon.Rpc.Messages.AgentWorkspace> WorkspaceMessages)
 		{
 			// Find a matching server, trying to use a previously selected one if possible
-			string? ServerAndPort = WorkspaceMessages.FirstOrDefault(x => x.ConfiguredCluster == Workspace.Cluster)?.ServerAndPort;
-			if (ServerAndPort == null)
+			string? BaseServerAndPort;
+			string? ServerAndPort;
+
+			HordeCommon.Rpc.Messages.AgentWorkspace? ExistingWorkspace = WorkspaceMessages.FirstOrDefault(x => x.ConfiguredCluster == Workspace.Cluster);
+			if(ExistingWorkspace != null)
+			{
+				BaseServerAndPort = ExistingWorkspace.BaseServerAndPort;
+				ServerAndPort = ExistingWorkspace.ServerAndPort;
+			}
+			else
 			{
 				if (Cluster == null)
 				{
@@ -824,6 +832,8 @@ namespace HordeServer.Models
 				{
 					return false;
 				}
+
+				BaseServerAndPort = Server.BaseServerAndPort;
 				ServerAndPort = Server.ServerAndPort;
 			}
 
@@ -846,6 +856,7 @@ namespace HordeServer.Models
 			Result.ConfiguredCluster = Workspace.Cluster;
 			Result.ConfiguredUserName = Workspace.UserName;
 			Result.Cluster = Cluster?.Name;
+			Result.BaseServerAndPort = BaseServerAndPort;
 			Result.ServerAndPort = ServerAndPort;
 			Result.UserName = Credentials?.UserName ?? Workspace.UserName;
 			Result.Password = Credentials?.Password;
diff --git a/Engine/Source/Programs/Horde/HordeServer/Tasks/Impl/ConformTaskSource.cs b/Engine/Source/Programs/Horde/HordeServer/Tasks/Impl/ConformTaskSource.cs
--- a/Engine/Source/Programs/Horde/HordeServer/Tasks/Impl/ConformTaskSource.cs
+++ b/Engine/Source/Programs/Horde/HordeServer/Tasks/Impl/ConformTaskSource.cs
@@ -250,10 +250,10 @@ namespace HordeServer.Tasks.Impl
 						return false;
 					}
 
-					PerforceServer? Server = Cluster.Servers.FirstOrDefault(x => x.ServerAndPort.Equals(Workspace.ServerAndPort, StringComparison.Ordinal));
+					PerforceServer? Server = Cluster.Servers.FirstOrDefault(x => x.ServerAndPort.Equals(Workspace.BaseServerAndPort, StringComparison.Ordinal));
 					if(Server == null)
 					{
-						Logger.LogWarning("Unable to find perforce server '{Server}' in cluster '{Cluster}' for conform", Workspace.ServerAndPort, Workspace.Cluster);
+						Logger.LogWarning("Unable to find perforce server '{Server}' in cluster '{Cluster}' for conform", Workspace.BaseServerAndPort, Workspace.Cluster);
 						return false;
 					}
 
