Horde: Only allow each server to be added to conform list once.

diff --git a/Engine/Source/Programs/Horde/HordeServer/Tasks/Impl/ConformTaskSource.cs b/Engine/Source/Programs/Horde/HordeServer/Tasks/Impl/ConformTaskSource.cs
--- a/Engine/Source/Programs/Horde/HordeServer/Tasks/Impl/ConformTaskSource.cs
+++ b/Engine/Source/Programs/Horde/HordeServer/Tasks/Impl/ConformTaskSource.cs
@@ -240,40 +240,43 @@ namespace HordeServer.Tasks.Impl
 					return false;
 				}
 
-				HashSet<string> Servers = new HashSet<string>(StringComparer.Ordinal);
+				HashSet<string> Servers = new HashSet<string>(StringComparer.OrdinalIgnoreCase);
 				foreach (HordeCommon.Rpc.Messages.AgentWorkspace Workspace in Workspaces)
 				{
-					PerforceCluster? Cluster = Globals.FindPerforceCluster(Workspace.Cluster);
-					if(Cluster == null)
+					if (Servers.Add(Workspace.ServerAndPort))
 					{
-						Logger.LogWarning("Unable to find perforce cluster '{Cluster}' for conform", Workspace.Cluster);
-						return false;
+						PerforceCluster? Cluster = Globals.FindPerforceCluster(Workspace.Cluster);
+						if (Cluster == null)
+						{
+							Logger.LogWarning("Unable to find perforce cluster '{Cluster}' for conform", Workspace.Cluster);
+							return false;
+						}
+
+						PerforceServer? Server = Cluster.Servers.FirstOrDefault(x => x.ServerAndPort.Equals(Workspace.BaseServerAndPort, StringComparison.Ordinal));
+						if (Server == null)
+						{
+							Logger.LogWarning("Unable to find perforce server '{Server}' in cluster '{Cluster}' for conform", Workspace.BaseServerAndPort, Workspace.Cluster);
+							return false;
+						}
+
+						ConformListServer? ConformServer = CurrentValue.Servers.FirstOrDefault(x => x.Cluster.Equals(Workspace.Cluster, StringComparison.Ordinal) && x.ServerAndPort.Equals(Workspace.ServerAndPort, StringComparison.Ordinal));
+						if (ConformServer == null)
+						{
+							ConformServer = new ConformListServer { Cluster = Workspace.Cluster, ServerAndPort = Workspace.ServerAndPort };
+							CurrentValue.Servers.Add(ConformServer);
+						}
+
+						if (Server.MaxConformCount != 0 && ConformServer.Entries.Count > Server.MaxConformCount)
+						{
+							return false;
+						}
+
+						ConformListEntry Entry = new ConformListEntry();
+						Entry.AgentId = AgentId;
+						Entry.LeaseId = LeaseId;
+						Entry.LastCheckTimeUtc = DateTime.UtcNow;
+						ConformServer.Entries.Add(Entry);
 					}
-
-					PerforceServer? Server = Cluster.Servers.FirstOrDefault(x => x.ServerAndPort.Equals(Workspace.BaseServerAndPort, StringComparison.Ordinal));
-					if(Server == null)
-					{
-						Logger.LogWarning("Unable to find perforce server '{Server}' in cluster '{Cluster}' for conform", Workspace.BaseServerAndPort, Workspace.Cluster);
-						return false;
-					}
-
-					ConformListServer? ConformServer = CurrentValue.Servers.FirstOrDefault(x => x.Cluster.Equals(Workspace.Cluster, StringComparison.Ordinal) && x.ServerAndPort.Equals(Workspace.ServerAndPort, StringComparison.Ordinal));
-					if (ConformServer == null)
-					{
-						ConformServer = new ConformListServer { Cluster = Workspace.Cluster, ServerAndPort = Workspace.ServerAndPort };
-						CurrentValue.Servers.Add(ConformServer);
-					}
-
-					if (Server.MaxConformCount != 0 && ConformServer.Entries.Count > Server.MaxConformCount)
-					{
-						return false;
-					}
-
-					ConformListEntry Entry = new ConformListEntry();
-					Entry.AgentId = AgentId;
-					Entry.LeaseId = LeaseId;
-					Entry.LastCheckTimeUtc = DateTime.UtcNow;
-					ConformServer.Entries.Add(Entry);
 				}
 
 				if (await ConformList.TryUpdateAsync(CurrentValue))
