Horde: Improve server lease display with status and links to batches
#rnx

diff --git a/Engine/Source/Programs/Horde/HordeDashboard/src/backend/Api.ts b/Engine/Source/Programs/Horde/HordeDashboard/src/backend/Api.ts
--- a/Engine/Source/Programs/Horde/HordeDashboard/src/backend/Api.ts
+++ b/Engine/Source/Programs/Horde/HordeDashboard/src/backend/Api.ts
@@ -410,6 +410,23 @@ export type BatchUpdatePoolRequest = {
 	deleted?: boolean;
 }
 
+export enum LeaseOutcome {
+    
+    /** Default value. */
+    Unspecified = "Unspecified",
+    
+    /** The lease was executed successfully */
+    Success = "Success",
+    
+    /** The lease was not executed succesfully, but cannot be run again. */
+    Failed = "Failed",
+    
+    /** The lease was cancelled by request */
+    Cancelled = "Cancelled"
+	
+  }
+
+
 /**Response for queries to find a particular lease within an agent */
 export type GetAgentLeaseResponse = {
 
@@ -445,6 +462,10 @@ export type GetAgentLeaseResponse = {
 
 	/**logId of this lease (parsed from Details array) */
 	logId: string;
+
+	/** Outcome of the lease */
+	outcome? : LeaseOutcome;
+
 }
 
 /**Information about an agent session */
diff --git a/Engine/Source/Programs/Horde/HordeServer/Api/Agents.cs b/Engine/Source/Programs/Horde/HordeServer/Api/Agents.cs
--- a/Engine/Source/Programs/Horde/HordeServer/Api/Agents.cs
+++ b/Engine/Source/Programs/Horde/HordeServer/Api/Agents.cs
@@ -1,6 +1,7 @@
 // Copyright Epic Games, Inc. All Rights Reserved.
 
 using EpicGames.Core;
+using HordeCommon;
 using HordeServer.Models;
 using HordeServer.Utilities;
 using System;
@@ -169,6 +170,12 @@ namespace HordeServer.Api
 		/// </summary>
 		public Dictionary<string, string>? Details { get; set; }
 
+		/// <summary>
+		/// Outcome of the lease
+		/// </summary>
+		public LeaseOutcome? Outcome { get; set; }
+
+
 		/// <summary>
 		/// Parameterless constructor for serialization
 		/// </summary>
@@ -190,7 +197,7 @@ namespace HordeServer.Api
 			this.StartTime = Lease.StartTime;
 			this.Executing = Lease.Active;
 			this.FinishTime = Lease.ExpiryTime;
-			this.Details = AgentLease.GetPayloadDetails(Lease.Payload);
+			this.Details = AgentLease.GetPayloadDetails(Lease.Payload);			
 		}
 
 		/// <summary>
@@ -207,6 +214,7 @@ namespace HordeServer.Api
 			this.Executing = (Lease.FinishTime == null);
 			this.FinishTime = Lease.FinishTime;
 			this.Details = AgentLease.GetPayloadDetails(Lease.Payload);
+			this.Outcome = Lease.Outcome;
 		}
 	}
 
