Horde: Change downtime service to log in the information channel.

diff --git a/Engine/Source/Programs/Horde/HordeServer/Services/DowntimeService.cs b/Engine/Source/Programs/Horde/HordeServer/Services/DowntimeService.cs
--- a/Engine/Source/Programs/Horde/HordeServer/Services/DowntimeService.cs
+++ b/Engine/Source/Programs/Horde/HordeServer/Services/DowntimeService.cs
@@ -75,7 +75,7 @@ namespace HordeServer.Services
 			DateTimeOffset Now = TimeZoneInfo.ConvertTime(DateTimeOffset.Now, Settings.CurrentValue.TimeZoneInfo);
 			bool bIsActive = Globals.ScheduledDowntime.Any(x => x.IsActive(Now));
 
-			Logger.LogDebug("Server time: {Time}. Downtime: {Downtime}. Schedule: {@Schedule}.", Now, bIsActive, Globals.ScheduledDowntime);
+			Logger.LogInformation("Server time: {Time}. Downtime: {Downtime}. Schedule: {@Schedule}.", Now, bIsActive, Globals.ScheduledDowntime);
 
 			if (bIsActive != IsDowntimeActive)
 			{
