Horde: Use batch error state in lease view
#rnx

diff --git a/Engine/Source/Programs/Horde/HordeDashboard/src/backend/Api.ts b/Engine/Source/Programs/Horde/HordeDashboard/src/backend/Api.ts
--- a/Engine/Source/Programs/Horde/HordeDashboard/src/backend/Api.ts
+++ b/Engine/Source/Programs/Horde/HordeDashboard/src/backend/Api.ts
@@ -423,7 +423,7 @@ export enum LeaseOutcome {
     
     /** The lease was cancelled by request */
     Cancelled = "Cancelled"
-	
+
   }
 
 
@@ -463,6 +463,9 @@ export type GetAgentLeaseResponse = {
 	/**logId of this lease (parsed from Details array) */
 	logId: string;
 
+	/**Batch data for this lease, queried from backend */
+	batch?: GetBatchResponse;
+
 	/** Outcome of the lease */
 	outcome? : LeaseOutcome;
 
diff --git a/Engine/Source/Programs/Horde/HordeDashboard/src/backend/Backend.ts b/Engine/Source/Programs/Horde/HordeDashboard/src/backend/Backend.ts
--- a/Engine/Source/Programs/Horde/HordeDashboard/src/backend/Backend.ts
+++ b/Engine/Source/Programs/Horde/HordeDashboard/src/backend/Backend.ts
@@ -3,7 +3,7 @@
 import { configure } from "mobx";
 import { isNumber } from 'util';
 import templateCache from '../backend/TemplateCache';
-import { AgentData, ArtifactData, BatchUpdatePoolRequest, ChangeSummaryData, CreateJobRequest, CreateJobResponse, CreatePoolRequest, CreateSoftwareResponse, CreateSubscriptionRequest, CreateSubscriptionResponse, DashboardPreference, EventData, GetArtifactZipRequest, GetUserResponse, GetGraphResponse, GetIssueStreamResponse, GetJobStepRefResponse, GetJobStepTraceResponse, GetJobTimingResponse, GetLogEventResponse, GetNotificationResponse, GetSoftwareResponse, GetSubscriptionResponse, IssueData, IssueQuery, JobData, JobQuery, LeaseData, LogData, LogLineData, PoolData, ProjectData, ScheduleData, ScheduleQuery, SearchLogFileResponse, SessionData, StreamData, TemplateData, TestData, UpdateAgentRequest, UpdateIssueRequest, UpdateJobRequest, UpdateNotificationsRequest, UpdatePoolRequest, UpdateStepRequest, UpdateStepResponse, UpdateUserRequest, GetUtilizationTelemetryResponse, TabType, GetJobsTabResponse, JobsTabColumnType, GetPerforceServerStatusResponse, GetDeviceResponse, GetDevicePlatformResponse, CreateDeviceResponse, CreateDeviceRequest, UpdateDeviceRequest, GetDevicePoolResponse, GetDeviceReservationResponse, UsersQuery } from './Api';
+import { AgentData, ArtifactData, BatchUpdatePoolRequest, ChangeSummaryData, CreateJobRequest, CreateJobResponse, CreatePoolRequest, CreateSoftwareResponse, CreateSubscriptionRequest, CreateSubscriptionResponse, DashboardPreference, EventData, GetArtifactZipRequest, GetUserResponse, GetGraphResponse, GetIssueStreamResponse, GetJobStepRefResponse, GetJobStepTraceResponse, GetJobTimingResponse, GetLogEventResponse, GetNotificationResponse, GetSoftwareResponse, GetSubscriptionResponse, IssueData, IssueQuery, JobData, JobQuery, LeaseData, LogData, LogLineData, PoolData, ProjectData, ScheduleData, ScheduleQuery, SearchLogFileResponse, SessionData, StreamData, TemplateData, TestData, UpdateAgentRequest, UpdateIssueRequest, UpdateJobRequest, UpdateNotificationsRequest, UpdatePoolRequest, UpdateStepRequest, UpdateStepResponse, UpdateUserRequest, GetUtilizationTelemetryResponse, TabType, GetJobsTabResponse, JobsTabColumnType, GetPerforceServerStatusResponse, GetDeviceResponse, GetDevicePlatformResponse, CreateDeviceResponse, CreateDeviceRequest, UpdateDeviceRequest, GetDevicePoolResponse, GetDeviceReservationResponse, UsersQuery, LeaseOutcome } from './Api';
 import dashboard from './Dashboard';
 import { ChallengeStatus, Fetch } from './Fetch';
 import graphCache, { GraphQuery } from './GraphCache';
@@ -136,13 +136,13 @@ export class Backend {
 
     }
 
-    getLeases(agentId: string, index: number, count: number): Promise<LeaseData[]> {
+    getLeases(agentId: string, index: number, count: number, includeBatches?:boolean): Promise<LeaseData[]> {
         const params: string[] = [];
         params.push("Index=" + index);
         params.push("Count=" + count);
         const paramString = params.join("&");
         return new Promise<LeaseData[]>((resolve, reject) => {
-            this.backend.get(`/api/v1/agents/${agentId}/leases?${paramString}`).then((response) => {
+            this.backend.get(`/api/v1/agents/${agentId}/leases?${paramString}`).then( async (response) => {
                 const leases = response.data as LeaseData[];
                 leases.forEach(function (lease) {
                     lease.startTime = new Date(Date.parse(lease.startTime as string));
@@ -155,6 +155,23 @@ export class Backend {
                         lease.type = lease.details.Type;
                     }
                 });
+
+                // get batch states
+                if (includeBatches) {
+
+                    const jobIds:Set<string> = new Set();
+                    leases.filter(lease => lease.jobId && lease.batchId).forEach(lease => jobIds.add(lease.jobId));
+                    if (jobIds.size) {
+                        const jobs = await this.getJobsByIds(Array.from(jobIds), {filter:"id,batches"});
+                        jobs.forEach(j => {
+                            const lease = leases.find(lease => lease.jobId === j.id && j.batches?.find(b => lease.batchId === b.id));
+                            if (lease) {
+                                lease.batch = j.batches?.find(b => lease.batchId === b.id );
+                            }
+                        });
+                    }    
+                }
+
                 resolve(leases);
             }).catch(reason => { reject(reason); });
         });
