Horde: Fix job batch lookup for leases
#rnx

diff --git a/Engine/Source/Programs/Horde/HordeDashboard/src/backend/Backend.ts b/Engine/Source/Programs/Horde/HordeDashboard/src/backend/Backend.ts
--- a/Engine/Source/Programs/Horde/HordeDashboard/src/backend/Backend.ts
+++ b/Engine/Source/Programs/Horde/HordeDashboard/src/backend/Backend.ts
@@ -163,10 +163,11 @@ export class Backend {
                     leases.filter(lease => lease.jobId && lease.batchId).forEach(lease => jobIds.add(lease.jobId));
                     if (jobIds.size) {
                         const jobs = await this.getJobsByIds(Array.from(jobIds), {filter:"id,batches"});
-                        jobs.forEach(j => {
-                            const lease = leases.find(lease => lease.jobId === j.id && j.batches?.find(b => lease.batchId === b.id));
-                            if (lease) {
-                                lease.batch = j.batches?.find(b => lease.batchId === b.id );
+
+                        leases.forEach(lease => {
+                            const job = jobs.find(j => lease.jobId === j.id);
+                            if (job) {
+                                lease.batch = job.batches?.find(b => lease.batchId === b.id );
                             }
                         });
                     }    
