Horde: Fix (caught) exception when checking for port being available when running tests.

diff --git a/Engine/Source/Programs/Horde/HordeServerTests/DatabaseRunner.cs b/Engine/Source/Programs/Horde/HordeServerTests/DatabaseRunner.cs
--- a/Engine/Source/Programs/Horde/HordeServerTests/DatabaseRunner.cs
+++ b/Engine/Source/Programs/Horde/HordeServerTests/DatabaseRunner.cs
@@ -4,7 +4,9 @@ using System;
 using System.Collections.Generic;
 using System.Diagnostics;
 using System.IO;
+using System.Linq;
 using System.Net;
+using System.Net.NetworkInformation;
 using System.Net.Sockets;
 using System.Reflection;
 using System.Runtime.InteropServices;
@@ -141,26 +143,15 @@ namespace HordeServerTests
 
 		private static bool IsPortAvailable(int Port)
 		{
-			TcpListener? ListenerAny = null;
-			TcpListener? ListenerLoopback = null;
-			try
-			{
-				ListenerAny = new TcpListener(IPAddress.Loopback, Port);
-				ListenerAny.Start();
-				ListenerLoopback = new TcpListener(IPAddress.Any, Port);
-				ListenerLoopback.Start();
-				return true;
-			}
-			catch (SocketException)
-			{
-			}
-			finally
+			IPGlobalProperties IpGlobalProperties = IPGlobalProperties.GetIPGlobalProperties();
+
+			IPEndPoint[] Listeners = IpGlobalProperties.GetActiveTcpListeners();
+			if (Listeners.Any(x => x.Port == Port))
 			{
-				ListenerAny?.Stop();
-				ListenerLoopback?.Stop();
+				return false;
 			}
 
-			return false;
+			return true;
 		}
 		
 		private static void DeleteDirectory(string Path)
