Horde: Basic implementation of storage API.

Current limitations:

* Only supports a filesystem backend
* No garbage collection
* Refs are stored in the database, and do not automatically expire.
* Validation of refs is slow and does not cache any data.

diff --git a/Engine/Source/Programs/Horde/HordeServer/HordeServer.csproj b/Engine/Source/Programs/Horde/HordeServer/HordeServer.csproj
--- a/Engine/Source/Programs/Horde/HordeServer/HordeServer.csproj
+++ b/Engine/Source/Programs/Horde/HordeServer/HordeServer.csproj
@@ -69,6 +69,7 @@
     <ProjectReference Include="..\..\Shared\EpicGames.Horde.Common\EpicGames.Horde.Common.csproj" />
     <ProjectReference Include="..\..\Shared\EpicGames.IoHash\EpicGames.IoHash.csproj" />
     <ProjectReference Include="..\..\Shared\EpicGames.Perforce\EpicGames.Perforce.csproj" />
+    <ProjectReference Include="..\..\Shared\EpicGames.Serialization\EpicGames.Serialization.csproj" />
     <ProjectReference Include="..\HordeCommon\HordeCommon.csproj" />
     <ProjectReference Include="..\..\..\ThirdParty\Perforce\P4Api.Net\p4api.net\p4api.net.csproj" />
   </ItemGroup>
diff --git a/Engine/Source/Programs/Horde/HordeServer/Models/Acl.cs b/Engine/Source/Programs/Horde/HordeServer/Models/Acl.cs
--- a/Engine/Source/Programs/Horde/HordeServer/Models/Acl.cs
+++ b/Engine/Source/Programs/Horde/HordeServer/Models/Acl.cs
@@ -404,6 +404,8 @@ namespace HordeServer.Models
 
 
 
+
+
 		//// NOTIFICATIONS ////
 
 		/// <summary>
@@ -411,6 +413,10 @@ namespace HordeServer.Models
 		/// </summary>
 		CreateSubscription,
 
+
+
+
+
 		//// DEVICES ////
 		
 		/// <summary>
@@ -421,8 +427,38 @@ namespace HordeServer.Models
 		/// <summary>
 		/// Ability to write devices
 		/// </summary>
-		DeviceWrite
+		DeviceWrite,
+
 
+
+
+
+		//// STORAGE ////
+
+		/// <summary>
+		/// Ability to read blobs from the storage service
+		/// </summary>
+		ReadBlobs,
+
+		/// <summary>
+		/// Ability to write blobs to the storage service
+		/// </summary>
+		WriteBlobs,
+
+		/// <summary>
+		/// Ability to read refs from the storage service
+		/// </summary>
+		ReadRefs,
+
+		/// <summary>
+		/// Ability to write refs to the storage service
+		/// </summary>
+		WriteRefs, 
+
+		/// <summary>
+		/// Ability to delete refs
+		/// </summary>
+		DeleteRefs,
 	}
 
 	/// <summary>
diff --git a/Engine/Source/Programs/Horde/HordeServer/Models/Config.cs b/Engine/Source/Programs/Horde/HordeServer/Models/Config.cs
--- a/Engine/Source/Programs/Horde/HordeServer/Models/Config.cs
+++ b/Engine/Source/Programs/Horde/HordeServer/Models/Config.cs
@@ -46,6 +46,11 @@ namespace HordeServer.Models
 		/// </summary>
 		public int MaxConformCount { get; set; }
 
+		/// <summary>
+		/// List of storage namespaces
+		/// </summary>
+		public StorageConfig? Storage { get; set; }
+
 		/// <summary>
 		/// Access control list
 		/// </summary>
@@ -200,5 +205,48 @@ namespace HordeServer.Models
 		/// </summary>
 		public string? PauseComment { get; set; }
 	}
+
+	/// <summary>
+	/// Configuration for storage system
+	/// </summary>
+	public class StorageConfig
+	{
+		/// <summary>
+		/// List of storage namespaces
+		/// </summary>
+		public List<NamespaceConfig> Namespaces { get; set; } = new List<NamespaceConfig>();
+	}
+
+	/// <summary>
+	/// Configuration for a storage namespace
+	/// </summary>
+	public class NamespaceConfig
+	{
+		/// <summary>
+		/// Identifier for this namespace
+		/// </summary>
+		public string Id { get; set; } = String.Empty;
+
+		/// <summary>
+		/// Buckets within this namespace
+		/// </summary>
+		public List<BucketConfig> Buckets { get; set; } = new List<BucketConfig>();
+
+		/// <summary>
+		/// Access control for this namespace
+		/// </summary>
+		public UpdateAclRequest? Acl { get; set; }
+	}
+
+	/// <summary>
+	/// Configuration for a bucket
+	/// </summary>
+	public class BucketConfig
+	{
+		/// <summary>
+		/// Identifier for the bucket
+		/// </summary>
+		public string Id { get; set; } = String.Empty;
+	}
 }
 
diff --git a/Engine/Source/Programs/Horde/HordeServer/ServerSettings.cs b/Engine/Source/Programs/Horde/HordeServer/ServerSettings.cs
--- a/Engine/Source/Programs/Horde/HordeServer/ServerSettings.cs
+++ b/Engine/Source/Programs/Horde/HordeServer/ServerSettings.cs
@@ -254,6 +254,16 @@ namespace HordeServer
 		/// </summary>
 		public DirectoryReference LocalLogsDirRef => DirectoryReference.Combine(Program.DataDir, LocalLogsDir);
 
+		/// <summary>
+		/// Local blob storage directory, if using type filesystem
+		/// </summary>
+		public string LocalBlobsDir { get; set; } = "Blobs";
+
+		/// <summary>
+		/// Gets the full path referred to by LocalStorageDir
+		/// </summary>
+		public DirectoryReference LocalBlobsDirRef => DirectoryReference.Combine(Program.DataDir, LocalBlobsDir);
+
 		/// <summary>
 		/// Local artifact storage directory, if using type filesystem
 		/// </summary>
diff --git a/Engine/Source/Programs/Horde/HordeServer/Startup.cs b/Engine/Source/Programs/Horde/HordeServer/Startup.cs
--- a/Engine/Source/Programs/Horde/HordeServer/Startup.cs
+++ b/Engine/Source/Programs/Horde/HordeServer/Startup.cs
@@ -71,6 +71,7 @@ using HordeServer.Notifications.Impl;
 using HordeServer.Notifications;
 using Microsoft.AspNetCore.Hosting.Server.Features;
 using System.Runtime.InteropServices;
+using HordeServer.Storage.Collections;
 
 namespace HordeServer
 {
@@ -316,6 +317,13 @@ namespace HordeServer
 			Services.AddSingleton<IUserCollection, UserCollectionV2>();
 			Services.AddSingleton<IDeviceCollection, DeviceCollection>();
 
+			// Storage
+			Services.AddSingleton<IBlobCollection, FileSystemBlobCollection>();
+			Services.AddSingleton<IBucketCollection, BucketCollection>();
+			Services.AddSingleton<INamespaceCollection, NamespaceCollection>();
+			Services.AddSingleton<IObjectCollection, ObjectCollection>();
+			Services.AddSingleton<IRefCollection, RefCollection>();
+
 			Services.AddSingleton(typeof(ISingletonDocument<>), typeof(SingletonDocument<>));
 
 			Services.AddSingleton<AutoscaleService>();
diff --git a/Engine/Source/Programs/Horde/HordeServer/Storage/Collections/BlobCollection.cs b/Engine/Source/Programs/Horde/HordeServer/Storage/Collections/BlobCollection.cs
new file mode 100644
--- /dev/null
+++ b/Engine/Source/Programs/Horde/HordeServer/Storage/Collections/BlobCollection.cs
@@ -0,0 +1,135 @@
+// Copyright Epic Games, Inc. All Rights Reserved.
+
+using EpicGames.Core;
+using HordeServer.Utilities;
+using Microsoft.Extensions.Options;
+using System;
+using System.Collections.Generic;
+using System.IO;
+using System.Linq;
+using System.Threading.Tasks;
+
+namespace HordeServer.Storage.Collections
+{
+	using NamespaceId = StringId<INamespace>;
+
+	class FileSystemBlobCollection : IBlobCollection
+	{
+		/// <summary>
+		/// Base directory for log files
+		/// </summary>
+		private readonly DirectoryReference BaseDir;
+
+		/// <summary>
+		/// Unique identifier for this instance
+		/// </summary>
+		private string InstanceId;
+
+		/// <summary>
+		/// Constructor
+		/// </summary>
+		/// <param name="Settings">Current Horde Settings</param>
+		public FileSystemBlobCollection(IOptions<ServerSettings> Settings)
+		{
+			ServerSettings CurrentSettings = Settings.Value;
+			BaseDir = CurrentSettings.LocalBlobsDirRef;
+			InstanceId = Guid.NewGuid().ToString("N");
+			DirectoryReference.CreateDirectory(BaseDir);
+		}
+
+		/// <inheritdoc/>
+		public async Task AddAsync(NamespaceId NamespaceId, IoHash Hash, Stream Input)
+		{
+			FileReference Location = GetPath(NamespaceId, Hash);
+			DirectoryReference.CreateDirectory(Location.Directory);
+
+			FileReference TempLocation = new FileReference($"{Location}.{InstanceId}");
+
+			Blake3.Hasher Hasher = Blake3.Hasher.New();
+			using (FileStream Output = FileReference.Open(TempLocation, FileMode.Create, FileAccess.Write, FileShare.Read))
+			{
+				byte[] ReadBuffer = new byte[64 * 1024];
+				Task<int> ReadTask = Task.Run(() => Input.ReadAsync(ReadBuffer, 0, ReadBuffer.Length));
+
+				byte[] WriteBuffer = new byte[64 * 1024];
+				int WriteLength = await ReadTask;
+
+				while (WriteLength > 0)
+				{
+					(WriteBuffer, ReadBuffer) = (ReadBuffer, WriteBuffer);
+
+					ReadTask = Task.Run(() => Input.ReadAsync(ReadBuffer, 0, ReadBuffer.Length));
+					Task WriteTask = Task.Run(() => Output.WriteAsync(WriteBuffer, 0, WriteLength));
+
+					Hasher.Update(WriteBuffer.AsSpan(0, WriteLength));
+
+					await Task.WhenAll(ReadTask, WriteTask);
+					WriteLength = await ReadTask;
+				}
+			}
+
+			IoHash ActualHash = new IoHash(Hasher);
+			if (ActualHash != Hash)
+			{
+				throw new ArgumentException($"Hash for blob is incorrect; should be {ActualHash}, not {Hash}");
+			}
+
+			try
+			{
+				FileReference.Move(TempLocation, Location, true);
+			}
+			catch (IOException) // Already exists
+			{
+				if (FileReference.Exists(Location))
+				{
+					FileReference.Delete(TempLocation);
+				}
+				else
+				{
+					throw;
+				}
+			}
+		}
+
+		/// <inheritdoc/>
+		public Task<Stream?> GetAsync(NamespaceId NamespaceId, IoHash Hash)
+		{
+			FileReference Location = GetPath(NamespaceId, Hash);
+			if (!FileReference.Exists(Location))
+			{
+				return Task.FromResult<Stream?>(null);
+			}
+
+			try
+			{
+				return Task.FromResult<Stream?>(FileReference.Open(Location, FileMode.Open, FileAccess.Read, FileShare.Read));
+			}
+			catch (DirectoryNotFoundException)
+			{
+				return Task.FromResult<Stream?>(null);
+			}
+			catch (FileNotFoundException)
+			{
+				return Task.FromResult<Stream?>(null);
+			}
+		}
+
+		/// <inheritdoc/>
+		public Task<List<IoHash>> ExistsAsync(NamespaceId NamespaceId, IEnumerable<IoHash> Hashes)
+		{
+			return Task.FromResult<List<IoHash>>(Hashes.Where(x => FileReference.Exists(GetPath(NamespaceId, x))).ToList());
+		}
+
+		/// <summary>
+		/// Gets the path for a blob
+		/// </summary>
+		/// <param name="NsId">The namespace id</param>
+		/// <param name="Hash">Hash of the blob</param>
+		/// <returns></returns>
+		FileReference GetPath(NamespaceId NsId, IoHash Hash)
+		{
+			string HashText = Hash.ToString();
+			return FileReference.Combine(BaseDir, NsId.ToString(), HashText[0..2], HashText[2..4], $"{HashText}.blob");
+		}
+	}
+}
diff --git a/Engine/Source/Programs/Horde/HordeServer/Storage/Collections/BucketCollection.cs b/Engine/Source/Programs/Horde/HordeServer/Storage/Collections/BucketCollection.cs
new file mode 100644
--- /dev/null
+++ b/Engine/Source/Programs/Horde/HordeServer/Storage/Collections/BucketCollection.cs
@@ -0,0 +1,71 @@
+// Copyright Epic Games, Inc. All Rights Reserved.
+
+using HordeServer.Models;
+using HordeServer.Services;
+using HordeServer.Utilities;
+using MongoDB.Driver;
+using System;
+using System.Collections.Generic;
+using System.Linq;
+using System.Threading.Tasks;
+
+namespace HordeServer.Storage.Collections
+{
+	using NamespaceId = StringId<INamespace>;
+	using BucketId = StringId<IBucket>;
+
+	class BucketCollection : IBucketCollection
+	{
+		class BucketCombinedId
+		{
+			public NamespaceId NamespaceId { get; set; }
+			public BucketId BucketId { get; set; }
+		}
+
+		class Bucket : IBucket
+		{
+			public BucketCombinedId Id { get; set; } = null!;
+
+			NamespaceId IBucket.NamespaceId => Id.NamespaceId;
+			BucketId IBucket.BucketId => Id.BucketId;
+
+			public bool Deleted { get; set; }
+		}
+
+		IMongoCollection<Bucket> Buckets;
+
+		public BucketCollection(DatabaseService DatabaseService)
+		{
+			this.Buckets = DatabaseService.GetCollection<Bucket>("Storage.Buckets");
+
+			if (!DatabaseService.ReadOnlyMode)
+			{
+				this.Buckets.Indexes.CreateOne(new CreateIndexModel<Bucket>(Builders<Bucket>.IndexKeys.Ascending(x => x.Id.NamespaceId)));
+			}
+		}
+
+		public async Task AddOrUpdateAsync(NamespaceId NamespaceId, BucketId BucketId, BucketConfig Bucket)
+		{
+			Bucket NewBucket = new Bucket();
+			NewBucket.Id = new BucketCombinedId { NamespaceId = NamespaceId, BucketId = BucketId };
+			await Buckets.ReplaceOneAsync(x => x.Id == NewBucket.Id, NewBucket, new ReplaceOptions { IsUpsert = true });
+		}
+
+		public async Task<List<IBucket>> FindAsync(NamespaceId NamespaceId)
+		{
+			return await Buckets.Find(x => x.Id.NamespaceId == NamespaceId).ToListAsync<Bucket, IBucket>();
+		}
+
+		public async Task<IBucket?> GetAsync(NamespaceId NamespaceId, BucketId BucketId)
+		{
+			BucketCombinedId Id = new BucketCombinedId { NamespaceId = NamespaceId, BucketId = BucketId };
+			return await Buckets.Find(x => x.Id == Id).FirstOrDefaultAsync();
+		}
+
+		public async Task RemoveAsync(NamespaceId NamespaceId, BucketId BucketId)
+		{
+			BucketCombinedId Id = new BucketCombinedId { NamespaceId = NamespaceId, BucketId = BucketId };
+			await Buckets.UpdateOneAsync(x => x.Id == Id, Builders<Bucket>.Update.Set(x => x.Deleted, true));
+		}
+	}
+}
diff --git a/Engine/Source/Programs/Horde/HordeServer/Storage/Collections/NamespaceCollection.cs b/Engine/Source/Programs/Horde/HordeServer/Storage/Collections/NamespaceCollection.cs
new file mode 100644
--- /dev/null
+++ b/Engine/Source/Programs/Horde/HordeServer/Storage/Collections/NamespaceCollection.cs
@@ -0,0 +1,91 @@
+// Copyright Epic Games, Inc. All Rights Reserved.
+
+using EpicGames.Core;
+using EpicGames.Serialization;
+using HordeCommon;
+using HordeServer.Models;
+using HordeServer.Services;
+using HordeServer.Utilities;
+using Microsoft.Extensions.Caching.Memory;
+using MongoDB.Bson;
+using MongoDB.Bson.Serialization.Attributes;
+using MongoDB.Driver;
+using System;
+using System.Collections.Generic;
+using System.Diagnostics;
+using System.Globalization;
+using System.IO;
+using System.IO.Compression;
+using System.Linq;
+using System.Security.Claims;
+using System.Text.Json;
+using System.Threading.Tasks;
+
+namespace HordeServer.Storage
+{
+	using NamespaceId = StringId<INamespace>;
+
+	class NamespaceCollection : INamespaceCollection, IDisposable
+	{
+		class Namespace : INamespace
+		{
+			public NamespaceId Id { get; set; }
+			public Acl? Acl { get; set; }
+			public bool Deleted { get; set; }
+		}
+
+		IMongoCollection<Namespace> Namespaces;
+		LazyCache<NamespaceId, INamespace?> NamespaceCache;
+
+		public NamespaceCollection(DatabaseService DatabaseService)
+		{
+			this.Namespaces = DatabaseService.GetCollection<Namespace>("Storage.Namespaces");
+			this.NamespaceCache = new LazyCache<NamespaceId, INamespace?>(GetAsync, new LazyCacheOptions { });
+		}
+
+		public void Dispose()
+		{
+			NamespaceCache.Dispose();
+		}
+
+		public async Task AddOrUpdateAsync(NamespaceId NamespaceId, NamespaceConfig Config)
+		{
+			Namespace Namespace = new Namespace();
+			Namespace.Id = NamespaceId;
+			Namespace.Acl = Acl.Merge(null, Config.Acl);
+			await Namespaces.ReplaceOneAsync(x => x.Id == Namespace.Id, Namespace, new ReplaceOptions { IsUpsert = true });
+		}
+
+		public async Task<List<INamespace>> FindAsync()
+		{
+			return await Namespaces.Find(x => !x.Deleted).ToListAsync<Namespace, INamespace>();
+		}
+
+		public async Task<INamespace?> GetAsync(NamespaceId NamespaceId)
+		{
+			return await Namespaces.Find(x => x.Id == NamespaceId).FirstOrDefaultAsync();
+		}
+
+		public async Task RemoveAsync(NamespaceId NamespaceId)
+		{
+			await Namespaces.UpdateOneAsync(x => x.Id == NamespaceId, Builders<Namespace>.Update.Set(x => x.Deleted, true));
+		}
+
+		public async Task<bool> AuthorizeAsync(NamespaceId NamespaceId, ClaimsPrincipal User, AclAction Action)
+		{
+			INamespace? Namespace = await NamespaceCache.GetAsync(NamespaceId);
+			if (Namespace == null)
+			{
+				return false;
+			}
+
+			bool? Result = Namespace.Acl?.Authorize(Action, User);
+			if (Result == null)
+			{
+				return false;
+			}
+
+			return Result.Value;
+		}
+	}
+}
diff --git a/Engine/Source/Programs/Horde/HordeServer/Storage/Collections/ObjectCollection.cs b/Engine/Source/Programs/Horde/HordeServer/Storage/Collections/ObjectCollection.cs
new file mode 100644
--- /dev/null
+++ b/Engine/Source/Programs/Horde/HordeServer/Storage/Collections/ObjectCollection.cs
@@ -0,0 +1,75 @@
+// Copyright Epic Games, Inc. All Rights Reserved.
+
+using EpicGames.Core;
+using EpicGames.Serialization;
+using HordeServer.Utilities;
+using System;
+using System.Collections.Generic;
+using System.IO;
+using System.Linq;
+using System.Threading.Tasks;
+
+namespace HordeServer.Storage.Collections
+{
+	using NamespaceId = StringId<INamespace>;
+
+	class ObjectCollection : IObjectCollection
+	{
+		IBlobCollection BlobCollection;
+
+		/// <summary>
+		/// Constructor
+		/// </summary>
+		/// <param name="BlobCollection"></param>
+		public ObjectCollection(IBlobCollection BlobCollection)
+		{
+			this.BlobCollection = BlobCollection;
+		}
+
+		/// <inheritdoc/>
+		public async Task AddAsync(NamespaceId NsId, IoHash Hash, CbObject Obj)
+		{
+			ReadOnlyMemory<byte> Memory = Obj.GetView();
+			using (ReadOnlyMemoryStream Stream = new ReadOnlyMemoryStream(Memory))
+			{
+				await BlobCollection.AddAsync(NsId, Hash, Stream);
+			}
+		}
+
+		/// <inheritdoc/>
+		public async Task<CbObject?> GetAsync(NamespaceId NsId, IoHash Hash)
+		{
+			byte[]? Data = await BlobCollection.GetByteArrayAsync(NsId, Hash);
+			if(Data == null)
+			{
+				return null;
+			}
+			return new CbObject(Data);
+		}
+
+		/// <inheritdoc/>
+		public Task<List<IoHash>> ExistsAsync(NamespaceId NsId, IEnumerable<IoHash> Hashes)
+		{
+			return BlobCollection.ExistsAsync(NsId, Hashes);
+		}
+	}
+
+	/// <summary>
+	/// Extension methods for <see cref="IObjectCollection"/>
+	/// </summary>
+	public static class ObjectCollectionExtensions
+	{
+		/// <summary>
+		/// Test whether a single object exists
+		/// </summary>
+		/// <param name="Collection"></param>
+		/// <param name="NamespaceId"></param>
+		/// <param name="Hash"></param>
+		/// <returns></returns>
+		public static async Task<bool> ExistsAsync(this IObjectCollection Collection, NamespaceId NamespaceId, IoHash Hash)
+		{
+			List<IoHash> Hashes = await Collection.ExistsAsync(NamespaceId, new[] { Hash });
+			return Hashes.Count > 0;
+		}
+	}
+}
diff --git a/Engine/Source/Programs/Horde/HordeServer/Storage/Collections/RefCollection.cs b/Engine/Source/Programs/Horde/HordeServer/Storage/Collections/RefCollection.cs
new file mode 100644
--- /dev/null
+++ b/Engine/Source/Programs/Horde/HordeServer/Storage/Collections/RefCollection.cs
@@ -0,0 +1,169 @@
+// Copyright Epic Games, Inc. All Rights Reserved.
+
+using EpicGames.Core;
+using EpicGames.Serialization;
+using HordeServer.Models;
+using HordeServer.Services;
+using HordeServer.Utilities;
+using MongoDB.Bson.Serialization.Attributes;
+using MongoDB.Driver;
+using System;
+using System.Collections.Generic;
+using System.Linq;
+using System.Threading.Tasks;
+
+namespace HordeServer.Storage.Collections
+{
+	using NamespaceId = StringId<INamespace>;
+	using BucketId = StringId<IBucket>;
+
+	class RefCollection : IRefCollection
+	{
+		class RefId
+		{
+			[BsonElement("ns")]
+			public NamespaceId NamespaceId { get; set; }
+	
+			[BsonElement("bkt")]
+			public BucketId BucketId { get; set; }
+
+			[BsonElement("name")]
+			public string Name { get; set; }
+
+			[BsonConstructor]
+			public RefId()
+			{
+				Name = null!;
+			}
+
+			public RefId(NamespaceId NamespaceId, BucketId BucketId, string Name)
+			{
+				this.NamespaceId = NamespaceId;
+				this.BucketId = BucketId;
+				this.Name = Name;
+			}
+		}
+
+		class Ref : IRef
+		{
+			public RefId Id { get; set; } = null!;
+
+			[BsonElement("t")]
+			public DateTime LastAccessTime { get; set; }
+
+			[BsonElement("d")]
+			public byte[] Data { get; set; } = null!;
+
+			[BsonElement("f")]
+			public bool Finalized { get; set; }
+
+			public NamespaceId NamespaceId => Id.NamespaceId;
+			public BucketId BucketId => Id.BucketId;
+			public string Name => Id.Name;
+
+			public CbObject Value => new CbObject(Data);
+		}
+
+		IMongoCollection<Ref> Refs;
+		IBlobCollection BlobCollection;
+		IObjectCollection ObjectCollection;
+
+		public RefCollection(DatabaseService DatabaseService, IBlobCollection BlobCollection, IObjectCollection ObjectCollection)
+		{
+			this.Refs = DatabaseService.GetCollection<Ref>("Storage.Refs");
+			this.BlobCollection = BlobCollection;
+			this.ObjectCollection = ObjectCollection;
+		}
+
+		public async Task<List<IoHash>> SetAsync(NamespaceId NamespaceId, BucketId BucketId, string Name, CbObject Value)
+		{
+			Ref NewRef = new Ref();
+			NewRef.Id = new RefId(NamespaceId, BucketId, Name);
+			NewRef.LastAccessTime = DateTime.UtcNow;
+			NewRef.Data = Value.GetView().ToArray();
+			await Refs.ReplaceOneAsync(x => x.Id == NewRef.Id, NewRef, new ReplaceOptions { IsUpsert = true });
+
+			return await FinalizeAsync(NewRef);
+		}
+
+		public async Task<List<IoHash>> FinalizeAsync(IRef Ref)
+		{
+			HashSet<IoHash> MissingHashes = new HashSet<IoHash>();
+			if (!Ref.Finalized)
+			{
+				HashSet<IoHash> CheckedObjectHashes = new HashSet<IoHash>();
+				HashSet<IoHash> CheckedBinaryHashes = new HashSet<IoHash>();
+				await FinalizeInternalAsync(Ref.NamespaceId, Ref.Value, CheckedObjectHashes, CheckedBinaryHashes, MissingHashes);
+
+				if (MissingHashes.Count == 0)
+				{
+					RefId RefId = new RefId(Ref.NamespaceId, Ref.BucketId, Ref.Name);
+					await Refs.UpdateOneAsync(x => x.Id == RefId, Builders<Ref>.Update.Set(x => x.Finalized, true));
+				}
+			}
+			return MissingHashes.ToList();
+		}
+
+		async Task FinalizeInternalAsync(NamespaceId NamespaceId, CbObject Object, HashSet<IoHash> CheckedObjectHashes, HashSet<IoHash> CheckedBinaryHashes, HashSet<IoHash> MissingHashes)
+		{
+			List<IoHash> NewObjectHashes = new List<IoHash>();
+			List<IoHash> NewBinaryHashes = new List<IoHash>();
+			Object.IterateAttachments(Field =>
+			{
+				IoHash AttachmentHash = Field.AsAttachment();
+				if (Field.IsObjectAttachment())
+				{
+					if (CheckedObjectHashes.Add(AttachmentHash))
+					{
+						NewObjectHashes.Add(AttachmentHash);
+					}
+				}
+				else
+				{
+					if (CheckedBinaryHashes.Add(AttachmentHash))
+					{
+						NewBinaryHashes.Add(AttachmentHash);
+					}
+				}
+			});
+
+			foreach (IoHash NewHash in NewObjectHashes)
+			{
+				CbObject? NextObject = await ObjectCollection.GetAsync(NamespaceId, NewHash);
+				if (NextObject == null)
+				{
+					MissingHashes.Add(NewHash);
+				}
+				else
+				{
+					await FinalizeInternalAsync(NamespaceId, NextObject, CheckedObjectHashes, CheckedBinaryHashes, MissingHashes);
+				}
+			}
+
+			if (NewBinaryHashes.Count > 0)
+			{
+				MissingHashes.UnionWith(NewBinaryHashes.Except(await BlobCollection.ExistsAsync(NamespaceId, NewBinaryHashes)));
+			}
+		}
+
+		public async Task<IRef?> GetAsync(NamespaceId NamespaceId, BucketId BucketId, string Name)
+		{
+			RefId RefId = new RefId(NamespaceId, BucketId, Name);
+			return await Refs.FindOneAndUpdateAsync<Ref>(x => x.Id == RefId, Builders<Ref>.Update.Set(x => x.LastAccessTime, DateTime.UtcNow), new FindOneAndUpdateOptions<Ref> { ReturnDocument = ReturnDocument.After });
+		}
+
+		public async Task<bool> TouchAsync(NamespaceId NamespaceId, BucketId BucketId, string Name)
+		{
+			RefId RefId = new RefId(NamespaceId, BucketId, Name);
+			UpdateResult Result = await Refs.UpdateOneAsync(x => x.Id == RefId, Builders<Ref>.Update.Set(x => x.LastAccessTime, DateTime.UtcNow));
+			return Result.MatchedCount > 0;
+		}
+
+		public async Task<bool> DeleteAsync(NamespaceId NamespaceId, BucketId BucketId, string Name)
+		{
+			RefId RefId = new RefId(NamespaceId, BucketId, Name);
+			DeleteResult Result = await Refs.DeleteOneAsync(x => x.Id == RefId);
+			return Result.DeletedCount > 0;
+		}
+	}
+}
diff --git a/Engine/Source/Programs/Horde/HordeServer/Storage/Controllers/BlobsController.cs b/Engine/Source/Programs/Horde/HordeServer/Storage/Controllers/BlobsController.cs
new file mode 100644
--- /dev/null
+++ b/Engine/Source/Programs/Horde/HordeServer/Storage/Controllers/BlobsController.cs
@@ -0,0 +1,168 @@
+// Copyright Epic Games, Inc. All Rights Reserved.
+
+using HordeServer.Api;
+using HordeServer.Services;
+using Microsoft.AspNetCore.Authorization;
+using Microsoft.AspNetCore.Http;
+using Microsoft.AspNetCore.Mvc;
+using Microsoft.Net.Http.Headers;
+using System;
+using System.IO;
+using System.Collections.Generic;
+using System.Linq;
+using System.Threading.Tasks;
+using HordeServer.Models;
+using MongoDB.Bson;
+using HordeServer.Utilities;
+using EpicGames.Core;
+using HordeServer.Storage.Collections;
+
+namespace HordeServer.Storage.Controllers
+{
+	using NamespaceId = StringId<INamespace>;
+
+	/// <summary>
+	/// Response from putting a blob
+	/// </summary>
+	public class PutBlobResponse
+	{
+		/// <summary>
+		/// Hash of the blob that was added
+		/// </summary>
+		public IoHash Hash { get; set; }
+	}
+
+	/// <summary>
+	/// Response for queries to test whether a set of blobs exist
+	/// </summary>
+	public class ExistsResponse
+	{
+		/// <summary>
+		/// List of blobs that exist
+		/// </summary>
+		public IoHash[]? Id { get; set; } = null;
+	}
+
+	/// <summary>
+	/// Controller for the /api/v1/software endpoint
+	/// </summary>
+	[ApiController]
+	[Authorize]
+	[Route("[controller]")]
+	public class BlobsController : ControllerBase
+	{
+		/// <summary>
+		/// Singleton instance of the namespace collection
+		/// </summary>
+		private readonly INamespaceCollection NamespaceCollection;
+
+		/// <summary>
+		/// Singleton instance of the blob collection
+		/// </summary>
+		private readonly IBlobCollection BlobCollection;
+
+		/// <summary>
+		/// Constructor
+		/// </summary>
+		/// <param name="NamespaceCollection">The namespace collection</param>
+		/// <param name="BlobCollection">The blob collection</param>
+		public BlobsController(INamespaceCollection NamespaceCollection, IBlobCollection BlobCollection)
+		{
+			this.NamespaceCollection = NamespaceCollection;
+			this.BlobCollection = BlobCollection;
+		}
+
+		/// <summary>
+		/// Submit a blob to the storage api. The payload hash will be verified against the received payload. 
+		/// The hash of the object is returned for a successful submission.
+		/// </summary>
+		/// <param name="NamespaceId">Namespace for the operation</param>
+		/// <param name="Hash">Hash of the blob. Must match the submitted data.</param>
+		/// <param name="File">The data to add</param>
+		/// <returns></returns>
+		[HttpPut]
+		[Route("/api/v1/blobs/{NamespaceId}/{Hash}")]
+		public async Task<ActionResult<PutBlobResponse>> PutBlobAsync(NamespaceId NamespaceId, IoHash Hash, IFormFile File)
+		{
+			if (!await NamespaceCollection.AuthorizeAsync(NamespaceId, User, AclAction.WriteBlobs))
+			{
+				return Forbid();
+			}
+
+			using (Stream Stream = File.OpenReadStream())
+			{
+				await BlobCollection.AddAsync(NamespaceId, Hash, Stream);
+			}
+
+			return new PutBlobResponse { Hash = Hash };
+		}
+
+		/// <summary>
+		/// Does payload exist, HTTP 200 returned if successful, otherwise a HTTP 404 if missing..
+		/// </summary>
+		/// <param name="NamespaceId">Namespace for the operation</param>
+		/// <param name="Hash">Hash of the blob to check</param>
+		/// <returns></returns>
+		[HttpHead]
+		[Route("/api/v1/blobs/{NamespaceId}/{Hash}")]
+		[ProducesResponseType(200)]
+		[ProducesResponseType(404)]
+		public async Task<ActionResult> DoesBlobExistAsync(NamespaceId NamespaceId, IoHash Hash)
+		{
+			if (!await NamespaceCollection.AuthorizeAsync(NamespaceId, User, AclAction.ReadBlobs))
+			{
+				return Forbid();
+			}
+			if (!await BlobCollection.ExistsAsync(NamespaceId, Hash))
+			{
+				return NotFound();
+			}
+
+			return Ok();
+		}
+
+		/// <summary>
+		/// Returns the payload with the given hash.
+		/// </summary>
+		/// <param name="NamespaceId">Namespace for the operation</param>
+		/// <param name="Hash">Hash of the blob. Must match the submitted data.</param>
+		/// <returns></returns>
+		[HttpGet]
+		[Route("/api/v1/blobs/{NamespaceId}/{Hash}")]
+		public async Task<ActionResult> GetBlobAsync(NamespaceId NamespaceId, IoHash Hash)
+		{
+			if (!await NamespaceCollection.AuthorizeAsync(NamespaceId, User, AclAction.ReadBlobs))
+			{
+				return Forbid();
+			}
+
+			Stream? Stream = await BlobCollection.GetAsync(NamespaceId, Hash);
+			if (Stream == null)
+			{
+				return NotFound();
+			}
+
+			return new FileStreamResult(Stream, new MediaTypeHeaderValue("application/octet-stream"));
+		}
+
+		/// <summary>
+		/// Send a list of query parameters called id for a batch lookup similar to the HEAD request.
+		/// </summary>
+		/// <param name="NamespaceId">Namespace for the operation</param>
+		/// <param name="Hashes">List of ids to check</param>
+		/// <returns></returns>
+		[HttpPost]
+		[Route("/api/v1/blobs/{NamespaceId}/exists")]
+		public async Task<ActionResult<ExistsResponse>> DoBlobsExistAsync(NamespaceId NamespaceId, [FromQuery(Name = "id")] IoHash[] Hashes)
+		{
+			if (!await NamespaceCollection.AuthorizeAsync(NamespaceId, User, AclAction.ReadBlobs))
+			{
+				return Forbid();
+			}
+
+			ExistsResponse Exists = new ExistsResponse();
+			Exists.Id = (await BlobCollection.ExistsAsync(NamespaceId, Hashes)).ToArray();
+			return Exists;
+		}
+	}
+}
diff --git a/Engine/Source/Programs/Horde/HordeServer/Storage/Controllers/ObjectsController.cs b/Engine/Source/Programs/Horde/HordeServer/Storage/Controllers/ObjectsController.cs
new file mode 100644
--- /dev/null
+++ b/Engine/Source/Programs/Horde/HordeServer/Storage/Controllers/ObjectsController.cs
@@ -0,0 +1,210 @@
+// Copyright Epic Games, Inc. All Rights Reserved.
+
+using HordeServer.Api;
+using HordeServer.Services;
+using Microsoft.AspNetCore.Authorization;
+using Microsoft.AspNetCore.Http;
+using Microsoft.AspNetCore.Mvc;
+using Microsoft.Net.Http.Headers;
+using System;
+using System.IO;
+using System.Collections.Generic;
+using System.Linq;
+using System.Threading.Tasks;
+using HordeServer.Models;
+using MongoDB.Bson;
+using HordeServer.Utilities;
+using EpicGames.Core;
+using HordeServer.Storage.Collections;
+using EpicGames.Serialization;
+
+namespace HordeServer.Storage.Controllers
+{
+	using NamespaceId = StringId<INamespace>;
+
+	/// <summary>
+	/// Format to return object data
+	/// </summary>
+	public enum ObjectFormat
+	{
+		/// <summary>
+		/// Json output
+		/// </summary>
+		Json,
+		
+		/// <summary>
+		/// Compact binary output
+		/// </summary>
+		CompactBinary,
+	}
+
+	/// <summary>
+	/// Request for negotiating list of objects to fetch
+	/// </summary>
+	public class FindObjectDeltaRequest
+	{
+		/// <summary>
+		/// List of objects that the client has
+		/// </summary>
+		public List<IoHash> Have { get; set; } = new List<IoHash>();
+
+		/// <summary>
+		/// List of objects that the client wants
+		/// </summary>
+		public List<IoHash> Want { get; set; } = new List<IoHash>();
+	}
+
+	/// <summary>
+	/// Response listing the objects that the client needs
+	/// </summary>
+	public class FindObjectDeltaResponse
+	{
+		/// <summary>
+		/// List of of objects that the client needs
+		/// </summary>
+		public List<IoHash> Need { get; set; } = new List<IoHash>();
+	}
+
+	/// <summary>
+	/// Response for putting an object
+	/// </summary>
+	public class PutObjectResponse
+	{
+		/// <summary>
+		/// List of missing hashes
+		/// </summary>
+		public IoHash Id { get; set; }
+	}
+
+	/// <summary>
+	/// Controller for the /api/v1/software endpoint
+	/// </summary>
+	[ApiController]
+	[Authorize]
+	[Route("[controller]")]
+	public class ObjectsController : ControllerBase
+	{
+		/// <summary>
+		/// Singleton instance of the namespace collection
+		/// </summary>
+		private readonly INamespaceCollection NamespaceCollection;
+
+		/// <summary>
+		/// Singleton instance of the object collection
+		/// </summary>
+		private readonly IObjectCollection ObjectCollection;
+
+		/// <summary>
+		/// Constructor
+		/// </summary>
+		/// <param name="NamespaceCollection">The namespace collection</param>
+		/// <param name="ObjectCollection">The object collection</param>
+		public ObjectsController(INamespaceCollection NamespaceCollection, IObjectCollection ObjectCollection)
+		{
+			this.NamespaceCollection = NamespaceCollection;
+			this.ObjectCollection = ObjectCollection;
+		}
+
+		/// <summary>
+		/// Submit a structured object to the storage api. The object hash will be verified against the received payload. 
+		/// </summary>
+		/// <param name="NamespaceId">Namespace for the operation</param>
+		/// <param name="Hash">Hash of the blob. Must match the submitted data.</param>
+		/// <param name="Data">The data to add</param>
+		/// <returns></returns>
+		[HttpPut]
+		[Route("/api/v1/objects/{NamespaceId}/{Hash}")]
+		[ProducesResponseType(typeof(PutObjectResponse), 200)]
+		public async Task<ActionResult> PutObjectAsync(NamespaceId NamespaceId, IoHash Hash, [FromBody] byte[] Data)
+		{
+			if (!await NamespaceCollection.AuthorizeAsync(NamespaceId, User, AclAction.WriteBlobs))
+			{
+				return Forbid();
+			}
+
+			await ObjectCollection.AddAsync(NamespaceId, Hash, new CbObject(Data));
+			return Ok(new PutObjectResponse { Id = Hash });
+		}
+
+		/// <summary>
+		/// Finds deltas between two object sets
+		/// </summary>
+		/// <param name="NamespaceId">Namespace for the operation</param>
+		/// <param name="Hash">The root object hash</param>
+		/// <param name="Request">Request object</param>
+		/// <returns></returns>
+		[HttpPost]
+		[Route("/api/v1/objects/{NamespaceId}/{Hash}/delta")]
+		public Task<ActionResult<FindObjectDeltaResponse>> FindObjectDeltaAsync(NamespaceId NamespaceId, IoHash Hash, [FromBody] FindObjectDeltaRequest Request)
+		{
+			throw new NotImplementedException();
+		}
+
+		/// <summary>
+		/// Does payload exist, HTTP 200 returned if successful, otherwise a HTTP 404 if missing..
+		/// </summary>
+		/// <param name="NamespaceId">Namespace for the operation</param>
+		/// <param name="Hash">Hash of the blob to check</param>
+		/// <returns></returns>
+		[HttpHead]
+		[Route("/api/v1/objects/{NamespaceId}/{Hash}")]
+		[ProducesResponseType(200)]
+		[ProducesResponseType(404)]
+		public async Task<ActionResult> DoesBlobExistAsync(NamespaceId NamespaceId, [FromRoute] IoHash Hash)
+		{
+			if (!await NamespaceCollection.AuthorizeAsync(NamespaceId, User, AclAction.ReadBlobs))
+			{
+				return Forbid();
+			}
+			if (!await ObjectCollection.ExistsAsync(NamespaceId, Hash))
+			{
+				return NotFound();
+			}
+			return Ok();
+		}
+
+		/// <summary>
+		/// Returns the object with the given hash.
+		/// </summary>
+		/// <param name="NamespaceId">Namespace for the operation</param>
+		/// <param name="Hash">Hash of the blob. Must match the submitted data.</param>
+		/// <returns></returns>
+		[HttpGet]
+		[Route("/api/v1/objects/{NamespaceId}/{Hash}")]
+		public async Task<ActionResult> GetObjectAsync(NamespaceId NamespaceId, IoHash Hash)
+		{
+			if (!await NamespaceCollection.AuthorizeAsync(NamespaceId, User, AclAction.ReadBlobs))
+			{
+				return Forbid();
+			}
+
+			CbObject? Object = await ObjectCollection.GetAsync(NamespaceId, Hash);
+			if(Object == null)
+			{
+				return NotFound();
+			}
+
+			return Ok(Object.GetView().ToArray());
+		}
+
+		/// <summary>
+		/// Send a list of query parameters called id for a batch lookup similar to the HEAD request.
+		/// </summary>
+		/// <param name="NamespaceId">Namespace for the operation</param>
+		/// <param name="Hashes">List of ids to check</param>
+		/// <returns></returns>
+		[HttpPost]
+		[Route("/api/v1/objects/{NamespaceId}/exists")]
+		public async Task<ActionResult<ExistsResponse>> DoBlobsExistAsync(NamespaceId NamespaceId, [FromQuery(Name = "id")] IoHash[] Hashes)
+		{
+			if (!await NamespaceCollection.AuthorizeAsync(NamespaceId, User, AclAction.ReadBlobs))
+			{
+				return Forbid();
+			}
+
+			ExistsResponse Response = new ExistsResponse();
+			Response.Id = (await ObjectCollection.ExistsAsync(NamespaceId, Hashes)).ToArray();
+			return Response;
+		}
+	}
+}
diff --git a/Engine/Source/Programs/Horde/HordeServer/Storage/Controllers/RefsController.cs b/Engine/Source/Programs/Horde/HordeServer/Storage/Controllers/RefsController.cs
new file mode 100644
--- /dev/null
+++ b/Engine/Source/Programs/Horde/HordeServer/Storage/Controllers/RefsController.cs
@@ -0,0 +1,168 @@
+// Copyright Epic Games, Inc. All Rights Reserved.
+
+using HordeServer.Api;
+using HordeServer.Services;
+using Microsoft.AspNetCore.Authorization;
+using Microsoft.AspNetCore.Http;
+using Microsoft.AspNetCore.Mvc;
+using Microsoft.Net.Http.Headers;
+using System;
+using System.IO;
+using System.Collections.Generic;
+using System.Linq;
+using System.Threading.Tasks;
+using HordeServer.Models;
+using MongoDB.Bson;
+using HordeServer.Utilities;
+using EpicGames.Core;
+using HordeServer.Storage.Collections;
+using EpicGames.Serialization;
+
+namespace HordeServer.Storage.Controllers
+{
+	using NamespaceId = StringId<INamespace>;
+	using BucketId = StringId<IBucket>;
+
+	/// <summary>
+	/// Response from trying to set a reference
+	/// </summary>
+	public class SetRefResponse
+	{
+		/// <summary>
+		/// List of missing hashes
+		/// </summary>
+		public List<IoHash> Need { get; set; } = new List<IoHash>();
+	}
+
+	/// <summary>
+	/// Controller for the /api/v1/software endpoint
+	/// </summary>
+	[ApiController]
+	[Authorize]
+	[Route("[controller]")]
+	public class RefsController : ControllerBase
+	{
+		/// <summary>
+		/// Singleton instance of the namespace collection
+		/// </summary>
+		private readonly INamespaceCollection NamespaceCollection;
+
+		/// <summary>
+		/// Singleton instance of the bucket collection
+		/// </summary>
+		private readonly IBucketCollection BucketCollection;
+
+		/// <summary>
+		/// Singleton instance of the ref collection
+		/// </summary>
+		private readonly IRefCollection RefCollection;
+
+		/// <summary>
+		/// Constructor
+		/// </summary>
+		/// <param name="NamespaceCollection">The namespace collection</param>
+		/// <param name="BucketCollection">The bucket collection</param>
+		/// <param name="RefCollection">The ref collection</param>
+		public RefsController(INamespaceCollection NamespaceCollection, IBucketCollection BucketCollection, IRefCollection RefCollection)
+		{
+			this.NamespaceCollection = NamespaceCollection;
+			this.BucketCollection = BucketCollection;
+			this.RefCollection = RefCollection;
+		}
+
+		/// <summary>
+		/// Submit a structured object to the storage api. The object hash will be verified against the received payload. 
+		/// </summary>
+		/// <param name="NamespaceId">Namespace for the operation</param>
+		/// <param name="BucketId">Name of the bucket.</param>
+		/// <param name="RefName">Name of the reference.</param>
+		/// <param name="Data">The emcpdedata to add</param>
+		/// <returns></returns>
+		[HttpPut]
+		[Route("/api/v1/refs/{NamespaceId}/{BucketId}/{RefName}")]
+		public async Task<ActionResult<SetRefResponse>> SetRefAsync(NamespaceId NamespaceId, BucketId BucketId, string RefName, [FromBody] byte[] Data)
+		{
+			if (!await NamespaceCollection.AuthorizeAsync(NamespaceId, User, AclAction.WriteRefs))
+			{
+				return Forbid();
+			}
+
+			CbObject Object = new CbObject(Data);
+
+			SetRefResponse Response = new SetRefResponse();
+			Response.Need = await RefCollection.SetAsync(NamespaceId, BucketId, RefName, Object);
+			return Response;
+		}
+
+		/// <summary>
+		/// Submit a structured object to the storage api. The object hash will be verified against the received payload. 
+		/// </summary>
+		/// <param name="NamespaceId">Namespace for the operation</param>
+		/// <param name="BucketId">Name of the bucket.</param>
+		/// <param name="RefName">Name of the reference.</param>
+		/// <returns></returns>
+		[HttpGet]
+		[Route("/api/v1/refs/{NsName}/{BktName}/{RefName}")]
+		public async Task<ActionResult> GetRefAsync(NamespaceId NamespaceId, BucketId BucketId, string RefName)
+		{
+			if (!await NamespaceCollection.AuthorizeAsync(NamespaceId, User, AclAction.ReadRefs))
+			{
+				return Forbid();
+			}
+
+			IRef? Ref = await RefCollection.GetAsync(NamespaceId, BucketId, RefName);
+			if (Ref == null)
+			{
+				return NotFound();
+			}
+
+			return Content(Ref.Value.ToJson(), new MediaTypeHeaderValue("application/json"));
+		}
+
+		/// <summary>
+		/// Does ref exist, HTTP 200 returned if successful, otherwise a HTTP 404 if missing.
+		/// </summary>
+		/// <param name="NamespaceId">Namespace for the operation</param>
+		/// <param name="BucketId">Name of the bucket.</param>
+		/// <param name="RefName">Name of the reference.</param>
+		/// <returns></returns>
+		[HttpHead]
+		[Route("/api/v1/refs/{NsName}/{BktName}/{RefName}")]
+		[ProducesResponseType(200)]
+		[ProducesResponseType(404)]
+		public async Task<ActionResult> DoesRefExistAsync(NamespaceId NamespaceId, BucketId BucketId, string RefName)
+		{
+			if (!await NamespaceCollection.AuthorizeAsync(NamespaceId, User, AclAction.ReadRefs))
+			{
+				return Forbid();
+			}
+			if (await RefCollection.GetAsync(NamespaceId, BucketId, RefName) == null)
+			{
+				return NotFound();
+			}
+			return Ok();
+		}
+
+		/// <summary>
+		/// Delete a reference from the store.
+		/// </summary>
+		/// <param name="NamespaceId">Namespace for the operation</param>
+		/// <param name="BucketId">Name of the bucket.</param>
+		/// <param name="RefName">Name of the reference.</param>
+		/// <returns></returns>
+		[HttpDelete]
+		[Route("/api/v1/refs/{NamespaceId}/{BucketId}/{RefName}")]
+		public async Task<ActionResult> DeleteRefAsync(NamespaceId NamespaceId, BucketId BucketId, string RefName)
+		{
+			if (!await NamespaceCollection.AuthorizeAsync(NamespaceId, User, AclAction.DeleteRefs))
+			{
+				return Forbid();
+			}
+			if (!await RefCollection.DeleteAsync(NamespaceId, BucketId, RefName))
+			{
+				return NotFound();
+			}
+			return Ok();
+		}
+	}
+}
diff --git a/Engine/Source/Programs/Horde/HordeServer/Storage/IBlobCollection.cs b/Engine/Source/Programs/Horde/HordeServer/Storage/IBlobCollection.cs
new file mode 100644
--- /dev/null
+++ b/Engine/Source/Programs/Horde/HordeServer/Storage/IBlobCollection.cs
@@ -0,0 +1,120 @@
+// Copyright Epic Games, Inc. All Rights Reserved.
+
+using EpicGames.Core;
+using HordeCommon;
+using HordeServer.Models;
+using HordeServer.Services;
+using HordeServer.Utilities;
+using MongoDB.Bson;
+using MongoDB.Bson.Serialization.Attributes;
+using MongoDB.Driver;
+using System;
+using System.Collections.Generic;
+using System.Globalization;
+using System.IO;
+using System.IO.Compression;
+using System.Linq;
+using System.Text.Json;
+using System.Threading.Tasks;
+
+namespace HordeServer.Storage
+{
+	using NamespaceId = StringId<INamespace>;
+
+	/// <summary>
+	/// Interface for a collection of blobs
+	/// </summary>
+	public interface IBlobCollection
+	{
+		/// <summary>
+		/// Adds an item to storage
+		/// </summary>
+		/// <param name="NamespaceId">Namespace to operate on</param>
+		/// <param name="Hash">Hash of the blob</param>
+		/// <param name="Blob">The blob stream data</param>
+		Task AddAsync(NamespaceId NamespaceId, IoHash Hash, Stream Blob);
+
+		/// <summary>
+		/// Opens a blob read stream
+		/// </summary>
+		/// <param name="NamespaceId">Namespace to operate on</param>
+		/// <param name="Hash">Hash of the blob</param>
+		/// <returns>Stream for the blob, or null if it does not exist</returns>
+		Task<Stream?> GetAsync(NamespaceId NamespaceId, IoHash Hash);
+
+		/// <summary>
+		/// Determines which of a set of blobs exist
+		/// </summary>
+		/// <param name="NamespaceId">Namespace to operate on</param>
+		/// <param name="Hashes"></param>
+		/// <returns></returns>
+		Task<List<IoHash>> ExistsAsync(NamespaceId NamespaceId, IEnumerable<IoHash> Hashes);
+	}
+
+	/// <summary>
+	/// Extension methods for <see cref="IBlobCollection"/>
+	/// </summary>
+	public static class BlobCollectionExtensions
+	{
+		/// <summary>
+		/// Adds a blob to the collection
+		/// </summary>
+		/// <param name="Collection"></param>
+		/// <param name="NamespaceId"></param>
+		/// <param name="Blob"></param>
+		/// <returns></returns>
+		public static async Task<IoHash> AddAsync(this IBlobCollection Collection, NamespaceId NamespaceId, ReadOnlyMemory<byte> Blob)
+		{
+			IoHash Hash = IoHash.Compute(Blob.Span);
+			await AddAsync(Collection, NamespaceId, Hash, Blob);
+			return Hash;
+		}
+
+		/// <summary>
+		/// Adds a blob to the collection
+		/// </summary>
+		/// <param name="Collection"></param>
+		/// <param name="NamespaceId"></param>
+		/// <param name="Hash"></param>
+		/// <param name="Blob"></param>
+		/// <returns></returns>
+		public static Task AddAsync(this IBlobCollection Collection, NamespaceId NamespaceId, IoHash Hash, ReadOnlyMemory<byte> Blob)
+		{
+			using ReadOnlyMemoryStream Stream = new ReadOnlyMemoryStream(Blob);
+			return Collection.AddAsync(NamespaceId, Hash, Stream);
+		}
+
+		/// <summary>
+		/// Gets a blob as a byte array
+		/// </summary>
+		/// <param name="Collection"></param>
+		/// <param name="NamespaceId"></param>
+		/// <param name="Hash"></param>
+		/// <returns></returns>
+		public static async Task<byte[]?> GetByteArrayAsync(this IBlobCollection Collection, NamespaceId NamespaceId, IoHash Hash)
+		{
+			using Stream? Stream = await Collection.GetAsync(NamespaceId, Hash);
+			if(Stream == null)
+			{
+				return null;
+			}
+
+			byte[] Data = new byte[Stream.Length];
+			Stream.Read(Data);
+			return Data;
+		}
+
+		/// <summary>
+		/// Test whether a single blob exists
+		/// </summary>
+		/// <param name="Collection"></param>
+		/// <param name="NamespaceId"></param>
+		/// <param name="Hash"></param>
+		/// <returns></returns>
+		public static async Task<bool> ExistsAsync(this IBlobCollection Collection, NamespaceId NamespaceId, IoHash Hash)
+		{
+			List<IoHash> Hashes = await Collection.ExistsAsync(NamespaceId, new[] { Hash });
+			return Hashes.Count > 0;
+		}
+	}
+}
diff --git a/Engine/Source/Programs/Horde/HordeServer/Storage/IBucket.cs b/Engine/Source/Programs/Horde/HordeServer/Storage/IBucket.cs
new file mode 100644
--- /dev/null
+++ b/Engine/Source/Programs/Horde/HordeServer/Storage/IBucket.cs
@@ -0,0 +1,45 @@
+// Copyright Epic Games, Inc. All Rights Reserved.
+
+using EpicGames.Core;
+using HordeCommon;
+using HordeServer.Models;
+using HordeServer.Services;
+using HordeServer.Utilities;
+using MongoDB.Bson;
+using MongoDB.Bson.Serialization.Attributes;
+using MongoDB.Driver;
+using System;
+using System.Collections.Generic;
+using System.Globalization;
+using System.IO;
+using System.IO.Compression;
+using System.Linq;
+using System.Text.Json;
+using System.Threading.Tasks;
+
+namespace HordeServer.Storage
+{
+	using NamespaceId = StringId<INamespace>;
+	using BucketId = StringId<IBucket>;
+
+	/// <summary>
+	/// Interface for a ref bucket
+	/// </summary>
+	public interface IBucket
+	{
+		/// <summary>
+		/// Namespace identifier
+		/// </summary>
+		public NamespaceId NamespaceId { get; }
+
+		/// <summary>
+		/// Bucket identifier
+		/// </summary>
+		public BucketId BucketId { get; }
+
+		/// <summary>
+		/// Whether the bucket has been deleted
+		/// </summary>
+		public bool Deleted { get; }
+	}
+}
diff --git a/Engine/Source/Programs/Horde/HordeServer/Storage/IBucketCollection.cs b/Engine/Source/Programs/Horde/HordeServer/Storage/IBucketCollection.cs
new file mode 100644
--- /dev/null
+++ b/Engine/Source/Programs/Horde/HordeServer/Storage/IBucketCollection.cs
@@ -0,0 +1,61 @@
+// Copyright Epic Games, Inc. All Rights Reserved.
+
+using EpicGames.Core;
+using HordeCommon;
+using HordeServer.Models;
+using HordeServer.Services;
+using HordeServer.Utilities;
+using MongoDB.Bson;
+using MongoDB.Bson.Serialization.Attributes;
+using MongoDB.Driver;
+using System;
+using System.Collections.Generic;
+using System.Globalization;
+using System.IO;
+using System.IO.Compression;
+using System.Linq;
+using System.Text.Json;
+using System.Threading.Tasks;
+
+namespace HordeServer.Storage
+{
+	using NamespaceId = StringId<INamespace>;
+	using BucketId = StringId<IBucket>;
+
+	/// <summary>
+	/// Interface for a collection of template documents
+	/// </summary>
+	public interface IBucketCollection
+	{
+		/// <summary>
+		/// Adds or updates a bucket configuration
+		/// </summary>
+		/// <param name="NamespaceId">Namespace identifier</param>
+		/// <param name="BucketId">Bucket identifier</param>
+		/// <param name="Bucket">New settings for the bucket</param>
+		Task AddOrUpdateAsync(NamespaceId NamespaceId, BucketId BucketId, BucketConfig Bucket);
+
+		/// <summary>
+		/// Attempts to get a bucket by identifier
+		/// </summary>
+		/// <param name="NamespaceId">Namespace identifier</param>
+		/// <param name="BucketId">Bucket identifier</param>
+		/// <returns>Bucket definition</returns>
+		Task<IBucket?> GetAsync(NamespaceId NamespaceId, BucketId BucketId);
+
+		/// <summary>
+		/// Remove a bucket definition
+		/// </summary>
+		/// <param name="NamespaceId">Namespace identifier</param>
+		/// <param name="BucketId">Bucket identifier</param>
+		/// <returns>Bucket definition</returns>
+		Task RemoveAsync(NamespaceId NamespaceId, BucketId BucketId);
+
+		/// <summary>
+		/// Find all buckets within a namespace
+		/// </summary>
+		/// <param name="NamespaceId">Namespace identifier</param>
+		/// <returns>List of buckets within the namespace</returns>
+		Task<List<IBucket>> FindAsync(NamespaceId NamespaceId);
+	}
+}
diff --git a/Engine/Source/Programs/Horde/HordeServer/Storage/INamespace.cs b/Engine/Source/Programs/Horde/HordeServer/Storage/INamespace.cs
new file mode 100644
--- /dev/null
+++ b/Engine/Source/Programs/Horde/HordeServer/Storage/INamespace.cs
@@ -0,0 +1,39 @@
+// Copyright Epic Games, Inc. All Rights Reserved.
+
+using EpicGames.Core;
+using HordeCommon;
+using HordeServer.Models;
+using HordeServer.Services;
+using HordeServer.Utilities;
+using MongoDB.Bson;
+using MongoDB.Bson.Serialization.Attributes;
+using MongoDB.Driver;
+using System;
+using System.Collections.Generic;
+using System.Globalization;
+using System.IO;
+using System.IO.Compression;
+using System.Linq;
+using System.Text.Json;
+using System.Threading.Tasks;
+
+namespace HordeServer.Storage
+{
+	using NamespaceId = StringId<INamespace>;
+
+	/// <summary>
+	/// Interface for a namespace
+	/// </summary>
+	public interface INamespace
+	{
+		/// <summary>
+		/// Identifier for the namespace
+		/// </summary>
+		NamespaceId Id { get; }
+
+		/// <summary>
+		/// Access control list for the namespace
+		/// </summary>
+		Acl? Acl { get; }
+	}
+}
diff --git a/Engine/Source/Programs/Horde/HordeServer/Storage/INamespaceCollection.cs b/Engine/Source/Programs/Horde/HordeServer/Storage/INamespaceCollection.cs
new file mode 100644
--- /dev/null
+++ b/Engine/Source/Programs/Horde/HordeServer/Storage/INamespaceCollection.cs
@@ -0,0 +1,65 @@
+// Copyright Epic Games, Inc. All Rights Reserved.
+
+using EpicGames.Core;
+using HordeCommon;
+using HordeServer.Models;
+using HordeServer.Services;
+using HordeServer.Utilities;
+using MongoDB.Bson;
+using MongoDB.Bson.Serialization.Attributes;
+using MongoDB.Driver;
+using System;
+using System.Collections.Generic;
+using System.Globalization;
+using System.IO;
+using System.IO.Compression;
+using System.Linq;
+using System.Security.Claims;
+using System.Text.Json;
+using System.Threading.Tasks;
+
+namespace HordeServer.Storage
+{
+	using NamespaceId = StringId<INamespace>;
+
+	/// <summary>
+	/// Interface for a collection of namespace documents
+	/// </summary>
+	public interface INamespaceCollection
+	{
+		/// <summary>
+		/// Adds or updates a bucket configuration
+		/// </summary>
+		/// <param name="NamespaceId">Namespace identifier</param>
+		/// <param name="Namespace">New settings for the bucket</param>
+		Task AddOrUpdateAsync(NamespaceId NamespaceId, NamespaceConfig Namespace);
+
+		/// <summary>
+		/// Attempts to get a namespace by identifier
+		/// </summary>
+		/// <param name="NamespaceId">Namespace identifier</param>
+		/// <returns>Namespace definition</returns>
+		Task<INamespace?> GetAsync(NamespaceId NamespaceId);
+
+		/// <summary>
+		/// Remove a namespace
+		/// </summary>
+		/// <param name="NamespaceId">Namespace identifier</param>
+		Task RemoveAsync(NamespaceId NamespaceId);
+
+		/// <summary>
+		/// Find all namespaces
+		/// </summary>
+		/// <returns>List of namespaces</returns>
+		Task<List<INamespace>> FindAsync();
+
+		/// <summary>
+		/// Authorize a user to perform the given action. May use cached data.
+		/// </summary>
+		/// <param name="NamespaceId">The namespace id</param>
+		/// <param name="User">User to authorize</param>
+		/// <param name="Action">Required permissions</param>
+		/// <returns>True if the user is authorized to perform the given action</returns>
+		Task<bool> AuthorizeAsync(NamespaceId NamespaceId, ClaimsPrincipal User, AclAction Action);
+	}
+}
diff --git a/Engine/Source/Programs/Horde/HordeServer/Storage/IObjectCollection.cs b/Engine/Source/Programs/Horde/HordeServer/Storage/IObjectCollection.cs
new file mode 100644
--- /dev/null
+++ b/Engine/Source/Programs/Horde/HordeServer/Storage/IObjectCollection.cs
@@ -0,0 +1,73 @@
+// Copyright Epic Games, Inc. All Rights Reserved.
+
+using EpicGames.Core;
+using EpicGames.Serialization;
+using HordeCommon;
+using HordeServer.Models;
+using HordeServer.Services;
+using HordeServer.Utilities;
+using MongoDB.Bson;
+using MongoDB.Bson.Serialization.Attributes;
+using MongoDB.Driver;
+using System;
+using System.Collections.Generic;
+using System.Globalization;
+using System.IO;
+using System.IO.Compression;
+using System.Linq;
+using System.Text.Json;
+using System.Threading.Tasks;
+
+namespace HordeServer.Storage
+{
+	using NamespaceId = StringId<INamespace>;
+
+	/// <summary>
+	/// Interface for a collection of objects
+	/// </summary>
+	public interface IObjectCollection
+	{
+		/// <summary>
+		/// Adds an item to storage
+		/// </summary>
+		/// <param name="NsId">The namespace id</param>
+		/// <param name="Hash">Hash of the object</param>
+		/// <param name="Data">The object data</param>
+		Task AddAsync(NamespaceId NsId, IoHash Hash, CbObject Data);
+
+		/// <summary>
+		/// Reads an object from storage
+		/// </summary>
+		/// <param name="NsId">The namespace id</param>
+		/// <param name="Hash">Hash of the object</param>
+		/// <returns>Stream for the blob, or null if it does not exist</returns>
+		Task<CbObject?> GetAsync(NamespaceId NsId, IoHash Hash);
+
+		/// <summary>
+		/// Determines which of a set of blobs exist
+		/// </summary>
+		/// <param name="NsId">The namespace id</param>
+		/// <param name="Hashes"></param>
+		/// <returns></returns>
+		Task<List<IoHash>> ExistsAsync(NamespaceId NsId, IEnumerable<IoHash> Hashes);
+	}
+
+	/// <summary>
+	/// Extension methods for <see cref="IObjectCollection"/>
+	/// </summary>
+	public static class ObjectCollectionExtensions
+	{
+		/// <summary>
+		/// Test whether a single blob exists
+		/// </summary>
+		/// <param name="Collection"></param>
+		/// <param name="NamespaceId"></param>
+		/// <param name="Hash"></param>
+		/// <returns></returns>
+		public static async Task<bool> ExistsAsync(this IObjectCollection Collection, NamespaceId NamespaceId, IoHash Hash)
+		{
+			List<IoHash> Hashes = await Collection.ExistsAsync(NamespaceId, new[] { Hash });
+			return Hashes.Count > 0;
+		}
+	}
+}
diff --git a/Engine/Source/Programs/Horde/HordeServer/Storage/IRef.cs b/Engine/Source/Programs/Horde/HordeServer/Storage/IRef.cs
new file mode 100644
--- /dev/null
+++ b/Engine/Source/Programs/Horde/HordeServer/Storage/IRef.cs
@@ -0,0 +1,61 @@
+// Copyright Epic Games, Inc. All Rights Reserved.
+
+using EpicGames.Core;
+using EpicGames.Serialization;
+using HordeCommon;
+using HordeServer.Models;
+using HordeServer.Services;
+using HordeServer.Utilities;
+using MongoDB.Bson;
+using MongoDB.Bson.Serialization.Attributes;
+using MongoDB.Driver;
+using System;
+using System.Collections.Generic;
+using System.Globalization;
+using System.IO;
+using System.IO.Compression;
+using System.Linq;
+using System.Text.Json;
+using System.Threading.Tasks;
+
+namespace HordeServer.Storage
+{
+	using NamespaceId = StringId<INamespace>;
+	using BucketId = StringId<IBucket>;
+
+	/// <summary>
+	/// Interface for a collection of template documents
+	/// </summary>
+	public interface IRef
+	{
+		/// <summary>
+		/// Namespace identifier
+		/// </summary>
+		NamespaceId NamespaceId { get; }
+
+		/// <summary>
+		/// Bucket identifier
+		/// </summary>
+		BucketId BucketId { get; }
+
+		/// <summary>
+		/// Ref name
+		/// </summary>
+		string Name { get; }
+
+		/// <summary>
+		/// Last access time
+		/// </summary>
+		DateTime LastAccessTime { get; }
+
+		/// <summary>
+		/// Value for the ref
+		/// </summary>
+		CbObject Value { get; }
+
+		/// <summary>
+		/// Whether the ref has been finalized
+		/// </summary>
+		bool Finalized { get; }
+	}
+}
diff --git a/Engine/Source/Programs/Horde/HordeServer/Storage/IRefCollection.cs b/Engine/Source/Programs/Horde/HordeServer/Storage/IRefCollection.cs
new file mode 100644
--- /dev/null
+++ b/Engine/Source/Programs/Horde/HordeServer/Storage/IRefCollection.cs
@@ -0,0 +1,75 @@
+// Copyright Epic Games, Inc. All Rights Reserved.
+
+using EpicGames.Core;
+using EpicGames.Serialization;
+using HordeCommon;
+using HordeServer.Models;
+using HordeServer.Services;
+using HordeServer.Utilities;
+using MongoDB.Bson;
+using MongoDB.Bson.Serialization.Attributes;
+using MongoDB.Driver;
+using System;
+using System.Collections.Generic;
+using System.Globalization;
+using System.IO;
+using System.IO.Compression;
+using System.Linq;
+using System.Text.Json;
+using System.Threading.Tasks;
+
+namespace HordeServer.Storage
+{
+	using NamespaceId = StringId<INamespace>;
+	using BucketId = StringId<IBucket>;
+
+	/// <summary>
+	/// Interface for a collection of ref documents
+	/// </summary>
+	public interface IRefCollection
+	{
+		/// <summary>
+		/// Sets the given reference
+		/// </summary>
+		/// <param name="NamespaceId">Namespace identifier</param>
+		/// <param name="BucketId">Bucket identifier</param>
+		/// <param name="Name">Name of the reference</param>
+		/// <param name="Value">New value for the reference, as a compact binary object</param>
+		/// <returns>List of missing references</returns>
+		Task<List<IoHash>> SetAsync(NamespaceId NamespaceId, BucketId BucketId, string Name, CbObject Value);
+
+		/// <summary>
+		/// Attempts to finalize a reference, turning its references into hard references
+		/// </summary>
+		/// <param name="Ref">The ref to finalize</param>
+		/// <returns></returns>
+		Task<List<IoHash>> FinalizeAsync(IRef Ref);
+
+		/// <summary>
+		/// Gets the given reference
+		/// </summary>
+		/// <param name="NamespaceId">Namespace identifier</param>
+		/// <param name="BucketId">Bucket identifier</param>
+		/// <param name="Name">Name of the reference</param>
+		/// <returns>The reference data if the ref exists</returns>
+		Task<IRef?> GetAsync(NamespaceId NamespaceId, BucketId BucketId, string Name);
+
+		/// <summary>
+		/// Touches the given reference, updating it's last access time
+		/// </summary>
+		/// <param name="NamespaceId">Namespace identifier</param>
+		/// <param name="BucketId">Bucket identifier</param>
+		/// <param name="Name">Name of the reference</param>
+		/// <returns>True if the ref exists</returns>
+		Task<bool> TouchAsync(NamespaceId NamespaceId, BucketId BucketId, string Name);
+
+		/// <summary>
+		/// Removes the given reference
+		/// </summary>
+		/// <param name="NamespaceId">Namespace identifier</param>
+		/// <param name="BucketId">Bucket identifier</param>
+		/// <param name="RefName">Name of the reference</param>
+		/// <returns>True if the ref was deleted, false if it did not exist</returns>
+		Task<bool> DeleteAsync(NamespaceId NamespaceId, BucketId BucketId, string RefName);
+	}
+}
diff --git a/Engine/Source/Programs/Horde/HordeServer/Utilities/LazyCache.cs b/Engine/Source/Programs/Horde/HordeServer/Utilities/LazyCache.cs
new file mode 100644
--- /dev/null
+++ b/Engine/Source/Programs/Horde/HordeServer/Utilities/LazyCache.cs
@@ -0,0 +1,125 @@
+using Microsoft.Extensions.Caching.Memory;
+using Microsoft.Extensions.Options;
+using System;
+using System.Collections.Concurrent;
+using System.Collections.Generic;
+using System.Diagnostics;
+using System.Linq;
+using System.Threading;
+using System.Threading.Tasks;
+
+namespace HordeServer.Utilities
+{
+	/// <summary>
+	/// Options for <see cref="LazyCache{TKey, TValue}"/>
+	/// </summary>
+	public class LazyCacheOptions
+	{
+		/// <summary>
+		/// Time after which a value will asynchronously be updated
+		/// </summary>
+		public TimeSpan? RefreshTime { get; set; } = TimeSpan.FromMinutes(1.0);
+
+		/// <summary>
+		/// Maximum age of any returned value. This will prevent a cached value being returned.
+		/// </summary>
+		public TimeSpan? MaxAge { get; set; } = TimeSpan.FromMinutes(2.0);
+	}
+
+	/// <summary>
+	/// Implements a cache which starts an asynchronous update of a value after a period of time.
+	/// </summary>
+	/// <typeparam name="TKey">Key for the cache</typeparam>
+	/// <typeparam name="TValue">Value for the cache</typeparam>
+	public sealed class LazyCache<TKey, TValue> : IDisposable where TKey : notnull
+	{
+		class Item
+		{
+			public Task<TValue>? CurrentTask;
+			public Stopwatch Timer = Stopwatch.StartNew();
+			public Task<TValue>? UpdateTask;
+		}
+
+		ConcurrentDictionary<TKey, Item> Dictionary = new ConcurrentDictionary<TKey, Item>();
+		Func<TKey, Task<TValue>> GetValueAsync;
+		LazyCacheOptions Options;
+
+		/// <summary>
+		/// Constructor
+		/// </summary>
+		/// <param name="GetValueAsync">Function used to get a value</param>
+		/// <param name="Options"></param>
+		public LazyCache(Func<TKey, Task<TValue>> GetValueAsync, LazyCacheOptions Options)
+		{
+			this.GetValueAsync = GetValueAsync;
+			this.Options = Options;
+		}
+
+		/// <inheritdoc/>
+		public void Dispose()
+		{
+			foreach (Item Item in Dictionary.Values)
+			{
+				Item.CurrentTask?.Wait();
+				Item.UpdateTask?.Wait();
+			}
+		}
+
+		/// <summary>
+		/// Gets the value associated with a key
+		/// </summary>
+		/// <param name="Key">The key to query</param>
+		/// <param name="MaxAge">Maximum age for values to return</param>
+		/// <returns></returns>
+		public Task<TValue> GetAsync(TKey Key, TimeSpan? MaxAge = null)
+		{
+			Item Item = Dictionary.GetOrAdd(Key, Key => new Item());
+
+			// Create the task to get the current value
+			Task<TValue> CurrentTask = InterlockedCreateTask(ref Item.CurrentTask, () => GetValueAsync(Key));
+
+			// If an update has completed, swap it out
+			Task<TValue>? UpdateTask = Item.UpdateTask;
+			if (UpdateTask != null && UpdateTask.IsCompleted)
+			{
+				Interlocked.CompareExchange(ref Item.CurrentTask, UpdateTask, CurrentTask);
+				Interlocked.CompareExchange(ref Item.UpdateTask, null, UpdateTask);
+				Item.Timer.Restart();
+			}
+
+			// Check if we need to update the value
+			TimeSpan Age = Item.Timer.Elapsed;
+			if (MaxAge != null && Age > MaxAge.Value)
+			{
+				return InterlockedCreateTask(ref Item.UpdateTask, () => GetValueAsync(Key));
+			}
+			if (Age > Options.RefreshTime)
+			{
+				InterlockedCreateTask(ref Item.UpdateTask, () => GetValueAsync(Key));
+			}
+
+			return CurrentTask;
+		}
+
+		/// <summary>
+		/// Creates a task, guaranteeing that only one task will be assigned to the given slot. Creates a cold task and only starts it once the variable is set.
+		/// </summary>
+		/// <param name="Value"></param>
+		/// <param name="CreateTask"></param>
+		/// <returns></returns>
+		static Task<TValue> InterlockedCreateTask(ref Task<TValue>? Value, Func<Task<TValue>> CreateTask)
+		{
+			Task<TValue>? CurrentTask = Value;
+			while (CurrentTask == null)
+			{
+				Task<Task<TValue>> NewTask = new Task<Task<TValue>>(CreateTask);
+				if (Interlocked.CompareExchange(ref Value, NewTask.Unwrap(), null) == null)
+				{
+					NewTask.Start();
+				}
+				CurrentTask = Value;
+			}
+			return CurrentTask;
+		}
+	}
+}
diff --git a/Engine/Source/Programs/Horde/HordeServerTests/HordeServerTests.csproj b/Engine/Source/Programs/Horde/HordeServerTests/HordeServerTests.csproj
--- a/Engine/Source/Programs/Horde/HordeServerTests/HordeServerTests.csproj
+++ b/Engine/Source/Programs/Horde/HordeServerTests/HordeServerTests.csproj
@@ -26,6 +26,7 @@
   </ItemGroup>
 
   <ItemGroup>
+    <ProjectReference Include="..\..\Shared\EpicGames.Serialization\EpicGames.Serialization.csproj" />
     <ProjectReference Include="..\HordeAgent\HordeAgent.csproj" Label="HordeAgent">
       <Aliases>HordeAgent</Aliases>
     </ProjectReference>
diff --git a/Engine/Source/Programs/Horde/HordeServerTests/StorageTests.cs b/Engine/Source/Programs/Horde/HordeServerTests/StorageTests.cs
new file mode 100644
--- /dev/null
+++ b/Engine/Source/Programs/Horde/HordeServerTests/StorageTests.cs
@@ -0,0 +1,160 @@
+using EpicGames.Core;
+using EpicGames.Serialization;
+using HordeServer.Storage;
+using HordeServer.Utilities;
+using Microsoft.Extensions.DependencyInjection;
+using Microsoft.VisualStudio.TestTools.UnitTesting;
+using System;
+using System.Collections.Generic;
+using System.IO;
+using System.Text;
+using System.Threading.Tasks;
+
+namespace HordeServerTests
+{
+	using NamespaceId = StringId<INamespace>;
+	using BucketId = StringId<IBucket>;
+
+	[TestClass]
+	public class StorageTests : TestSetup
+	{
+		public StorageTests()
+			: base(DatabaseIntegrationTest.GetDatabaseService())
+		{
+		}
+
+		[TestMethod]
+		public async Task BlobCollectionTest()
+		{
+			NamespaceId NamespaceId = new NamespaceId("test-ns");
+
+			IBlobCollection BlobCollection = ServiceProvider.GetRequiredService<IBlobCollection>();
+
+			byte[] TestData = Encoding.UTF8.GetBytes("Hello world");
+			IoHash Hash = IoHash.Compute(TestData);
+
+			Assert.IsNull(await BlobCollection.GetAsync(NamespaceId, Hash));
+			Assert.IsFalse(await BlobCollection.ExistsAsync(NamespaceId, Hash));
+
+			IoHash ReturnedHash = await BlobCollection.AddAsync(NamespaceId, TestData);
+			Assert.AreEqual(Hash, ReturnedHash);
+			Assert.IsTrue(await BlobCollection.ExistsAsync(NamespaceId, Hash));
+
+			byte[]? StoredData = await BlobCollection.GetByteArrayAsync(NamespaceId, Hash);
+			Assert.IsNotNull(StoredData);
+			Assert.IsTrue(TestData.AsSpan().SequenceEqual(StoredData));
+			Assert.IsTrue(await BlobCollection.ExistsAsync(NamespaceId, Hash));
+
+			NamespaceId OtherNamespaceId = new NamespaceId("other-ns");
+			Assert.IsNull(await BlobCollection.GetAsync(OtherNamespaceId, Hash));
+			Assert.IsFalse(await BlobCollection.ExistsAsync(OtherNamespaceId, Hash));
+		}
+
+		[TestMethod]
+		public async Task ObjectCollectionTest()
+		{
+			NamespaceId NamespaceId = new NamespaceId("test-ns");
+
+			IObjectCollection ObjectCollection = ServiceProvider.GetRequiredService<IObjectCollection>();
+
+			CbObject ObjectD = CbObject.Build(Writer =>
+			{
+				Writer.AddString("hello", "field");
+			});
+			IoHash HashD = ObjectD.GetHash();
+
+			Assert.IsFalse(await ObjectCollection.ExistsAsync(NamespaceId, HashD));
+			Assert.IsNull(await ObjectCollection.GetAsync(NamespaceId, HashD));
+
+			await ObjectCollection.AddAsync(NamespaceId, HashD, ObjectD);
+
+			Assert.IsTrue(await ObjectCollection.ExistsAsync(NamespaceId, HashD));
+
+			CbObject? ReturnedObjectD = await ObjectCollection.GetAsync(NamespaceId, HashD);
+			Assert.IsNotNull(ReturnedObjectD);
+			Assert.IsTrue(ReturnedObjectD!.GetView().Span.SequenceEqual(ObjectD.GetView().Span));
+		}
+
+		[TestMethod]
+		public async Task RefCollectionTest()
+		{
+			IBlobCollection BlobCollection = ServiceProvider.GetRequiredService<IBlobCollection>();
+			IObjectCollection ObjectCollection = ServiceProvider.GetRequiredService<IObjectCollection>();
+			IRefCollection RefCollection = ServiceProvider.GetRequiredService<IRefCollection>();
+
+			byte[] BlobA = Encoding.UTF8.GetBytes("This is blob A");
+			IoHash HashA = IoHash.Compute(BlobA);
+
+			byte[] BlobB = Encoding.UTF8.GetBytes("This is blob B");
+			IoHash HashB = IoHash.Compute(BlobB);
+
+			byte[] BlobC = Encoding.UTF8.GetBytes("This is blob C");
+			IoHash HashC = IoHash.Compute(BlobC);
+
+			CbObject ObjectD = CbObject.Build(Writer =>
+			{
+				Writer.AddBinaryAttachment(HashA, "A");
+				Writer.AddBinaryAttachment(HashB, "B");
+				Writer.AddBinaryAttachment(HashC, "C");
+			});
+			IoHash HashD = ObjectD.GetHash();
+
+			CbObject ObjectE = CbObject.Build(Writer =>
+			{
+				Writer.AddBinaryAttachment(HashA, "A");
+				Writer.AddCompactBinaryAttachment(HashD, "D");
+			});
+			IoHash HashE = ObjectE.GetHash();
+
+			NamespaceId NamespaceId = new NamespaceId("test-ns");
+			BucketId BucketId = new BucketId("test-bkt");
+			const string RefName = "refname";
+
+			// Check that setting the ref to E fails with the correct missing hashes
+			List<IoHash> MissingHashes = await RefCollection.SetAsync(NamespaceId, BucketId, RefName, ObjectE);
+			Assert.AreEqual(2, MissingHashes.Count);
+			Assert.IsTrue(MissingHashes.Contains(HashA));
+			Assert.IsTrue(MissingHashes.Contains(HashD));
+
+			IRef? Ref = await RefCollection.GetAsync(NamespaceId, BucketId, RefName);
+			Assert.IsNotNull(Ref);
+			Assert.IsFalse(Ref!.Finalized);
+			Assert.IsTrue(Ref!.Value.GetView().Span.SequenceEqual(ObjectE.GetView().Span));
+
+			// Add object D, and check that changes the missing hashes to just the blobs
+			await ObjectCollection.AddAsync(NamespaceId, HashD, ObjectD);
+
+			MissingHashes = await RefCollection.FinalizeAsync(Ref);
+			Assert.AreEqual(3, MissingHashes.Count);
+			Assert.IsTrue(MissingHashes.Contains(HashA));
+			Assert.IsTrue(MissingHashes.Contains(HashB));
+			Assert.IsTrue(MissingHashes.Contains(HashC));
+
+			Ref = await RefCollection.GetAsync(NamespaceId, BucketId, RefName);
+			Assert.IsNotNull(Ref);
+			Assert.IsFalse(Ref!.Finalized);
+
+			// Add blobs A, B and C and check that the object can be finalized
+			await BlobCollection.AddAsync(NamespaceId, HashA, BlobA);
+			await BlobCollection.AddAsync(NamespaceId, HashB, BlobB);
+			await BlobCollection.AddAsync(NamespaceId, HashC, BlobC);
+
+			MissingHashes = await RefCollection.FinalizeAsync(Ref);
+			Assert.AreEqual(0, MissingHashes.Count);
+
+			Ref = await RefCollection.GetAsync(NamespaceId, BucketId, RefName);
+			Assert.IsNotNull(Ref);
+			Assert.IsTrue(Ref!.Finalized);
+
+			// Add a new ref to existing objects and check it finalizes correctly 
+			const string RefName2 = "refname";
+
+			MissingHashes = await RefCollection.SetAsync(NamespaceId, BucketId, RefName2, ObjectE);
+			Assert.AreEqual(0, MissingHashes.Count);
+
+			Ref = await RefCollection.GetAsync(NamespaceId, BucketId, RefName2);
+			Assert.IsNotNull(Ref);
+			Assert.IsTrue(Ref!.Finalized);
+		}
+	}
+}
diff --git a/Engine/Source/Programs/Horde/HordeServerTests/TestSetup.cs b/Engine/Source/Programs/Horde/HordeServerTests/TestSetup.cs
--- a/Engine/Source/Programs/Horde/HordeServerTests/TestSetup.cs
+++ b/Engine/Source/Programs/Horde/HordeServerTests/TestSetup.cs
@@ -6,6 +6,7 @@ using System.Reflection;
 using System.Security.Claims;
 using System.Threading.Tasks;
 using Datadog.Trace;
+using EpicGames.Core;
 using HordeCommon;
 using HordeServer;
 using HordeServer.Collections;
@@ -21,6 +22,7 @@ using HordeServer.Rpc;
 using HordeServer.Services;
 using HordeServer.Services.Impl;
 using HordeServer.Storage;
+using HordeServer.Storage.Collections;
 using HordeServer.Storage.Impl;
 using HordeServer.Storage.Services;
 using HordeServer.Tasks.Impl;
@@ -49,7 +51,7 @@ namespace HordeServerTests
 	public class TestSetup
 	{
 		public IServiceProvider ServiceProvider { get; }
-		public FakeClock Clock { get; }
+		public FakeClock Clock { get; set; } = null!;
 
 		public IGraphCollection GraphCollection => ServiceProvider.GetRequiredService<IGraphCollection>();
 		public INotificationTriggerCollection NotificationTriggerCollection => ServiceProvider.GetRequiredService<INotificationTriggerCollection>();
@@ -106,15 +108,32 @@ namespace HordeServerTests
 			IConfiguration Config = new ConfigurationBuilder().Build();
 
 			IServiceCollection Services = new ServiceCollection();
-			Services.Configure<ServerSettings>(Settings =>
-			{
-				Settings.AdminClaimType = HordeClaimTypes.Role;
-				Settings.AdminClaimValue = "app-horde-admins";
-			});
-			Services.AddLogging(Builder => Builder.AddConsole());
 			Services.AddSingleton(DbService);
+			Services.Configure<ServerSettings>(ConfigureSettings);
 			Services.AddSingleton<IConfiguration>(Config);
 
+			ConfigureServices(Services);
+
+			ServiceProvider = Services.BuildServiceProvider();
+		}
+
+		protected virtual void ConfigureSettings(ServerSettings Settings)
+		{
+			DirectoryReference DataDir = new DirectoryReference("TempData");
+			FileUtils.ForceDeleteDirectoryContents(DataDir);
+
+			Settings.AdminClaimType = HordeClaimTypes.Role;
+			Settings.AdminClaimValue = "app-horde-admins";
+
+			Settings.LocalLogsDir = DataDir.FullName;
+			Settings.LocalArtifactsDir = DataDir.FullName;
+			Settings.LocalBlobsDir = DataDir.FullName;
+		}
+
+		protected virtual void ConfigureServices(IServiceCollection Services)
+		{
+			Services.AddLogging(Builder => Builder.AddConsole());
+
 			Services.AddSingleton<StackExchange.Redis.IDatabase>(new Mock<StackExchange.Redis.IDatabase>().Object);
 
 			Services.AddSingleton<IAgentCollection, AgentCollection>();
@@ -185,13 +204,17 @@ namespace HordeServerTests
 			Services.AddSingleton<ContentStorageService>();
 			Services.AddSingleton<ExecutionService>();
 
+			Services.AddSingleton<IBlobCollection, FileSystemBlobCollection>();
+			Services.AddSingleton<IObjectCollection, ObjectCollection>();
+			Services.AddSingleton<IRefCollection, RefCollection>();
+			Services.AddSingleton<INamespaceCollection, NamespaceCollection>();
+			Services.AddSingleton<IBucketCollection, BucketCollection>();
+
 			Services.AddSingleton<FileSystemStorageBackend>();
 			Services.AddSingleton<IStorageService, SimpleStorageService>(SP => new SimpleStorageService(SP.GetRequiredService<FileSystemStorageBackend>()));
 
 			Services.AddSingleton<ISingletonDocument<GlobalPermissions>>(new SingletonDocumentStub<GlobalPermissions>());
 			Services.AddSingleton<ISingletonDocument<AgentSoftwareChannels>>(new SingletonDocumentStub<AgentSoftwareChannels>());
-
-			ServiceProvider = Services.BuildServiceProvider();
 		}
 
 		public async Task CreateFixture(bool ForceNewFixture = false)
diff --git a/Engine/Source/Programs/Shared/EpicGames.IoHash/IoHash.cs b/Engine/Source/Programs/Shared/EpicGames.IoHash/IoHash.cs
--- a/Engine/Source/Programs/Shared/EpicGames.IoHash/IoHash.cs
+++ b/Engine/Source/Programs/Shared/EpicGames.IoHash/IoHash.cs
@@ -148,6 +148,15 @@ namespace EpicGames.Core
 		/// Tests whether A is less than or equal to B
 		/// </summary>
 		public static bool operator <=(IoHash A, IoHash B) => A.CompareTo(B) <= 0;
+
+		/// <summary>
+		/// Convert a Blake3Hash to an IoHash
+		/// </summary>
+		/// <param name="Hash"></param>
+		public static implicit operator IoHash(Blake3Hash Hash)
+		{
+			return new IoHash(Hash.Memory.Slice(0, NumBytes));
+		}
 	}
 
 	/// <summary>
diff --git a/Engine/Source/Programs/Shared/EpicGames.Serialization/CompactBinary.cs b/Engine/Source/Programs/Shared/EpicGames.Serialization/CompactBinary.cs
--- a/Engine/Source/Programs/Shared/EpicGames.Serialization/CompactBinary.cs
+++ b/Engine/Source/Programs/Shared/EpicGames.Serialization/CompactBinary.cs
@@ -1966,6 +1966,20 @@ namespace EpicGames.Serialization
 			InnerField = new CbField(Buffer, FieldType);
 		}
 
+		/// <summary>
+		/// Builds an object by calling a delegate with a writer
+		/// </summary>
+		/// <param name="Build"></param>
+		/// <returns></returns>
+		public static CbObject Build(Action<CbWriter> Build)
+		{
+			CbWriter Writer = new CbWriter();
+			Writer.BeginObject();
+			Build(Writer);
+			Writer.EndObject();
+			return new CbObject(Writer.Save());
+		}
+
 		/// <inheritdoc cref="FindView(ReadOnlyUtf8String)"/>
 		public CbField Find(ReadOnlyUtf8String Name) => FindView(Name);
 
@@ -2094,6 +2108,22 @@ namespace EpicGames.Serialization
 		/// <param name="Visitor"></param>
 		public void IterateAttachments(Action<CbField> Visitor) => CreateViewIterator().IterateRangeAttachments(Visitor);
 
+		/// <summary>
+		/// Creates a view of the object, excluding the name
+		/// </summary>
+		/// <returns></returns>
+		public ReadOnlyMemory<byte> GetView()
+		{
+			ReadOnlyMemory<byte> Memory;
+			if (!TryGetView(out Memory))
+			{
+				byte[] Data = new byte[GetSize()];
+				CopyTo(Data);
+				Memory = Data;
+			}
+			return Memory;
+		}
+
 		/// <summary>
 		/// Try to get a view of the object as it would be serialized, such as by CopyTo.
 		/// 
diff --git a/Engine/Source/Programs/Shared/EpicGames.Serialization/CompactBinaryWriter.cs b/Engine/Source/Programs/Shared/EpicGames.Serialization/CompactBinaryWriter.cs
--- a/Engine/Source/Programs/Shared/EpicGames.Serialization/CompactBinaryWriter.cs
+++ b/Engine/Source/Programs/Shared/EpicGames.Serialization/CompactBinaryWriter.cs
@@ -20,6 +20,11 @@ namespace EpicGames.Serialization
 			public byte[]? Payload;
 		}
 
+		public CbObject ToObject()
+		{
+			return new CbObject(Save());
+		}
+
 		public byte[] Save()
 		{
 			int bufferLength = 1; // field type;
