Missing copyright notices.

diff --git a/Engine/Source/Programs/Horde/HordeServer/Utilities/LazyCache.cs b/Engine/Source/Programs/Horde/HordeServer/Utilities/LazyCache.cs
--- a/Engine/Source/Programs/Horde/HordeServer/Utilities/LazyCache.cs
+++ b/Engine/Source/Programs/Horde/HordeServer/Utilities/LazyCache.cs
@@ -1,3 +1,5 @@
+// Copyright Epic Games, Inc. All Rights Reserved.
+
 using Microsoft.Extensions.Caching.Memory;
 using Microsoft.Extensions.Options;
 using System;
diff --git a/Engine/Source/Programs/Horde/HordeServerTests/StorageTests.cs b/Engine/Source/Programs/Horde/HordeServerTests/StorageTests.cs
--- a/Engine/Source/Programs/Horde/HordeServerTests/StorageTests.cs
+++ b/Engine/Source/Programs/Horde/HordeServerTests/StorageTests.cs
@@ -1,3 +1,5 @@
+// Copyright Epic Games, Inc. All Rights Reserved.
+
 using EpicGames.Core;
 using EpicGames.Serialization;
 using HordeServer.Storage;
