Horde: Use previous implementation of IsPortAvailable on Linux.

diff --git a/Engine/Source/Programs/Horde/HordeServerTests/DatabaseRunner.cs b/Engine/Source/Programs/Horde/HordeServerTests/DatabaseRunner.cs
--- a/Engine/Source/Programs/Horde/HordeServerTests/DatabaseRunner.cs
+++ b/Engine/Source/Programs/Horde/HordeServerTests/DatabaseRunner.cs
@@ -143,15 +143,41 @@ namespace HordeServerTests
 
 		private static bool IsPortAvailable(int Port)
 		{
-			IPGlobalProperties IpGlobalProperties = IPGlobalProperties.GetIPGlobalProperties();
+			if (RuntimeInformation.IsOSPlatform(OSPlatform.Windows))
+			{
+				IPGlobalProperties IpGlobalProperties = IPGlobalProperties.GetIPGlobalProperties();
+
+				IPEndPoint[] Listeners = IpGlobalProperties.GetActiveTcpListeners();
+				if (Listeners.Any(x => x.Port == Port))
+				{
+					return false;
+				}
 
-			IPEndPoint[] Listeners = IpGlobalProperties.GetActiveTcpListeners();
-			if (Listeners.Any(x => x.Port == Port))
+				return true;
+			}
+			else
 			{
+				TcpListener? ListenerAny = null;
+				TcpListener? ListenerLoopback = null;
+				try
+				{
+					ListenerAny = new TcpListener(IPAddress.Loopback, Port);
+					ListenerAny.Start();
+					ListenerLoopback = new TcpListener(IPAddress.Any, Port);
+					ListenerLoopback.Start();
+					return true;
+				}
+				catch (SocketException)
+				{
+				}
+				finally
+				{
+					ListenerAny?.Stop();
+					ListenerLoopback?.Stop();
+				}
+
 				return false;
 			}
-
-			return true;
 		}
 		
 		private static void DeleteDirectory(string Path)
