Fix exception running tests.

diff --git a/Engine/Source/Programs/Horde/HordeServer/Services/PerforceService.cs b/Engine/Source/Programs/Horde/HordeServer/Services/PerforceService.cs
--- a/Engine/Source/Programs/Horde/HordeServer/Services/PerforceService.cs
+++ b/Engine/Source/Programs/Horde/HordeServer/Services/PerforceService.cs
@@ -632,7 +632,7 @@ namespace HordeServer.Services
 		/// <inheritdoc/>
 		public override Task<int> DuplicateShelvedChangeAsync(int ShelvedChange)
 		{
-			throw new NotImplementedException();
+			return Task.FromResult(0);
 		}
 
 		/// <summary>
