Horde: Prevent tests failing if something has an open handle to a file in the storage directory.

diff --git a/Engine/Source/Programs/Horde/HordeServerTests/TestSetup.cs b/Engine/Source/Programs/Horde/HordeServerTests/TestSetup.cs
--- a/Engine/Source/Programs/Horde/HordeServerTests/TestSetup.cs
+++ b/Engine/Source/Programs/Horde/HordeServerTests/TestSetup.cs
@@ -119,12 +119,19 @@ namespace HordeServerTests
 
 		protected virtual void ConfigureSettings(ServerSettings Settings)
 		{
-			DirectoryReference DataDir = new DirectoryReference("TempData");
-			FileUtils.ForceDeleteDirectoryContents(DataDir);
+			DirectoryReference BaseDir = new DirectoryReference("TempData");
+			try
+			{
+				FileUtils.ForceDeleteDirectoryContents(BaseDir);
+			}
+			catch
+			{
+			}
 
 			Settings.AdminClaimType = HordeClaimTypes.Role;
 			Settings.AdminClaimValue = "app-horde-admins";
 
+			DirectoryReference DataDir = DirectoryReference.Combine(BaseDir, Guid.NewGuid().ToString("N"));
 			Settings.LocalLogsDir = DataDir.FullName;
 			Settings.LocalArtifactsDir = DataDir.FullName;
 			Settings.LocalBlobsDir = DataDir.FullName;
