Horde: Remove static state for fixtures to fix automated tests.

diff --git a/Engine/Source/Programs/Horde/HordeServer/Services/ArtifactService.cs b/Engine/Source/Programs/Horde/HordeServer/Services/ArtifactService.cs
--- a/Engine/Source/Programs/Horde/HordeServer/Services/ArtifactService.cs
+++ b/Engine/Source/Programs/Horde/HordeServer/Services/ArtifactService.cs
@@ -91,6 +91,11 @@ namespace HordeServer.Services
 				Logger.LogError("Tried to find artifact that didn't exist at {0}!", ArtifactFile.FullName);
 				throw;
 			}
+			catch (DirectoryNotFoundException)
+			{
+				Logger.LogError("Tried to find artifact that didn't exist at {0}!", ArtifactFile.FullName);
+				throw;
+			}
 		}
 
 		/// <inheritdoc/>
diff --git a/Engine/Source/Programs/Horde/HordeServerTests/ControllerIntegrationTest.cs b/Engine/Source/Programs/Horde/HordeServerTests/ControllerIntegrationTest.cs
--- a/Engine/Source/Programs/Horde/HordeServerTests/ControllerIntegrationTest.cs
+++ b/Engine/Source/Programs/Horde/HordeServerTests/ControllerIntegrationTest.cs
@@ -35,17 +35,17 @@ namespace HordeServerTests
 
     public class ControllerIntegrationTest
     {
-        private static HttpClient? _client;
-        private static TestWebApplicationFactory<Startup>? _factory;
+        private HttpClient? _client;
+        private TestWebApplicationFactory<Startup>? _factory;
         protected readonly HttpClient client;
-        private static Fixture? _fixture;
+        private Fixture? _fixture;
 
         public ControllerIntegrationTest()
         {
             client = GetClientForTestServer();
         }
 
-        public static async Task<Fixture> GetFixture()
+        public async Task<Fixture> GetFixture()
         {
             if (_fixture != null) return _fixture;
 
@@ -65,7 +65,7 @@ namespace HordeServerTests
         }
 
         [MethodImpl(MethodImplOptions.Synchronized)]
-        public static HttpClient GetClientForTestServer()
+        public HttpClient GetClientForTestServer()
         {
             if (_client != null) return _client;
 
@@ -75,7 +75,7 @@ namespace HordeServerTests
         }
 
         [MethodImpl(MethodImplOptions.Synchronized)]
-        static TestWebApplicationFactory<Startup> GetFactory()
+        TestWebApplicationFactory<Startup> GetFactory()
         {
             if (_factory != null) return _factory;
             _factory = new TestWebApplicationFactory<Startup>();
diff --git a/Engine/Source/Programs/Horde/HordeServerTests/Fixture.cs b/Engine/Source/Programs/Horde/HordeServerTests/Fixture.cs
--- a/Engine/Source/Programs/Horde/HordeServerTests/Fixture.cs
+++ b/Engine/Source/Programs/Horde/HordeServerTests/Fixture.cs
@@ -32,15 +32,10 @@ namespace HordeServerTests
 		public IAgent Agent1 { get; private set; } = null!;
 		public string Agent1Name { get; private set; } = null!;
 
-		private static Fixture? _fixture;
-
 		public static async Task<Fixture> Create(bool ForceNewFixture, IGraphCollection GraphCollection, TemplateService TemplateService, JobService JobService, ArtifactService ArtifactService, StreamService StreamService, AgentService AgentService, IPerforceService PerforceService)
 		{
-			if (_fixture == null || ForceNewFixture)
-			{
-				_fixture = new Fixture();
-				await _fixture.Populate(GraphCollection, TemplateService, JobService, ArtifactService, StreamService, AgentService, PerforceService);
-			}
+			Fixture _fixture = new Fixture();
+			await _fixture.Populate(GraphCollection, TemplateService, JobService, ArtifactService, StreamService, AgentService, PerforceService);
 			
 			(PerforceService as PerforceServiceStub)?.AddChange("//UE5/Main", 112233, "leet.coder", "Did stuff", new []{"file.cpp"});
 			(PerforceService as PerforceServiceStub)?.AddChange("//UE5/Main", 1111, "swarm", "A shelved CL here", new []{"renderer.cpp"});
diff --git a/Engine/Source/Programs/Horde/HordeServerTests/TestSetup.cs b/Engine/Source/Programs/Horde/HordeServerTests/TestSetup.cs
--- a/Engine/Source/Programs/Horde/HordeServerTests/TestSetup.cs
+++ b/Engine/Source/Programs/Horde/HordeServerTests/TestSetup.cs
@@ -119,7 +119,7 @@ namespace HordeServerTests
 
 		protected virtual void ConfigureSettings(ServerSettings Settings)
 		{
-			DirectoryReference BaseDir = new DirectoryReference("TempData");
+			DirectoryReference BaseDir = DirectoryReference.Combine(Program.DataDir, "Tests");
 			try
 			{
 				FileUtils.ForceDeleteDirectoryContents(BaseDir);
@@ -132,9 +132,9 @@ namespace HordeServerTests
 			Settings.AdminClaimValue = "app-horde-admins";
 
 			DirectoryReference DataDir = DirectoryReference.Combine(BaseDir, Guid.NewGuid().ToString("N"));
-			Settings.LocalLogsDir = DataDir.FullName;
-			Settings.LocalArtifactsDir = DataDir.FullName;
-			Settings.LocalBlobsDir = DataDir.FullName;
+			Settings.LocalLogsDir = DirectoryReference.Combine(DataDir, "Logs").FullName;
+			Settings.LocalArtifactsDir = DirectoryReference.Combine(DataDir, "Artifacts").FullName;
+			Settings.LocalBlobsDir = DirectoryReference.Combine(DataDir, "Blobs").FullName;
 		}
 
 		protected virtual void ConfigureServices(IServiceCollection Services)
