Horde: Always drop the database when running tests.

diff --git a/Engine/Source/Programs/Horde/HordeServerTests/DatabaseIntegrationTest.cs b/Engine/Source/Programs/Horde/HordeServerTests/DatabaseIntegrationTest.cs
--- a/Engine/Source/Programs/Horde/HordeServerTests/DatabaseIntegrationTest.cs
+++ b/Engine/Source/Programs/Horde/HordeServerTests/DatabaseIntegrationTest.cs
@@ -57,7 +57,7 @@ namespace HordeServerTests
 	        return TestSetup;
         }
 
-        public static DatabaseService GetDatabaseService(bool ForceDropDatabase = false)
+        public static DatabaseService GetDatabaseService(bool ForceDropDatabase = true)
         {
             if (_mongoDbRunner == null)
             {
