Horde: Print the start of the next downtime window to the log.

diff --git a/Engine/Source/Programs/Horde/HordeServer/Services/DowntimeService.cs b/Engine/Source/Programs/Horde/HordeServer/Services/DowntimeService.cs
--- a/Engine/Source/Programs/Horde/HordeServer/Services/DowntimeService.cs
+++ b/Engine/Source/Programs/Horde/HordeServer/Services/DowntimeService.cs
@@ -75,7 +75,20 @@ namespace HordeServer.Services
 			DateTimeOffset Now = TimeZoneInfo.ConvertTime(DateTimeOffset.Now, Settings.CurrentValue.TimeZoneInfo);
 			bool bIsActive = Globals.ScheduledDowntime.Any(x => x.IsActive(Now));
 
-			Logger.LogInformation("Server time: {Time}. Downtime: {Downtime}. Schedule: {@Schedule}.", Now, bIsActive, Globals.ScheduledDowntime);
+			DateTimeOffset? Next = null;
+			foreach (ScheduledDowntime Schedule in Globals.ScheduledDowntime)
+			{
+				DateTimeOffset Start = Schedule.GetNext(Now).StartTime;
+				if(Next == null || Start < Next)
+				{
+					Next = Start;
+				}
+			}
+
+			if (Next != null)
+			{
+				Logger.LogInformation("Server time: {Time}. Downtime: {Downtime}. Next: {Next}.", Now, bIsActive, Next.Value);
+			}
 
 			if (bIsActive != IsDowntimeActive)
 			{
