Horde: Mirror ThirdParty folder structure for runtime dependencies.

diff --git a/Engine/Source/Programs/Horde/HordeServer/HordeServer.csproj b/Engine/Source/Programs/Horde/HordeServer/HordeServer.csproj
--- a/Engine/Source/Programs/Horde/HordeServer/HordeServer.csproj
+++ b/Engine/Source/Programs/Horde/HordeServer/HordeServer.csproj
@@ -83,6 +83,7 @@
     </Content>
     <Content Include="..\ThirdParty\**" Condition="$([MSBuild]::IsOSPlatform('Windows'))">
       <CopyToOutputDirectory>PreserveNewest</CopyToOutputDirectory>
+      <Link>ThirdParty/%(RecursiveDir)/%(Filename)%(Extension)</Link>
     </Content>
   </ItemGroup>
 
