Horde: Add a debug endpoint to test default AWS configuration.

diff --git a/Engine/Source/Programs/Horde/HordeServer/Controllers/DebugController.cs b/Engine/Source/Programs/Horde/HordeServer/Controllers/DebugController.cs
--- a/Engine/Source/Programs/Horde/HordeServer/Controllers/DebugController.cs
+++ b/Engine/Source/Programs/Horde/HordeServer/Controllers/DebugController.cs
@@ -40,7 +40,8 @@ using HordeServer.Storage;
 using Microsoft.Extensions.Hosting;
 using Microsoft.Extensions.Options;
 using StatsdClient;
-using P4Debugging = Perforce.P4.P4Debugging; 
+using P4Debugging = Perforce.P4.P4Debugging;
+using Amazon.S3;
 
 namespace HordeServer.Controllers
 {
@@ -60,17 +61,33 @@ namespace HordeServer.Controllers
 
 		IDogStatsd DogStatsd;
 
+		IAmazonS3 AwsTest;
+
 		/// <summary>
 		/// Constructor
 		/// </summary>
 		/// <param name="RequestTrackerService"></param>
 		/// <param name="ApplicationLifetime"></param>
 		/// <param name="DogStatsd"></param>
-		public PublicDebugController(RequestTrackerService RequestTrackerService, IHostApplicationLifetime ApplicationLifetime, IDogStatsd DogStatsd)
+		/// <param name="AwsTest"></param>
+		public PublicDebugController(RequestTrackerService RequestTrackerService, IHostApplicationLifetime ApplicationLifetime, IDogStatsd DogStatsd, IAmazonS3 AwsTest)
 		{
 			this.RequestTrackerService = RequestTrackerService;
 			this.ApplicationLifetime = ApplicationLifetime;
 			this.DogStatsd = DogStatsd;
+			this.AwsTest = AwsTest;
+		}
+
+		/// <summary>
+		/// Returns an object from aws
+		/// </summary>
+		/// <returns>Http result</returns>
+		[HttpGet]
+		[Route("/api/v1/debug/awsdata")]
+		public async Task<ActionResult> GetAwsData([FromQuery] string BucketName, [FromQuery] string Path)
+		{
+			Stream Stream = await AwsTest.GetObjectStreamAsync(BucketName, Path, new Dictionary<string, object>());
+			return new FileStreamResult(Stream, new MediaTypeHeaderValue("application/octet-stream"));
 		}
 
 		/// <summary>
diff --git a/Engine/Source/Programs/Horde/HordeServer/Startup.cs b/Engine/Source/Programs/Horde/HordeServer/Startup.cs
--- a/Engine/Source/Programs/Horde/HordeServer/Startup.cs
+++ b/Engine/Source/Programs/Horde/HordeServer/Startup.cs
@@ -72,6 +72,7 @@ using HordeServer.Notifications;
 using Microsoft.AspNetCore.Hosting.Server.Features;
 using System.Runtime.InteropServices;
 using HordeServer.Storage.Collections;
+using Amazon.S3;
 
 namespace HordeServer
 {
@@ -416,6 +417,9 @@ namespace HordeServer
 
 			Services.AddSingleton<DeviceService>();
 
+			Services.AddDefaultAWSOptions(Configuration.GetAWSOptions());
+			Services.AddAWSService<IAmazonS3>();
+
 			ConfigureStorageProvider(Services, Settings);
 			ConfigureLogStorage(Services);
 			//			ConfigureLogFileWriteCache(Services, Settings);
