Horde: Manually assume roles for default AWS credentials if specified in config file.

diff --git a/Engine/Source/Programs/Horde/HordeServer/Startup.cs b/Engine/Source/Programs/Horde/HordeServer/Startup.cs
--- a/Engine/Source/Programs/Horde/HordeServer/Startup.cs
+++ b/Engine/Source/Programs/Horde/HordeServer/Startup.cs
@@ -72,7 +72,9 @@ using HordeServer.Notifications;
 using Microsoft.AspNetCore.Hosting.Server.Features;
 using System.Runtime.InteropServices;
 using HordeServer.Storage.Collections;
+using Amazon.Extensions.NETCore.Setup;
 using Amazon.S3;
+using Amazon.Runtime;
 
 namespace HordeServer
 {
@@ -417,7 +419,12 @@ namespace HordeServer
 
 			Services.AddSingleton<DeviceService>();
 
-			Services.AddDefaultAWSOptions(Configuration.GetAWSOptions());
+			AWSOptions Options = Configuration.GetAWSOptions();
+			if (Settings.S3CredentialType == "AssumeRole" && !String.IsNullOrEmpty(Settings.S3AssumeArn) && Options.Credentials != null)
+			{
+				Options.Credentials = new AssumeRoleAWSCredentials(Options.Credentials, Settings.S3AssumeArn, "Horde");
+			}
+			Services.AddDefaultAWSOptions(Options);
 			Services.AddAWSService<IAmazonS3>();
 
 			ConfigureStorageProvider(Services, Settings);
