Horde: Print types of AWS credential providers in dev.

diff --git a/Engine/Source/Programs/Horde/HordeServer/Startup.cs b/Engine/Source/Programs/Horde/HordeServer/Startup.cs
--- a/Engine/Source/Programs/Horde/HordeServer/Startup.cs
+++ b/Engine/Source/Programs/Horde/HordeServer/Startup.cs
@@ -420,6 +420,10 @@ namespace HordeServer
 			Services.AddSingleton<DeviceService>();
 
 			AWSOptions Options = Configuration.GetAWSOptions();
+
+			Serilog.Log.Logger.Information("AWS: {Key}={Value}", "CredentialsType", Options.Credentials?.GetType()?.Name);
+			Serilog.Log.Logger.Information("AWS: {Key}={Value}", "DefaultCredentialsType", FallbackCredentialsFactory.GetCredentials()?.GetType()?.Name);
+
 			if (Settings.S3CredentialType == "AssumeRole" && !String.IsNullOrEmpty(Settings.S3AssumeArn) && Options.Credentials != null)
 			{
 				Options.Credentials = new AssumeRoleAWSCredentials(Options.Credentials, Settings.S3AssumeArn, "Horde");
