Horde: Refactor storage API to use common backend with existing log storage (adding support for S3, in-memory, etc...).

diff --git a/Engine/Source/Programs/Horde/HordeServer/Controllers/DebugController.cs b/Engine/Source/Programs/Horde/HordeServer/Controllers/DebugController.cs
--- a/Engine/Source/Programs/Horde/HordeServer/Controllers/DebugController.cs
+++ b/Engine/Source/Programs/Horde/HordeServer/Controllers/DebugController.cs
@@ -515,7 +515,7 @@ namespace HordeServer.Controllers
 				return Forbid();
 			}
 
-			ReadOnlyMemory<byte>? Data = await StorageProvider.ReadAsync(Path);
+			ReadOnlyMemory<byte>? Data = await StorageProvider.ReadBytesAsync(Path);
 			if (Data == null)
 			{
 				return NotFound();
diff --git a/Engine/Source/Programs/Horde/HordeServer/Logs/Storage/Impl/PersistentLogStorage.cs b/Engine/Source/Programs/Horde/HordeServer/Logs/Storage/Impl/PersistentLogStorage.cs
--- a/Engine/Source/Programs/Horde/HordeServer/Logs/Storage/Impl/PersistentLogStorage.cs
+++ b/Engine/Source/Programs/Horde/HordeServer/Logs/Storage/Impl/PersistentLogStorage.cs
@@ -33,7 +33,7 @@ namespace HordeServer.Logs.Readers
 		/// </summary>
 		/// <param name="StorageProvider">The storage provider</param>
 		/// <param name="Logger">Logging provider</param>
-		public PersistentLogStorage(IStorageBackend StorageProvider, ILogger Logger)
+		public PersistentLogStorage(IStorageBackend StorageProvider, ILogger<PersistentLogStorage> Logger)
 		{
 			this.StorageProvider = StorageProvider;
 			this.Logger = Logger;
@@ -42,7 +42,6 @@ namespace HordeServer.Logs.Readers
 		/// <inheritdoc/>
 		public void Dispose()
 		{
-			StorageProvider.Dispose();
 		}
 
 		/// <inheritdoc/>
@@ -51,7 +50,7 @@ namespace HordeServer.Logs.Readers
 			Logger.LogDebug("Reading log {LogId} index length {Length} from persistent storage", LogId, Length);
 
 			string Path = $"{LogId}/index_{Length}";
-			ReadOnlyMemory<byte>? Data = await StorageProvider.ReadAsync(Path);
+			ReadOnlyMemory<byte>? Data = await StorageProvider.ReadBytesAsync(Path);
 			if (Data == null)
 			{
 				return null;
@@ -66,7 +65,7 @@ namespace HordeServer.Logs.Readers
 
 			string Path = $"{LogId}/index_{Length}";
 			ReadOnlyMemory<byte> Data = IndexData.ToByteArray();
-			return StorageProvider.WriteAsync(Path, Data);
+			return StorageProvider.WriteBytesAsync(Path, Data);
 		}
 
 		/// <inheritdoc/>
@@ -75,7 +74,7 @@ namespace HordeServer.Logs.Readers
 			Logger.LogDebug("Reading log {LogId} chunk offset {Offset} from persistent storage", LogId, Offset);
 
 			string Path = $"{LogId}/offset_{Offset}";
-			ReadOnlyMemory<byte>? Data = await StorageProvider.ReadAsync(Path);
+			ReadOnlyMemory<byte>? Data = await StorageProvider.ReadBytesAsync(Path);
 			if(Data == null)
 			{
 				return null;
@@ -103,7 +102,7 @@ namespace HordeServer.Logs.Readers
 			Writer.WriteLogChunkData(ChunkData);
 			Writer.CheckOffset(Data.Length);
 
-			return StorageProvider.WriteAsync(Path, Data);
+			return StorageProvider.WriteBytesAsync(Path, Data);
 		}
 	}
 }
diff --git a/Engine/Source/Programs/Horde/HordeServer/ServerSettings.cs b/Engine/Source/Programs/Horde/HordeServer/ServerSettings.cs
--- a/Engine/Source/Programs/Horde/HordeServer/ServerSettings.cs
+++ b/Engine/Source/Programs/Horde/HordeServer/ServerSettings.cs
@@ -55,7 +55,6 @@ namespace HordeServer
 		Aws,
 	}
 
-	
 	/// <summary>
 	/// Settings for a remote execution instance
 	///
@@ -245,24 +244,14 @@ namespace HordeServer
 		public StorageProviderType ExternalStorageProviderType { get; set; } = StorageProviderType.FileSystem;
 
 		/// <summary>
-		/// Local log/artifact storage directory, if using type filesystem
-		/// </summary>
-		public string LocalLogsDir { get; set; } = "Logs";
-
-		/// <summary>
-		/// Gets the full path referred to by LocalStorageDir
-		/// </summary>
-		public DirectoryReference LocalLogsDirRef => DirectoryReference.Combine(Program.DataDir, LocalLogsDir);
-
-		/// <summary>
-		/// Local blob storage directory, if using type filesystem
+		/// Local storage directory, if using type filesystem
 		/// </summary>
-		public string LocalBlobsDir { get; set; } = "Blobs";
+		public string LocalStorageDir { get; set; } = "Storage";
 
 		/// <summary>
 		/// Gets the full path referred to by LocalStorageDir
 		/// </summary>
-		public DirectoryReference LocalBlobsDirRef => DirectoryReference.Combine(Program.DataDir, LocalBlobsDir);
+		public DirectoryReference LocalStorageDirRef => DirectoryReference.Combine(Program.DataDir, LocalStorageDir);
 
 		/// <summary>
 		/// Local artifact storage directory, if using type filesystem
@@ -282,7 +271,7 @@ namespace HordeServer
 		/// <summary>
 		/// Arn to assume for s3.  "Basic", "AssumeRole", "AssumeRoleWebIdentity" only
 		/// </summary>
-		public string S3CredentialType { get; set; } = null!;
+		public string? S3CredentialType { get; set; }
 
 		/// <summary>
 		/// S3 Client username (used in Basic auth type only)
diff --git a/Engine/Source/Programs/Horde/HordeServer/Startup.cs b/Engine/Source/Programs/Horde/HordeServer/Startup.cs
--- a/Engine/Source/Programs/Horde/HordeServer/Startup.cs
+++ b/Engine/Source/Programs/Horde/HordeServer/Startup.cs
@@ -51,7 +51,6 @@ using HordeCommon;
 using Microsoft.AspNetCore.DataProtection;
 using StackExchange.Redis;
 using HordeServer.Storage;
-using HordeServer.Storage.Impl;
 using HordeServer.Logs;
 using HordeServer.Logs.Builder;
 using HordeServer.Logs.Storage;
@@ -62,7 +61,6 @@ using HordeServer.Services.Impl;
 using Microsoft.Extensions.DependencyInjection.Extensions;
 using HordeServer.Authentication;
 using HordeServer.Rpc;
-using HordeServer.Storage.Services;
 using HordeServer.Tasks.Impl;
 using HordeServer.Tasks;
 using StatsdClient;
@@ -75,6 +73,8 @@ using HordeServer.Storage.Collections;
 using Amazon.Extensions.NETCore.Setup;
 using Amazon.S3;
 using Amazon.Runtime;
+using HordeServer.Storage.Backends;
+using HordeServer.Storage.Services;
 
 namespace HordeServer
 {
@@ -321,7 +321,7 @@ namespace HordeServer
 			Services.AddSingleton<IDeviceCollection, DeviceCollection>();
 
 			// Storage
-			Services.AddSingleton<IBlobCollection, FileSystemBlobCollection>();
+			Services.AddSingleton<IBlobCollection, BlobCollection>();
 			Services.AddSingleton<IBucketCollection, BucketCollection>();
 			Services.AddSingleton<INamespaceCollection, NamespaceCollection>();
 			Services.AddSingleton<IObjectCollection, ObjectCollection>();
@@ -414,9 +414,6 @@ namespace HordeServer
 			Services.AddSingleton<ContentStorageService>();
 			Services.AddSingleton<ExecutionService>();
 
-			Services.AddSingleton<FileSystemStorageBackend>();
-			Services.AddSingleton<IStorageService, SimpleStorageService>(SP => new SimpleStorageService(SP.GetRequiredService<FileSystemStorageBackend>()));
-
 			Services.AddSingleton<DeviceService>();
 
 			AWSOptions Options = Configuration.GetAWSOptions();
@@ -424,12 +421,11 @@ namespace HordeServer
 			{
 				Options.Credentials = new AssumeRoleAWSCredentials(FallbackCredentialsFactory.GetCredentials(), Settings.S3AssumeArn, "Horde");
 			}
-
 			Services.AddDefaultAWSOptions(Options);
-			Services.AddAWSService<IAmazonS3>();
 
-			ConfigureStorageProvider(Services, Settings);
+			ConfigureStorageBackend(Services, Settings);
 			ConfigureLogStorage(Services);
+			Services.AddSingleton<IStorageService, SimpleStorageService>();
 			//			ConfigureLogFileWriteCache(Services, Settings);
 
 			AuthenticationBuilder AuthBuilder = Services.AddAuthentication(Options =>
@@ -593,7 +589,7 @@ namespace HordeServer
 			Options.Converters.Add(new JsonDateTimeConverter());
 		}
 
-		private static void ConfigureStorageProvider(IServiceCollection Services, ServerSettings Settings)
+		private static void ConfigureStorageBackend(IServiceCollection Services, ServerSettings Settings)
 		{
 			switch (Settings.ExternalStorageProviderType)
 			{
@@ -628,9 +624,12 @@ namespace HordeServer
 					return new RedisLogBuilder(RedisService.ConnectionPool, Provider.GetRequiredService<ILogger<RedisLogBuilder>>());
 				}
 			});
+
+			Services.AddSingleton<PersistentLogStorage>();
+
 			Services.AddSingleton<ILogStorage>(Provider =>
 			{
-				ILogStorage Storage = new PersistentLogStorage(Provider.GetRequiredService<IStorageBackend>(), Provider.GetRequiredService<ILogger<PersistentLogStorage>>());
+				ILogStorage Storage = Provider.GetRequiredService<PersistentLogStorage>();
 
 //				IDatabase? RedisDb = Provider.GetService<IDatabase>();
 //				if (RedisDb != null)
diff --git a/Engine/Source/Programs/Horde/HordeServer/Storage/Backends/AwsStorageBackend.cs b/Engine/Source/Programs/Horde/HordeServer/Storage/Backends/AwsStorageBackend.cs
--- a/Engine/Source/Programs/Horde/HordeServer/Storage/Backends/AwsStorageBackend.cs
+++ b/Engine/Source/Programs/Horde/HordeServer/Storage/Backends/AwsStorageBackend.cs
@@ -4,21 +4,18 @@ using Amazon.Extensions.NETCore.Setup;
 using Amazon.Runtime;
 using Amazon.S3;
 using Amazon.S3.Model;
-using EpicGames.Core;
 using HordeServer.Utilities;
 using Microsoft.Extensions.Configuration;
 using Microsoft.Extensions.Logging;
 using Microsoft.Extensions.Options;
 using System;
-using System.Collections.Generic;
 using System.Diagnostics.CodeAnalysis;
 using System.Globalization;
 using System.IO;
-using System.Linq;
 using System.Threading;
 using System.Threading.Tasks;
 
-namespace HordeServer.Storage.Impl
+namespace HordeServer.Storage.Backends
 {
 	/// <summary>
 	/// Exception wrapper for S3 requests
@@ -40,7 +37,7 @@ namespace HordeServer.Storage.Impl
 	/// <summary>
 	/// FileStorage implementation using an s3 bucket
 	/// </summary>
-	public sealed class AwsStorageBackend : IStorageBackend
+	public sealed class AwsStorageBackend : IStorageBackend, IDisposable
 	{
 		/// <summary>
 		/// S3 Client
@@ -65,100 +62,104 @@ namespace HordeServer.Storage.Impl
 		/// <summary>
 		/// Constructor
 		/// </summary>
-		/// <param name="Settings">Current Horde Settings</param>
-		/// <param name="Configuration">Configuration</param>
+		/// <param name="AwsOptions">The AWS options</param>
+		/// <param name="ServerSettings">Options for configuring AWS</param>
 		/// <param name="Logger">Logger interface</param>
-		public AwsStorageBackend(IOptionsMonitor<ServerSettings> Settings, IConfiguration Configuration, ILogger<AwsStorageBackend> Logger)
+		public AwsStorageBackend(AWSOptions AwsOptions, IOptions<ServerSettings> ServerSettings, ILogger<AwsStorageBackend> Logger)
 		{
-			ServerSettings CurrentSettings = Settings.CurrentValue;
-			AWSOptions Options = Configuration.GetAWSOptions();
-			Options.Credentials = GetCredentials(CurrentSettings);
-
-			this.Client = Options.CreateServiceClient<IAmazonS3>();
-			this.BucketName = CurrentSettings.S3LogBucketName;
+			ServerSettings OptionsValue = ServerSettings.Value;
+			this.Client = AwsOptions.CreateServiceClient<IAmazonS3>();
+			this.BucketName = OptionsValue.S3LogBucketName;
 			this.Semaphore = new SemaphoreSlim(16);
 			this.Logger = Logger;
 		}
 
-		enum AWSCredentialsType
+		/// <inheritdoc/>
+		public void Dispose()
 		{
-			Basic,
-			AssumeRole,
-			AssumeRoleWebIdentity
+			Client.Dispose();
+			Semaphore.Dispose();
 		}
 
-		/// <summary>
-		/// Gets auth credentials based on type
-		/// </summary>
-		/// <param name="Settings">Horde settings</param>
-		/// <returns></returns>
-		static AWSCredentials GetCredentials(ServerSettings Settings)
+		class WrappedResponseStream : Stream
 		{
-			AWSCredentialsType AuthType;
-			if (!Enum.TryParse(Settings.S3CredentialType, true, out AuthType))
+			IDisposable Semaphore;
+			GetObjectResponse Response;
+			Stream ResponseStream;
+
+			public WrappedResponseStream(IDisposable Semaphore, GetObjectResponse Response)
 			{
-				throw new Exception("Could not determine auth type from appsettings. Should be Basic, AssumeRole, or AssumeRoleWebIdentity!");
+				this.Semaphore = Semaphore;
+				this.Response = Response;
+				this.ResponseStream = Response.ResponseStream;
 			}
 
-			switch (AuthType)
+			public override bool CanRead => true;
+			public override bool CanSeek => false;
+			public override bool CanWrite => false;
+			public override long Length => ResponseStream.Length;
+			public override long Position { get => ResponseStream.Position; set => throw new NotSupportedException(); }
+
+			public override void Flush() => ResponseStream.Flush();
+
+			public override int Read(byte[] buffer, int offset, int count) => ResponseStream.Read(buffer, offset, count);
+			public override ValueTask<int> ReadAsync(Memory<byte> buffer, CancellationToken cancellationToken = default) => ResponseStream.ReadAsync(buffer, cancellationToken);
+			public override Task<int> ReadAsync(byte[] buffer, int offset, int count, CancellationToken cancellationToken) => ResponseStream.ReadAsync(buffer, offset, count, cancellationToken);
+
+			public override long Seek(long offset, SeekOrigin origin) => throw new NotSupportedException();
+			public override void SetLength(long value) => throw new NotSupportedException();
+			public override void Write(byte[] buffer, int offset, int count) => throw new NotSupportedException();
+
+			protected override void Dispose(bool disposing)
 			{
-				case AWSCredentialsType.Basic:
-					return new BasicAWSCredentials(Settings.S3ClientKeyId, Settings.S3ClientSecret);
-				case AWSCredentialsType.AssumeRole:
-					return new AssumeRoleAWSCredentials(FallbackCredentialsFactory.GetCredentials(), Settings.S3AssumeArn, "Horde");
-				case AWSCredentialsType.AssumeRoleWebIdentity:
-					return AssumeRoleWithWebIdentityCredentials.FromEnvironmentVariables();
-				default:
-					throw new Exception($"Invalid authentication type: {AuthType}");
+				base.Dispose(disposing);
+
+				Semaphore.Dispose();
+				Response.Dispose();
+				ResponseStream.Dispose();
 			}
-		}
 
-		/// <inheritdoc/>
-		public void Dispose()
-		{
-			Client.Dispose();
-			Semaphore.Dispose();
-		}
+			public override async ValueTask DisposeAsync()
+			{
+				await base.DisposeAsync();
 
-		/// <inheritdoc/>
-		public Task<bool> TouchAsync(string Path)
-		{
-			throw new NotImplementedException();
+				Semaphore.Dispose();
+				Response.Dispose();
+				await ResponseStream.DisposeAsync();
+			}
 		}
 
 		/// <inheritdoc/>
-		public async Task<ReadOnlyMemory<byte>?> ReadAsync(string Path)
+		public async Task<Stream?> ReadAsync(string Path)
 		{
-			using (await Semaphore.UseWaitAsync())
+			IDisposable? Lock = null;
+			GetObjectResponse? Response = null;
+			try
 			{
+				Lock = await Semaphore.UseWaitAsync();
+
 				GetObjectRequest NewGetRequest = new GetObjectRequest();
 				NewGetRequest.BucketName = BucketName;
 				NewGetRequest.Key = Path;
 
-				try
-				{
-					using (GetObjectResponse GetResponse = await this.Client.GetObjectAsync(NewGetRequest))
-					{
-						using (Stream ResponseStream = GetResponse.ResponseStream)
-						{
-							using (MemoryStream OutputStream = new MemoryStream())
-							{
-								await ResponseStream.CopyToAsync(OutputStream);
-								return OutputStream.ToArray();
-							}
-						}
-					}
-				}
-				catch(Exception Ex)
-				{
-					throw new AwsException($"Unable to read from bucket {BucketName} path {Path}", Ex);
-				}
+				Response = await Client.GetObjectAsync(NewGetRequest);
+
+				return new WrappedResponseStream(Lock, Response);
+			}
+			catch (Exception Ex)
+			{
+				Logger.LogWarning(Ex, "Unable to read {Path} from S3", Path);
+
+				Lock?.Dispose();
+				Response?.Dispose();
+
+				return null;
 			}
 		}
 
 		/// <inheritdoc/>
 		[SuppressMessage("Design", "CA1031:Do not catch general exception types", Justification = "<Pending>")]
-		public async Task WriteAsync(string Path, ReadOnlyMemory<byte> Data)
+		public async Task WriteAsync(string Path, Stream InputStream)
 		{
 			TimeSpan[] RetryTimes =
 			{
@@ -171,14 +172,14 @@ namespace HordeServer.Storage.Impl
 			{
 				try
 				{
-					await WriteInternalAsync(Path, Data);
+					await WriteInternalAsync(Path, InputStream);
 					Logger.LogDebug("Written data to {Path}", Path);
 					break;
 				}
 				catch (Exception Ex)
 				{
 					Logger.LogError(Ex, "Unable to write data to {Path} ({Attempt}/{AttemptCount})", Path, Attempt + 1, RetryTimes.Length + 1);
-					if(Attempt >= RetryTimes.Length)
+					if (Attempt >= RetryTimes.Length)
 					{
 						throw new AwsException($"Unable to write to bucket {BucketName} path {Path}", Ex);
 					}
@@ -189,26 +190,43 @@ namespace HordeServer.Storage.Impl
 		}
 
 		/// <inheritdoc/>
-		public async Task WriteInternalAsync(string Path, ReadOnlyMemory<byte> Data)
+		public async Task WriteInternalAsync(string Path, Stream InputStream)
 		{
 			using (await Semaphore.UseWaitAsync())
 			{
-				using (ReadOnlyMemoryStream InputStream = new ReadOnlyMemoryStream(Data)) 
-				{
-					PutObjectRequest NewUploadRequest = new PutObjectRequest();
-					NewUploadRequest.BucketName = BucketName;
-					NewUploadRequest.Key = Path;
-					NewUploadRequest.InputStream = InputStream;
-					NewUploadRequest.Metadata.Add("bytes-written", Data.Length.ToString(CultureInfo.InvariantCulture));
-					await Client.PutObjectAsync(NewUploadRequest);
-				}
+				PutObjectRequest NewUploadRequest = new PutObjectRequest();
+				NewUploadRequest.BucketName = BucketName;
+				NewUploadRequest.Key = Path;
+				NewUploadRequest.InputStream = InputStream;
+				NewUploadRequest.Metadata.Add("bytes-written", InputStream.Length.ToString(CultureInfo.InvariantCulture));
+				await Client.PutObjectAsync(NewUploadRequest);
 			}
 		}
 
 		/// <inheritdoc/>
-		public Task DeleteAsync(string Path)
+		public async Task DeleteAsync(string Path)
+		{
+			DeleteObjectRequest NewDeleteRequest = new DeleteObjectRequest();
+			NewDeleteRequest.BucketName = BucketName;
+			NewDeleteRequest.Key = Path;
+			await Client.DeleteObjectAsync(NewDeleteRequest);
+		}
+
+		/// <inheritdoc/>
+		public async Task<bool> ExistsAsync(string Path)
 		{
-			throw new NotImplementedException();
+			try
+			{
+				GetObjectMetadataRequest Request = new GetObjectMetadataRequest();
+				Request.BucketName = BucketName;
+				Request.Key = Path;
+				await Client.GetObjectMetadataAsync(Request);
+				return true;
+			}
+			catch (AmazonS3Exception Ex) when (Ex.StatusCode == System.Net.HttpStatusCode.NotFound)
+			{
+				return false;
+			}
 		}
 	}
 }
diff --git a/Engine/Source/Programs/Horde/HordeServer/Storage/Backends/FileSystemStorageBackend.cs b/Engine/Source/Programs/Horde/HordeServer/Storage/Backends/FileSystemStorageBackend.cs
--- a/Engine/Source/Programs/Horde/HordeServer/Storage/Backends/FileSystemStorageBackend.cs
+++ b/Engine/Source/Programs/Horde/HordeServer/Storage/Backends/FileSystemStorageBackend.cs
@@ -1,19 +1,19 @@
 // Copyright Epic Games, Inc. All Rights Reserved.
 
 using EpicGames.Core;
+using HordeServer.Utilities;
+using Microsoft.Extensions.Logging;
 using Microsoft.Extensions.Options;
 using System;
 using System.Collections.Generic;
 using System.IO;
 using System.Linq;
+using System.Threading;
 using System.Threading.Tasks;
 
-namespace HordeServer.Storage.Impl
+namespace HordeServer.Storage.Backends
 {
-	/// <summary>
-	/// IStorageProvider implementation which uses the file system
-	/// </summary>
-	public sealed class FileSystemStorageBackend : IStorageBackend
+	class FileSystemStorageBackend : IStorageBackend
 	{
 		/// <summary>
 		/// Base directory for log files
@@ -21,100 +21,93 @@ namespace HordeServer.Storage.Impl
 		private readonly DirectoryReference BaseDir;
 
 		/// <summary>
-		/// Constructor
+		/// Unique identifier for this instance
 		/// </summary>
-		/// <param name="Settings">Current Horde Settings</param>
-		public FileSystemStorageBackend(IOptionsMonitor<ServerSettings> Settings)
-		{
-			ServerSettings CurrentSettings = Settings.CurrentValue;
-			BaseDir = CurrentSettings.LocalLogsDirRef;
-			DirectoryReference.CreateDirectory(BaseDir);
-		}
-
-		/// <inheritdoc/>
-		public void Dispose()
-		{
-		}
+		private string InstanceId;
 
 		/// <summary>
-		/// Tests if a given path exists
+		/// Unique id for each write
 		/// </summary>
-		/// <param name="Path">Path to check</param>
-		/// <returns>True if the path exists</returns>
-		public bool Exists(string Path)
-		{
-			FileReference Location = FileReference.Combine(BaseDir, Path);
-			return FileReference.Exists(Location);
-		}
+		private int UniqueId;
 
 		/// <summary>
-		/// Touches the given path
+		/// Constructor
 		/// </summary>
-		/// <param name="Path"></param>
-		/// <returns></returns>
-		public Task<bool> TouchAsync(string Path)
+		/// <param name="Settings">Current Horde Settings</param>
+		public FileSystemStorageBackend(IOptions<ServerSettings> Settings)
 		{
-			FileReference Location = FileReference.Combine(BaseDir, Path);
-			if (FileReference.Exists(Location))
-			{
-				try
-				{
-					FileReference.SetLastWriteTimeUtc(Location, DateTime.UtcNow);
-					return Task.FromResult(true);
-				}
-				catch(FileNotFoundException)
-				{
-					// Does not exist
-				}
-			}
-			return Task.FromResult(false);
+			ServerSettings CurrentSettings = Settings.Value;
+			this.BaseDir = CurrentSettings.LocalStorageDirRef;
+			this.InstanceId = Guid.NewGuid().ToString("N");
+			DirectoryReference.CreateDirectory(BaseDir);
 		}
 
 		/// <inheritdoc/>
-		public async Task<ReadOnlyMemory<byte>?> ReadAsync(string Path)
+		public Task<Stream?> ReadAsync(string Path)
 		{
 			FileReference Location = FileReference.Combine(BaseDir, Path);
 			if (!FileReference.Exists(Location))
 			{
-				return null;
+				return Task.FromResult<Stream?>(null);
 			}
 
 			try
 			{
-				return await FileReference.ReadAllBytesAsync(Location);
+				return Task.FromResult<Stream?>(FileReference.Open(Location, FileMode.Open, FileAccess.Read, FileShare.Read));
 			}
 			catch (DirectoryNotFoundException)
 			{
-				return null;
+				return Task.FromResult<Stream?>(null);
 			}
-			catch(FileNotFoundException)
+			catch (FileNotFoundException)
 			{
-				return null;
+				return Task.FromResult<Stream?>(null);
 			}
 		}
 
 		/// <inheritdoc/>
-		public async Task WriteAsync(string Path, ReadOnlyMemory<byte> Data)
+		public async Task WriteAsync(string Path, Stream Stream)
 		{
-			FileReference Location = FileReference.Combine(BaseDir, Path);
-			DirectoryReference.CreateDirectory(Location.Directory);
-
-			FileReference TempLocation = new FileReference(Location.FullName + ".tmp");
-			using (FileStream Stream = FileReference.Open(TempLocation, FileMode.Create, FileAccess.Write, FileShare.Read))
+			FileReference FinalLocation = FileReference.Combine(BaseDir, Path);
+			if (!FileReference.Exists(FinalLocation))
 			{
-				await Stream.WriteAsync(Data);
-			}
+				// Write to a temp file first
+				int NewUniqueId = Interlocked.Increment(ref UniqueId);
 
-			try
-			{
-				FileReference.Move(TempLocation, Location, true);
-			}
-			catch(IOException)
-			{
-				// Already exists
+				DirectoryReference.CreateDirectory(FinalLocation.Directory);
+				FileReference TempLocation = new FileReference($"{FinalLocation}.{InstanceId}.{UniqueId:x8}");
+
+				using (Stream OutputStream = FileReference.Open(TempLocation, FileMode.Create, FileAccess.Write, FileShare.Read))
+				{
+					await Stream.CopyToAsync(OutputStream);
+				}
+
+				// Move the temp file into place
+				try
+				{
+					FileReference.Move(TempLocation, FinalLocation, true);
+				}
+				catch (IOException) // Already exists
+				{
+					if (FileReference.Exists(FinalLocation))
+					{
+						FileReference.Delete(TempLocation);
+					}
+					else
+					{
+						throw;
+					}
+				}
 			}
 		}
 
+		/// <inheritdoc/>
+		public Task<bool> ExistsAsync(string Path)
+		{
+			FileReference Location = FileReference.Combine(BaseDir, Path);
+			return Task.FromResult(FileReference.Exists(Location));
+		}
+
 		/// <inheritdoc/>
 		public Task DeleteAsync(string Path)
 		{
diff --git a/Engine/Source/Programs/Horde/HordeServer/Storage/Backends/RelayStorageBackend.cs b/Engine/Source/Programs/Horde/HordeServer/Storage/Backends/RelayStorageBackend.cs
--- a/Engine/Source/Programs/Horde/HordeServer/Storage/Backends/RelayStorageBackend.cs
+++ b/Engine/Source/Programs/Horde/HordeServer/Storage/Backends/RelayStorageBackend.cs
@@ -10,12 +10,12 @@ using System.Net.Http;
 using System.Net.Http.Headers;
 using System.Threading.Tasks;
 
-namespace HordeServer.Storage.Impl
+namespace HordeServer.Storage.Backends
 {
 	/// <summary>
 	/// Implementation of ILogFileStorage which forwards requests to another server
 	/// </summary>
-	public sealed class RelayStorageBackend : IStorageBackend
+	public sealed class RelayStorageBackend : IStorageBackend, IDisposable
 	{
 		/// <summary>
 		/// The client to connect with
@@ -36,9 +36,9 @@ namespace HordeServer.Storage.Impl
 		/// Constructor
 		/// </summary>
 		/// <param name="Settings">Settings for the server instance</param>
-		public RelayStorageBackend(IOptionsMonitor<ServerSettings> Settings)
+		public RelayStorageBackend(IOptions<ServerSettings> Settings)
 		{
-			ServerSettings CurrentSettings = Settings.CurrentValue;
+			ServerSettings CurrentSettings = Settings.Value;
 			if (CurrentSettings.LogRelayServer == null)
 			{
 				throw new InvalidDataException("Missing LogRelayServer in server configuration");
@@ -60,19 +60,12 @@ namespace HordeServer.Storage.Impl
 		public void Dispose()
 		{
 			Client.Dispose();
-			LocalStorage.Dispose();
 		}
 
 		/// <inheritdoc/>
-		public Task<bool> TouchAsync(string Path)
+		public async Task<Stream?> ReadAsync(string Path)
 		{
-			throw new NotImplementedException();
-		}
-
-		/// <inheritdoc/>
-		public async Task<ReadOnlyMemory<byte>?> ReadAsync(string Path)
-		{
-			ReadOnlyMemory<byte>? LocalResult = await LocalStorage.ReadAsync(Path);
+			Stream? LocalResult = await LocalStorage.ReadAsync(Path);
 			if (LocalResult != null)
 			{
 				return LocalResult;
@@ -84,8 +77,8 @@ namespace HordeServer.Storage.Impl
 				if (Response.IsSuccessStatusCode)
 				{
 					byte[] ResponseData = await Response.Content.ReadAsByteArrayAsync();
-					await LocalStorage.WriteAsync(Path, ResponseData);
-					return ResponseData;
+					await LocalStorage.WriteBytesAsync(Path, ResponseData);
+					return new MemoryStream(ResponseData);
 				}
 			}
 
@@ -93,9 +86,15 @@ namespace HordeServer.Storage.Impl
 		}
 
 		/// <inheritdoc/>
-		public Task WriteAsync(string Path, ReadOnlyMemory<byte> Data)
+		public Task WriteAsync(string Path, Stream Stream)
+		{
+			return LocalStorage.WriteAsync(Path, Stream);
+		}
+
+		/// <inheritdoc/>
+		public Task<bool> ExistsAsync(string Path)
 		{
-			return LocalStorage.WriteAsync(Path, Data);
+			throw new NotSupportedException();
 		}
 
 		/// <inheritdoc/>
diff --git a/Engine/Source/Programs/Horde/HordeServer/Storage/Backends/TransientStorageBackend.cs b/Engine/Source/Programs/Horde/HordeServer/Storage/Backends/TransientStorageBackend.cs
--- a/Engine/Source/Programs/Horde/HordeServer/Storage/Backends/TransientStorageBackend.cs
+++ b/Engine/Source/Programs/Horde/HordeServer/Storage/Backends/TransientStorageBackend.cs
@@ -1,11 +1,13 @@
 // Copyright Epic Games, Inc. All Rights Reserved.
 
 using System;
+using System.Collections.Concurrent;
 using System.Collections.Generic;
+using System.IO;
 using System.Linq;
 using System.Threading.Tasks;
 
-namespace HordeServer.Storage.Impl
+namespace HordeServer.Storage.Backends
 {
 	/// <summary>
 	/// In-memory implementation of ILogFileStorage
@@ -15,44 +17,42 @@ namespace HordeServer.Storage.Impl
 		/// <summary>
 		/// Data storage
 		/// </summary>
-		Dictionary<string, byte[]> PathToData = new Dictionary<string, byte[]>();
+		ConcurrentDictionary<string, byte[]> PathToData = new ConcurrentDictionary<string, byte[]>();
 
 		/// <inheritdoc/>
-		public void Dispose()
-		{
-		}
-
-		/// <inheritdoc/>
-		public Task<bool> TouchAsync(string Path)
-		{
-			return Task.FromResult(PathToData.ContainsKey(Path));
-		}
-
-		/// <inheritdoc/>
-		public Task<ReadOnlyMemory<byte>?> ReadAsync(string Path)
+		public Task<Stream?> ReadAsync(string Path)
 		{
 			byte[]? Data;
 			if (PathToData.TryGetValue(Path, out Data))
 			{
-				return Task.FromResult<ReadOnlyMemory<byte>?>(Data);
+				return Task.FromResult<Stream?>(new MemoryStream(Data, false));
 			}
 			else
 			{
-				return Task.FromResult<ReadOnlyMemory<byte>?>(null);
+				return Task.FromResult<Stream?>(null);
 			}
 		}
 
 		/// <inheritdoc/>
-		public Task WriteAsync(string Path, ReadOnlyMemory<byte> Data)
+		public async Task WriteAsync(string Path, Stream Stream)
 		{
-			PathToData[Path] = Data.Span.ToArray();
-			return Task.CompletedTask;
+			using (MemoryStream Buffer = new MemoryStream())
+			{
+				await Stream.CopyToAsync(Buffer);
+				PathToData[Path] = Buffer.ToArray();
+			}
+		}
+
+		/// <inheritdoc/>
+		public Task<bool> ExistsAsync(string Path)
+		{
+			return Task.FromResult(PathToData.ContainsKey(Path));
 		}
 
 		/// <inheritdoc/>
 		public Task DeleteAsync(string Path)
 		{
-			PathToData.Remove(Path);
+			PathToData.TryRemove(Path, out _);
 			return Task.CompletedTask;
 		}
 	}
diff --git a/Engine/Source/Programs/Horde/HordeServer/Storage/Collections/BlobCollection.cs b/Engine/Source/Programs/Horde/HordeServer/Storage/Collections/BlobCollection.cs
--- a/Engine/Source/Programs/Horde/HordeServer/Storage/Collections/BlobCollection.cs
+++ b/Engine/Source/Programs/Horde/HordeServer/Storage/Collections/BlobCollection.cs
@@ -13,111 +13,49 @@ namespace HordeServer.Storage.Collections
 {
 	using NamespaceId = StringId<INamespace>;
 
-	class FileSystemBlobCollection : IBlobCollection
+	class BlobCollection : IBlobCollection
 	{
 		/// <summary>
-		/// Base directory for log files
+		/// The inner storage provider
 		/// </summary>
-		private readonly DirectoryReference BaseDir;
-
-		/// <summary>
-		/// Unique identifier for this instance
-		/// </summary>
-		private string InstanceId;
+		IStorageBackend StorageBackend;
 
 		/// <summary>
 		/// Constructor
 		/// </summary>
-		/// <param name="Settings">Current Horde Settings</param>
-		public FileSystemBlobCollection(IOptions<ServerSettings> Settings)
+		/// <param name="StorageBackend"></param>
+		public BlobCollection(IStorageBackend StorageBackend)
 		{
-			ServerSettings CurrentSettings = Settings.Value;
-			BaseDir = CurrentSettings.LocalBlobsDirRef;
-			InstanceId = Guid.NewGuid().ToString("N");
-			DirectoryReference.CreateDirectory(BaseDir);
+			this.StorageBackend = StorageBackend;
 		}
 
 		/// <inheritdoc/>
-		public async Task AddAsync(NamespaceId NamespaceId, IoHash Hash, Stream Input)
+		public Task<Stream?> ReadAsync(NamespaceId NamespaceId, IoHash Hash)
 		{
-			FileReference Location = GetPath(NamespaceId, Hash);
-			DirectoryReference.CreateDirectory(Location.Directory);
-
-			FileReference TempLocation = new FileReference($"{Location}.{InstanceId}");
-
-			Blake3.Hasher Hasher = Blake3.Hasher.New();
-			using (FileStream Output = FileReference.Open(TempLocation, FileMode.Create, FileAccess.Write, FileShare.Read))
-			{
-				byte[] ReadBuffer = new byte[64 * 1024];
-				Task<int> ReadTask = Task.Run(() => Input.ReadAsync(ReadBuffer, 0, ReadBuffer.Length));
-
-				byte[] WriteBuffer = new byte[64 * 1024];
-				int WriteLength = await ReadTask;
-
-				while (WriteLength > 0)
-				{
-					(WriteBuffer, ReadBuffer) = (ReadBuffer, WriteBuffer);
-
-					ReadTask = Task.Run(() => Input.ReadAsync(ReadBuffer, 0, ReadBuffer.Length));
-					Task WriteTask = Task.Run(() => Output.WriteAsync(WriteBuffer, 0, WriteLength));
-
-					Hasher.Update(WriteBuffer.AsSpan(0, WriteLength));
-
-					await Task.WhenAll(ReadTask, WriteTask);
-					WriteLength = await ReadTask;
-				}
-			}
-
-			IoHash ActualHash = new IoHash(Hasher);
-			if (ActualHash != Hash)
-			{
-				throw new ArgumentException($"Hash for blob is incorrect; should be {ActualHash}, not {Hash}");
-			}
+			string Path = GetPath(NamespaceId, Hash);
+			return StorageBackend.ReadAsync(Path);
+		}
 
-			try
-			{
-				FileReference.Move(TempLocation, Location, true);
-			}
-			catch (IOException) // Already exists
-			{
-				if (FileReference.Exists(Location))
-				{
-					FileReference.Delete(TempLocation);
-				}
-				else
-				{
-					throw;
-				}
-			}
+		/// <inheritdoc/>
+		public Task WriteAsync(NamespaceId NamespaceId, IoHash Hash, Stream Stream)
+		{
+			string Path = GetPath(NamespaceId, Hash);
+			return StorageBackend.WriteAsync(Path, Stream);
 		}
 
 		/// <inheritdoc/>
-		public Task<Stream?> GetAsync(NamespaceId NamespaceId, IoHash Hash)
+		public async Task<List<IoHash>> ExistsAsync(NamespaceId NamespaceId, IEnumerable<IoHash> Hashes)
 		{
-			FileReference Location = GetPath(NamespaceId, Hash);
-			if (!FileReference.Exists(Location))
+			List<Task<IoHash?>> Tasks = new List<Task<IoHash?>>();
+			foreach (IoHash Hash in Hashes)
 			{
-				return Task.FromResult<Stream?>(null);
+				string Path = GetPath(NamespaceId, Hash);
+				Tasks.Add(Task.Run(async () => await StorageBackend.ExistsAsync(Path)? (IoHash?)Hash : null));
 			}
 
-			try
-			{
-				return Task.FromResult<Stream?>(FileReference.Open(Location, FileMode.Open, FileAccess.Read, FileShare.Read));
-			}
-			catch (DirectoryNotFoundException)
-			{
-				return Task.FromResult<Stream?>(null);
-			}
-			catch (FileNotFoundException)
-			{
-				return Task.FromResult<Stream?>(null);
-			}
-		}
+			await Task.WhenAll(Tasks);
 
-		/// <inheritdoc/>
-		public Task<List<IoHash>> ExistsAsync(NamespaceId NamespaceId, IEnumerable<IoHash> Hashes)
-		{
-			return Task.FromResult<List<IoHash>>(Hashes.Where(x => FileReference.Exists(GetPath(NamespaceId, x))).ToList());
+			return Tasks.Where(x => x.Result != null).Select(x => (IoHash)x.Result!).ToList();
 		}
 
 		/// <summary>
@@ -126,10 +64,10 @@ namespace HordeServer.Storage.Collections
 		/// <param name="NsId">The namespace id</param>
 		/// <param name="Hash">Hash of the blob</param>
 		/// <returns></returns>
-		FileReference GetPath(NamespaceId NsId, IoHash Hash)
+		static string GetPath(NamespaceId NsId, IoHash Hash)
 		{
 			string HashText = Hash.ToString();
-			return FileReference.Combine(BaseDir, NsId.ToString(), HashText[0..2], HashText[2..4], $"{HashText}.blob");
+			return $"blobs/{NsId}/{HashText[0..2]}/{HashText[2..4]}/{HashText}.blob";
 		}
 	}
 }
diff --git a/Engine/Source/Programs/Horde/HordeServer/Storage/Collections/ObjectCollection.cs b/Engine/Source/Programs/Horde/HordeServer/Storage/Collections/ObjectCollection.cs
--- a/Engine/Source/Programs/Horde/HordeServer/Storage/Collections/ObjectCollection.cs
+++ b/Engine/Source/Programs/Horde/HordeServer/Storage/Collections/ObjectCollection.cs
@@ -30,16 +30,13 @@ namespace HordeServer.Storage.Collections
 		public async Task AddAsync(NamespaceId NsId, IoHash Hash, CbObject Obj)
 		{
 			ReadOnlyMemory<byte> Memory = Obj.GetView();
-			using (ReadOnlyMemoryStream Stream = new ReadOnlyMemoryStream(Memory))
-			{
-				await BlobCollection.AddAsync(NsId, Hash, Stream);
-			}
+			await BlobCollection.WriteBytesAsync(NsId, Hash, Memory);
 		}
 
 		/// <inheritdoc/>
 		public async Task<CbObject?> GetAsync(NamespaceId NsId, IoHash Hash)
 		{
-			byte[]? Data = await BlobCollection.GetByteArrayAsync(NsId, Hash);
+			byte[]? Data = await BlobCollection.ReadBytesAsync(NsId, Hash);
 			if(Data == null)
 			{
 				return null;
diff --git a/Engine/Source/Programs/Horde/HordeServer/Storage/Controllers/BlobsController.cs b/Engine/Source/Programs/Horde/HordeServer/Storage/Controllers/BlobsController.cs
--- a/Engine/Source/Programs/Horde/HordeServer/Storage/Controllers/BlobsController.cs
+++ b/Engine/Source/Programs/Horde/HordeServer/Storage/Controllers/BlobsController.cs
@@ -91,7 +91,7 @@ namespace HordeServer.Storage.Controllers
 
 			using (Stream Stream = File.OpenReadStream())
 			{
-				await BlobCollection.AddAsync(NamespaceId, Hash, Stream);
+				await BlobCollection.WriteAsync(NamespaceId, Hash, Stream);
 			}
 
 			return new PutBlobResponse { Hash = Hash };
@@ -136,7 +136,7 @@ namespace HordeServer.Storage.Controllers
 				return Forbid();
 			}
 
-			Stream? Stream = await BlobCollection.GetAsync(NamespaceId, Hash);
+			Stream? Stream = await BlobCollection.ReadAsync(NamespaceId, Hash);
 			if (Stream == null)
 			{
 				return NotFound();
diff --git a/Engine/Source/Programs/Horde/HordeServer/Storage/IBlobCollection.cs b/Engine/Source/Programs/Horde/HordeServer/Storage/IBlobCollection.cs
--- a/Engine/Source/Programs/Horde/HordeServer/Storage/IBlobCollection.cs
+++ b/Engine/Source/Programs/Horde/HordeServer/Storage/IBlobCollection.cs
@@ -27,20 +27,20 @@ namespace HordeServer.Storage
 	public interface IBlobCollection
 	{
 		/// <summary>
-		/// Adds an item to storage
+		/// Opens a blob read stream
 		/// </summary>
 		/// <param name="NamespaceId">Namespace to operate on</param>
 		/// <param name="Hash">Hash of the blob</param>
-		/// <param name="Blob">The blob stream data</param>
-		Task AddAsync(NamespaceId NamespaceId, IoHash Hash, Stream Blob);
+		/// <returns>Stream for the blob, or null if it does not exist</returns>
+		Task<Stream?> ReadAsync(NamespaceId NamespaceId, IoHash Hash);
 
 		/// <summary>
-		/// Opens a blob read stream
+		/// Adds an item to storage
 		/// </summary>
 		/// <param name="NamespaceId">Namespace to operate on</param>
 		/// <param name="Hash">Hash of the blob</param>
-		/// <returns>Stream for the blob, or null if it does not exist</returns>
-		Task<Stream?> GetAsync(NamespaceId NamespaceId, IoHash Hash);
+		/// <param name="Stream">The stream to write</param>
+		Task WriteAsync(NamespaceId NamespaceId, IoHash Hash, Stream Stream);
 
 		/// <summary>
 		/// Determines which of a set of blobs exist
@@ -57,31 +57,41 @@ namespace HordeServer.Storage
 	public static class BlobCollectionExtensions
 	{
 		/// <summary>
-		/// Adds a blob to the collection
+		/// Gets a blob as a byte array
 		/// </summary>
 		/// <param name="Collection"></param>
 		/// <param name="NamespaceId"></param>
-		/// <param name="Blob"></param>
+		/// <param name="Hash"></param>
 		/// <returns></returns>
-		public static async Task<IoHash> AddAsync(this IBlobCollection Collection, NamespaceId NamespaceId, ReadOnlyMemory<byte> Blob)
+		public static async Task<byte[]?> ReadBytesAsync(this IBlobCollection Collection, NamespaceId NamespaceId, IoHash Hash)
 		{
-			IoHash Hash = IoHash.Compute(Blob.Span);
-			await AddAsync(Collection, NamespaceId, Hash, Blob);
-			return Hash;
+			using (MemoryStream OutputStream = new MemoryStream())
+			{
+				using (Stream? InputStream = await Collection.ReadAsync(NamespaceId, Hash))
+				{
+					if(InputStream == null)
+					{
+						return null;
+					}
+
+					await InputStream.CopyToAsync(OutputStream);
+					return OutputStream.ToArray();
+				}
+			}
 		}
 
 		/// <summary>
-		/// Adds a blob to the collection
+		/// Gets a blob as a byte array
 		/// </summary>
 		/// <param name="Collection"></param>
 		/// <param name="NamespaceId"></param>
-		/// <param name="Hash"></param>
-		/// <param name="Blob"></param>
+		/// <param name="Data">The data to be written</param>
 		/// <returns></returns>
-		public static Task AddAsync(this IBlobCollection Collection, NamespaceId NamespaceId, IoHash Hash, ReadOnlyMemory<byte> Blob)
+		public static async Task<IoHash> WriteBytesAsync(this IBlobCollection Collection, NamespaceId NamespaceId, ReadOnlyMemory<byte> Data)
 		{
-			using ReadOnlyMemoryStream Stream = new ReadOnlyMemoryStream(Blob);
-			return Collection.AddAsync(NamespaceId, Hash, Stream);
+			IoHash Hash = IoHash.Compute(Data.Span);
+			await WriteBytesAsync(Collection, NamespaceId, Hash, Data);
+			return Hash;
 		}
 
 		/// <summary>
@@ -89,19 +99,15 @@ namespace HordeServer.Storage
 		/// </summary>
 		/// <param name="Collection"></param>
 		/// <param name="NamespaceId"></param>
-		/// <param name="Hash"></param>
+		/// <param name="Hash">Expected hash of the data</param>
+		/// <param name="Data">The data to be written</param>
 		/// <returns></returns>
-		public static async Task<byte[]?> GetByteArrayAsync(this IBlobCollection Collection, NamespaceId NamespaceId, IoHash Hash)
+		public static async Task WriteBytesAsync(this IBlobCollection Collection, NamespaceId NamespaceId, IoHash Hash, ReadOnlyMemory<byte> Data)
 		{
-			using Stream? Stream = await Collection.GetAsync(NamespaceId, Hash);
-			if(Stream == null)
+			using (ReadOnlyMemoryStream Stream = new ReadOnlyMemoryStream(Data))
 			{
-				return null;
+				await Collection.WriteAsync(NamespaceId, Hash, Stream);
 			}
-
-			byte[] Data = new byte[Stream.Length];
-			Stream.Read(Data);
-			return Data;
 		}
 
 		/// <summary>
diff --git a/Engine/Source/Programs/Horde/HordeServer/Storage/IStorageBackend.cs b/Engine/Source/Programs/Horde/HordeServer/Storage/IStorageBackend.cs
--- a/Engine/Source/Programs/Horde/HordeServer/Storage/IStorageBackend.cs
+++ b/Engine/Source/Programs/Horde/HordeServer/Storage/IStorageBackend.cs
@@ -1,44 +1,87 @@
 // Copyright Epic Games, Inc. All Rights Reserved.
 
+using EpicGames.Core;
 using System;
 using System.Collections.Generic;
+using System.IO;
 using System.Linq;
 using System.Threading.Tasks;
 
 namespace HordeServer.Storage
 {
 	/// <summary>
-	/// Interface for a log file storage implementation
+	/// Interface for a traditional location-addressed storage provider
 	/// </summary>
-	public interface IStorageBackend : IDisposable
+	public interface IStorageBackend
 	{
 		/// <summary>
-		/// Updates the timestamp on the given file. Touching the file ensures that it will not be deleted in a "reasonable" timeframe.
+		/// Opens a read stream for the given path.
 		/// </summary>
-		/// <param name="Path">Path to update</param>
-		/// <returns>True if the file exists, false otherwise</returns>
-		Task<bool> TouchAsync(string Path);
+		/// <param name="Path">Relative path within the bucket</param>
+		/// <returns></returns>
+		Task<Stream?> ReadAsync(string Path);
 
 		/// <summary>
-		/// Gets a chunk from a storage device
+		/// Writes a stream to the given path. If the stream throws an exception during read, the write will be aborted.
 		/// </summary>
-		/// <param name="Path">The path to read from</param>
-		/// <returns>The data to create from</returns>
-		Task<ReadOnlyMemory<byte>?> ReadAsync(string Path);
+		/// <param name="Path">Relative path within the bucket</param>
+		/// <param name="Stream">Stream to write</param>
+		Task WriteAsync(string Path, Stream Stream);
 
 		/// <summary>
-		/// Tries to update a log chunk after verifying it's the correct version
+		/// Tests whether the given path exists
 		/// </summary>
-		/// <param name="Path">The path to write to</param>
-		/// <param name="Data">Buffer containing the data to write</param>
-		/// <returns>Async task</returns>
-		Task WriteAsync(string Path, ReadOnlyMemory<byte> Data);
+		/// <param name="Path">Relative path within the bucket</param>
+		/// <returns></returns>
+		Task<bool> ExistsAsync(string Path);
 
 		/// <summary>
 		/// Deletes a file with the given path
 		/// </summary>
-		/// <param name="Path">Path of the file to delete</param>
+		/// <param name="Path">Relative path within the bucket</param>
 		/// <returns>Async task</returns>
 		Task DeleteAsync(string Path);
 	}
+
+	/// <summary>
+	/// Extension methods for <see cref="IStorageBackend"/>
+	/// </summary>
+	public static class StorageBackendExtensions
+	{
+		/// <summary>
+		/// Writes a block of memory to storage
+		/// </summary>
+		/// <param name="StorageBackend"></param>
+		/// <param name="Path"></param>
+		/// <returns></returns>
+		public static async Task<ReadOnlyMemory<byte>> ReadBytesAsync(this IStorageBackend StorageBackend, string Path)
+		{
+			using (Stream? InputStream = await StorageBackend.ReadAsync(Path))
+			{
+				using (MemoryStream OutputStream = new MemoryStream())
+				{
+					if (InputStream != null)
+					{
+						await InputStream.CopyToAsync(OutputStream);
+					}
+					return OutputStream.ToArray();
+				}
+			}
+		}
+
+		/// <summary>
+		/// Writes a block of memory to storage
+		/// </summary>
+		/// <param name="StorageBackend"></param>
+		/// <param name="Path"></param>
+		/// <param name="Data"></param>
+		/// <returns></returns>
+		public static async Task WriteBytesAsync(this IStorageBackend StorageBackend, string Path, ReadOnlyMemory<byte> Data)
+		{
+			using (ReadOnlyMemoryStream Stream = new ReadOnlyMemoryStream(Data))
+			{
+				await StorageBackend.WriteAsync(Path, Stream);
+			}
+		}
+	}
 }
diff --git a/Engine/Source/Programs/Horde/HordeServer/Storage/Services/SimpleStorageService.cs b/Engine/Source/Programs/Horde/HordeServer/Storage/Services/SimpleStorageService.cs
--- a/Engine/Source/Programs/Horde/HordeServer/Storage/Services/SimpleStorageService.cs
+++ b/Engine/Source/Programs/Horde/HordeServer/Storage/Services/SimpleStorageService.cs
@@ -1,7 +1,6 @@
 // Copyright Epic Games, Inc. All Rights Reserved.
 
 using EpicGames.Core;
-using HordeServer.Storage.Impl;
 using HordeServer.Storage.Primitives;
 using Microsoft.Extensions.Options;
 using System;
@@ -57,7 +56,7 @@ namespace HordeServer.Storage.Services
 		public async Task<ReadOnlyMemory<byte>?> TryGetBlobAsync(IoHash Hash, DateTime Deadline = default)
 		{
 			string Path = GetItemPath(BlobsPath, Hash);
-			return await Backend.ReadAsync(Path);
+			return await Backend.ReadBytesAsync(Path);
 		}
 
 		/// <inheritdoc/>
@@ -70,9 +69,9 @@ namespace HordeServer.Storage.Services
 			}
 
 			string Path = GetItemPath(BlobsPath, Hash);
-			if (!await Backend.TouchAsync(Path))
+			if (!await Backend.ExistsAsync(Path))
 			{
-				await Backend.WriteAsync(Path, Value);
+				await Backend.WriteBytesAsync(Path, Value);
 			}
 		}
 
@@ -93,21 +92,20 @@ namespace HordeServer.Storage.Services
 			}
 			else
 			{
-				await Backend.WriteAsync(Path, Value.Value.Memory);
+				await Backend.WriteBytesAsync(Path, Value.Value.Memory);
 			}
 		}
 
 		/// <inheritdoc/>
-		public async Task TouchRefAsync(IoHash Key)
+		public Task TouchRefAsync(IoHash Key)
 		{
-			string Path = GetItemPath(RefsPath, Key);
-			await Backend.TouchAsync(Path);
+			return Task.CompletedTask;
 		}
 
 		/// <inheritdoc/>
 		public async Task<IoHash?> GetRefAsync(IoHash Key, TimeSpan MaxDrift = default)
 		{
-			ReadOnlyMemory<byte>? Memory = await Backend.ReadAsync(GetItemPath(RefsPath, Key));
+			ReadOnlyMemory<byte>? Memory = await Backend.ReadBytesAsync(GetItemPath(RefsPath, Key));
 			if (Memory == null)
 			{
 				return null;
diff --git a/Engine/Source/Programs/Horde/HordeServerTests/IssueServiceTests.cs b/Engine/Source/Programs/Horde/HordeServerTests/IssueServiceTests.cs
--- a/Engine/Source/Programs/Horde/HordeServerTests/IssueServiceTests.cs
+++ b/Engine/Source/Programs/Horde/HordeServerTests/IssueServiceTests.cs
@@ -17,7 +17,6 @@ using HordeServer.Models;
 using HordeServer.Notifications;
 using HordeServer.Services;
 using HordeServer.Services.Impl;
-using HordeServer.Storage.Impl;
 using HordeServer.Utilities;
 using HordeServerTests.Stubs.Collections;
 using HordeServerTests.Stubs.Services;
diff --git a/Engine/Source/Programs/Horde/HordeServerTests/KeyValueTreeTests.cs b/Engine/Source/Programs/Horde/HordeServerTests/KeyValueTreeTests.cs
--- a/Engine/Source/Programs/Horde/HordeServerTests/KeyValueTreeTests.cs
+++ b/Engine/Source/Programs/Horde/HordeServerTests/KeyValueTreeTests.cs
@@ -2,7 +2,7 @@
 
 using EpicGames.Core;
 using HordeServer.Storage;
-using HordeServer.Storage.Impl;
+using HordeServer.Storage.Backends;
 using HordeServer.Storage.Primitives;
 using HordeServer.Storage.Services;
 using Microsoft.VisualStudio.TestTools.UnitTesting;
diff --git a/Engine/Source/Programs/Horde/HordeServerTests/LogFileServiceTests.cs b/Engine/Source/Programs/Horde/HordeServerTests/LogFileServiceTests.cs
--- a/Engine/Source/Programs/Horde/HordeServerTests/LogFileServiceTests.cs
+++ b/Engine/Source/Programs/Horde/HordeServerTests/LogFileServiceTests.cs
@@ -19,7 +19,7 @@ using HordeServer.Logs.Storage;
 using HordeServer.Logs.Storage.Impl;
 using HordeServer.Models;
 using HordeServer.Services;
-using HordeServer.Storage.Impl;
+using HordeServer.Storage.Backends;
 using HordeServer.Utilities;
 using Microsoft.Extensions.Configuration;
 using Microsoft.Extensions.Logging;
@@ -47,7 +47,7 @@ namespace HordeServerTests
             IConfiguration Config = new ConfigurationBuilder().Build();
 
 			ILogBuilder LogBuilder = new RedisLogBuilder(GetRedisConnectionPool(), NullLogger.Instance);
-			ILogStorage LogStorage = new PersistentLogStorage(new TransientStorageBackend(), NullLogger.Instance);
+			ILogStorage LogStorage = new PersistentLogStorage(new TransientStorageBackend(), NullLogger<PersistentLogStorage>.Instance);
 			LogFileService = new LogFileService(LogFileCollection, null!, LogBuilder, LogStorage, Logger);
         }
 
diff --git a/Engine/Source/Programs/Horde/HordeServerTests/StorageTests.cs b/Engine/Source/Programs/Horde/HordeServerTests/StorageTests.cs
--- a/Engine/Source/Programs/Horde/HordeServerTests/StorageTests.cs
+++ b/Engine/Source/Programs/Horde/HordeServerTests/StorageTests.cs
@@ -35,20 +35,20 @@ namespace HordeServerTests
 			byte[] TestData = Encoding.UTF8.GetBytes("Hello world");
 			IoHash Hash = IoHash.Compute(TestData);
 
-			Assert.IsNull(await BlobCollection.GetAsync(NamespaceId, Hash));
+			Assert.IsNull(await BlobCollection.ReadAsync(NamespaceId, Hash));
 			Assert.IsFalse(await BlobCollection.ExistsAsync(NamespaceId, Hash));
 
-			IoHash ReturnedHash = await BlobCollection.AddAsync(NamespaceId, TestData);
+			IoHash ReturnedHash = await BlobCollection.WriteBytesAsync(NamespaceId, TestData);
 			Assert.AreEqual(Hash, ReturnedHash);
 			Assert.IsTrue(await BlobCollection.ExistsAsync(NamespaceId, Hash));
 
-			byte[]? StoredData = await BlobCollection.GetByteArrayAsync(NamespaceId, Hash);
+			byte[]? StoredData = await BlobCollection.ReadBytesAsync(NamespaceId, Hash);
 			Assert.IsNotNull(StoredData);
 			Assert.IsTrue(TestData.AsSpan().SequenceEqual(StoredData));
 			Assert.IsTrue(await BlobCollection.ExistsAsync(NamespaceId, Hash));
 
 			NamespaceId OtherNamespaceId = new NamespaceId("other-ns");
-			Assert.IsNull(await BlobCollection.GetAsync(OtherNamespaceId, Hash));
+			Assert.IsNull(await BlobCollection.ReadAsync(OtherNamespaceId, Hash));
 			Assert.IsFalse(await BlobCollection.ExistsAsync(OtherNamespaceId, Hash));
 		}
 
@@ -137,9 +137,9 @@ namespace HordeServerTests
 			Assert.IsFalse(Ref!.Finalized);
 
 			// Add blobs A, B and C and check that the object can be finalized
-			await BlobCollection.AddAsync(NamespaceId, HashA, BlobA);
-			await BlobCollection.AddAsync(NamespaceId, HashB, BlobB);
-			await BlobCollection.AddAsync(NamespaceId, HashC, BlobC);
+			await BlobCollection.WriteBytesAsync(NamespaceId, HashA, BlobA);
+			await BlobCollection.WriteBytesAsync(NamespaceId, HashB, BlobB);
+			await BlobCollection.WriteBytesAsync(NamespaceId, HashC, BlobC);
 
 			MissingHashes = await RefCollection.FinalizeAsync(Ref);
 			Assert.AreEqual(0, MissingHashes.Count);
diff --git a/Engine/Source/Programs/Horde/HordeServerTests/TestSetup.cs b/Engine/Source/Programs/Horde/HordeServerTests/TestSetup.cs
--- a/Engine/Source/Programs/Horde/HordeServerTests/TestSetup.cs
+++ b/Engine/Source/Programs/Horde/HordeServerTests/TestSetup.cs
@@ -22,8 +22,8 @@ using HordeServer.Rpc;
 using HordeServer.Services;
 using HordeServer.Services.Impl;
 using HordeServer.Storage;
+using HordeServer.Storage.Backends;
 using HordeServer.Storage.Collections;
-using HordeServer.Storage.Impl;
 using HordeServer.Storage.Services;
 using HordeServer.Tasks.Impl;
 using HordeServer.Utilities;
@@ -132,9 +132,8 @@ namespace HordeServerTests
 			Settings.AdminClaimValue = "app-horde-admins";
 
 			DirectoryReference DataDir = DirectoryReference.Combine(BaseDir, Guid.NewGuid().ToString("N"));
-			Settings.LocalLogsDir = DirectoryReference.Combine(DataDir, "Logs").FullName;
 			Settings.LocalArtifactsDir = DirectoryReference.Combine(DataDir, "Artifacts").FullName;
-			Settings.LocalBlobsDir = DirectoryReference.Combine(DataDir, "Blobs").FullName;
+			Settings.LocalStorageDir = DataDir.FullName;
 		}
 
 		protected virtual void ConfigureServices(IServiceCollection Services)
@@ -211,14 +210,14 @@ namespace HordeServerTests
 			Services.AddSingleton<ContentStorageService>();
 			Services.AddSingleton<ExecutionService>();
 
-			Services.AddSingleton<IBlobCollection, FileSystemBlobCollection>();
+			Services.AddSingleton<IBlobCollection, BlobCollection>();
 			Services.AddSingleton<IObjectCollection, ObjectCollection>();
 			Services.AddSingleton<IRefCollection, RefCollection>();
 			Services.AddSingleton<INamespaceCollection, NamespaceCollection>();
 			Services.AddSingleton<IBucketCollection, BucketCollection>();
 
-			Services.AddSingleton<FileSystemStorageBackend>();
-			Services.AddSingleton<IStorageService, SimpleStorageService>(SP => new SimpleStorageService(SP.GetRequiredService<FileSystemStorageBackend>()));
+			Services.AddSingleton<IStorageBackend, FileSystemStorageBackend>();
+			Services.AddSingleton<IStorageService, SimpleStorageService>(SP => new SimpleStorageService(SP.GetRequiredService<IStorageBackend>()));
 
 			Services.AddSingleton<ISingletonDocument<GlobalPermissions>>(new SingletonDocumentStub<GlobalPermissions>());
 			Services.AddSingleton<ISingletonDocument<AgentSoftwareChannels>>(new SingletonDocumentStub<AgentSoftwareChannels>());
