Horde: Reconcile differences between storage implementations for artifact service vs log service.

Now uses a common implementation with generic specializations for dependency injection, similar to ILogger<T>.

diff --git a/Engine/Source/Programs/Horde/HordeServer/Controllers/ArtifactsController.cs b/Engine/Source/Programs/Horde/HordeServer/Controllers/ArtifactsController.cs
--- a/Engine/Source/Programs/Horde/HordeServer/Controllers/ArtifactsController.cs
+++ b/Engine/Source/Programs/Horde/HordeServer/Controllers/ArtifactsController.cs
@@ -232,7 +232,7 @@ namespace HordeServer.Controllers
 			}
 
 			// Fun, filestream result automatically closes the stream!
-			return new FileStreamResult(ArtifactService.OpenArtifactReadStream(Artifact), Artifact.MimeType);
+			return new FileStreamResult(await ArtifactService.OpenArtifactReadStreamAsync(Artifact), Artifact.MimeType);
 		}
 		
 		/// <summary>
@@ -261,7 +261,7 @@ namespace HordeServer.Controllers
 			}
 
 			Artifact Artifact = Artifacts[0];
-			return new FileStreamResult(ArtifactService.OpenArtifactReadStream(Artifact), Artifact.MimeType);
+			return new FileStreamResult(await ArtifactService.OpenArtifactReadStreamAsync(Artifact), Artifact.MimeType);
 		}
 
 		/// <summary>
@@ -334,7 +334,7 @@ namespace HordeServer.Controllers
 				return Forbid();
 			}
 
-			return new InlineFileStreamResult(ArtifactService.OpenArtifactReadStream(Artifact), Artifact.MimeType, Path.GetFileName(Artifact.Name));
+			return new InlineFileStreamResult(await ArtifactService.OpenArtifactReadStreamAsync(Artifact), Artifact.MimeType, Path.GetFileName(Artifact.Name));
 		}
 
 		/// <summary>
@@ -400,7 +400,7 @@ namespace HordeServer.Controllers
 					{
 						foreach (Artifact Artifact in ZipArtifacts)
 						{
-							await using (System.IO.Stream ArtifactStream = ArtifactService.OpenArtifactReadStream(Artifact))
+							await using (System.IO.Stream ArtifactStream = await ArtifactService.OpenArtifactReadStreamAsync(Artifact))
 							{
 								// tack on the step name into the directory if it exists
 								string StepName = string.Empty;
diff --git a/Engine/Source/Programs/Horde/HordeServer/Controllers/JobsController.cs b/Engine/Source/Programs/Horde/HordeServer/Controllers/JobsController.cs
--- a/Engine/Source/Programs/Horde/HordeServer/Controllers/JobsController.cs
+++ b/Engine/Source/Programs/Horde/HordeServer/Controllers/JobsController.cs
@@ -1192,7 +1192,7 @@ namespace HordeServer.Controllers
 			}
 
 			Artifact Artifact = Artifacts[0];
-			return new FileStreamResult(ArtifactService.OpenArtifactReadStream(Artifact), Artifact.MimeType);
+			return new FileStreamResult(await ArtifactService.OpenArtifactReadStreamAsync(Artifact), Artifact.MimeType);
 		}
 
 		/// <summary>
@@ -1229,7 +1229,7 @@ namespace HordeServer.Controllers
 			{
 				if (Artifact.Name.Equals("trace.json", StringComparison.OrdinalIgnoreCase))
 				{
-					return new FileStreamResult(ArtifactService.OpenArtifactReadStream(Artifact), "text/json");
+					return new FileStreamResult(await ArtifactService.OpenArtifactReadStreamAsync(Artifact), "text/json");
 				}
 			}
 			return NotFound();
diff --git a/Engine/Source/Programs/Horde/HordeServer/Logs/Storage/Impl/PersistentLogStorage.cs b/Engine/Source/Programs/Horde/HordeServer/Logs/Storage/Impl/PersistentLogStorage.cs
--- a/Engine/Source/Programs/Horde/HordeServer/Logs/Storage/Impl/PersistentLogStorage.cs
+++ b/Engine/Source/Programs/Horde/HordeServer/Logs/Storage/Impl/PersistentLogStorage.cs
@@ -2,6 +2,7 @@
 
 using EpicGames.Core;
 using HordeServer.Models;
+using HordeServer.Services;
 using HordeServer.Storage;
 using HordeServer.Utilities;
 using Microsoft.Extensions.Logging;
@@ -33,7 +34,7 @@ namespace HordeServer.Logs.Readers
 		/// </summary>
 		/// <param name="StorageProvider">The storage provider</param>
 		/// <param name="Logger">Logging provider</param>
-		public PersistentLogStorage(IStorageBackend StorageProvider, ILogger<PersistentLogStorage> Logger)
+		public PersistentLogStorage(IStorageBackend<PersistentLogStorage> StorageProvider, ILogger<PersistentLogStorage> Logger)
 		{
 			this.StorageProvider = StorageProvider;
 			this.Logger = Logger;
diff --git a/Engine/Source/Programs/Horde/HordeServer/ServerSettings.cs b/Engine/Source/Programs/Horde/HordeServer/ServerSettings.cs
--- a/Engine/Source/Programs/Horde/HordeServer/ServerSettings.cs
+++ b/Engine/Source/Programs/Horde/HordeServer/ServerSettings.cs
@@ -1,7 +1,12 @@
 // Copyright Epic Games, Inc. All Rights Reserved.
 
 using EpicGames.Core;
+using HordeServer.Logs.Readers;
+using HordeServer.Models;
 using HordeServer.Services;
+using HordeServer.Storage;
+using HordeServer.Storage.Backends;
+using HordeServer.Storage.Collections;
 using HordeServer.Utilities;
 using System;
 using System.Collections.Generic;
@@ -246,12 +251,7 @@ namespace HordeServer
 		/// <summary>
 		/// Local storage directory, if using type filesystem
 		/// </summary>
-		public string LocalStorageDir { get; set; } = "Storage";
-
-		/// <summary>
-		/// Gets the full path referred to by LocalStorageDir
-		/// </summary>
-		public DirectoryReference LocalStorageDirRef => DirectoryReference.Combine(Program.DataDir, LocalStorageDir);
+		public string LocalLogsDir { get; set; } = "Logs";
 
 		/// <summary>
 		/// Local artifact storage directory, if using type filesystem
@@ -259,9 +259,9 @@ namespace HordeServer
 		public string LocalArtifactsDir { get; set; } = "Artifacts";
 
 		/// <summary>
-		/// Gets the full path referred to by LocalStorageDir
+		/// Local artifact storage directory, if using type filesystem
 		/// </summary>
-		public DirectoryReference LocalArtifactsDirRef => DirectoryReference.Combine(Program.DataDir, LocalArtifactsDir);
+		public string LocalBlobsDir { get; set; } = "Blobs";
 
 		/// <summary>
 		/// S3 bucket region for logfile storage
@@ -298,6 +298,11 @@ namespace HordeServer
 		/// </summary>
 		public string S3ArtifactBucketName { get; set; } = null!;
 
+		/// <summary>
+		/// S3 blob bucket name
+		/// </summary>
+		public string S3BlobBucketName { get; set; } = null!;
+
 		/// <summary>
 		/// When using a relay storage provider, specifies the remote server to use
 		/// </summary>
diff --git a/Engine/Source/Programs/Horde/HordeServer/Services/ArtifactService.cs b/Engine/Source/Programs/Horde/HordeServer/Services/ArtifactService.cs
--- a/Engine/Source/Programs/Horde/HordeServer/Services/ArtifactService.cs
+++ b/Engine/Source/Programs/Horde/HordeServer/Services/ArtifactService.cs
@@ -6,6 +6,7 @@ using Amazon.S3;
 using Amazon.S3.Model;
 using HordeServer.Api;
 using HordeServer.Models;
+using HordeServer.Storage;
 using HordeServer.Utilities;
 using Microsoft.Extensions.Configuration;
 using Microsoft.Extensions.Hosting;
@@ -24,308 +25,15 @@ using System.Threading.Tasks;
 
 namespace HordeServer.Services
 {
-	/// <summary>
-	/// Interface for a artifact storage implementation
-	/// </summary>
-	public interface IArtifactStorage
-	{
-		/// <summary>
-		/// Gets an artifact from a storage device
-		/// </summary>
-		/// <param name="JobId">the job id</param>
-		/// <param name="StepId">the step id</param>
-		/// <param name="Name">The artifact name</param>
-		/// <returns>Stream to an artifact if found, null otherwise</returns>
-		System.IO.Stream OpenReadStream(ObjectId JobId, SubResourceId? StepId, string Name);
-
-		/// <summary>
-		/// Writes artifact data
-		/// </summary>
-		/// <param name="JobId">the job id</param>
-		/// <param name="StepId">the step id</param>
-		/// <param name="Name">The artifact name</param>
-		/// <param name="Data">The data to write</param>
-		/// <returns>Whether or not the update completed</returns>
-		Task WriteAsync(ObjectId JobId, SubResourceId? StepId, string Name, System.IO.Stream Data);
-	}
-
-	/// <summary>
-	/// FileStorage implementation using the file system
-	/// </summary>
-	public sealed class FSExternalArtifactStorage : IArtifactStorage
-	{
-		/// <summary>
-		/// Base directory for artifacts
-		/// </summary>
-		private readonly DirectoryInfo BaseDir;
-
-		/// <summary>
-		/// Information Logger
-		/// </summary>
-		private readonly ILogger<ArtifactService> Logger;
-
-		/// <summary>
-		/// Constructor
-		/// </summary>
-		/// <param name="CurrentSettings">Current Horde Settings</param>
-		/// <param name="Logger">Logger interface</param>
-		public FSExternalArtifactStorage(ServerSettings CurrentSettings, ILogger<ArtifactService> Logger)
-		{
-			this.Logger = Logger;
-			this.BaseDir = new DirectoryInfo(CurrentSettings.LocalArtifactsDirRef.FullName);
-
-			Directory.CreateDirectory(this.BaseDir.FullName);
-		}
-
-		/// <inheritdoc/>
-		public System.IO.Stream OpenReadStream(ObjectId JobId, SubResourceId? StepId, string Name)
-		{
-			string FilePath = Path.Combine(BaseDir.ToString(), JobId.ToString(), StepId.ToString() ?? string.Empty, Name);
-			FileInfo ArtifactFile = new FileInfo(FilePath);
-			try
-			{
-				return new FileStream(ArtifactFile.FullName, FileMode.Open, FileAccess.Read, FileShare.ReadWrite, 4096, true);
-			}
-			catch (FileNotFoundException)
-			{
-				Logger.LogError("Tried to find artifact that didn't exist at {0}!", ArtifactFile.FullName);
-				throw;
-			}
-			catch (DirectoryNotFoundException)
-			{
-				Logger.LogError("Tried to find artifact that didn't exist at {0}!", ArtifactFile.FullName);
-				throw;
-			}
-		}
-
-		/// <inheritdoc/>
-		public async Task WriteAsync(ObjectId JobId, SubResourceId? StepId, string Name, System.IO.Stream Data)
-		{
-			string FilePath = Path.Combine(BaseDir.FullName, JobId.ToString(), StepId.ToString() ?? string.Empty, Name);
-			FileInfo ArtifactFile = new FileInfo(FilePath);
-			Directory.CreateDirectory(ArtifactFile.Directory.FullName);
-			using (FileStream FileWriter = File.Open(ArtifactFile.FullName, FileMode.Create, FileAccess.Write, FileShare.Read))
-			{
-				await Data.CopyToAsync(FileWriter);
-			}
-		}
-	}
-
-	/// <summary>
-	/// FileStorage implementation using an s3 bucket
-	/// </summary>
-	public sealed class S3ExternalArtifactStorage : IArtifactStorage
-	{
-		/// <summary>
-		/// S3 Client
-		/// </summary>
-		private readonly IAmazonS3 Client;
-
-		/// <summary>
-		/// S3 bucket name
-		/// </summary>
-		private readonly string BucketName;
-
-		/// <summary>
-		/// Information Logger
-		/// </summary>
-		private readonly ILogger<ArtifactService> Logger;
-
-		/// <summary>
-		/// Constructor
-		/// </summary>
-		/// <param name="CurrentSettings">Current Horde Settings</param>
-		/// <param name="Configuration">Configuration</param>
-		/// <param name="Logger">Logger interface</param>
-		public S3ExternalArtifactStorage(ServerSettings CurrentSettings, IConfiguration Configuration, ILogger<ArtifactService> Logger)
-		{
-			this.Logger = Logger;
-			AWSOptions Options = Configuration.GetAWSOptions();
-			Options.Credentials = GetCredentials(CurrentSettings);
-
-			this.Client = Options.CreateServiceClient<IAmazonS3>();
-			this.BucketName = CurrentSettings.S3ArtifactBucketName;
-		}
-
-		enum AWSCredentialsType
-		{
-			Basic,
-			AssumeRole,
-			AssumeRoleWebIdentity
-		}
-
-		/// <summary>
-		/// Gets auth credentials based on type
-		/// </summary>
-		/// <param name="Settings">Horde settings</param>
-		/// <returns></returns>
-		public AWSCredentials GetCredentials(ServerSettings Settings)
-		{
-			AWSCredentialsType AuthType;
-			if (!Enum.TryParse(Settings.S3CredentialType, true, out AuthType))
-			{
-				Logger.LogError("Could not determine auth type from appsettings. Should be Basic, AssumeRole, or AssumeRoleWebIdentity!");
-				return null!;
-			}
-			switch (AuthType)
-			{
-				case AWSCredentialsType.Basic:
-					return new BasicAWSCredentials(Settings.S3ClientKeyId, Settings.S3ClientSecret);
-				case AWSCredentialsType.AssumeRole:
-					return new AssumeRoleAWSCredentials(FallbackCredentialsFactory.GetCredentials(), Settings.S3AssumeArn, "Horde");
-				case AWSCredentialsType.AssumeRoleWebIdentity:
-					return AssumeRoleWithWebIdentityCredentials.FromEnvironmentVariables();
-				default:
-					return null!;
-			}
-		}
-
-		/// <inheritdoc/>
-		public System.IO.Stream OpenReadStream(ObjectId JobId, SubResourceId? StepId, string Name)
-		{
-			string BucketPath = Path.Combine(JobId.ToString(), StepId.ToString() ?? string.Empty, Name);
-			try
-			{
-				GetObjectRequest NewGetRequest = new GetObjectRequest
-				{
-					BucketName = this.BucketName,
-					Key = BucketPath
-				};
-
-				return this.Client.GetObjectAsync(NewGetRequest).Result.ResponseStream;
-			}
-			catch (AmazonS3Exception Ex)
-			{
-				if (Ex.ErrorCode == "NoSuchKey")
-				{
-					Logger.LogError("Tried to find artifact that didn't exist at {0}!", BucketPath);
-				}
-				throw;
-			}
-		}
-
-		/// <inheritdoc/>
-		public async Task WriteAsync(ObjectId JobId, SubResourceId? StepId, string Name, System.IO.Stream Stream)
-		{
-			const int MinPartSize = 5 * 1024 * 1024;
-
-			string Key = Path.Combine(JobId.ToString(), StepId.ToString() ?? string.Empty, Name);
-
-			long StreamLen = Stream.Length;
-			if (StreamLen < MinPartSize)
-			{
-				// Read the data into memory
-				byte[] Buffer = new byte[StreamLen];
-				await ReadExactLengthAsync(Stream, Buffer, (int)StreamLen);
-
-				// Upload it to S3
-				using (MemoryStream InputStream = new MemoryStream(Buffer))
-				{
-					PutObjectRequest UploadRequest = new PutObjectRequest();
-					UploadRequest.BucketName = BucketName;
-					UploadRequest.Key = Key;
-					UploadRequest.InputStream = InputStream;
-					await Client.PutObjectAsync(UploadRequest);
-				}
-			}
-			else
-			{
-				// Initiate a multi-part upload
-				InitiateMultipartUploadRequest InitiateRequest = new InitiateMultipartUploadRequest();
-				InitiateRequest.BucketName = BucketName;
-				InitiateRequest.Key = Key;
-
-				InitiateMultipartUploadResponse InitiateResponse = await Client.InitiateMultipartUploadAsync(InitiateRequest);
-				try
-				{
-					// Buffer for reading the data
-					byte[] Buffer = new byte[MinPartSize];
-
-					// Upload all the parts
-					List<PartETag> PartTags = new List<PartETag>();
-					for(long StreamPos = 0; StreamPos < StreamLen; )
-					{
-						// Read the next chunk of data into the buffer
-						int BufferLen = (int)Math.Min((long)MinPartSize, StreamLen - StreamPos);
-						await ReadExactLengthAsync(Stream, Buffer, BufferLen);
-						StreamPos += BufferLen;
-
-						// Upload the part
-						using (MemoryStream InputStream = new MemoryStream(Buffer, 0, BufferLen))
-						{
-							UploadPartRequest PartRequest = new UploadPartRequest();
-							PartRequest.BucketName = BucketName;
-							PartRequest.Key = Key;
-							PartRequest.UploadId = InitiateResponse.UploadId;
-							PartRequest.InputStream = InputStream;
-							PartRequest.PartSize = BufferLen;
-							PartRequest.PartNumber = PartTags.Count + 1;
-							PartRequest.IsLastPart = (StreamPos == StreamLen);
-
-							UploadPartResponse PartResponse = await Client.UploadPartAsync(PartRequest);
-							PartTags.Add(new PartETag(PartResponse.PartNumber, PartResponse.ETag));
-						}
-					}
-
-					// Mark the upload as complete
-					CompleteMultipartUploadRequest CompleteRequest = new CompleteMultipartUploadRequest();
-					CompleteRequest.BucketName = BucketName;
-					CompleteRequest.Key = Key;
-					CompleteRequest.UploadId = InitiateResponse.UploadId;
-					CompleteRequest.PartETags = PartTags;
-					await Client.CompleteMultipartUploadAsync(CompleteRequest);
-				}
-				catch
-				{
-					// Abort the upload
-					AbortMultipartUploadRequest AbortRequest = new AbortMultipartUploadRequest();
-					AbortRequest.BucketName = BucketName;
-					AbortRequest.Key = Key;
-					AbortRequest.UploadId = InitiateResponse.UploadId;
-					await Client.AbortMultipartUploadAsync(AbortRequest);
-
-					throw;
-				}
-			}
-		}
-
-		/// <summary>
-		/// Reads data of an exact length into a stream
-		/// </summary>
-		/// <param name="Stream">The stream to read from</param>
-		/// <param name="Buffer">The buffer to read into</param>
-		/// <param name="Length">Length of the data to read</param>
-		/// <returns>Async task</returns>
-		static async Task ReadExactLengthAsync(System.IO.Stream Stream, byte[] Buffer, int Length)
-		{
-			int BufferPos = 0;
-			while (BufferPos < Length)
-			{
-				int BytesRead = await Stream.ReadAsync(Buffer, BufferPos, Length - BufferPos);
-				if (BytesRead == 0)
-				{
-					throw new InvalidOperationException("Unexpected end of stream");
-				}
-				BufferPos += BytesRead;
-			}
-		}
-	}
-
 	/// <summary>
 	/// Wraps functionality for manipulating artifacts
 	/// </summary>
 	public class ArtifactService
 	{
-		private enum ProviderTypes
-		{
-			S3,
-			FileSystem
-		};
-
 		/// <summary>
 		/// Backend storage container
 		/// </summary>
-		private readonly IArtifactStorage ExternalStorageProvider;
+		private readonly IStorageBackend StorageBackend;
 
 		/// <summary>
 		/// Information Logger
@@ -341,29 +49,13 @@ namespace HordeServer.Services
 		/// Constructor
 		/// </summary>
 		/// <param name="DatabaseService">The database service</param>
-		/// <param name="Configuration">The configuration</param>
-		/// <param name="Settings">Connection settings</param>
+		/// <param name="StorageBackend">The storage backend</param>
 		/// <param name="Logger">Log interface</param>
-		public ArtifactService(DatabaseService DatabaseService, IOptionsMonitor<ServerSettings> Settings, IConfiguration Configuration, ILogger<ArtifactService> Logger)
+		public ArtifactService(DatabaseService DatabaseService, IStorageBackend<ArtifactService> StorageBackend, ILogger<ArtifactService> Logger)
 		{
+			this.StorageBackend = StorageBackend;
 			this.Logger = Logger;
 
-			// Initialize Cache provider
-			ServerSettings CurrentSettings = Settings.CurrentValue;
-			switch (CurrentSettings.ExternalStorageProviderType)
-			{
-				case StorageProviderType.FileSystem:
-					this.ExternalStorageProvider = new FSExternalArtifactStorage(CurrentSettings, Logger);
-					break;
-				case StorageProviderType.S3:
-					this.ExternalStorageProvider = new S3ExternalArtifactStorage(CurrentSettings, Configuration, Logger);
-					break;
-				default:
-					// shouldn't be possible to get here, but default to file system worst case.
-					this.ExternalStorageProvider = new FSExternalArtifactStorage(CurrentSettings, Logger);
-					break;
-			}
-
 			// Initialize Artifacts table
 			Artifacts = DatabaseService.Artifacts;
 			if (!DatabaseService.ReadOnlyMode)
@@ -385,7 +77,7 @@ namespace HordeServer.Services
 		{
 			// upload first
 			string ArtifactName = ValidateName(Name);
-			await ExternalStorageProvider.WriteAsync(JobId, StepId, ArtifactName, Data);
+			await StorageBackend.WriteAsync(GetPath(JobId, StepId, ArtifactName), Data);
 
 			// then create entry
 			Artifact NewArtifact = new Artifact(JobId, StepId, ArtifactName, Data.Length, MimeType);
@@ -477,7 +169,7 @@ namespace HordeServer.Services
 
 				// re-upload the data to external
 				string ArtifactName = ValidateName(Artifact.Name);
-				await ExternalStorageProvider.WriteAsync(Artifact.JobId, Artifact.StepId, ArtifactName, NewData);
+				await StorageBackend.WriteAsync(GetPath(Artifact.JobId, Artifact.StepId, ArtifactName), NewData);
 
 				if (await TryUpdateArtifactAsync(Artifact, UpdateBuilder.Combine(Updates)))
 				{
@@ -494,9 +186,33 @@ namespace HordeServer.Services
 		/// </summary>
 		/// <param name="Artifact">The artifact</param>
 		/// <returns>The chunk data</returns>
-		public System.IO.Stream OpenArtifactReadStream(Artifact Artifact)
+		public async Task<System.IO.Stream> OpenArtifactReadStreamAsync(Artifact Artifact)
+		{
+			System.IO.Stream? Stream = await StorageBackend.ReadAsync(GetPath(Artifact.JobId, Artifact.StepId, Artifact.Name));
+			if (Stream == null)
+			{
+				throw new Exception($"Unable to get artifact {Artifact.Id}");
+			}
+			return Stream;
+		}
+
+		/// <summary>
+		/// Get the path for an artifact
+		/// </summary>
+		/// <param name="JobId"></param>
+		/// <param name="StepId"></param>
+		/// <param name="Name"></param>
+		/// <returns></returns>
+		private static string GetPath(ObjectId JobId, SubResourceId? StepId, string Name)
 		{
-			return ExternalStorageProvider.OpenReadStream(Artifact.JobId, Artifact.StepId, Artifact.Name);
+			if (StepId == null)
+			{
+				return $"{JobId}/{Name}";
+			}
+			else
+			{
+				return $"{JobId}/{StepId.Value}/{Name}";
+			}
 		}
 
 		/// <summary>
diff --git a/Engine/Source/Programs/Horde/HordeServer/Startup.cs b/Engine/Source/Programs/Horde/HordeServer/Startup.cs
--- a/Engine/Source/Programs/Horde/HordeServer/Startup.cs
+++ b/Engine/Source/Programs/Horde/HordeServer/Startup.cs
@@ -359,7 +359,7 @@ namespace HordeServer
 						StatsdServerName = DatadogAgentHost,
 						StatsdPort = 8125,
 					};
-					
+
 					DogStatsdService DogStatsdService = new DogStatsdService();
 					DogStatsdService.Configure(Config);
 					return DogStatsdService;
@@ -384,7 +384,7 @@ namespace HordeServer
 			else if (!String.IsNullOrEmpty(Settings.PerforceBridge))
 			{
 				// Python P4 bridge service
-				Services.AddSingleton<IPerforceService, BridgePerforceService>(SP => new BridgePerforceService(Settings.PerforceBridge!, SP.GetRequiredService<ILogger<BridgePerforceService>>()));				
+				Services.AddSingleton<IPerforceService, BridgePerforceService>(SP => new BridgePerforceService(Settings.PerforceBridge!, SP.GetRequiredService<ILogger<BridgePerforceService>>()));
 			}
 			else if (!String.IsNullOrEmpty(Settings.P4BridgeServer))
 			{
@@ -416,14 +416,18 @@ namespace HordeServer
 
 			Services.AddSingleton<DeviceService>();
 
-			AWSOptions Options = Configuration.GetAWSOptions();
-			if (Settings.S3CredentialType == "AssumeRole" && !String.IsNullOrEmpty(Settings.S3AssumeArn))
+			AWSOptions AwsOptions = Configuration.GetAWSOptions();
+			if (Settings.S3AssumeArn != null)
 			{
-				Options.Credentials = new AssumeRoleAWSCredentials(FallbackCredentialsFactory.GetCredentials(), Settings.S3AssumeArn, "Horde");
+				AwsOptions.Credentials = new AssumeRoleAWSCredentials(FallbackCredentialsFactory.GetCredentials(), Settings.S3AssumeArn, "Horde");
 			}
-			Services.AddDefaultAWSOptions(Options);
+			Services.AddSingleton(AwsOptions);
+
+			Services.AddSingleton(new StorageBackendSettings<PersistentLogStorage> { Type = Settings.ExternalStorageProviderType, BaseDir = Settings.LocalLogsDir, BucketName = Settings.S3LogBucketName });
+			Services.AddSingleton(new StorageBackendSettings<ArtifactService> { Type = Settings.ExternalStorageProviderType, BaseDir = Settings.LocalArtifactsDir, BucketName = Settings.S3ArtifactBucketName });
+			Services.AddSingleton(new StorageBackendSettings<BlobCollection> { Type = Settings.ExternalStorageProviderType, BaseDir = Settings.LocalBlobsDir, BucketName = Settings.S3LogBucketName, BucketPath = "blobs/" });
+			Services.AddSingleton(typeof(IStorageBackend<>), typeof(StorageBackendFactory<>));
 
-			ConfigureStorageBackend(Services, Settings);
 			ConfigureLogStorage(Services);
 			Services.AddSingleton<IStorageService, SimpleStorageService>();
 			//			ConfigureLogFileWriteCache(Services, Settings);
@@ -589,25 +593,59 @@ namespace HordeServer
 			Options.Converters.Add(new JsonDateTimeConverter());
 		}
 
-		private static void ConfigureStorageBackend(IServiceCollection Services, ServerSettings Settings)
+		public class StorageBackendSettings<T> : IFileSystemStorageOptions, IAwsStorageOptions
+		{
+			/// <summary>
+			/// The type of storage backend to use
+			/// </summary>
+			public StorageProviderType Type { get; set; }
+
+			/// <inheritdoc/>
+			public string? BaseDir { get; set; }
+
+			/// <inheritdoc/>
+			public string? BucketName { get; set; }
+
+			/// <inheritdoc/>
+			public string? BucketPath { get; set; }
+		}
+
+		public class StorageBackendFactory<T> : IStorageBackend<T>
 		{
-			switch (Settings.ExternalStorageProviderType)
+			IStorageBackend Inner;
+
+			public StorageBackendFactory(IServiceProvider ServiceProvider, StorageBackendSettings<T> Options)
 			{
-				case StorageProviderType.FileSystem:
-					Services.AddSingleton<IStorageBackend, FileSystemStorageBackend>();
-					break;
-				case StorageProviderType.S3:
-					Services.AddSingleton<IStorageBackend, AwsStorageBackend>();
-					break;
-				case StorageProviderType.Transient:
-					Services.AddSingleton<IStorageBackend, TransientStorageBackend>();
-					break;
-				case StorageProviderType.Relay:
-					Services.AddSingleton<IStorageBackend, RelayStorageBackend>();
-					break;
-				default:
-					throw new NotImplementedException();
+				switch (Options.Type)
+				{
+					case StorageProviderType.S3:
+						Inner = new AwsStorageBackend(ServiceProvider.GetRequiredService<AWSOptions>(), Options, ServiceProvider.GetRequiredService<ILogger<AwsStorageBackend>>());
+						break;
+					case StorageProviderType.FileSystem:
+						Inner = new FileSystemStorageBackend(Options);
+						break;
+					case StorageProviderType.Transient:
+						Inner = new TransientStorageBackend();
+						break;
+					case StorageProviderType.Relay:
+						Inner = new RelayStorageBackend(ServiceProvider.GetRequiredService<IOptions<ServerSettings>>());
+						break;
+					default:
+						throw new NotImplementedException();
+				}
 			}
+
+			/// <inheritdoc/>
+			public Task<Stream?> ReadAsync(string Path) => Inner.ReadAsync(Path);
+
+			/// <inheritdoc/>
+			public Task WriteAsync(string Path, Stream Stream) => Inner.WriteAsync(Path, Stream);
+
+			/// <inheritdoc/>
+			public Task DeleteAsync(string Path) => Inner.DeleteAsync(Path);
+
+			/// <inheritdoc/>
+			public Task<bool> ExistsAsync(string Path) => Inner.ExistsAsync(Path);
 		}
 
 		private static void ConfigureLogStorage(IServiceCollection Services)
diff --git a/Engine/Source/Programs/Horde/HordeServer/Storage/Backends/AwsStorageBackend.cs b/Engine/Source/Programs/Horde/HordeServer/Storage/Backends/AwsStorageBackend.cs
--- a/Engine/Source/Programs/Horde/HordeServer/Storage/Backends/AwsStorageBackend.cs
+++ b/Engine/Source/Programs/Horde/HordeServer/Storage/Backends/AwsStorageBackend.cs
@@ -9,9 +9,11 @@ using Microsoft.Extensions.Configuration;
 using Microsoft.Extensions.Logging;
 using Microsoft.Extensions.Options;
 using System;
+using System.Collections.Generic;
 using System.Diagnostics.CodeAnalysis;
 using System.Globalization;
 using System.IO;
+using System.Text;
 using System.Threading;
 using System.Threading.Tasks;
 
@@ -34,6 +36,22 @@ namespace HordeServer.Storage.Backends
 		}
 	}
 
+	/// <summary>
+	/// Options for AWS
+	/// </summary>
+	public interface IAwsStorageOptions
+	{
+		/// <summary>
+		/// Name of the bucket to use
+		/// </summary>
+		public string? BucketName { get; }
+
+		/// <summary>
+		/// Base path within the bucket 
+		/// </summary>
+		public string? BucketPath { get; }
+	}
+
 	/// <summary>
 	/// FileStorage implementation using an s3 bucket
 	/// </summary>
@@ -45,9 +63,9 @@ namespace HordeServer.Storage.Backends
 		private readonly IAmazonS3 Client;
 
 		/// <summary>
-		/// S3 bucket name
+		/// Options for AWs
 		/// </summary>
-		private readonly string BucketName;
+		private IAwsStorageOptions Options;
 
 		/// <summary>
 		/// Semaphore for connecting to AWS
@@ -62,14 +80,13 @@ namespace HordeServer.Storage.Backends
 		/// <summary>
 		/// Constructor
 		/// </summary>
-		/// <param name="AwsOptions">The AWS options</param>
-		/// <param name="ServerSettings">Options for configuring AWS</param>
+		/// <param name="AwsOptions">AWS options</param>
+		/// <param name="Options">Storage options</param>
 		/// <param name="Logger">Logger interface</param>
-		public AwsStorageBackend(AWSOptions AwsOptions, IOptions<ServerSettings> ServerSettings, ILogger<AwsStorageBackend> Logger)
+		public AwsStorageBackend(AWSOptions AwsOptions, IAwsStorageOptions Options, ILogger<AwsStorageBackend> Logger)
 		{
-			ServerSettings OptionsValue = ServerSettings.Value;
 			this.Client = AwsOptions.CreateServiceClient<IAmazonS3>();
-			this.BucketName = OptionsValue.S3LogBucketName;
+			this.Options = Options;
 			this.Semaphore = new SemaphoreSlim(16);
 			this.Logger = Logger;
 		}
@@ -129,9 +146,28 @@ namespace HordeServer.Storage.Backends
 			}
 		}
 
+		string GetFullPath(string Path)
+		{
+			if (Options.BucketPath == null)
+			{
+				return Path;
+			}
+
+			StringBuilder Result = new StringBuilder();
+			Result.Append(Options.BucketPath);
+			if (Result.Length > 0 && Result[Result.Length - 1] != '/')
+			{
+				Result.Append('/');
+			}
+			Result.Append(Path);
+			return Result.ToString();
+		}
+
 		/// <inheritdoc/>
 		public async Task<Stream?> ReadAsync(string Path)
 		{
+			string FullPath = GetFullPath(Path);
+
 			IDisposable? Lock = null;
 			GetObjectResponse? Response = null;
 			try
@@ -139,8 +175,8 @@ namespace HordeServer.Storage.Backends
 				Lock = await Semaphore.UseWaitAsync();
 
 				GetObjectRequest NewGetRequest = new GetObjectRequest();
-				NewGetRequest.BucketName = BucketName;
-				NewGetRequest.Key = Path;
+				NewGetRequest.BucketName = Options.BucketName;
+				NewGetRequest.Key = FullPath;
 
 				Response = await Client.GetObjectAsync(NewGetRequest);
 
@@ -148,7 +184,7 @@ namespace HordeServer.Storage.Backends
 			}
 			catch (Exception Ex)
 			{
-				Logger.LogWarning(Ex, "Unable to read {Path} from S3", Path);
+				Logger.LogWarning(Ex, "Unable to read {Path} from S3", FullPath);
 
 				Lock?.Dispose();
 				Response?.Dispose();
@@ -168,20 +204,22 @@ namespace HordeServer.Storage.Backends
 				TimeSpan.FromSeconds(10.0),
 			};
 
+			string FullPath = GetFullPath(Path);
 			for (int Attempt = 0; ; Attempt++)
 			{
 				try
 				{
-					await WriteInternalAsync(Path, InputStream);
+					using IDisposable Lock = await Semaphore.UseWaitAsync();
+					await WriteInternalAsync(FullPath, InputStream);
 					Logger.LogDebug("Written data to {Path}", Path);
 					break;
 				}
 				catch (Exception Ex)
 				{
-					Logger.LogError(Ex, "Unable to write data to {Path} ({Attempt}/{AttemptCount})", Path, Attempt + 1, RetryTimes.Length + 1);
+					Logger.LogError(Ex, "Unable to write data to {Path} ({Attempt}/{AttemptCount})", FullPath, Attempt + 1, RetryTimes.Length + 1);
 					if (Attempt >= RetryTimes.Length)
 					{
-						throw new AwsException($"Unable to write to bucket {BucketName} path {Path}", Ex);
+						throw new AwsException($"Unable to write to bucket {Options.BucketName} path {FullPath}", Ex);
 					}
 				}
 
@@ -190,16 +228,99 @@ namespace HordeServer.Storage.Backends
 		}
 
 		/// <inheritdoc/>
-		public async Task WriteInternalAsync(string Path, Stream InputStream)
+		public async Task WriteInternalAsync(string FullPath, Stream Stream)
+		{
+			const int MinPartSize = 5 * 1024 * 1024;
+
+			long StreamLen = Stream.Length;
+			if (StreamLen < MinPartSize)
+			{
+				PutObjectRequest UploadRequest = new PutObjectRequest();
+				UploadRequest.BucketName = Options.BucketName;
+				UploadRequest.Key = FullPath;
+				UploadRequest.InputStream = Stream;
+				UploadRequest.Metadata.Add("bytes-written", StreamLen.ToString(CultureInfo.InvariantCulture));
+				await Client.PutObjectAsync(UploadRequest);
+			}
+			else
+			{
+				// Initiate a multi-part upload
+				InitiateMultipartUploadRequest InitiateRequest = new InitiateMultipartUploadRequest();
+				InitiateRequest.BucketName = Options.BucketName;
+				InitiateRequest.Key = FullPath;
+
+				InitiateMultipartUploadResponse InitiateResponse = await Client.InitiateMultipartUploadAsync(InitiateRequest);
+				try
+				{
+					// Buffer for reading the data
+					byte[] Buffer = new byte[MinPartSize];
+
+					// Upload all the parts
+					List<PartETag> PartTags = new List<PartETag>();
+					for (long StreamPos = 0; StreamPos < StreamLen;)
+					{
+						// Read the next chunk of data into the buffer
+						int BufferLen = (int)Math.Min((long)MinPartSize, StreamLen - StreamPos);
+						await ReadExactLengthAsync(Stream, Buffer, BufferLen);
+						StreamPos += BufferLen;
+
+						// Upload the part
+						using (MemoryStream InputStream = new MemoryStream(Buffer, 0, BufferLen))
+						{
+							UploadPartRequest PartRequest = new UploadPartRequest();
+							PartRequest.BucketName = Options.BucketName;
+							PartRequest.Key = FullPath;
+							PartRequest.UploadId = InitiateResponse.UploadId;
+							PartRequest.InputStream = InputStream;
+							PartRequest.PartSize = BufferLen;
+							PartRequest.PartNumber = PartTags.Count + 1;
+							PartRequest.IsLastPart = (StreamPos == StreamLen);
+
+							UploadPartResponse PartResponse = await Client.UploadPartAsync(PartRequest);
+							PartTags.Add(new PartETag(PartResponse.PartNumber, PartResponse.ETag));
+						}
+					}
+
+					// Mark the upload as complete
+					CompleteMultipartUploadRequest CompleteRequest = new CompleteMultipartUploadRequest();
+					CompleteRequest.BucketName = Options.BucketName;
+					CompleteRequest.Key = FullPath;
+					CompleteRequest.UploadId = InitiateResponse.UploadId;
+					CompleteRequest.PartETags = PartTags;
+					await Client.CompleteMultipartUploadAsync(CompleteRequest);
+				}
+				catch
+				{
+					// Abort the upload
+					AbortMultipartUploadRequest AbortRequest = new AbortMultipartUploadRequest();
+					AbortRequest.BucketName = Options.BucketName;
+					AbortRequest.Key = FullPath;
+					AbortRequest.UploadId = InitiateResponse.UploadId;
+					await Client.AbortMultipartUploadAsync(AbortRequest);
+
+					throw;
+				}
+			}
+		}
+
+		/// <summary>
+		/// Reads data of an exact length into a stream
+		/// </summary>
+		/// <param name="Stream">The stream to read from</param>
+		/// <param name="Buffer">The buffer to read into</param>
+		/// <param name="Length">Length of the data to read</param>
+		/// <returns>Async task</returns>
+		static async Task ReadExactLengthAsync(System.IO.Stream Stream, byte[] Buffer, int Length)
 		{
-			using (await Semaphore.UseWaitAsync())
+			int BufferPos = 0;
+			while (BufferPos < Length)
 			{
-				PutObjectRequest NewUploadRequest = new PutObjectRequest();
-				NewUploadRequest.BucketName = BucketName;
-				NewUploadRequest.Key = Path;
-				NewUploadRequest.InputStream = InputStream;
-				NewUploadRequest.Metadata.Add("bytes-written", InputStream.Length.ToString(CultureInfo.InvariantCulture));
-				await Client.PutObjectAsync(NewUploadRequest);
+				int BytesRead = await Stream.ReadAsync(Buffer, BufferPos, Length - BufferPos);
+				if (BytesRead == 0)
+				{
+					throw new InvalidOperationException("Unexpected end of stream");
+				}
+				BufferPos += BytesRead;
 			}
 		}
 
@@ -207,8 +328,8 @@ namespace HordeServer.Storage.Backends
 		public async Task DeleteAsync(string Path)
 		{
 			DeleteObjectRequest NewDeleteRequest = new DeleteObjectRequest();
-			NewDeleteRequest.BucketName = BucketName;
-			NewDeleteRequest.Key = Path;
+			NewDeleteRequest.BucketName = Options.BucketName;
+			NewDeleteRequest.Key = GetFullPath(Path);
 			await Client.DeleteObjectAsync(NewDeleteRequest);
 		}
 
@@ -218,8 +339,8 @@ namespace HordeServer.Storage.Backends
 			try
 			{
 				GetObjectMetadataRequest Request = new GetObjectMetadataRequest();
-				Request.BucketName = BucketName;
-				Request.Key = Path;
+				Request.BucketName = Options.BucketName;
+				Request.Key = GetFullPath(Path);
 				await Client.GetObjectMetadataAsync(Request);
 				return true;
 			}
diff --git a/Engine/Source/Programs/Horde/HordeServer/Storage/Backends/FileSystemStorageBackend.cs b/Engine/Source/Programs/Horde/HordeServer/Storage/Backends/FileSystemStorageBackend.cs
--- a/Engine/Source/Programs/Horde/HordeServer/Storage/Backends/FileSystemStorageBackend.cs
+++ b/Engine/Source/Programs/Horde/HordeServer/Storage/Backends/FileSystemStorageBackend.cs
@@ -13,7 +13,21 @@ using System.Threading.Tasks;
 
 namespace HordeServer.Storage.Backends
 {
-	class FileSystemStorageBackend : IStorageBackend
+	/// <summary>
+	/// Options for the filesystem backend
+	/// </summary>
+	public interface IFileSystemStorageOptions
+	{
+		/// <summary>
+		/// Base directory for storing files
+		/// </summary>
+		public string? BaseDir { get; }
+	}
+
+	/// <summary>
+	/// Storage backend using the filesystem
+	/// </summary>
+	public class FileSystemStorageBackend : IStorageBackend
 	{
 		/// <summary>
 		/// Base directory for log files
@@ -33,11 +47,10 @@ namespace HordeServer.Storage.Backends
 		/// <summary>
 		/// Constructor
 		/// </summary>
-		/// <param name="Settings">Current Horde Settings</param>
-		public FileSystemStorageBackend(IOptions<ServerSettings> Settings)
+		/// <param name="Options">Current Horde Settings</param>
+		public FileSystemStorageBackend(IFileSystemStorageOptions Options)
 		{
-			ServerSettings CurrentSettings = Settings.Value;
-			this.BaseDir = CurrentSettings.LocalStorageDirRef;
+			this.BaseDir = DirectoryReference.Combine(Program.DataDir, Options.BaseDir ?? "Storage");
 			this.InstanceId = Guid.NewGuid().ToString("N");
 			DirectoryReference.CreateDirectory(BaseDir);
 		}
diff --git a/Engine/Source/Programs/Horde/HordeServer/Storage/Backends/RelayStorageBackend.cs b/Engine/Source/Programs/Horde/HordeServer/Storage/Backends/RelayStorageBackend.cs
--- a/Engine/Source/Programs/Horde/HordeServer/Storage/Backends/RelayStorageBackend.cs
+++ b/Engine/Source/Programs/Horde/HordeServer/Storage/Backends/RelayStorageBackend.cs
@@ -17,6 +17,11 @@ namespace HordeServer.Storage.Backends
 	/// </summary>
 	public sealed class RelayStorageBackend : IStorageBackend, IDisposable
 	{
+		class FileSystemSettings : IFileSystemStorageOptions
+		{
+			public string? BaseDir { get; set; }
+		}
+
 		/// <summary>
 		/// The client to connect with
 		/// </summary>
@@ -53,7 +58,7 @@ namespace HordeServer.Storage.Backends
 			Client = new HttpClient();
 			Client.DefaultRequestHeaders.Authorization = new AuthenticationHeaderValue("Bearer", CurrentSettings.LogRelayBearerToken);
 
-			LocalStorage = new FileSystemStorageBackend(Settings);
+			LocalStorage = new FileSystemStorageBackend(new FileSystemSettings { BaseDir = "RelayCache" });
 		}
 
 		/// <inheritdoc/>
diff --git a/Engine/Source/Programs/Horde/HordeServer/Storage/Backends/TransientStorageBackend.cs b/Engine/Source/Programs/Horde/HordeServer/Storage/Backends/TransientStorageBackend.cs
--- a/Engine/Source/Programs/Horde/HordeServer/Storage/Backends/TransientStorageBackend.cs
+++ b/Engine/Source/Programs/Horde/HordeServer/Storage/Backends/TransientStorageBackend.cs
@@ -12,7 +12,7 @@ namespace HordeServer.Storage.Backends
 	/// <summary>
 	/// In-memory implementation of ILogFileStorage
 	/// </summary>
-	public sealed class TransientStorageBackend : IStorageBackend
+	public class TransientStorageBackend : IStorageBackend
 	{
 		/// <summary>
 		/// Data storage
diff --git a/Engine/Source/Programs/Horde/HordeServer/Storage/Collections/BlobCollection.cs b/Engine/Source/Programs/Horde/HordeServer/Storage/Collections/BlobCollection.cs
--- a/Engine/Source/Programs/Horde/HordeServer/Storage/Collections/BlobCollection.cs
+++ b/Engine/Source/Programs/Horde/HordeServer/Storage/Collections/BlobCollection.cs
@@ -24,7 +24,7 @@ namespace HordeServer.Storage.Collections
 		/// Constructor
 		/// </summary>
 		/// <param name="StorageBackend"></param>
-		public BlobCollection(IStorageBackend StorageBackend)
+		public BlobCollection(IStorageBackend<BlobCollection> StorageBackend)
 		{
 			this.StorageBackend = StorageBackend;
 		}
diff --git a/Engine/Source/Programs/Horde/HordeServer/Storage/IStorageBackend.cs b/Engine/Source/Programs/Horde/HordeServer/Storage/IStorageBackend.cs
--- a/Engine/Source/Programs/Horde/HordeServer/Storage/IStorageBackend.cs
+++ b/Engine/Source/Programs/Horde/HordeServer/Storage/IStorageBackend.cs
@@ -43,11 +43,57 @@ namespace HordeServer.Storage
 		Task DeleteAsync(string Path);
 	}
 
+	/// <summary>
+	/// Generic version of IStorageBackend, to allow for dependency injection of different singletons
+	/// </summary>
+	/// <typeparam name="T">Type distinguishing different singletons</typeparam>
+	public interface IStorageBackend<T> : IStorageBackend
+	{
+	}
+
 	/// <summary>
 	/// Extension methods for <see cref="IStorageBackend"/>
 	/// </summary>
 	public static class StorageBackendExtensions
 	{
+		/// <summary>
+		/// Wrapper for <see cref="IStorageBackend"/>
+		/// </summary>
+		/// <typeparam name="T"></typeparam>
+		class StorageBackend<T> : IStorageBackend<T>
+		{
+			IStorageBackend Inner;
+
+			/// <summary>
+			/// Constructor
+			/// </summary>
+			/// <param name="Inner"></param>
+			public StorageBackend(IStorageBackend Inner) => this.Inner = Inner;
+
+			/// <inheritdoc/>
+			public Task<Stream?> ReadAsync(string Path) => Inner.ReadAsync(Path);
+
+			/// <inheritdoc/>
+			public Task WriteAsync(string Path, Stream Stream) => Inner.WriteAsync(Path, Stream);
+
+			/// <inheritdoc/>
+			public Task DeleteAsync(string Path) => Inner.DeleteAsync(Path);
+
+			/// <inheritdoc/>
+			public Task<bool> ExistsAsync(string Path) => Inner.ExistsAsync(Path);
+		}
+
+		/// <summary>
+		/// Creates a typed wrapper around the given storage backend
+		/// </summary>
+		/// <typeparam name="T"></typeparam>
+		/// <param name="Backend"></param>
+		/// <returns></returns>
+		public static IStorageBackend<T> ForType<T>(this IStorageBackend Backend)
+		{
+			return new StorageBackend<T>(Backend);
+		}
+
 		/// <summary>
 		/// Writes a block of memory to storage
 		/// </summary>
diff --git a/Engine/Source/Programs/Horde/HordeServer/Storage/Services/SimpleStorageService.cs b/Engine/Source/Programs/Horde/HordeServer/Storage/Services/SimpleStorageService.cs
--- a/Engine/Source/Programs/Horde/HordeServer/Storage/Services/SimpleStorageService.cs
+++ b/Engine/Source/Programs/Horde/HordeServer/Storage/Services/SimpleStorageService.cs
@@ -1,6 +1,7 @@
 // Copyright Epic Games, Inc. All Rights Reserved.
 
 using EpicGames.Core;
+using HordeServer.Services;
 using HordeServer.Storage.Primitives;
 using Microsoft.Extensions.Options;
 using System;
@@ -36,7 +37,7 @@ namespace HordeServer.Storage.Services
 		/// Constructor
 		/// </summary>
 		/// <param name="Backend">The storage backend</param>
-		public SimpleStorageService(IStorageBackend Backend)
+		public SimpleStorageService(IStorageBackend<ArtifactService> Backend)
 		{
 			this.Backend = Backend;
 		}
diff --git a/Engine/Source/Programs/Horde/HordeServerTests/KeyValueTreeTests.cs b/Engine/Source/Programs/Horde/HordeServerTests/KeyValueTreeTests.cs
--- a/Engine/Source/Programs/Horde/HordeServerTests/KeyValueTreeTests.cs
+++ b/Engine/Source/Programs/Horde/HordeServerTests/KeyValueTreeTests.cs
@@ -1,6 +1,7 @@
 // Copyright Epic Games, Inc. All Rights Reserved.
 
 using EpicGames.Core;
+using HordeServer.Services;
 using HordeServer.Storage;
 using HordeServer.Storage.Backends;
 using HordeServer.Storage.Primitives;
@@ -66,7 +67,7 @@ namespace HordeServerTests
 		[TestMethod]
 		public async Task TestBasicInsertion()
 		{
-			TransientStorageBackend Backend = new TransientStorageBackend();
+			IStorageBackend<ArtifactService> Backend = new TransientStorageBackend().ForType<ArtifactService>();
 			SimpleStorageService Storage = new SimpleStorageService(Backend);
 
 			BlobRef<TestTreeMax2> Tree = await Storage.CreateKeyValueTreeAsync<TestTreeMax2>();
@@ -121,7 +122,7 @@ namespace HordeServerTests
 		{
 			Random Random = new Random(0x12345678);
 
-			TransientStorageBackend Backend = new TransientStorageBackend();
+			IStorageBackend<ArtifactService> Backend = new TransientStorageBackend().ForType<ArtifactService>();
 			SimpleStorageService Storage = new SimpleStorageService(Backend);
 
 			BlobRef<TestTreeMax256> Tree = await Storage.CreateKeyValueTreeAsync<TestTreeMax256>();
diff --git a/Engine/Source/Programs/Horde/HordeServerTests/LogFileServiceTests.cs b/Engine/Source/Programs/Horde/HordeServerTests/LogFileServiceTests.cs
--- a/Engine/Source/Programs/Horde/HordeServerTests/LogFileServiceTests.cs
+++ b/Engine/Source/Programs/Horde/HordeServerTests/LogFileServiceTests.cs
@@ -19,6 +19,7 @@ using HordeServer.Logs.Storage;
 using HordeServer.Logs.Storage.Impl;
 using HordeServer.Models;
 using HordeServer.Services;
+using HordeServer.Storage;
 using HordeServer.Storage.Backends;
 using HordeServer.Utilities;
 using Microsoft.Extensions.Configuration;
@@ -47,7 +48,7 @@ namespace HordeServerTests
             IConfiguration Config = new ConfigurationBuilder().Build();
 
 			ILogBuilder LogBuilder = new RedisLogBuilder(GetRedisConnectionPool(), NullLogger.Instance);
-			ILogStorage LogStorage = new PersistentLogStorage(new TransientStorageBackend(), NullLogger<PersistentLogStorage>.Instance);
+			ILogStorage LogStorage = new PersistentLogStorage(new TransientStorageBackend().ForType<PersistentLogStorage>(), NullLogger<PersistentLogStorage>.Instance);
 			LogFileService = new LogFileService(LogFileCollection, null!, LogBuilder, LogStorage, Logger);
         }
 
diff --git a/Engine/Source/Programs/Horde/HordeServerTests/RpcServiceTest.cs b/Engine/Source/Programs/Horde/HordeServerTests/RpcServiceTest.cs
--- a/Engine/Source/Programs/Horde/HordeServerTests/RpcServiceTest.cs
+++ b/Engine/Source/Programs/Horde/HordeServerTests/RpcServiceTest.cs
@@ -271,7 +271,7 @@ namespace HordeServerTests
 
 			Artifact? Artifact = await TestSetup.ArtifactService.GetArtifactAsync(ObjectId.Parse(Res.Id));
 			Assert.IsNotNull(Artifact);
-			Stream Stream = TestSetup.ArtifactService.OpenArtifactReadStream(Artifact!);
+			Stream Stream = await TestSetup.ArtifactService.OpenArtifactReadStreamAsync(Artifact!);
 			StreamReader Reader = new StreamReader(Stream);
 			string text = Reader.ReadToEnd();
 			Assert.AreEqual(DataStr, text);
diff --git a/Engine/Source/Programs/Horde/HordeServerTests/TestSetup.cs b/Engine/Source/Programs/Horde/HordeServerTests/TestSetup.cs
--- a/Engine/Source/Programs/Horde/HordeServerTests/TestSetup.cs
+++ b/Engine/Source/Programs/Horde/HordeServerTests/TestSetup.cs
@@ -14,6 +14,7 @@ using HordeServer.Collections.Impl;
 using HordeServer.Controllers;
 using HordeServer.Logs;
 using HordeServer.Logs.Builder;
+using HordeServer.Logs.Readers;
 using HordeServer.Logs.Storage.Impl;
 using HordeServer.Models;
 using HordeServer.Notifications;
@@ -130,10 +131,6 @@ namespace HordeServerTests
 
 			Settings.AdminClaimType = HordeClaimTypes.Role;
 			Settings.AdminClaimValue = "app-horde-admins";
-
-			DirectoryReference DataDir = DirectoryReference.Combine(BaseDir, Guid.NewGuid().ToString("N"));
-			Settings.LocalArtifactsDir = DirectoryReference.Combine(DataDir, "Artifacts").FullName;
-			Settings.LocalStorageDir = DataDir.FullName;
 		}
 
 		protected virtual void ConfigureServices(IServiceCollection Services)
@@ -216,8 +213,12 @@ namespace HordeServerTests
 			Services.AddSingleton<INamespaceCollection, NamespaceCollection>();
 			Services.AddSingleton<IBucketCollection, BucketCollection>();
 
-			Services.AddSingleton<IStorageBackend, FileSystemStorageBackend>();
-			Services.AddSingleton<IStorageService, SimpleStorageService>(SP => new SimpleStorageService(SP.GetRequiredService<IStorageBackend>()));
+			Services.AddSingleton(new Startup.StorageBackendSettings<PersistentLogStorage> { Type = StorageProviderType.Transient });
+			Services.AddSingleton(new Startup.StorageBackendSettings<ArtifactService> { Type = StorageProviderType.Transient });
+			Services.AddSingleton(new Startup.StorageBackendSettings<BlobCollection> { Type = StorageProviderType.Transient });
+			Services.AddSingleton(typeof(IStorageBackend<>), typeof(Startup.StorageBackendFactory<>));
+
+			Services.AddSingleton<IStorageService, SimpleStorageService>();
 
 			Services.AddSingleton<ISingletonDocument<GlobalPermissions>>(new SingletonDocumentStub<GlobalPermissions>());
 			Services.AddSingleton<ISingletonDocument<AgentSoftwareChannels>>(new SingletonDocumentStub<AgentSoftwareChannels>());
