Horde: Add RPC interfaces for storage service.

diff --git a/Engine/Source/Programs/Horde/HordeServer/Startup.cs b/Engine/Source/Programs/Horde/HordeServer/Startup.cs
--- a/Engine/Source/Programs/Horde/HordeServer/Startup.cs
+++ b/Engine/Source/Programs/Horde/HordeServer/Startup.cs
@@ -780,6 +780,10 @@ namespace HordeServer
 				Endpoints.MapGrpcService<CapabilitiesService>();
 				Endpoints.MapGrpcService<ContentStorageService>();
 				Endpoints.MapGrpcService<ExecutionService>();
+
+				Endpoints.MapGrpcService<BlobRpc>();
+				Endpoints.MapGrpcService<RefRpc>();
+
 				Endpoints.MapGrpcReflectionService();
 
 				Endpoints.MapControllers();
diff --git a/Engine/Source/Programs/Horde/HordeServer/Storage/Collections/BlobRpc.cs b/Engine/Source/Programs/Horde/HordeServer/Storage/Collections/BlobRpc.cs
new file mode 100644
--- /dev/null
+++ b/Engine/Source/Programs/Horde/HordeServer/Storage/Collections/BlobRpc.cs
@@ -0,0 +1,104 @@
+// Copyright Epic Games, Inc. All Rights Reserved.
+
+using System;
+using System.Collections.Generic;
+using System.IO;
+using System.Linq;
+using System.Threading.Tasks;
+using EpicGames.Core;
+using EpicGames.Horde.Common.Storage;
+using Google.Protobuf;
+using Grpc.Core;
+using HordeServer.Utilities;
+
+namespace HordeServer.Storage.Collections
+{
+	using NamespaceId = StringId<INamespace>;
+
+	class BlobRpc : BlobService.BlobServiceBase
+	{
+		IBlobCollection BlobCollection;
+
+		/// <summary>
+		/// Constructor
+		/// </summary>
+		/// <param name="BlobCollection"></param>
+		public BlobRpc(IBlobCollection BlobCollection)
+		{
+			this.BlobCollection = BlobCollection;
+		}
+
+		/// <inheritdoc/>
+		public override async Task<AddBlobResponse> Add(AddBlobRequest Request, ServerCallContext Context)
+		{
+			NamespaceId NamespaceId = new NamespaceId(Request.NamespaceId);
+			IoHash Hash = new IoHash(Request.Hash.ToByteArray());
+			await BlobCollection.WriteBytesAsync(NamespaceId, Hash, Request.Data.ToByteArray());
+			return new AddBlobResponse();
+		}
+
+		/// <inheritdoc/>
+		public override async Task<GetBlobResponse> Get(GetBlobRequest Request, ServerCallContext Context)
+		{
+			NamespaceId NamespaceId = new NamespaceId(Request.NamespaceId);
+			IoHash Hash = new IoHash(Request.Hash.ToByteArray());
+
+			Stream? Stream = await BlobCollection.ReadAsync(NamespaceId, Hash);
+			if(Stream == null)
+			{
+				throw new RpcException(new Status(StatusCode.NotFound, "Blob not found"));
+			}
+
+			GetBlobResponse Response = new GetBlobResponse();
+			Response.Data = await ByteString.FromStreamAsync(Stream);
+			return Response;
+		}
+
+		/// <inheritdoc/>
+		public override async Task<ExistsResponse> Exists(ExistsRequest Request, ServerCallContext Context)
+		{
+			NamespaceId NamespaceId = new NamespaceId(Request.NamespaceId);
+			List<IoHash> Hashes = ParseIoHashList(Request.Hashes);
+
+			List<IoHash> Results = await BlobCollection.ExistsAsync(NamespaceId, Hashes);
+
+			ExistsResponse Response = new ExistsResponse();
+			Response.Hashes = BuildIoHashList(Results);
+			return Response;
+		}
+
+		/// <summary>
+		/// 
+		/// </summary>
+		/// <param name="ByteString"></param>
+		/// <returns></returns>
+		public static List<IoHash> ParseIoHashList(ByteString ByteString)
+		{
+			List<IoHash> Hashes = new List<IoHash>();
+			for (int Offset = 0; Offset < ByteString.Length; Offset += IoHash.NumBytes)
+			{
+				Hashes.Add(new IoHash(ByteString.Span.Slice(Offset, IoHash.NumBytes).ToArray()));
+			}
+			return Hashes;
+		}
+
+		/// <summary>
+		/// 
+		/// </summary>
+		/// <param name="Hashes"></param>
+		/// <returns></returns>
+		public static ByteString BuildIoHashList(List<IoHash> Hashes)
+		{
+			byte[] Data = new byte[Hashes.Count * IoHash.NumBytes];
+
+			Span<byte> Output = Data;
+			foreach (IoHash Hash in Hashes)
+			{
+				Hash.Span.CopyTo(Output);
+				Output = Output.Slice(IoHash.NumBytes);
+			}
+
+			return ByteString.CopyFrom(Data);
+		}
+	}
+}
diff --git a/Engine/Source/Programs/Horde/HordeServer/Storage/Collections/RefRpc.cs b/Engine/Source/Programs/Horde/HordeServer/Storage/Collections/RefRpc.cs
new file mode 100644
--- /dev/null
+++ b/Engine/Source/Programs/Horde/HordeServer/Storage/Collections/RefRpc.cs
@@ -0,0 +1,93 @@
+// Copyright Epic Games, Inc. All Rights Reserved.
+
+using System;
+using System.Collections.Generic;
+using System.IO;
+using System.Linq;
+using System.Threading.Tasks;
+using EpicGames.Core;
+using EpicGames.Horde.Common.Storage;
+using EpicGames.Serialization;
+using Google.Protobuf;
+using Google.Protobuf.WellKnownTypes;
+using Grpc.Core;
+using HordeServer.Utilities;
+
+namespace HordeServer.Storage.Collections
+{
+	using NamespaceId = StringId<INamespace>;
+	using BucketId = StringId<IBucket>;
+
+	class RefRpc : RefService.RefServiceBase
+	{
+		IRefCollection RefCollection;
+
+		/// <summary>
+		/// Constructor
+		/// </summary>
+		/// <param name="RefCollection"></param>
+		public RefRpc(IRefCollection RefCollection)
+		{
+			this.RefCollection = RefCollection;
+		}
+
+		/// <inheritdoc/>
+		public override async Task<SetRefResponse> Set(SetRefRequest Request, ServerCallContext Context)
+		{
+			NamespaceId NamespaceId = new NamespaceId(Request.NamespaceId);
+			BucketId BucketId = new BucketId(Request.BucketId);
+			List<IoHash> MissingHashes = await RefCollection.SetAsync(NamespaceId, BucketId, Request.Name, new CbObject(Request.Value.ToByteArray()));
+
+			SetRefResponse Response = new SetRefResponse();
+			Response.MissingHashes = BlobRpc.BuildIoHashList(MissingHashes);
+			return Response;
+		}
+
+		/// <inheritdoc/>
+		public override async Task<GetRefResponse> Get(GetRefRequest Request, ServerCallContext Context)
+		{
+			NamespaceId NamespaceId = new NamespaceId(Request.NamespaceId);
+			BucketId BucketId = new BucketId(Request.BucketId);
+
+			IRef? Ref = await RefCollection.GetAsync(NamespaceId, BucketId, Request.Name);
+			if(Ref == null)
+			{
+				throw new RpcException(new Status(StatusCode.NotFound, $"Unable to find ref {NamespaceId}/{BucketId}/{Request.Name}"));
+			}
+
+			GetRefResponse Response = new GetRefResponse();
+			Response.LastAccessTime = Timestamp.FromDateTime(Ref.LastAccessTime);
+			Response.Value = ByteString.CopyFrom(Ref.Value.GetView().Span);
+			Response.Finalized = Ref.Finalized;
+			return Response;
+		}
+
+		/// <inheritdoc/>
+		public override async Task<TouchResponse> Touch(TouchRequest Request, ServerCallContext Context)
+		{
+			NamespaceId NamespaceId = new NamespaceId(Request.NamespaceId);
+			BucketId BucketId = new BucketId(Request.BucketId);
+
+			if (!await RefCollection.TouchAsync(NamespaceId, BucketId, Request.Name))
+			{
+				throw new RpcException(new Status(StatusCode.NotFound, $"Unable to find ref {NamespaceId}/{BucketId}/{Request.Name}"));
+			}
+
+			return new TouchResponse();
+		}
+
+		/// <inheritdoc/>
+		public override async Task<DeleteResponse> Delete(DeleteRequest Request, ServerCallContext Context)
+		{
+			NamespaceId NamespaceId = new NamespaceId(Request.NamespaceId);
+			BucketId BucketId = new BucketId(Request.BucketId);
+
+			if (!await RefCollection.DeleteAsync(NamespaceId, BucketId, Request.Name))
+			{
+				throw new RpcException(new Status(StatusCode.NotFound, $"Unable to find ref {NamespaceId}/{BucketId}/{Request.Name}"));
+			}
+
+			return new DeleteResponse();
+		}
+	}
+}
