Horde: Build blob/ref store RPCs for server.

diff --git a/Engine/Source/Programs/Horde/HordeServer/HordeServer.csproj b/Engine/Source/Programs/Horde/HordeServer/HordeServer.csproj
--- a/Engine/Source/Programs/Horde/HordeServer/HordeServer.csproj
+++ b/Engine/Source/Programs/Horde/HordeServer/HordeServer.csproj
@@ -62,6 +62,8 @@
     <Protobuf Include="..\..\Shared\EpicGames.Horde.Common\Protos\horde\action_rpc.proto" ProtoRoot="..\..\Shared\EpicGames.Horde.Common\Protos" GrpcServices="Server" />
     <Protobuf Include="..\..\Shared\EpicGames.Horde.Common\Protos\horde\horde_rpc.proto" ProtoRoot="..\..\Shared\EpicGames.Horde.Common\Protos" GrpcServices="Server" />
     <Protobuf Include="..\..\Shared\EpicGames.Horde.Common\Protos\horde\tasks\action_task.proto" ProtoRoot="..\..\Shared\EpicGames.Horde.Common\Protos" GrpcServices="Server" />
+    <Protobuf Include="..\..\Shared\EpicGames.Horde.Common\Protos\horde\storage\ref_service.proto" ProtoRoot="..\..\Shared\EpicGames.Horde.Common\Protos" GrpcServices="Server" />
+    <Protobuf Include="..\..\Shared\EpicGames.Horde.Common\Protos\horde\storage\blob_service.proto" ProtoRoot="..\..\Shared\EpicGames.Horde.Common\Protos" GrpcServices="Server" />
   </ItemGroup>
 
   <ItemGroup>
