Horde: Replace epic hostname references in package.json and stream summary
#rnx

diff --git a/Engine/Source/Programs/Horde/HordeDashboard/package.json b/Engine/Source/Programs/Horde/HordeDashboard/package.json
--- a/Engine/Source/Programs/Horde/HordeDashboard/package.json
+++ b/Engine/Source/Programs/Horde/HordeDashboard/package.json
@@ -3,7 +3,7 @@
   "version": "0.1.0",
   "private": true,
   "main": "index.js",
-  "proxy": "https://horde.devtools.epicgames.com",
+  "proxy": "https://localhost:5001",
   "dependencies": {
     "@blueprintjs/icons": "^3.13.0",
     "@datadog/browser-logs": "^1.12.4",
