Horde: Remove code to clear downtime from globals list. Downtime was not triggering due to repeated schedules being removed.

diff --git a/Engine/Source/Programs/Horde/HordeServer/Models/Globals.cs b/Engine/Source/Programs/Horde/HordeServer/Models/Globals.cs
--- a/Engine/Source/Programs/Horde/HordeServer/Models/Globals.cs
+++ b/Engine/Source/Programs/Horde/HordeServer/Models/Globals.cs
@@ -376,15 +376,5 @@ namespace HordeServer.Models
 				return PerforceClusters.FirstOrDefault(x => String.Equals(x.Name, Name, StringComparison.OrdinalIgnoreCase));
 			}
 		}
-
-		/// <inheritdoc/>
-		public override void PostLoad()
-		{
-			base.PostLoad();
-
-			DateTimeOffset Now = DateTimeOffset.Now;
-			Notices.RemoveAll(x => x.FinishTime != null && x.FinishTime < Now);
-			ScheduledDowntime.RemoveAll(x => x.FinishTime < Now);
-		}
 	}
 }
diff --git a/Engine/Source/Programs/Horde/HordeServer/Services/ConfigService.cs b/Engine/Source/Programs/Horde/HordeServer/Services/ConfigService.cs
--- a/Engine/Source/Programs/Horde/HordeServer/Services/ConfigService.cs
+++ b/Engine/Source/Programs/Horde/HordeServer/Services/ConfigService.cs
@@ -33,7 +33,7 @@ namespace HordeServer.Services
 		/// <summary>
 		/// Config file version number
 		/// </summary>
-		const int Version = 4;
+		const int Version = 5;
 
 		/// <summary>
 		/// Database service
