HordeAgent: Terminate remote execution process if cancelled

#rnx

diff --git a/Engine/Source/Programs/Horde/HordeAgent/Execution/ActionExecutor.cs b/Engine/Source/Programs/Horde/HordeAgent/Execution/ActionExecutor.cs
--- a/Engine/Source/Programs/Horde/HordeAgent/Execution/ActionExecutor.cs
+++ b/Engine/Source/Programs/Horde/HordeAgent/Execution/ActionExecutor.cs
@@ -92,8 +92,9 @@ namespace HordeAgent
 		/// <param name="ActionTask">Task to execute</param>
 		/// <param name="SandboxDir">Directory to use as a sandbox for execution</param>
 		/// <param name="ActionTaskStartTime">When the agent received the action task. Used for reporting execution stats in REAPI.</param>
+		/// <param name="CancellationToken">Token used to cancel the operation</param>
 		/// <returns>The action result</returns>
-		public async Task<ActionResult> ExecuteActionAsync(string LeaseId, ActionTask ActionTask, DirectoryReference SandboxDir, DateTimeOffset ActionTaskStartTime)
+		public async Task<ActionResult> ExecuteActionAsync(string LeaseId, ActionTask ActionTask, DirectoryReference SandboxDir, DateTimeOffset ActionTaskStartTime, CancellationToken CancellationToken)
 		{
 			Action? Action = await Storage.GetProtoMessageAsync<Action>(InstanceName, ActionTask.Digest);
 			if (Action == null)
@@ -150,7 +151,11 @@ namespace HordeAgent
 					using (MemoryStream StdOutStream = new MemoryStream())
 					using (MemoryStream StdErrStream = new MemoryStream())
 					{
-						await Task.WhenAll(Process.StdOut.CopyToAsync(StdOutStream), Process.StdErr.CopyToAsync(StdErrStream));
+						await Task.WhenAll(Process.StdOut.CopyToAsync(StdOutStream, CancellationToken), Process.StdErr.CopyToAsync(StdErrStream, CancellationToken));
+						if (CancellationToken.IsCancellationRequested)
+						{
+							throw new RpcException(new Grpc.Core.Status(StatusCode.Cancelled, "Action Cancelled"));
+						}
 						StdOutData = StdOutStream.ToArray();
 						StdErrData = StdErrStream.ToArray();
 					}
diff --git a/Engine/Source/Programs/Horde/HordeAgent/Services/WorkerService.cs b/Engine/Source/Programs/Horde/HordeAgent/Services/WorkerService.cs
--- a/Engine/Source/Programs/Horde/HordeAgent/Services/WorkerService.cs
+++ b/Engine/Source/Programs/Horde/HordeAgent/Services/WorkerService.cs
@@ -803,7 +803,7 @@ namespace HordeAgent.Services
 
 				try
 				{
-					await Executor.ExecuteActionAsync(LeaseId, ActionTask, LeaseDir, ActionTaskStartTime);
+					await Executor.ExecuteActionAsync(LeaseId, ActionTask, LeaseDir, ActionTaskStartTime, CancellationToken);
 				}
 				catch (Exception re)
 				{
