Horde: Publish HordeAgent symbols to the network.

diff --git a/Engine/Source/Programs/Horde/HordeAgent/HordeAgent.csproj b/Engine/Source/Programs/Horde/HordeAgent/HordeAgent.csproj
--- a/Engine/Source/Programs/Horde/HordeAgent/HordeAgent.csproj
+++ b/Engine/Source/Programs/Horde/HordeAgent/HordeAgent.csproj
@@ -45,9 +45,8 @@
     <PackageReference Include="System.ServiceProcess.ServiceController" Version="4.7.0" />
   </ItemGroup>
 
-  <ItemGroup Condition="'$(SymStoreRoot)' != ''">
+  <ItemGroup Condition="'$(FullSymbolsDir)' != ''">
     <PackageReference Include="Microsoft.DiaSymReader.Pdb2Pdb" Version="1.1.0-beta2-20370-01" />
-    <PackageReference Include="Microsoft.SourceLink.GitHub" Version="1.0.0" PrivateAssets="All" />
   </ItemGroup>
 
   <ItemGroup>
@@ -114,15 +113,13 @@
 
   <PropertyGroup>
 	<Pdb2PdbPath>$(NuGetPackageRoot)microsoft.diasymreader.pdb2pdb\1.1.0-beta2-20370-01\tools\Pdb2Pdb.exe</Pdb2PdbPath>
-	<SymStorePath>C:\Program Files (x86)\Windows Kits\10\Debuggers\x64\symstore.exe</SymStorePath>
   </PropertyGroup>
 
-  <Target Name="PublishSymbols" AfterTargets="AfterBuild" Condition="'$(SymStoreRoot)' != ''">
-	<MakeDir Directories="$(OutDir)Symbols" />
-	<Exec Command="&quot;$(Pdb2PdbPath)&quot; &quot;$(OutDir)HordeAgent.dll&quot; /out &quot;$(OutDir)Symbols\HordeAgent.pdb&quot; /nowarn 21" IgnoreExitCode="false" />
-	<Exec Command="&quot;$(Pdb2PdbPath)&quot; &quot;$(OutDir)EpicGames.Core.dll&quot; /out &quot;$(OutDir)Symbols\EpicGames.Core.pdb&quot; /nowarn 21" IgnoreExitCode="false" />
-	<Exec Command="&quot;$(Pdb2PdbPath)&quot; &quot;$(OutDir)EpicGames.Perforce.dll&quot; /out &quot;$(OutDir)Symbols\EpicGames.Perforce.pdb&quot; /nowarn 21" IgnoreExitCode="false" />
-	<Exec Command="&quot;$(SymStorePath)&quot; add /f &quot;$(OutDir)Symbols\*.pdb&quot; /s &quot;$(SymStoreRoot)&quot; /t HordeAgent" />
+  <Target Name="PublishSymbols" AfterTargets="AfterBuild" Condition="'$(FullSymbolsDir)' != ''">
+	<MakeDir Directories="$(FullSymbolsDir)" />
+	<Exec Command="&quot;$(Pdb2PdbPath)&quot; &quot;$(OutDir)HordeAgent.dll&quot; /out &quot;$(FullSymbolsDir)\HordeAgent.pdb&quot; /nowarn 21" IgnoreExitCode="false" />
+	<Exec Command="&quot;$(Pdb2PdbPath)&quot; &quot;$(OutDir)EpicGames.Core.dll&quot; /out &quot;$(FullSymbolsDir)\EpicGames.Core.pdb&quot; /nowarn 21" IgnoreExitCode="false" />
+	<Exec Command="&quot;$(Pdb2PdbPath)&quot; &quot;$(OutDir)EpicGames.Perforce.dll&quot; /out &quot;$(FullSymbolsDir)\EpicGames.Perforce.pdb&quot; /nowarn 21" IgnoreExitCode="false" />
   </Target>
   
 </Project>
