Horde: Enable reshelving changes in all branches.

diff --git a/Engine/Source/Programs/Horde/HordeServer/Services/JobService.cs b/Engine/Source/Programs/Horde/HordeServer/Services/JobService.cs
--- a/Engine/Source/Programs/Horde/HordeServer/Services/JobService.cs
+++ b/Engine/Source/Programs/Horde/HordeServer/Services/JobService.cs
@@ -201,11 +201,6 @@ namespace HordeServer.Services
 			this.Logger = Logger;
 		}
 
-		static bool ShouldClonePreflightChange(StreamId StreamId)
-		{
-			return StreamId == new StreamId("ue5-main");
-		}
-
 		/// <summary>
 		/// Creates a new job
 		/// </summary>
@@ -238,7 +233,7 @@ namespace HordeServer.Services
 			ObjectId JobIdValue = JobId ?? ObjectId.GenerateNewId();
 			using IDisposable Scope = Logger.BeginScope("CreateJobAsync({JobId})", JobIdValue);
 
-			if (PreflightChange != null && ShouldClonePreflightChange(StreamId))
+			if (PreflightChange != null)
 			{
 				ClonedPreflightChange = await CloneShelvedChangeAsync(ClonedPreflightChange ?? PreflightChange.Value);
 			}
@@ -979,14 +974,7 @@ namespace HordeServer.Services
 				int ClonedPreflightChange = Job.ClonedPreflightChange;
 				if (ClonedPreflightChange == 0)
 				{
-					if (ShouldClonePreflightChange(Job.StreamId))
-					{
-						ClonedPreflightChange = await CloneShelvedChangeAsync(Job.PreflightChange);
-					}
-					else
-					{
-						ClonedPreflightChange = Job.PreflightChange;
-					}
+					ClonedPreflightChange = await CloneShelvedChangeAsync(Job.PreflightChange);
 				}
 
 				Logger.LogInformation("Updating description for {ClonedPreflightChange}", ClonedPreflightChange);
@@ -1015,19 +1003,9 @@ namespace HordeServer.Services
 				Logger.LogError(Ex, "Unable to submit shelved change");
 			}
 
-			if (ShouldClonePreflightChange(Job.StreamId))
+			if (Change != null && Job.ClonedPreflightChange != 0)
 			{
-				if (Change != null && Job.ClonedPreflightChange != 0)
-				{
-					await DeleteShelvedChangeAsync(Job.PreflightChange);
-				}
-			}
-			else
-			{
-				if (Change != null && Job.PreflightChange != 0)
-				{
-					await DeleteShelvedChangeAsync(Job.PreflightChange);
-				}
+				await DeleteShelvedChangeAsync(Job.PreflightChange);
 			}
 
 			for (; ; )
