Horde: Generate full PDBs for symbol store.

diff --git a/Engine/Source/Programs/Horde/HordeAgent/HordeAgent.csproj b/Engine/Source/Programs/Horde/HordeAgent/HordeAgent.csproj
--- a/Engine/Source/Programs/Horde/HordeAgent/HordeAgent.csproj
+++ b/Engine/Source/Programs/Horde/HordeAgent/HordeAgent.csproj
@@ -115,11 +115,16 @@
 	<Pdb2PdbPath>$(NuGetPackageRoot)microsoft.diasymreader.pdb2pdb\1.1.0-beta2-20370-01\tools\Pdb2Pdb.exe</Pdb2PdbPath>
   </PropertyGroup>
 
-  <Target Name="PublishSymbols" AfterTargets="AfterBuild" Condition="'$(FullSymbolsDir)' != ''">
-	<MakeDir Directories="$(FullSymbolsDir)" />
-	<Exec Command="&quot;$(Pdb2PdbPath)&quot; &quot;$(OutDir)HordeAgent.dll&quot; /out &quot;$(FullSymbolsDir)\HordeAgent.pdb&quot; /nowarn 21" IgnoreExitCode="false" />
-	<Exec Command="&quot;$(Pdb2PdbPath)&quot; &quot;$(OutDir)EpicGames.Core.dll&quot; /out &quot;$(FullSymbolsDir)\EpicGames.Core.pdb&quot; /nowarn 21" IgnoreExitCode="false" />
-	<Exec Command="&quot;$(Pdb2PdbPath)&quot; &quot;$(OutDir)EpicGames.Perforce.dll&quot; /out &quot;$(FullSymbolsDir)\EpicGames.Perforce.pdb&quot; /nowarn 21" IgnoreExitCode="false" />
+  <Target Name="PublishSymbols" AfterTargets="AfterBuild">
+    <ItemGroup>
+      <OutputFiles Include="$(OutDir)**"/>
+      <BinaryFilesBase Include="$(OutDir)\*.dll"/>
+      <BinaryFiles Include="@(BinaryFilesBase)" Condition="Exists('%(RelativeDir)%(FileName).pdb')"/>
+    </ItemGroup>
+
+    <MakeDir Directories="$(FullSymbolsDir)" />
+    <Copy SourceFiles="@(Files)" DestinationFolder="$(FullSymbolsDir)"/>
+    <Exec Command="&quot;$(Pdb2PdbPath)&quot; &quot;%(BinaryFiles.Identity)&quot; /out &quot;$(FullSymbolsDir)\%(BinaryFiles.FileName).pdb&quot; /nowarn 21" IgnoreExitCode="false" />
   </Target>
   
 </Project>
