Horde: Fix local building when user does not have the diasymreader.converter package already in the NuGet package cache.

diff --git a/Engine/Source/Programs/Horde/HordeAgent/HordeAgent.csproj b/Engine/Source/Programs/Horde/HordeAgent/HordeAgent.csproj
--- a/Engine/Source/Programs/Horde/HordeAgent/HordeAgent.csproj
+++ b/Engine/Source/Programs/Horde/HordeAgent/HordeAgent.csproj
@@ -115,7 +115,7 @@
 	<Pdb2PdbPath>$(NuGetPackageRoot)microsoft.diasymreader.pdb2pdb\1.1.0-beta2-20370-01\tools\Pdb2Pdb.exe</Pdb2PdbPath>
   </PropertyGroup>
 
-  <Target Name="PublishSymbols" AfterTargets="AfterBuild">
+  <Target Name="PublishSymbols" AfterTargets="AfterBuild" Condition="'$(FullSymbolsDir)' != ''">
     <ItemGroup>
       <OutputFiles Include="$(OutDir)**"/>
       <BinaryFilesBase Include="$(OutDir)\*.dll"/>
