Horde: Fix global config not being updated on server pods other than the first one that observes a submitted change.

diff --git a/Engine/Source/Programs/Horde/HordeServer/Services/ConfigService.cs b/Engine/Source/Programs/Horde/HordeServer/Services/ConfigService.cs
--- a/Engine/Source/Programs/Horde/HordeServer/Services/ConfigService.cs
+++ b/Engine/Source/Programs/Horde/HordeServer/Services/ConfigService.cs
@@ -93,13 +93,12 @@ namespace HordeServer.Services
 		}
 
 		GlobalConfig? CachedGlobalConfig;
+		string? CachedGlobalConfigRevision;
 		Dictionary<ProjectId, ProjectConfig> CachedProjectConfigs = new Dictionary<ProjectId, ProjectConfig>();
 		Dictionary<ProjectId, string?> CachedLogoRevisions = new Dictionary<ProjectId, string?>();
 
 		async Task UpdateConfigAsync(string ConfigPath)
 		{
-			Logger.LogInformation("Updating configuration from {ConfigPath}", ConfigPath);
-
 			// Update the globals singleton
 			GlobalConfig GlobalConfig;
 			for (; ; )
@@ -111,9 +110,7 @@ namespace HordeServer.Services
 				}
 
 				string Revision = GlobalRevisions.First().Value;
-
-				Globals Globals = await DatabaseService.GetGlobalsAsync();
-				if (Globals.ConfigRevision != Revision || CachedGlobalConfig == null)
+				if (CachedGlobalConfig == null || Revision != CachedGlobalConfigRevision)
 				{
 					Logger.LogInformation("Caching global config from {Revision}", Revision);
 					CachedGlobalConfig = await ReadDataAsync<GlobalConfig>(ConfigPath);
@@ -122,11 +119,14 @@ namespace HordeServer.Services
 						Logger.LogWarning("Unable to update any projects or streams due to invalid global config.");
 						return;
 					}
+					CachedGlobalConfigRevision = Revision;
 				}
 				GlobalConfig = CachedGlobalConfig;
 
+				Globals Globals = await DatabaseService.GetGlobalsAsync();
 				if (Globals.ConfigRevision == Revision)
 				{
+					Logger.LogInformation("Updating configuration from {ConfigPath}", Globals.ConfigRevision);
 					break;
 				}
 
