Horde: Prevent infinite loop in background services when enabling read-only database mode.

diff --git a/Engine/Source/Programs/Horde/HordeServer/Services/PoolUpdateService.cs b/Engine/Source/Programs/Horde/HordeServer/Services/PoolUpdateService.cs
--- a/Engine/Source/Programs/Horde/HordeServer/Services/PoolUpdateService.cs
+++ b/Engine/Source/Programs/Horde/HordeServer/Services/PoolUpdateService.cs
@@ -116,7 +116,9 @@ namespace HordeServer.Services
 					if (!AgentWorkspace.SetEquals(CurrentPool.Workspaces, NewWorkspaces) || CurrentPool.Workspaces.Count != NewWorkspaces.Count)
 					{
 						Logger.LogInformation("New workspaces for pool {Pool}:{Workspaces}", CurrentPool.Id, String.Join("", NewWorkspaces.Select(x => $"\n  Identifier=\"{x.Identifier}\", Stream={x.Stream}")));
-						if (await Pools.TryUpdateAsync(CurrentPool, NewWorkspaces: NewWorkspaces) == null)
+
+						IPool? Result = await Pools.TryUpdateAsync(CurrentPool, NewWorkspaces: NewWorkspaces);
+						if (Result == null)
 						{
 							Logger.LogInformation("Pool modified; will retry");
 							bRetryUpdate = true;
diff --git a/Engine/Source/Programs/Horde/HordeServer/Utilities/ElectedBackgroundService.cs b/Engine/Source/Programs/Horde/HordeServer/Utilities/ElectedBackgroundService.cs
--- a/Engine/Source/Programs/Horde/HordeServer/Utilities/ElectedBackgroundService.cs
+++ b/Engine/Source/Programs/Horde/HordeServer/Utilities/ElectedBackgroundService.cs
@@ -48,7 +48,7 @@ namespace HordeServer.Utilities
 		/// <param name="ElectionTimeout"></param>
 		/// <param name="Logger">Logging device</param>
 		public ElectedBackgroundService(DatabaseService DatabaseService, ObjectId ServiceId, TimeSpan PollInterval, TimeSpan ElectionTimeout, ILogger Logger)
-			: base(null, Logger)
+			: base(TimeSpan.FromSeconds(10.0), Logger)
 		{
 			this.StateAccessor = new SingletonDocument<State>(DatabaseService, ServiceId);
 			this.PollInterval = PollInterval;
diff --git a/Engine/Source/Programs/Horde/HordeServer/Utilities/TickedBackgroundService.cs b/Engine/Source/Programs/Horde/HordeServer/Utilities/TickedBackgroundService.cs
--- a/Engine/Source/Programs/Horde/HordeServer/Utilities/TickedBackgroundService.cs
+++ b/Engine/Source/Programs/Horde/HordeServer/Utilities/TickedBackgroundService.cs
@@ -21,7 +21,7 @@ namespace HordeServer.Utilities
 		/// <summary>
 		/// Frequency that the service will be ticked
 		/// </summary>
-		public TimeSpan? Interval { get; set; }
+		public TimeSpan Interval { get; set; }
 
 		/// <summary>
 		/// Logger interface
@@ -38,7 +38,7 @@ namespace HordeServer.Utilities
 		/// </summary>
 		/// <param name="Interval">Interval for ticking the service</param>
 		/// <param name="Logger">Logger for debug messages and errors</param>
-		public TickedBackgroundService(TimeSpan? Interval, ILogger Logger)
+		public TickedBackgroundService(TimeSpan Interval, ILogger Logger)
 		{
 			this.Interval = Interval;
 			this.Logger = Logger;
@@ -59,12 +59,7 @@ namespace HordeServer.Utilities
 			using (CancellationTask StoppingTask = new CancellationTask(StoppingToken))
 			{
 				// Wait a random amount of time before the first iteration
-				TimeSpan? IntervalCopy = Interval;
-				if(IntervalCopy.HasValue)
-				{
-					TimeSpan Delay = new Random().NextDouble() * IntervalCopy.Value;
-					await Task.WhenAny(StoppingTask.Task, Task.Delay(Delay));
-				}
+				await Task.WhenAny(StoppingTask.Task, Task.Delay(new Random().NextDouble() * Interval));
 
 				// Enter the main tick loop
 				Stopwatch Timer = new Stopwatch();
@@ -98,14 +93,10 @@ namespace HordeServer.Utilities
 					}
 
 					// Wait until the next interval has elapsed
-					IntervalCopy = Interval;
-					if(IntervalCopy.HasValue)
+					TimeSpan Delay = Interval - Timer.Elapsed;
+					if(Delay > TimeSpan.Zero)
 					{
-						TimeSpan Delay = IntervalCopy.Value - Timer.Elapsed;
-						if(Delay > TimeSpan.Zero)
-						{
-							await Task.WhenAny(Task.Delay(Delay), StoppingTask.Task);
-						}
+						await Task.WhenAny(Task.Delay(Delay), StoppingTask.Task);
 					}
 				}
 			}
