Horde: Add support for lease cancellation
#rnx

diff --git a/Engine/Source/Programs/Horde/HordeDashboard/src/backend/Api.ts b/Engine/Source/Programs/Horde/HordeDashboard/src/backend/Api.ts
--- a/Engine/Source/Programs/Horde/HordeDashboard/src/backend/Api.ts
+++ b/Engine/Source/Programs/Horde/HordeDashboard/src/backend/Api.ts
@@ -287,7 +287,7 @@ export type IssueQuery = {
 
 export type UsersQuery = {
 	ids?: string[];
-	nameRegex? : string;
+	nameRegex?: string;
 	index?: number;
 	count?: number;
 	includeAvatar?: boolean;
@@ -410,21 +410,37 @@ export type BatchUpdatePoolRequest = {
 	deleted?: boolean;
 }
 
+/** Lease state for agent leases */
+export enum LeaseState {
+
+	Unspecified = "Unspecified",
+
+	Pending = "Pending",
+
+	Active = "Active",
+
+	Completed = "Completed",
+
+	Cancelled = "Cancelled"
+
+}
+
+
 export enum LeaseOutcome {
-    
-    /** Default value. */
-    Unspecified = "Unspecified",
-    
-    /** The lease was executed successfully */
-    Success = "Success",
-    
-    /** The lease was not executed succesfully, but cannot be run again. */
-    Failed = "Failed",
-    
-    /** The lease was cancelled by request */
-    Cancelled = "Cancelled"
-
-  }
+
+	/** Default value. */
+	Unspecified = "Unspecified",
+
+	/** The lease was executed successfully */
+	Success = "Success",
+
+	/** The lease was not executed succesfully, but cannot be run again. */
+	Failed = "Failed",
+
+	/** The lease was cancelled by request */
+	Cancelled = "Cancelled"
+
+}
 
 
 /**Response for queries to find a particular lease within an agent */
@@ -467,7 +483,9 @@ export type GetAgentLeaseResponse = {
 	batch?: GetBatchResponse;
 
 	/** Outcome of the lease */
-	outcome? : LeaseOutcome;
+	outcome?: LeaseOutcome;
+
+	state?: LeaseState;
 
 }
 
@@ -2748,21 +2766,21 @@ export type GetPerforceServerStatusResponse = {
 
 /** Get object response which describes a device platform */
 export type GetDevicePlatformResponse = {
-	
-	/** Unique id of device platform */	
+
+	/** Unique id of device platform */
 	id: string;
-	
+
 	/** Friendly name of device platform */
-	name:string;
+	name: string;
 
 	/** Platform vendor models */
-	modelIds:string[];
+	modelIds: string[];
 }
 
 
 /** Device pool response object */
 export type GetDevicePoolResponse = {
-	
+
 	/**  Id of  the device pool */
 	id: string;
 
@@ -2799,43 +2817,43 @@ export type CreateDeviceResponse = {
 	id: string;
 
 }
-	
+
 /** Get response object which describes a device */
 export type GetDeviceResponse = {
-	
+
 	/** The unique id of the device */
 	id: string;
-	
+
 	/** The platform of the device */
 	platformId: string;
-	
+
 	/** The pool the device belongs to */
 	poolId: string;
-	
+
 	/** The friendly name of the device */
 	name: string;
-	
+
 	/** Whether the device is currently enabled */
 	enabled: boolean;
-	
+
 	/**  The address of the device (if it allows network connections) */
 	address?: string;
-	
+
 	/** The vendor model id of the device */
 	modelId?: string;
 
 	/** Any notes provided for the device */
-	notes? : string;	
-	
-	/** The UTC time when a device problem was reported */		
+	notes?: string;
+
+	/** The UTC time when a device problem was reported */
 	problemTime?: Date | string;
 
-	/** The UTC time when a device was set for maintenance */		
+	/** The UTC time when a device was set for maintenance */
 	maintenanceTime?: Date | string;
-	
+
 }
 
-/** Request object for updating a device */		
+/** Request object for updating a device */
 export type UpdateDeviceRequest = {
 
 	/** The pool to assign device to */
@@ -2847,11 +2865,11 @@ export type UpdateDeviceRequest = {
 	/** Whether the device is enabled */
 	enabled?: boolean;
 
-	/** Whether to clear any problem state */ 
-	problem?:boolean;
+	/** Whether to clear any problem state */
+	problem?: boolean;
 
 	/** Whether the device should be put into maintenance mode */
-	maintenance?:boolean;
+	maintenance?: boolean;
 
 	/** New address or hostname of device */
 	address?: string;
@@ -2863,40 +2881,45 @@ export type UpdateDeviceRequest = {
 	notes?: string;
 }
 
-	
+
 /** A reservation containing one or more devices */
 
 export type GetDeviceReservationResponse = {
-		
+
 	/** Randomly generated unique id for this reservation */
 	id: string;
 
 	/** Which device pool the reservation is in	*/
 	poolId: string;
-	
+
 	/** The reserved device ids	*/
-	devices:string[];
-	
-	/** JobID holding reservation */		
+	devices: string[];
+
+	/** JobID holding reservation */
 	jobId?: string;
-	
+
 	/** Job step id holding reservation	*/
 	stepId?: string;
-	
+
 	/** Reservations held by a user	*/
 	userId?: string;
-	
+
 	/** The hostname of machine holding reservation	*/
-	hostname?:string;
+	hostname?: string;
 
 	/** The optional details of the reservation	*/
-	reservationDetails?:string;
-	
-	/** The UTC time when the reservation was created */		
+	reservationDetails?: string;
+
+	/** The UTC time when the reservation was created */
 	createTimeUtc: Date | string;
-	
+
 	/** The legacy reservation system guid, to be removed once can update Gauntlet client in all streams */
 	legacyGuid: string;
 }
 
+/** Updates an existing lease */
+export type UpdateLeaseRequest = {
+	/** Mark this lease as aborted */
+	aborted?: boolean
+}
 
diff --git a/Engine/Source/Programs/Horde/HordeDashboard/src/backend/Backend.ts b/Engine/Source/Programs/Horde/HordeDashboard/src/backend/Backend.ts
--- a/Engine/Source/Programs/Horde/HordeDashboard/src/backend/Backend.ts
+++ b/Engine/Source/Programs/Horde/HordeDashboard/src/backend/Backend.ts
@@ -3,7 +3,7 @@
 import { configure } from "mobx";
 import { isNumber } from 'util';
 import templateCache from '../backend/TemplateCache';
-import { AgentData, ArtifactData, BatchUpdatePoolRequest, ChangeSummaryData, CreateJobRequest, CreateJobResponse, CreatePoolRequest, CreateSoftwareResponse, CreateSubscriptionRequest, CreateSubscriptionResponse, DashboardPreference, EventData, GetArtifactZipRequest, GetUserResponse, GetGraphResponse, GetIssueStreamResponse, GetJobStepRefResponse, GetJobStepTraceResponse, GetJobTimingResponse, GetLogEventResponse, GetNotificationResponse, GetSoftwareResponse, GetSubscriptionResponse, IssueData, IssueQuery, JobData, JobQuery, LeaseData, LogData, LogLineData, PoolData, ProjectData, ScheduleData, ScheduleQuery, SearchLogFileResponse, SessionData, StreamData, TemplateData, TestData, UpdateAgentRequest, UpdateIssueRequest, UpdateJobRequest, UpdateNotificationsRequest, UpdatePoolRequest, UpdateStepRequest, UpdateStepResponse, UpdateUserRequest, GetUtilizationTelemetryResponse, TabType, GetJobsTabResponse, JobsTabColumnType, GetPerforceServerStatusResponse, GetDeviceResponse, GetDevicePlatformResponse, CreateDeviceResponse, CreateDeviceRequest, UpdateDeviceRequest, GetDevicePoolResponse, GetDeviceReservationResponse, UsersQuery, LeaseOutcome } from './Api';
+import { AgentData, ArtifactData, BatchUpdatePoolRequest, ChangeSummaryData, CreateJobRequest, CreateJobResponse, CreatePoolRequest, CreateSoftwareResponse, CreateSubscriptionRequest, CreateSubscriptionResponse, DashboardPreference, EventData, GetArtifactZipRequest, GetUserResponse, GetGraphResponse, GetIssueStreamResponse, GetJobStepRefResponse, GetJobStepTraceResponse, GetJobTimingResponse, GetLogEventResponse, GetNotificationResponse, GetSoftwareResponse, GetSubscriptionResponse, IssueData, IssueQuery, JobData, JobQuery, LeaseData, LogData, LogLineData, PoolData, ProjectData, ScheduleData, ScheduleQuery, SearchLogFileResponse, SessionData, StreamData, TemplateData, TestData, UpdateAgentRequest, UpdateIssueRequest, UpdateJobRequest, UpdateNotificationsRequest, UpdatePoolRequest, UpdateStepRequest, UpdateStepResponse, UpdateUserRequest, GetUtilizationTelemetryResponse, TabType, GetJobsTabResponse, JobsTabColumnType, GetPerforceServerStatusResponse, GetDeviceResponse, GetDevicePlatformResponse, CreateDeviceResponse, CreateDeviceRequest, UpdateDeviceRequest, GetDevicePoolResponse, GetDeviceReservationResponse, UsersQuery, LeaseOutcome, UpdateLeaseRequest } from './Api';
 import dashboard from './Dashboard';
 import { ChallengeStatus, Fetch } from './Fetch';
 import graphCache, { GraphQuery } from './GraphCache';
@@ -1141,6 +1141,16 @@ export class Backend {
         });        
     }
 
+    updateLease(leaseId: string, update:UpdateLeaseRequest): Promise<void> {
+        return new Promise<void>((resolve, reject) => {
+            this.backend.put(`/api/v1/leases/${leaseId}`, update).then(() => {
+                resolve();
+            }).catch(reason => {
+                reject(reason);
+            });
+        });        
+    }
+
 
     deleteDevice(deviceId: string): Promise<void> {
         return new Promise<void>((resolve, reject) => {
diff --git a/Engine/Source/Programs/Horde/HordeServer/Api/Agents.cs b/Engine/Source/Programs/Horde/HordeServer/Api/Agents.cs
--- a/Engine/Source/Programs/Horde/HordeServer/Api/Agents.cs
+++ b/Engine/Source/Programs/Horde/HordeServer/Api/Agents.cs
@@ -175,6 +175,10 @@ namespace HordeServer.Api
 		/// </summary>
 		public LeaseOutcome? Outcome { get; set; }
 
+		/// <summary>
+		/// State of the lease (for AgentLeases)
+		/// </summary>
+		public LeaseState? State { get; set; }
 
 		/// <summary>
 		/// Parameterless constructor for serialization
@@ -194,6 +198,7 @@ namespace HordeServer.Api
 			this.Id = Lease.Id.ToString();
 			this.Name = Lease.Name;
 			this.LogId = Lease.LogId?.ToString();
+			this.State = Lease.State;
 			this.StartTime = Lease.StartTime;
 			this.Executing = Lease.Active;
 			this.FinishTime = Lease.ExpiryTime;
diff --git a/Engine/Source/Programs/Horde/HordeServer/Api/Leases.cs b/Engine/Source/Programs/Horde/HordeServer/Api/Leases.cs
new file mode 100644
--- /dev/null
+++ b/Engine/Source/Programs/Horde/HordeServer/Api/Leases.cs
@@ -0,0 +1,30 @@
+// Copyright Epic Games, Inc. All Rights Reserved.
+
+using EpicGames.Core;
+using HordeCommon;
+using HordeServer.Models;
+using HordeServer.Utilities;
+using System;
+using System.Collections.Generic;
+using System.ComponentModel.DataAnnotations;
+using System.Linq;
+using System.Security.Claims;
+using System.Text.Json.Serialization;
+using System.Threading.Tasks;
+
+namespace HordeServer.Api
+{
+
+	/// <summary>
+	/// Updates an existing lease
+	/// </summary>
+	public class UpdateLeaseRequest
+	{
+		/// <summary>
+		/// Mark this lease as aborted
+		/// </summary>
+		public bool? Aborted { get; set; }
+
+	}
+
+}
diff --git a/Engine/Source/Programs/Horde/HordeServer/Controllers/LeasesController.cs b/Engine/Source/Programs/Horde/HordeServer/Controllers/LeasesController.cs
--- a/Engine/Source/Programs/Horde/HordeServer/Controllers/LeasesController.cs
+++ b/Engine/Source/Programs/Horde/HordeServer/Controllers/LeasesController.cs
@@ -20,6 +20,7 @@ using System.Reflection;
 using System.Security.Claims;
 using System.Text.Json;
 using System.Threading.Tasks;
+using Microsoft.Extensions.Logging;
 
 using PoolId = HordeServer.Utilities.StringId<HordeServer.Models.IPool>;
 using AgentSoftwareChannelName = HordeServer.Utilities.StringId<HordeServer.Services.AgentSoftwareChannels>;
@@ -44,15 +45,22 @@ namespace HordeServer.Controllers
 		/// </summary>
 		AgentService AgentService;
 
+		/// <summary>
+		/// Logger instance
+		/// </summary>
+		private ILogger<LeasesController> Logger;
+
 		/// <summary>
 		/// Constructor
 		/// </summary>
 		/// <param name="AclService">The ACL service singleton</param>
 		/// <param name="AgentService">The agent service</param>
-		public LeasesController(AclService AclService, AgentService AgentService)
+		/// <param name="Logger">The controller logger</param>
+		public LeasesController(AclService AclService, AgentService AgentService, ILogger<LeasesController> Logger)
 		{
 			this.AclService = AclService;
 			this.AgentService = AgentService;
+			this.Logger = Logger;
 		}
 
 		/// <summary>
@@ -125,5 +133,58 @@ namespace HordeServer.Controllers
 
 			return new GetAgentLeaseResponse(Lease);
 		}
+
+		/// <summary>
+		/// Update a particular lease
+		/// </summary>
+		/// <param name="LeaseId">Unique id of the particular lease</param>
+		/// <returns>Lease matching the given id</returns>
+		[HttpPut]
+		[Route("/api/v1/leases/{LeaseId}")]
+		public async Task<ActionResult> UpdateLeaseAsync(string LeaseId, [FromBody] UpdateLeaseRequest Request)
+		{
+			ObjectId LeaseIdValue = LeaseId.ToObjectId();
+
+			// only update supported right now is abort
+			if (!Request.Aborted.HasValue || !Request.Aborted.Value)
+			{
+				return Ok();
+			}
+
+			ILease? Lease = await AgentService.GetLeaseAsync(LeaseIdValue);
+			if (Lease == null)
+			{
+				return NotFound();
+			}
+
+			IAgent? Agent = await AgentService.GetAgentAsync(Lease.AgentId);
+			if (Agent == null)
+			{
+				return NotFound();
+			}
+
+			if (!await AgentService.AuthorizeAsync(Agent, AclAction.AdminWrite, User, null))
+			{
+				return Forbid();
+			}
+
+			AgentLease? AgentLease = Agent.Leases.FirstOrDefault(x => x.Id == LeaseIdValue);
+
+			if (AgentLease == null)
+			{
+				return NotFound();
+			}
+
+			if (!AgentLease.IsActionLease() && !AgentLease.IsConformLease())
+			{
+				Logger.LogError("Lease abort only supported on action and conform leases for now, {LeaseId}", LeaseId);
+				return Ok();
+			}
+
+			await AgentService.CancelLeaseAsync(Agent, LeaseIdValue);
+
+			return Ok();
+		}
+
 	}
 }
diff --git a/Engine/Source/Programs/Horde/HordeServer/Models/Agent.cs b/Engine/Source/Programs/Horde/HordeServer/Models/Agent.cs
--- a/Engine/Source/Programs/Horde/HordeServer/Models/Agent.cs
+++ b/Engine/Source/Programs/Horde/HordeServer/Models/Agent.cs
@@ -339,6 +339,23 @@ namespace HordeServer.Models
 			}
 			return false;
 		}
+		/// <summary>
+		/// Determines if this is a remote action lease
+		/// </summary>
+		/// <returns>True if this is nn action lease</returns>
+		public bool IsActionLease()
+		{
+			if (Payload != null)
+			{
+				Any BasePayload = Any.Parser.ParseFrom(Payload);
+				if (BasePayload.Is(ActionTask.Descriptor))
+				{
+					return true;
+				}
+			}
+			return false;
+		}
+
 
 		/// <summary>
 		/// Gets user-readable payload information
@@ -377,6 +394,14 @@ namespace HordeServer.Models
 					Details["SoftwareId"] = UpgradeTask.SoftwareId;
 					Details["LogId"] = UpgradeTask.LogId;
 				}
+
+				ActionTask ActionTask;
+				if (BasePayload.TryUnpack(out ActionTask))
+				{
+					Details["Type"] = "Action";
+					Details["LogId"] = ActionTask.LogId;
+				}
+
 			}
 			return Details;
 		}
diff --git a/Engine/Source/Programs/Horde/HordeServer/Services/AgentService.cs b/Engine/Source/Programs/Horde/HordeServer/Services/AgentService.cs
--- a/Engine/Source/Programs/Horde/HordeServer/Services/AgentService.cs
+++ b/Engine/Source/Programs/Horde/HordeServer/Services/AgentService.cs
@@ -443,6 +443,35 @@ namespace HordeServer.Services
 			return Agent;
 		}
 
+		/// <summary>
+		/// Cancels the specified agent lease
+		/// </summary>
+		/// <param name="Agent">Agent to cancel the lease on</param>
+		/// <param name="LeaseId">The lease id to cancel</param>
+		/// <returns></returns>
+		public async Task<bool> CancelLeaseAsync(IAgent Agent, ObjectId LeaseId)
+		{		
+			int Index = 0;
+			while (Index < Agent.Leases.Count && Agent.Leases[Index].Id != LeaseId)
+			{
+				Index++;
+			}
+
+			if (Index == Agent.Leases.Count)
+			{
+				return false;
+			}
+
+			if (Agent.Leases[Index].State == LeaseState.Cancelled)
+			{
+				return false;
+			}
+
+			await Agents.TryCancelLeaseAsync(Agent, Index);
+
+			return true;
+		}
+
 		/// <summary>
 		/// Updates the state of the current agent session
 		/// </summary>
diff --git a/Engine/Source/Programs/Horde/HordeServer/Services/JobService.cs b/Engine/Source/Programs/Horde/HordeServer/Services/JobService.cs
--- a/Engine/Source/Programs/Horde/HordeServer/Services/JobService.cs
+++ b/Engine/Source/Programs/Horde/HordeServer/Services/JobService.cs
@@ -426,7 +426,7 @@ namespace HordeServer.Services
 				{
 					Index++;
 				}
-				if (Index == -1)
+				if (Index == Agent.Leases.Count)
 				{
 					break;
 				}
