Horde: Add max RAM info to agent properties.

diff --git a/Engine/Source/Programs/Horde/HordeAgent/Services/WorkerService.cs b/Engine/Source/Programs/Horde/HordeAgent/Services/WorkerService.cs
--- a/Engine/Source/Programs/Horde/HordeAgent/Services/WorkerService.cs
+++ b/Engine/Source/Programs/Horde/HordeAgent/Services/WorkerService.cs
@@ -1396,6 +1396,25 @@ namespace HordeAgent.Services
 				}
 
 				// Add RAM info
+				using (ManagementObjectSearcher Searcher = new ManagementObjectSearcher("select Capacity from Win32_PhysicalMemory"))
+				{
+					ulong TotalCapacity = 0;
+					foreach (ManagementObject Row in Searcher.Get())
+					{
+						object? Capacity = Row.GetPropertyValue("Capacity");
+						if (Capacity is ulong)
+						{
+							TotalCapacity += (ulong)Capacity;
+						}
+					}
+
+					if (TotalCapacity > 0)
+					{
+						PrimaryDevice.Properties.Add($"RAM={TotalCapacity / (1024 * 1024 * 1024)}");
+					}
+				}
+
+				// Add GPU info
 				using (ManagementObjectSearcher Searcher = new ManagementObjectSearcher("select Name, DriverVersion, AdapterRAM from Win32_VideoController"))
 				{
 					int Index = 0;
