Horde: Attempt to fix S3 auth failures in prod.

diff --git a/Engine/Source/Programs/Horde/HordeServer/Startup.cs b/Engine/Source/Programs/Horde/HordeServer/Startup.cs
--- a/Engine/Source/Programs/Horde/HordeServer/Startup.cs
+++ b/Engine/Source/Programs/Horde/HordeServer/Startup.cs
@@ -417,7 +417,7 @@ namespace HordeServer
 			Services.AddSingleton<DeviceService>();
 
 			AWSOptions AwsOptions = Configuration.GetAWSOptions();
-			if (Settings.S3AssumeArn != null)
+			if (Settings.S3CredentialType == "AssumeRole" && Settings.S3AssumeArn != null)
 			{
 				AwsOptions.Credentials = new AssumeRoleAWSCredentials(FallbackCredentialsFactory.GetCredentials(), Settings.S3AssumeArn, "Horde");
 			}
