Horde: Revert to older method for AWS authentication.

diff --git a/Engine/Source/Programs/Horde/HordeServer/Startup.cs b/Engine/Source/Programs/Horde/HordeServer/Startup.cs
--- a/Engine/Source/Programs/Horde/HordeServer/Startup.cs
+++ b/Engine/Source/Programs/Horde/HordeServer/Startup.cs
@@ -421,6 +421,10 @@ namespace HordeServer
 			{
 				AwsOptions.Credentials = new AssumeRoleAWSCredentials(FallbackCredentialsFactory.GetCredentials(), Settings.S3AssumeArn, "Horde");
 			}
+			else if(Settings.S3CredentialType == "AssumeRoleWebIdentity")
+			{
+				AwsOptions.Credentials = AssumeRoleWithWebIdentityCredentials.FromEnvironmentVariables();
+			}
 			Services.AddSingleton(AwsOptions);
 
 			Services.AddSingleton(new StorageBackendSettings<PersistentLogStorage> { Type = Settings.ExternalStorageProviderType, BaseDir = Settings.LocalLogsDir, BucketName = Settings.S3LogBucketName });
