Horde: Additional logging for S3 auth.

diff --git a/Engine/Source/Programs/Horde/HordeServer/Storage/Backends/AwsStorageBackend.cs b/Engine/Source/Programs/Horde/HordeServer/Storage/Backends/AwsStorageBackend.cs
--- a/Engine/Source/Programs/Horde/HordeServer/Storage/Backends/AwsStorageBackend.cs
+++ b/Engine/Source/Programs/Horde/HordeServer/Storage/Backends/AwsStorageBackend.cs
@@ -89,6 +89,8 @@ namespace HordeServer.Storage.Backends
 			this.Options = Options;
 			this.Semaphore = new SemaphoreSlim(16);
 			this.Logger = Logger;
+
+			Logger.LogInformation("Created AWS storage backend for bucket {BucketName} using credentials {Credentials} {CredentialsStr}", Options.BucketName, AwsOptions.Credentials.GetType(), AwsOptions.Credentials.ToString());
 		}
 
 		/// <inheritdoc/>
@@ -235,6 +237,7 @@ namespace HordeServer.Storage.Backends
 			long StreamLen = Stream.Length;
 			if (StreamLen < MinPartSize)
 			{
+				Logger.LogInformation("Writing to bucket {BucketName} key {FullPath}", Options.BucketName, FullPath);
 				PutObjectRequest UploadRequest = new PutObjectRequest();
 				UploadRequest.BucketName = Options.BucketName;
 				UploadRequest.Key = FullPath;
