Horde: Fix warning in leases controller.

diff --git a/Engine/Source/Programs/Horde/HordeServer/Controllers/LeasesController.cs b/Engine/Source/Programs/Horde/HordeServer/Controllers/LeasesController.cs
--- a/Engine/Source/Programs/Horde/HordeServer/Controllers/LeasesController.cs
+++ b/Engine/Source/Programs/Horde/HordeServer/Controllers/LeasesController.cs
@@ -138,6 +138,7 @@ namespace HordeServer.Controllers
 		/// Update a particular lease
 		/// </summary>
 		/// <param name="LeaseId">Unique id of the particular lease</param>
+		/// <param name="Request"></param>
 		/// <returns>Lease matching the given id</returns>
 		[HttpPut]
 		[Route("/api/v1/leases/{LeaseId}")]
