Horde: Buffer data in memory before uploading to S3. Seems to cause incorrect hash errors when using the stream directly (AWS SDK needs to seek through the stream to compute hash?)

diff --git a/Engine/Source/Programs/Horde/HordeServer/Storage/Backends/AwsStorageBackend.cs b/Engine/Source/Programs/Horde/HordeServer/Storage/Backends/AwsStorageBackend.cs
--- a/Engine/Source/Programs/Horde/HordeServer/Storage/Backends/AwsStorageBackend.cs
+++ b/Engine/Source/Programs/Horde/HordeServer/Storage/Backends/AwsStorageBackend.cs
@@ -237,13 +237,20 @@ namespace HordeServer.Storage.Backends
 			long StreamLen = Stream.Length;
 			if (StreamLen < MinPartSize)
 			{
-				Logger.LogInformation("Writing to bucket {BucketName} key {FullPath}", Options.BucketName, FullPath);
-				PutObjectRequest UploadRequest = new PutObjectRequest();
-				UploadRequest.BucketName = Options.BucketName;
-				UploadRequest.Key = FullPath;
-				UploadRequest.InputStream = Stream;
-				UploadRequest.Metadata.Add("bytes-written", StreamLen.ToString(CultureInfo.InvariantCulture));
-				await Client.PutObjectAsync(UploadRequest);
+				// Read the data into memory. Errors with hash value not matching if we don't do this (?)
+				byte[] Buffer = new byte[StreamLen];
+				await ReadExactLengthAsync(Stream, Buffer, (int)StreamLen);
+
+				// Upload it to S3
+				using (MemoryStream InputStream = new MemoryStream(Buffer))
+				{
+					PutObjectRequest UploadRequest = new PutObjectRequest();
+					UploadRequest.BucketName = Options.BucketName;
+					UploadRequest.Key = FullPath;
+					UploadRequest.InputStream = InputStream;
+					UploadRequest.Metadata.Add("bytes-written", StreamLen.ToString(CultureInfo.InvariantCulture));
+					await Client.PutObjectAsync(UploadRequest);
+				}
 			}
 			else
 			{
