Horde: Dont's strip search string from Okta redirect
#rnx

diff --git a/Engine/Source/Programs/Horde/HordeDashboard/src/backend/Fetch.ts b/Engine/Source/Programs/Horde/HordeDashboard/src/backend/Fetch.ts
--- a/Engine/Source/Programs/Horde/HordeDashboard/src/backend/Fetch.ts
+++ b/Engine/Source/Programs/Horde/HordeDashboard/src/backend/Fetch.ts
@@ -57,7 +57,7 @@ export class Fetch {
 
                 if (challenge === ChallengeStatus.Unauthorized) {
                     
-                    this.login(window.location.toString());
+                    this.login(window.location.toString() + window.location.search);
                     return;
                 }
 
