Horde: Remove endpoint to modify Perforce settings through REST API. This should be done through config files now.

diff --git a/Engine/Source/Programs/Horde/HordeServer/Controllers/PerforceController.cs b/Engine/Source/Programs/Horde/HordeServer/Controllers/PerforceController.cs
--- a/Engine/Source/Programs/Horde/HordeServer/Controllers/PerforceController.cs
+++ b/Engine/Source/Programs/Horde/HordeServer/Controllers/PerforceController.cs
@@ -62,30 +62,6 @@ namespace HordeServer.Controllers
 			return Responses;
 		}
 
-		/// <summary>
-		/// Update Perforce server settings
-		/// </summary>
-		[HttpPut]
-		[Route("/api/v1/perforce/settings")]
-		public async Task<ActionResult> UpdatePerforceSettingsAsync([FromBody] List<PerforceCluster> Clusters)
-		{
-			if (!await AclService.AuthorizeAsync(AclAction.AdminWrite, User))
-			{
-				return Forbid();
-			}
-
-			for (; ; )
-			{
-				Globals Globals = await DatabaseService.GetGlobalsAsync();
-				Globals.PerforceClusters = Clusters;
-
-				if (await DatabaseService.TryUpdateSingletonAsync(Globals))
-				{
-					return Ok("Settings have been updated");
-				}
-			}
-		}
-
 		/// <summary>
 		/// Gets the current perforce settinsg
 		/// </summary>
