Horde: Remove unused implementations of IPerforceService. Now using P4API all the way!

diff --git a/Engine/Source/Programs/Horde/HordeServer/ServerSettings.cs b/Engine/Source/Programs/Horde/HordeServer/ServerSettings.cs
--- a/Engine/Source/Programs/Horde/HordeServer/ServerSettings.cs
+++ b/Engine/Source/Programs/Horde/HordeServer/ServerSettings.cs
@@ -333,11 +333,6 @@ namespace HordeServer
 		/// </summary>
 		public string? ScheduleTimeZone { get; set; }
 
-		/// <summary>
-		/// Address of the Perforce bridge
-		/// </summary>
-		public string? PerforceBridge { get; set; }
-
 		/// <summary>
 		/// Token for interacting with Slack
 		/// </summary>
@@ -373,27 +368,6 @@ namespace HordeServer
 		/// </summary>
 		public Uri DashboardUrl { get; set; } = new Uri("https://localhost:3000");
 		
-		/// <summary>
-		/// The URL to Helix Swarm server
-		/// </summary>
-		public Uri? HelixSwarmServerUrl { get; set; }
-		
-		/// <summary>
-		/// Username to use when communicating with Helix Swarm
-		/// </summary>
-		public string? HelixSwarmUsername { get; set; }
-		
-		/// <summary>
-		/// Password to use when communicating with Helix Swarm
-		/// </summary>
-		public string? HelixSwarmPassword { get; set; }
-		
-		/// <summary>
-		/// Projects in Helix Swarm to monitor, including the state of reviews in each project
-		/// Multiple projects can be specified, separated by comma.
-		/// </summary>
-		public string? HelixSwarmProjects { get; set; }
-
 		/// <summary>
 		/// The p4 bridge server
 		/// </summary>
@@ -414,11 +388,6 @@ namespace HordeServer
 		/// </summary>
 		public bool P4BridgeCanImpersonate { get; set; } = false;
 
-		/// <summary>
-		/// Whether to use the (temporary) P4 router while transitioning to p4 api 
-		/// </summary>
-		public bool P4UseRouter { get; set; } = false;
-
 		/// <summary>
 		/// Set the minimum size of the global thread pool
 		/// This value has been found in need of tweaking to avoid timeouts with the Redis client during bursts
@@ -458,23 +427,5 @@ namespace HordeServer
 		/// Whether to open a browser on startup
 		/// </summary>
 		public bool OpenBrowser { get; set; } = false;
-
-		/// <summary>
-		/// Check if Helix Swarm is enabled
-		/// </summary>
-		/// <returns>True if enabled</returns>
-		public bool IsSwarmEnabled()
-		{
-			return HelixSwarmServerUrl != null;
-		}
-		
-		/// <summary>
-		/// Parse the comma-separated string of Swarm projects into an actual array
-		/// </summary>
-		/// <returns>Array of Swarm projects</returns>
-		public string[] GetSwarmProjects()
-		{
-			return HelixSwarmProjects != null ? HelixSwarmProjects.Split(",") : Array.Empty<string>();
-		}
 	}
 }
diff --git a/Engine/Source/Programs/Horde/HordeServer/Services/PerforceService.cs b/Engine/Source/Programs/Horde/HordeServer/Services/PerforceService.cs
--- a/Engine/Source/Programs/Horde/HordeServer/Services/PerforceService.cs
+++ b/Engine/Source/Programs/Horde/HordeServer/Services/PerforceService.cs
@@ -209,922 +209,4 @@ namespace HordeServer.Services
 			return Changes[0].Number;
 		}
 	}
-
-	/// <summary>
-	/// Base class for IPerforceService implementations
-	/// </summary>
-	public abstract class PerforceServiceBase : IPerforceService
-	{
-		/// <inheritdoc/>
-		public abstract Task<PerforceUserInfo?> GetUserInfoAsync(string UserName);
-
-		/// <inheritdoc/>
-		public abstract Task<int> CreateNewChangeAsync(string StreamName, string Path);
-
-		/// <inheritdoc/>
-		public abstract Task<List<ChangeSummary>> GetChangesAsync(string StreamName, int? MinChange, int? MaxChange, int Results, string? ImpersonateUser);
-		
-		/// <inheritdoc/>
-		public abstract Task<List<ChangeDetails>> GetChangeDetailsAsync(string StreamName, IReadOnlyList<int> ChangeNumbers, string? ImpersonateUser);
-
-		/// <inheritdoc/>
-		public abstract Task<string> CreateTicket(string ImpersonateUser);
-
-		/// <inheritdoc/>
-		public abstract Task<List<FileSummary>> FindFilesAsync(IEnumerable<string> Paths);
-
-		/// <inheritdoc/>
-		public abstract Task<byte[]> PrintAsync(string Path);
-
-		/// <inheritdoc/>
-		public abstract Task<int> DuplicateShelvedChangeAsync(int ShelvedChange);
-
-		/// <inheritdoc/>
-		public abstract Task<(int? Change, string Message)> SubmitShelvedChangeAsync(int Change, int OriginalChange);
-
-		/// <inheritdoc/>
-		public abstract Task DeleteShelvedChangeAsync(int ShelvedChange);
-
-		/// <inheritdoc/>
-		public abstract Task UpdateChangelistDescription(int Change, string Description);
-
-
-		/// <inheritdoc/>
-		public virtual async Task<int> GetCodeChangeAsync(string StreamName, int Change)
-		{
-			int MaxChange = Change;
-			for(; ;)
-			{
-				// Query for the changes before this point
-				List<ChangeSummary> Changes = await GetChangesAsync(StreamName, null, MaxChange, 10, null);
-				Serilog.Log.Logger.Information("Finding last code change in {Stream} before {MaxChange}: {NumResults}", StreamName, MaxChange, Changes.Count);
-				if (Changes.Count == 0)
-				{
-					return 0;
-				}
-
-				// Get the details for them
-				List<ChangeDetails> DetailsList = await GetChangeDetailsAsync(StreamName, Changes.ConvertAll(x => x.Number), null);
-				foreach (ChangeDetails Details in DetailsList.OrderByDescending(x => x.Number))
-				{
-					ChangeContentFlags ContentFlags = Details.GetContentFlags();
-					Serilog.Log.Logger.Information("Change {Change} = {Flags}", Details.Number, ContentFlags.ToString());
-					if ((Details.GetContentFlags() & ChangeContentFlags.ContainsCode) != 0)
-					{
-						return Details.Number;
-					}
-				}
-
-				// Loop round again, adjusting our maximum changelist number
-				MaxChange = Changes.Min(x => x.Number) - 1;
-			}
-		}
-
-		/// <summary>
-		/// Gets the wildcard filter for a particular range of changes
-		/// </summary>
-		/// <param name="BasePath">Base path to find files for</param>
-		/// <param name="MinChange">Minimum changelist number to query</param>
-		/// <param name="MaxChange">Maximum changelist number to query</param>
-		/// <returns>Filter string</returns>
-		public static string GetFilter(string BasePath, int? MinChange, int? MaxChange)
-		{
-			StringBuilder Filter = new StringBuilder(BasePath);
-			if (MinChange != null && MaxChange != null)
-			{
-				Filter.Append($"@{MinChange},{MaxChange}");
-			}
-			else if (MinChange != null)
-			{
-				Filter.Append($"@>={MinChange}");
-			}
-			else if (MaxChange != null)
-			{
-				Filter.Append($"@<={MaxChange}");
-			}
-			return Filter.ToString();
-		}
-
-		/// <summary>
-		/// Gets a stream-relative path from a depot path
-		/// </summary>
-		/// <param name="DepotFile">The depot file to check</param>
-		/// <param name="StreamName">Name of the stream</param>
-		/// <param name="RelativePath">Receives the relative path to the file</param>
-		/// <returns>True if the stream-relative path was returned</returns>
-		public static bool TryGetStreamRelativePath(string DepotFile, string StreamName, [NotNullWhen(true)] out string? RelativePath)
-		{
-			if (DepotFile.StartsWith(StreamName, StringComparison.OrdinalIgnoreCase) && DepotFile.Length > StreamName.Length && DepotFile[StreamName.Length] == '/')
-			{
-				RelativePath = DepotFile.Substring(StreamName.Length);
-				return true;
-			}
-
-			Match Match = Regex.Match(DepotFile, "^//[^/]+/[^/]+(/.*)$");
-			if (Match.Success)
-			{
-				RelativePath = Match.Groups[1].Value;
-				return true;
-			}
-
-			RelativePath = null;
-			return false;
-		}
-	}
-
-	/// <summary>
-	/// Empty implementation of IPerforceService
-	/// </summary>
-	public class EmptyPerforceService : IPerforceService
-	{
-		/// <inheritdoc/>
-		public Task<PerforceUserInfo?> GetUserInfoAsync(string UserName)
-		{
-			throw new NotImplementedException();
-		}
-
-		/// <inheritdoc/>
-		public Task<int> CreateNewChangeAsync(string StreamName, string Path)
-		{
-			throw new NotImplementedException();
-		}
-
-		/// <inheritdoc/>
-		public Task<List<ChangeSummary>> GetChangesAsync(string StreamName, int? MinChange, int? MaxChange, int Results, string? ImpersonateUser)
-		{
-			return Task.FromResult<List<ChangeSummary>>(new List<ChangeSummary>());
-		}
-
-		/// <inheritdoc/>
-		public Task<List<ChangeDetails>> GetChangeDetailsAsync(string StreamName, IReadOnlyList<int> ChangeNumbers, string? ImpersonateUser)
-		{
-			throw new NotImplementedException();
-		}
-
-		/// <inheritdoc />
-		public Task<string> CreateTicket(string ImpersonateUser)
-		{
-			throw new NotImplementedException();
-		}
-
-		/// <inheritdoc/>
-		public Task<int> GetCodeChangeAsync(string StreamName, int Change)
-		{
-			throw new NotImplementedException();
-		}
-
-		/// <inheritdoc/>
-		public Task<List<FileSummary>> FindFilesAsync(IEnumerable<string> Paths)
-		{
-			throw new NotImplementedException();
-		}
-
-		/// <inheritdoc/>
-		public Task<byte[]> PrintAsync(string Path)
-		{
-			throw new NotImplementedException();
-		}
-
-		/// <inheritdoc/>
-		public Task<int> DuplicateShelvedChangeAsync(int ShelvedChange)
-		{
-			throw new NotImplementedException();
-		}
-
-		/// <inheritdoc/>
-		public Task<(int? Change, string Message)> SubmitShelvedChangeAsync(int ShelvedChange, int OriginalChange)
-		{
-			throw new NotImplementedException();
-		}
-
-		/// <inheritdoc/>
-		public Task DeleteShelvedChangeAsync(int ShelvedChange)
-		{
-			throw new NotImplementedException();
-		}
-
-		/// <inheritdoc/>
-		public Task UpdateChangelistDescription(int Change, string Description)
-		{
-			throw new NotImplementedException();
-		}
-	}
-
-	/// <summary>
-	/// Local implementation of the Perforce service
-	/// </summary>
-	public class LocalPerforceService : PerforceServiceBase
-	{
-		/// <summary>
-		/// The Perforce connection
-		/// </summary>
-		PerforceConnection Perforce;
-
-		/// <summary>
-		/// Cached client connections
-		/// </summary>
-		Dictionary<string, Task<PerforceConnection>> Clients = new Dictionary<string, Task<PerforceConnection>>(StringComparer.OrdinalIgnoreCase);
-
-		/// <summary>
-		/// Information about the Perforce server
-		/// </summary>
-		Task<InfoRecord> InfoTask;
-
-		/// <summary>
-		/// Constructor
-		/// </summary>
-		/// <param name="Logger">Logger</param>
-		public LocalPerforceService(ILogger<LocalPerforceService> Logger)
-		{
-			Perforce = new PerforceConnection(null, null, Logger);
-			InfoTask = Perforce.GetInfoAsync(InfoOptions.None, CancellationToken.None);
-		}
-
-		/// <inheritdoc/>
-		public async override Task<PerforceUserInfo?> GetUserInfoAsync(string UserName)
-		{
-			PerforceResponseList<UserRecord> UserResponses = await Perforce.TryGetUserAsync(UserName, CancellationToken.None);
-			foreach(PerforceResponse<UserRecord> UserResponse in UserResponses)
-			{
-				if (UserResponse.Succeeded && UserResponse.Data.Access != default)
-				{
-					PerforceUserInfo UserInfo = new PerforceUserInfo();
-					UserInfo.Name = UserResponse.Data.FullName;
-					UserInfo.Email = UserResponse.Data.Email;
-					return UserInfo;
-				}
-			}
-			return null;
-		}
-
-		/// <inheritdoc/>
-		public async override Task<int> CreateNewChangeAsync(string StreamName, string FilePath)
-		{
-			string ClientName = $"Horde-NewChange-{Guid.NewGuid().ToString("N")}";
-			DirectoryReference ClientDir = DirectoryReference.Combine(new DirectoryReference(Path.GetTempPath()), ClientName);
-
-			InfoRecord Info = await InfoTask;
-
-			ClientRecord Client = new ClientRecord(ClientName, Info.UserName!, ClientDir.FullName);
-			Client.Stream = StreamName;
-			await Perforce.CreateClientAsync(Client, CancellationToken.None);
-
-			PerforceConnection PerforceClient = new PerforceConnection(Perforce) { ClientName = ClientName };
-			try
-			{
-				ChangeRecord ChangeRecord = new ChangeRecord();
-				ChangeRecord.Client = ClientName;
-				ChangeRecord.Description = "New change for Horde job";
-				ChangeRecord = await PerforceClient.CreateChangeAsync(ChangeRecord, CancellationToken.None);
-
-				string ClientFilePath = $"//{ClientName}/{FilePath.TrimStart('/')}";
-				for (; ; )
-				{
-					await PerforceClient.TrySyncAsync(SyncOptions.Force, -1, new[] { ClientFilePath }, CancellationToken.None);
-					await PerforceClient.TryEditAsync(ChangeRecord.Number, new[] { ClientFilePath }, CancellationToken.None);
-
-					PerforceResponse<SubmitRecord> Response = await PerforceClient.TrySubmitAsync(ChangeRecord.Number, SubmitOptions.None, CancellationToken.None);
-					if (Response.Info == null || !Response.Info.Data.Contains(" - must resolve", StringComparison.Ordinal))
-					{
-						return Response.Data.ChangeNumber;
-					}
-
-					await PerforceClient.RevertAsync(-1, null, RevertOptions.None, new[] { "//..." }, CancellationToken.None);
-				}
-			}
-			finally
-			{
-				await PerforceClient.TryRevertAsync(-1, null, RevertOptions.None, new[] { "//..." }, CancellationToken.None);
-				await Perforce.TryDeleteClientAsync(DeleteClientOptions.None, ClientName, CancellationToken.None);
-			}
-		}
-
-		/// <inheritdoc/>
-		public override async Task<List<ChangeSummary>> GetChangesAsync(string StreamName, int? MinChange, int? MaxChange, int Results, string? ImpersonateUser)
-		{
-			PerforceConnection Client = await CreateReadOnlyClientAsync(StreamName, ImpersonateUser);
-
-			List<ChangeSummary> Changes = new List<ChangeSummary>();
-
-			string Filter = GetFilter($"//{Client.ClientName}/...", MinChange, MaxChange);
-
-			List<ChangesRecord> ChangeRecords = await Client.GetChangesAsync(ChangesOptions.LongOutput, Results, ChangeStatus.Submitted, new[] { Filter }, CancellationToken.None);
-			foreach (ChangesRecord ChangeRecord in ChangeRecords)
-			{
-				Changes.Add(new ChangeSummary(ChangeRecord.Number, ChangeRecord.User ?? String.Empty, ChangeRecord.Description ?? String.Empty));
-			}
-
-			return Changes;
-		}
-
-		/// <inheritdoc/>
-		public override async Task<List<ChangeDetails>> GetChangeDetailsAsync(string StreamName, IReadOnlyList<int> ChangeNumbers, string? ImpersonateUser)
-		{
-			// Create a client for this stream, in case it's a task stream
-			PerforceConnection Client = await CreateReadOnlyClientAsync(StreamName, ImpersonateUser);
-
-			// Get the Perforce description
-			List<DescribeRecord> DescribeRecords = await Client.DescribeAsync(ChangeNumbers.ToArray(), CancellationToken.None);
-
-			// Convert it to ChangeDetails objects
-			List<ChangeDetails> Results = new List<ChangeDetails>();
-			foreach (DescribeRecord DescribeRecord in DescribeRecords)
-			{
-				List<string> Files = new List<string>();
-				foreach (DescribeFileRecord DescribeFile in DescribeRecord.Files)
-				{
-					string? RelativePath;
-					if (TryGetStreamRelativePath(DescribeFile.DepotFile, StreamName, out RelativePath))
-					{
-						Files.Add(RelativePath);
-					}
-				}
-				Results.Add(new ChangeDetails(DescribeRecord.Number, DescribeRecord.User, DescribeRecord.Description, Files));
-			}
-			return Results;
-		}
-
-		/// <inheritdoc />
-		public override Task<string> CreateTicket(string ImpersonateUser)
-		{
-			throw new NotImplementedException();
-		}
-
-		/// <inheritdoc/>
-		public override async Task<List<FileSummary>> FindFilesAsync(IEnumerable<string> Paths)
-		{
-			List<FileSummary> Results = new List<FileSummary>();
-
-			PerforceResponseList<FStatRecord> Responses = await Perforce.TryFStatAsync(FStatOptions.ShortenOutput, Paths.ToArray(), CancellationToken.None);
-			foreach (PerforceResponse<FStatRecord> Response in Responses)
-			{
-				if (Response.Succeeded)
-				{
-					FStatRecord Record = Response.Data;
-					if (Record.DepotFile != null)
-					{
-						Results.Add(new FileSummary(Record.DepotFile, Record.HeadAction != FileAction.Delete && Record.HeadAction != FileAction.MoveDelete, Record.HeadChange));
-					}
-				}
-			}
-
-			return Results;
-		}
-
-		/// <inheritdoc/>
-		public override async Task<byte[]> PrintAsync(string DepotPath)
-		{
-			string TempFile = Path.GetTempFileName();
-			try
-			{
-				await Perforce.PrintAsync(TempFile, DepotPath, CancellationToken.None);
-				return await File.ReadAllBytesAsync(TempFile);
-			}
-			finally
-			{
-				if (File.Exists(TempFile))
-				{
-					File.SetAttributes(TempFile, FileAttributes.Normal);
-					File.Delete(TempFile);
-				}
-			}
-		}
-
-		/// <inheritdoc/>
-		public override async Task<(int? Change, string Message)> SubmitShelvedChangeAsync(int Change, int OriginalChange)
-		{
-			PerforceResponse<SubmitRecord> Response = await Perforce.TrySubmitShelvedAsync(Change, CancellationToken.None);
-			if (Response.Succeeded)
-			{
-				return (Response.Data.ChangeNumber, $"Submitted in CL {Response.Data.ChangeNumber}");
-			}
-			else
-			{
-				return (null, Response.Error!.Data);
-			}
-		}
-
-		/// <summary>
-		/// Creates a read-only client for the given stream
-		/// </summary>
-		/// <param name="StreamName">The stream to query</param>
-		/// <param name="ImpersonateUser">Name of the user to impersonate</param>
-		/// <returns>Client connection</returns>
-		async Task<PerforceConnection> CreateReadOnlyClientAsync(string StreamName, string? ImpersonateUser)
-		{
-			InfoRecord Info = await InfoTask;
-			string UserName = ImpersonateUser ?? Info.UserName!;
-			string HostName = Info.ClientHost!;
-			string ClientName = $"Horde-{UserName}-{HostName}-{StreamName.Replace('/', '+').Trim('+')}";
-
-			Task<PerforceConnection>? Result;
-			lock (Clients)
-			{
-				if (!Clients.TryGetValue(ClientName, out Result))
-				{
-					Result = CreateNewReadOnlyClientAsync(StreamName, UserName, ClientName);
-					Clients.Add(ClientName, Result);
-				}
-			}
-			return await Result;
-		}
-
-		/// <inheritdoc/>
-		public override Task<int> DuplicateShelvedChangeAsync(int ShelvedChange)
-		{
-			return Task.FromResult(0);
-		}
-
-		/// <summary>
-		/// Creates a new read-only client for the given stream
-		/// </summary>
-		/// <param name="StreamName">The stream to query</param>
-		/// <param name="UserName">Name of the user to own the client</param>
-		/// <param name="ClientName">Name of the client</param>
-		/// <returns>Client connection</returns>
-		async Task<PerforceConnection> CreateNewReadOnlyClientAsync(string StreamName, string UserName, string ClientName)
-		{
-			await Task.Yield();
-
-			PerforceConnection PerforceClient = new PerforceConnection(Perforce) { UserName = UserName };
-
-			ClientRecord Client = new ClientRecord(ClientName, UserName, "C:\\");
-			Client.Stream = StreamName;
-			await PerforceClient.CreateClientAsync(Client, CancellationToken.None);
-
-			PerforceClient.ClientName = ClientName;
-			return PerforceClient;
-		}
-
-		/// <inheritdoc/>
-		public override Task DeleteShelvedChangeAsync(int ShelvedChange)
-		{
-			throw new NotImplementedException();
-		}
-
-		/// <inheritdoc/>
-		public override Task UpdateChangelistDescription(int Change, string Description)
-		{
-			throw new NotImplementedException();
-		}
-	}
-
-	/// <summary>
-	/// Implementation of IPerforceService which connects to a HordePerforceBridge instance
-	/// </summary>
-	public sealed class BridgePerforceService : PerforceServiceBase, IDisposable
-	{
-#pragma warning disable CA1812
-		class BridgeChange
-		{
-			public string Change { get; set; } = String.Empty;
-			public string Description { get; set; } = String.Empty;
-			public string User { get; set; } = String.Empty;
-		}
-
-		class BridgeChanges
-		{
-			public List<BridgeChange>? Summaries { get; set; }
-
-			public List<string>? Errors { get; set; } = null;
-		}
-
-		class BridgeChangeFile
-		{
-			public string DepotFile { get; set; } = String.Empty;
-		}
-		
-		class BridgeUserData
-		{
-			public string? Name { get; set; } = String.Empty;
-			public string? Email { get; set; } = String.Empty;
-		}
-
-		class BridgeChangeDetail
-		{
-			public string Change { get; set; } = String.Empty;
-			public string Description { get; set; } = String.Empty;
-			public string User { get; set; } = String.Empty;
-			public List<BridgeChangeFile> Files { get; set; } = new List<BridgeChangeFile>();
-		}
-
-		class BridgeChangeDetails
-		{
-			public List<BridgeChangeDetail> Changes { get; set; } = new List<BridgeChangeDetail>();
-
-			public List<string>? Errors { get; set; } = null;
-		}
-
-		class BridgeTicket
-		{
-			public string Ticket { get; set; } = String.Empty;
-			public List<string>? Errors { get; set; } = null;
-		}
-
-		class BridgeUser
-		{
-			public BridgeUserData? User { get; set; } = null;
-			public List<string>? Errors { get; set; } = null;
-		}
-
-		class BridgeChangelistNumber
-		{
-			public string Change { get; set; } = String.Empty;
-			public string Message { get; set; } = string.Empty;
-			public List<string>? Errors { get; set; } = null;
-		}
-
-		class BridgeFStat
-		{
-			public string DepotFile { get; set; } = String.Empty;
-			public bool Exists { get; set; } = false;
-			public string HeadChange { get; set; } = String.Empty;
-			public string Errors { get; set; } = String.Empty;
-		}
-
-		class BridgePrint
-		{
-			public string Data { get; set; } = String.Empty;
-			public List<string>? Errors { get; set; } = null;
-			public bool Binary { get; set; } = false;
-		}
-#pragma warning restore CA1812
-
-		HttpClient Client;
-		AsyncPolicy<HttpResponseMessage> RetryPolicy;
-		Uri Server;
-		ILogger Logger;
-
-		/// <summary>
-		/// Constructor
-		/// </summary>
-		public BridgePerforceService(string Address, ILogger Logger)
-		{
-			this.Logger = Logger;
-
-			Client = new HttpClient();
-			Client.Timeout = TimeSpan.FromSeconds(300.0);
-			Server = new Uri(Address);
-			RetryPolicy = HttpPolicyExtensions.HandleTransientHttpError().WaitAndRetryAsync(3, Attempt => TimeSpan.FromSeconds(Math.Pow(2.0, Attempt)));
-		}
-
-		/// <inheritdoc/>
-		public void Dispose()
-		{
-			Client.Dispose();
-		}
-
-		/// <inheritdoc/>
-		public override async Task<PerforceUserInfo?> GetUserInfoAsync(string UserName)
-		{
-			Uri Uri = new Uri(Server, $"/api/v1/user?username={UserName}");
-
-			HttpResponseMessage Response = await RetryPolicy.ExecuteAsync(() => Client.GetAsync(Uri));
-			if (!Response.IsSuccessStatusCode)
-			{
-				Logger.LogError("GET to {Uri} returned {Code} ({Response})", Uri, Response.StatusCode, await Response.Content.ReadAsStringAsync());
-				throw new Exception("CreateTicket call to PerforceBridge failed");
-			}
-
-			byte[] Data = await Response.Content.ReadAsByteArrayAsync();
-			BridgeUser Result = JsonSerializer.Deserialize<BridgeUser>(Data.AsSpan());
-
-			if (Result.Errors != null)
-			{
-				Logger.LogError("{Uri} reported errors: {Errors}", Uri, string.Join("\n", Result.Errors));
-				throw new Exception($"{Uri} reported errors: {string.Join("\n", Result.Errors)}");
-			}
-
-			if(Result.User == null)
-			{
-				Logger.LogWarning("User info query returned not found for username {UserName}", UserName);
-				return null;
-			}
-
-			return new PerforceUserInfo { Name = Result.User.Name, Email = Result.User.Email };
-		}
-
-		/// <inheritdoc/>
-		public override async Task<int> CreateNewChangeAsync(string StreamName, string Path)
-		{
-			List<string> Arguments = new List<string>();
-			Arguments.Add($"stream={StreamName}");
-			Arguments.Add($"path={Path}");
-
-			Uri Uri = new Uri(Server, $"/api/v1/createnewchange?{String.Join("&", Arguments)}");
-
-			HttpResponseMessage Response = await RetryPolicy.ExecuteAsync(() => Client.GetAsync(Uri));
-			if (!Response.IsSuccessStatusCode)
-			{
-				Logger.LogError("GET to {Uri} returned {Code} ({Response})", Uri, Response.StatusCode, await Response.Content.ReadAsStringAsync());
-				throw new Exception("Query to PerforceBridge failed");
-			}
-
-			byte[] Data = await Response.Content.ReadAsByteArrayAsync();
-
-			BridgeChangelistNumber Result = JsonSerializer.Deserialize<BridgeChangelistNumber>(Data);
-			if (Result.Errors != null)
-			{
-				Logger.LogError("{Uri} reported errors: {Errors}", Uri, string.Join("\n", Result.Errors));
-				throw new Exception($"{Uri} reported errors: {string.Join("\n", Result.Errors)}");
-			}
-
-			int ChangeNumber;
-			if(!int.TryParse(Result.Change, out ChangeNumber))
-			{
-				throw new Exception($"Missing changelist number from response. Data: {Encoding.UTF8.GetString(Data)}");
-			}
-
-			return ChangeNumber;
-		}
-
-		/// <inheritdoc/>
-		public override async Task<List<ChangeSummary>> GetChangesAsync(string StreamName, int? MinChange, int? MaxChange, int Results, string? ImpersonateUser)
-		{
-			List<string> Arguments = new List<string>();
-			Arguments.Add($"stream={StreamName}");
-			if (MinChange != null)
-			{
-				Arguments.Add($"minchange={MinChange.Value}");
-			}
-			if (MaxChange != null)
-			{
-				Arguments.Add($"maxchange={MaxChange.Value}");
-			}
-			if (ImpersonateUser != null)
-			{
-				Arguments.Add($"p4iuser={ImpersonateUser}");
-			}
-			Arguments.Add($"maxresults={Results}");
-
-			Uri Uri = new Uri(Server, $"/api/v1/changes?{String.Join("&", Arguments)}");
-
-			HttpResponseMessage Response = await RetryPolicy.ExecuteAsync(() => Client.GetAsync(Uri));
-			if (!Response.IsSuccessStatusCode)
-			{
-				Logger.LogError("GET to {Uri} returned {Code} ({Response})", Uri, Response.StatusCode, await Response.Content.ReadAsStringAsync());
-				throw new Exception("Query to PerforceBridge failed");
-			}
-
-			byte[] Data = await Response.Content.ReadAsByteArrayAsync();
-
-			BridgeChanges Result = JsonSerializer.Deserialize<BridgeChanges>(Data.AsSpan());
-			if(Result.Errors != null)
-			{
-				Logger.LogError("{Uri} reported errors: {Errors}", Uri, string.Join("\n", Result.Errors));
-				throw new Exception($"{Uri} reported errors: {string.Join("\n", Result.Errors)}");
-			}
-
-			try
-			{
-				if (Result.Summaries == null)
-				{
-					return new List<ChangeSummary>();
-				}
-				else
-				{
-					return Result.Summaries.ConvertAll(x => new ChangeSummary(int.Parse(x.Change, CultureInfo.InvariantCulture), x.User, x.Description));
-				}
-			}
-			catch(Exception Ex)
-			{
-				throw new Exception($"Unable to parse changelist summaries from response: {Encoding.UTF8.GetString(Data)}", Ex); 
-			}
-		}
-
-		/// <inheritdoc/>
-		public override async Task<List<ChangeDetails>> GetChangeDetailsAsync(string StreamName, IReadOnlyList<int> ChangeNumbers, string? ImpersonateUser)
-		{
-			List<ChangeDetails> Results = new List<ChangeDetails>();
-
-			List<string> Arguments = new List<string>();
-			Arguments.Add($"change={string.Join(';', ChangeNumbers)}");
-			if(ImpersonateUser != null)
-			{
-				Arguments.Add($"iuser={ImpersonateUser}");
-			}
-
-			Uri Uri = new Uri(Server, $"/api/v1/describe?{String.Join("&", Arguments)}");
-
-			HttpResponseMessage Response = await RetryPolicy.ExecuteAsync(() => Client.GetAsync(Uri));
-			if (!Response.IsSuccessStatusCode)
-			{
-				Logger.LogError("GET to {Uri} returned {Code} ({Response})", Uri, Response.StatusCode, await Response.Content.ReadAsStringAsync());
-				throw new Exception("Query to PerforceBridge failed");
-			}
-
-			byte[] Data = await Response.Content.ReadAsByteArrayAsync();
-
-			BridgeChangeDetails Result = JsonSerializer.Deserialize<BridgeChangeDetails>(Data.AsSpan());
-
-			if (Result.Errors != null)
-			{
-				Logger.LogError("{Uri} reported errors: {Errors}", Uri, string.Join("\n", Result.Errors));
-				throw new Exception($"{Uri} reported errors: {string.Join("\n", Result.Errors)}");
-			}
-
-			foreach (BridgeChangeDetail Details in Result.Changes)
-			{ 
-				List<string> Files = new List<string>();
-				foreach (BridgeChangeFile File in Details.Files)
-				{
-					string? RelativePath;
-					if (TryGetStreamRelativePath(File.DepotFile, StreamName, out RelativePath))
-					{
-						Files.Add(RelativePath);
-					}
-				}
-				Results.Add(new ChangeDetails(int.Parse(Details.Change, CultureInfo.InvariantCulture), Details.User, Details.Description, Files));				
-			}
-			return Results;
-		}
-
-		/// <inheritdoc />
-		public override async Task<string> CreateTicket(string ImpersonateUser)
-		{
-			Uri Uri = new Uri(Server, $"/api/v1/ticket?p4iuser={ImpersonateUser}");
-
-			HttpResponseMessage Response = await RetryPolicy.ExecuteAsync(() => Client.GetAsync(Uri));
-			if (!Response.IsSuccessStatusCode)
-			{
-				Logger.LogError("GET to {Uri} returned {Code} ({Response})", Uri, Response.StatusCode, await Response.Content.ReadAsStringAsync());
-				throw new Exception("CreateTicket call to PerforceBridge failed");
-			}
-
-			byte[] Data = await Response.Content.ReadAsByteArrayAsync();
-			BridgeTicket Result = JsonSerializer.Deserialize<BridgeTicket>(Data.AsSpan());
-
-			if (Result.Errors != null)
-			{
-				Logger.LogError("{Uri} reported errors: {Errors}", Uri, string.Join("\n", Result.Errors));
-				throw new Exception($"{Uri} reported errors: {string.Join("\n", Result.Errors)}");
-			}
-
-			return Result.Ticket;
-		}
-
-		/// <inheritdoc/>
-		public override async Task<List<FileSummary>> FindFilesAsync(IEnumerable<string> Paths)
-		{
-			List<FileSummary> Results = new List<FileSummary>();
-			Uri Uri = new Uri(Server, $"/api/v1/findfiles?paths={string.Join(';', Paths)}");
-			HttpResponseMessage Response = await RetryPolicy.ExecuteAsync(() => Client.GetAsync(Uri));
-			if (!Response.IsSuccessStatusCode)
-			{
-				Logger.LogError("GET to {Uri} returned {Code} ({Response})", Uri, Response.StatusCode, await Response.Content.ReadAsStringAsync());
-				throw new Exception("Query to PerforceBridge failed");
-			}
-
-			byte[] Data = await Response.Content.ReadAsByteArrayAsync();
-			List<BridgeFStat> FStatResults = JsonSerializer.Deserialize<List<BridgeFStat>>(Data.AsSpan());
-			foreach(BridgeFStat Result in FStatResults)
-			{
-				if(Result.DepotFile != null)
-				{
-					Results.Add(new FileSummary(Result.DepotFile, Result.Exists, int.Parse(Result.HeadChange, CultureInfo.InvariantCulture), Result.Errors));
-				}
-			}
-			return Results;
-		}
-
-		/// <inheritdoc/>
-		public override async Task<byte[]> PrintAsync(string Path)
-		{
-			List<FileSummary> Results = new List<FileSummary>();
-			Uri Uri = new Uri(Server, $"/api/v1/print?file={Path}");
-			HttpResponseMessage Response = await RetryPolicy.ExecuteAsync(() => Client.GetAsync(Uri));
-			if (!Response.IsSuccessStatusCode)
-			{
-				Logger.LogError("GET to {Uri} returned {Code} ({Response})", Uri, Response.StatusCode, await Response.Content.ReadAsStringAsync());
-				throw new Exception("Query to PerforceBridge failed");
-			}
-
-			byte[] Data = await Response.Content.ReadAsByteArrayAsync();
-			BridgePrint Result = JsonSerializer.Deserialize<BridgePrint>(Data.AsSpan());
-
-			if (Result.Errors != null)
-			{
-				Logger.LogError("{Uri} reported errors: {Errors}", Uri, string.Join("\n", Result.Errors));
-				throw new Exception($"{Uri} reported errors: {string.Join("\n", Result.Errors)}");
-			}
-
-			if(Result.Data == null)
-			{
-				Logger.LogError("{Uri} resulted in no data being returned. Was {Path} deleted?", Uri, Path);
-				throw new Exception($"{Uri} resulted in no data being returned. Was {Path} deleted?");
-			}
-
-			if(Result.Binary)
-			{
-				return Convert.FromBase64String(Result.Data);
-			}
-			return Encoding.Default.GetBytes(Result.Data);
-		}
-
-		/// <inheritdoc/>
-		public override async Task<int> DuplicateShelvedChangeAsync(int ShelvedChange)
-		{
-			List<string> Arguments = new List<string>();
-			Arguments.Add($"change={ShelvedChange}");
-
-			Uri Uri = new Uri(Server, $"/api/v2/duplicateshelvedchange?{String.Join("&", Arguments)}");
-
-			HttpResponseMessage Response = await RetryPolicy.ExecuteAsync(() => Client.GetAsync(Uri));
-			if (!Response.IsSuccessStatusCode)
-			{
-				Logger.LogError("GET to {Uri} returned {Code} ({Response})", Uri, Response.StatusCode, await Response.Content.ReadAsStringAsync());
-				throw new Exception("Query to PerforceBridge failed");
-			}
-
-			byte[] Data = await Response.Content.ReadAsByteArrayAsync();
-			BridgeChangelistNumber Result = JsonSerializer.Deserialize<BridgeChangelistNumber>(Data.AsSpan());
-
-			if (Result.Errors != null)
-			{
-				Logger.LogError("{Uri} reported errors: {Errors}", Uri, string.Join("\n", Result.Errors));
-				throw new Exception($"{Uri} reported errors: {string.Join("\n", Result.Errors)}");
-			}
-
-			return int.Parse(Result.Change, CultureInfo.InvariantCulture);
-
-		}
-		/// <inheritdoc/>
-		public override async Task<(int? Change, string Message)> SubmitShelvedChangeAsync(int Change, int OriginalChange)
-		{
-			List<string> Arguments = new List<string>();
-			Arguments.Add($"change={Change}");
-			Arguments.Add($"originalchange={OriginalChange}");
-
-			Uri Uri = new Uri(Server, $"/api/v2/submitshelvedchange?{String.Join("&", Arguments)}");
-
-			HttpResponseMessage Response = await RetryPolicy.ExecuteAsync(() => Client.GetAsync(Uri));
-			if (!Response.IsSuccessStatusCode)
-			{
-				Logger.LogError("GET to {Uri} returned {Code} ({Response})", Uri, Response.StatusCode, await Response.Content.ReadAsStringAsync());
-				return (null, "Query to PerforceBridge failed");
-			}
-
-			byte[] Data = await Response.Content.ReadAsByteArrayAsync();
-			BridgeChangelistNumber Result = JsonSerializer.Deserialize<BridgeChangelistNumber>(Data.AsSpan());
-
-			if (Result.Errors != null)
-			{
-				Logger.LogError("{Uri} reported errors: {Errors}", Uri, string.Join("\n", Result.Errors));
-				return (null, string.Join("\n", Result.Errors));
-				//throw new Exception($"{Uri} reported errors: {string.Join("\n", Result.Errors)}");
-			}
-
-			return (int.Parse(Result.Change, CultureInfo.InvariantCulture), Result.Message);
-		}
-
-		/// <inheritdoc/>
-		public override async Task DeleteShelvedChangeAsync(int ShelvedChange)
-		{
-			List<string> Arguments = new List<string>();
-			Arguments.Add($"change={ShelvedChange}");
-
-			Uri Uri = new Uri(Server, $"/api/v2/deletechange?{String.Join("&", Arguments)}");
-
-			HttpResponseMessage Response = await RetryPolicy.ExecuteAsync(() => Client.GetAsync(Uri));
-			if (!Response.IsSuccessStatusCode)
-			{
-				Logger.LogError("GET to {Uri} returned {Code} ({Response})", Uri, Response.StatusCode, await Response.Content.ReadAsStringAsync());
-				throw new Exception("Query to PerforceBridge failed");
-			}
-
-			byte[] Data = await Response.Content.ReadAsByteArrayAsync();
-			BridgeChangelistNumber Result = JsonSerializer.Deserialize<BridgeChangelistNumber>(Data.AsSpan());
-
-			if (Result.Errors != null)
-			{
-				Logger.LogError("{Uri} reported errors: {Errors}", Uri, string.Join("\n", Result.Errors));
-				throw new Exception($"{Uri} reported errors: {string.Join("\n", Result.Errors)}");
-			}
-		}
-
-		/// <inheritdoc/>
-		public override async Task UpdateChangelistDescription(int Change, string Description)
-		{
-			Uri Uri = new Uri(Server, $"/api/v2/updatechangedescription");
-
-			using (HttpContent Payload = new StringContent($"{{ \"Change\": {Change}, \"Description\": \"{JsonEncodedText.Encode(Description)}\" }}", Encoding.UTF8, "application/json"))
-			{
-				HttpResponseMessage Response = await RetryPolicy.ExecuteAsync(() => Client.PutAsync(Uri, Payload));
-				if (!Response.IsSuccessStatusCode)
-				{
-					Logger.LogError("GET to {Uri} returned {Code} ({Response})", Uri, Response.StatusCode, await Response.Content.ReadAsStringAsync());
-					throw new Exception("Query to PerforceBridge failed");
-				}
-
-				byte[] Data = await Response.Content.ReadAsByteArrayAsync();
-				BridgeChangelistNumber Result = JsonSerializer.Deserialize<BridgeChangelistNumber>(Data.AsSpan());
-
-				if (Result.Errors != null)
-				{
-					Logger.LogError("{Uri} reported errors: {Errors}", Uri, string.Join("\n", Result.Errors));
-					throw new Exception($"{Uri} reported errors: {string.Join("\n", Result.Errors)}");
-				}
-			}
-		}
-	}
 }
diff --git a/Engine/Source/Programs/Horde/HordeServer/Services/PerforceServiceApi.cs b/Engine/Source/Programs/Horde/HordeServer/Services/PerforceServiceApi.cs
--- a/Engine/Source/Programs/Horde/HordeServer/Services/PerforceServiceApi.cs
+++ b/Engine/Source/Programs/Horde/HordeServer/Services/PerforceServiceApi.cs
@@ -15,13 +15,14 @@ using Microsoft.Extensions.Options;
 using System.Runtime.InteropServices;
 using System.Globalization;
 using System.Text.RegularExpressions;
+using System.Diagnostics.CodeAnalysis;
 
 namespace HordeServer.Services
 {
 	/// <summary>
 	/// P4API implementation of the Perforce service
 	/// </summary>
-	public class P4APIPerforceService : PerforceServiceBase
+	public class P4APIPerforceService : IPerforceService
 	{
 
 		IOptionsMonitor<ServerSettings> Settings;
@@ -426,7 +427,7 @@ namespace HordeServer.Services
 
 
 		/// <inheritdoc/>
-		public override async Task<List<ChangeSummary>> GetChangesAsync(string StreamName, int? MinChange, int? MaxChange, int Results, string? ImpersonateUser)
+		public async Task<List<ChangeSummary>> GetChangesAsync(string StreamName, int? MinChange, int? MaxChange, int Results, string? ImpersonateUser)
 		{
 
 			return await Task.Run(() =>
@@ -460,7 +461,7 @@ namespace HordeServer.Services
 		}
 
 		/// <inheritdoc/>
-		public override async Task<List<ChangeDetails>> GetChangeDetailsAsync(string StreamName, IReadOnlyList<int> ChangeNumbers, string? ImpersonateUser)
+		public async Task<List<ChangeDetails>> GetChangeDetailsAsync(string StreamName, IReadOnlyList<int> ChangeNumbers, string? ImpersonateUser)
 		{
 			return await Task.Run(() =>
 			{
@@ -536,7 +537,7 @@ namespace HordeServer.Services
 		}
 
 		/// <inheritdoc />
-		public override async Task<string> CreateTicket(string ImpersonateUser)
+		public async Task<string> CreateTicket(string ImpersonateUser)
 		{
 			return await Task.Run(() =>
 			{
@@ -554,7 +555,7 @@ namespace HordeServer.Services
 		}
 
 		/// <inheritdoc/>
-		public override async Task<List<FileSummary>> FindFilesAsync(IEnumerable<string> Paths)
+		public async Task<List<FileSummary>> FindFilesAsync(IEnumerable<string> Paths)
 		{
 			return await Task.Run(() =>
 			{
@@ -592,7 +593,7 @@ namespace HordeServer.Services
 		}
 
 		/// <inheritdoc/>
-		public override async Task<byte[]> PrintAsync(string DepotPath)
+		public async Task<byte[]> PrintAsync(string DepotPath)
 		{
 			return await Task.Run(() =>
 			{
@@ -623,7 +624,7 @@ namespace HordeServer.Services
 		}
 
 		/// <inheritdoc/>
-		public override async Task<int> DuplicateShelvedChangeAsync(int ShelvedChange)
+		public async Task<int> DuplicateShelvedChangeAsync(int ShelvedChange)
 		{
 			return await Task.Run(() =>
 			{
@@ -673,7 +674,7 @@ namespace HordeServer.Services
 		}
 
 		/// <inheritdoc/>
-		public override async Task DeleteShelvedChangeAsync(int ShelvedChange)
+		public async Task DeleteShelvedChangeAsync(int ShelvedChange)
 		{
 			await Task.Run(() =>
 			{
@@ -710,7 +711,7 @@ namespace HordeServer.Services
 		}
 
 		/// <inheritdoc/>
-		public override async Task UpdateChangelistDescription(int Change, string Description)
+		public async Task UpdateChangelistDescription(int Change, string Description)
 		{
 			await Task.Run(() =>
 			{
@@ -739,7 +740,7 @@ namespace HordeServer.Services
 		}
 
 		/// <inheritdoc/>
-		public async override Task<int> CreateNewChangeAsync(string StreamName, string FilePath)
+		public async Task<int> CreateNewChangeAsync(string StreamName, string FilePath)
 		{
 			return await Task.Run(() =>
 			{
@@ -905,7 +906,7 @@ namespace HordeServer.Services
 		}
 
 		/// <inheritdoc/>
-		public override async Task<(int? Change, string Message)> SubmitShelvedChangeAsync(int Change, int OriginalChange)
+		public async Task<(int? Change, string Message)> SubmitShelvedChangeAsync(int Change, int OriginalChange)
 		{
 			return await Task.Run(() => 
 			{
@@ -1016,7 +1017,7 @@ namespace HordeServer.Services
 		}
 
 		/// <inheritdoc/>		
-		public override async Task<PerforceUserInfo?> GetUserInfoAsync(string UserName)
+		public async Task<PerforceUserInfo?> GetUserInfoAsync(string UserName)
 		{
 			return await Task.Run(() =>
 			{
@@ -1109,6 +1110,87 @@ namespace HordeServer.Services
 
 		}
 
+		/// <inheritdoc/>
+		public virtual async Task<int> GetCodeChangeAsync(string StreamName, int Change)
+		{
+			int MaxChange = Change;
+			for (; ; )
+			{
+				// Query for the changes before this point
+				List<ChangeSummary> Changes = await GetChangesAsync(StreamName, null, MaxChange, 10, null);
+				Serilog.Log.Logger.Information("Finding last code change in {Stream} before {MaxChange}: {NumResults}", StreamName, MaxChange, Changes.Count);
+				if (Changes.Count == 0)
+				{
+					return 0;
+				}
+
+				// Get the details for them
+				List<ChangeDetails> DetailsList = await GetChangeDetailsAsync(StreamName, Changes.ConvertAll(x => x.Number), null);
+				foreach (ChangeDetails Details in DetailsList.OrderByDescending(x => x.Number))
+				{
+					ChangeContentFlags ContentFlags = Details.GetContentFlags();
+					Serilog.Log.Logger.Information("Change {Change} = {Flags}", Details.Number, ContentFlags.ToString());
+					if ((Details.GetContentFlags() & ChangeContentFlags.ContainsCode) != 0)
+					{
+						return Details.Number;
+					}
+				}
+
+				// Loop round again, adjusting our maximum changelist number
+				MaxChange = Changes.Min(x => x.Number) - 1;
+			}
+		}
+
+		/// <summary>
+		/// Gets the wildcard filter for a particular range of changes
+		/// </summary>
+		/// <param name="BasePath">Base path to find files for</param>
+		/// <param name="MinChange">Minimum changelist number to query</param>
+		/// <param name="MaxChange">Maximum changelist number to query</param>
+		/// <returns>Filter string</returns>
+		public static string GetFilter(string BasePath, int? MinChange, int? MaxChange)
+		{
+			StringBuilder Filter = new StringBuilder(BasePath);
+			if (MinChange != null && MaxChange != null)
+			{
+				Filter.Append($"@{MinChange},{MaxChange}");
+			}
+			else if (MinChange != null)
+			{
+				Filter.Append($"@>={MinChange}");
+			}
+			else if (MaxChange != null)
+			{
+				Filter.Append($"@<={MaxChange}");
+			}
+			return Filter.ToString();
+		}
+
+		/// <summary>
+		/// Gets a stream-relative path from a depot path
+		/// </summary>
+		/// <param name="DepotFile">The depot file to check</param>
+		/// <param name="StreamName">Name of the stream</param>
+		/// <param name="RelativePath">Receives the relative path to the file</param>
+		/// <returns>True if the stream-relative path was returned</returns>
+		public static bool TryGetStreamRelativePath(string DepotFile, string StreamName, [NotNullWhen(true)] out string? RelativePath)
+		{
+			if (DepotFile.StartsWith(StreamName, StringComparison.OrdinalIgnoreCase) && DepotFile.Length > StreamName.Length && DepotFile[StreamName.Length] == '/')
+			{
+				RelativePath = DepotFile.Substring(StreamName.Length);
+				return true;
+			}
+
+			Match Match = Regex.Match(DepotFile, "^//[^/]+/[^/]+(/.*)$");
+			if (Match.Success)
+			{
+				RelativePath = Match.Groups[1].Value;
+				return true;
+			}
+
+			RelativePath = null;
+			return false;
+		}
 	}
 
 }
diff --git a/Engine/Source/Programs/Horde/HordeServer/Services/PerforceServiceRouter.cs b/Engine/Source/Programs/Horde/HordeServer/Services/PerforceServiceRouter.cs
deleted file mode 100644
diff --git a/Engine/Source/Programs/Horde/HordeServer/Startup.cs b/Engine/Source/Programs/Horde/HordeServer/Startup.cs
--- a/Engine/Source/Programs/Horde/HordeServer/Startup.cs
+++ b/Engine/Source/Programs/Horde/HordeServer/Startup.cs
@@ -373,28 +373,7 @@ namespace HordeServer
 			Services.AddSingleton<LifetimeService>();
 			Services.AddSingleton<ILogFileService, LogFileService>();
 			Services.AddSingleton<INotificationService, NotificationService>();
-
-			if (Settings.P4UseRouter && !String.IsNullOrEmpty(Settings.PerforceBridge) && !String.IsNullOrEmpty(Settings.P4BridgeServer))
-			{
-				// using p4 service router, which will be removed along with the python p4 service
-				// once the p4 api service is hardened
-				Services.AddSingleton<IPerforceService, PerforceServiceRouter>();
-
-			}
-			else if (!String.IsNullOrEmpty(Settings.PerforceBridge))
-			{
-				// Python P4 bridge service
-				Services.AddSingleton<IPerforceService, BridgePerforceService>(SP => new BridgePerforceService(Settings.PerforceBridge!, SP.GetRequiredService<ILogger<BridgePerforceService>>()));
-			}
-			else if (!String.IsNullOrEmpty(Settings.P4BridgeServer))
-			{
-				// Native P4 API
-				Services.AddSingleton<IPerforceService, P4APIPerforceService>();
-			}
-			else
-			{
-				Services.AddSingleton<IPerforceService, LocalPerforceService>();
-			}
+			Services.AddSingleton<IPerforceService, P4APIPerforceService>();
 
 			Services.AddSingleton<PerforceLoadBalancer>();
 			Services.AddSingleton<PoolService>();
