Horde: Rename P4APIPerforceService to PerforceService, since it's now the only implementation.

diff --git a/Engine/Source/Programs/Horde/HordeServer/Services/IPerforceService.cs b/Engine/Source/Programs/Horde/HordeServer/Services/IPerforceService.cs
new file mode 100644
--- /dev/null
+++ b/Engine/Source/Programs/Horde/HordeServer/Services/IPerforceService.cs
@@ -0,0 +1,212 @@
+// Copyright Epic Games, Inc. All Rights Reserved.
+
+using EpicGames.Core;
+using EpicGames.Perforce;
+using HordeServer.Models;
+using Microsoft.Extensions.Logging;
+using Polly;
+using Polly.Extensions.Http;
+using System;
+using System.Collections.Generic;
+using System.Diagnostics.CodeAnalysis;
+using System.Globalization;
+using System.IO;
+using System.Linq;
+using System.Net.Http;
+using System.Text;
+using System.Text.Json;
+using System.Text.RegularExpressions;
+using System.Threading;
+using System.Threading.Tasks;
+
+namespace HordeServer.Services
+{
+	/// <summary>
+	/// Information about a Perforce user
+	/// </summary>
+	public class PerforceUserInfo
+	{
+		/// <summary>
+		/// Full name of the user
+		/// </summary>
+		public string? Name { get; set; }
+
+		/// <summary>
+		/// User's email address
+		/// </summary>
+		public string? Email { get; set; }
+	}
+
+	/// <summary>
+	/// Wrapper around Perforce functionality. Can use a local p4.exe client for development purposes, or a separate HordePerforceBridge instance over REST for deployments.
+	/// </summary>
+	public interface IPerforceService
+	{
+		/// <summary>
+		/// Create a new changelist by submitting the given file
+		/// </summary>
+		/// <param name="StreamName">The stream to query</param>
+		/// <param name="Path">Path for the file to submit</param>
+		/// <returns>New changelist number</returns>
+		public Task<int> CreateNewChangeAsync(string StreamName, string Path);
+
+		/// <summary>
+		/// Gets the code change corresponding to an actual change submitted to a stream
+		/// </summary>
+		/// <param name="StreamName"></param>
+		/// <param name="Change">The changelist number to query</param>
+		/// <returns>Code change for the latest change</returns>
+		public Task<int> GetCodeChangeAsync(string StreamName, int Change);
+
+		/// <summary>
+		/// Gets information about the given user
+		/// </summary>
+		/// <param name="UserName">The user name</param>
+		/// <returns>User information</returns>
+		public Task<PerforceUserInfo?> GetUserInfoAsync(string UserName);
+
+		/// <summary>
+		/// Gets the latest change for a particular stream
+		/// </summary>
+		/// <param name="StreamName">The stream to query</param>
+		/// <param name="MinChange">The minimum changelist number</param>
+		/// <param name="MaxChange">The maximum changelist number</param>
+		/// <param name="Results">Number of results to return</param>
+		/// <param name="ImpersonateUser">Name of the user to impersonate</param>
+		/// <returns>Latest changelist number</returns>
+		public Task<List<ChangeSummary>> GetChangesAsync(string StreamName, int? MinChange, int? MaxChange, int Results, string? ImpersonateUser);
+
+		/// <summary>
+		/// Gets the latest change for a particular stream
+		/// </summary>
+		/// <param name="StreamName">The stream to query</param>
+		/// <param name="ChangeNumbers">Change numbers to query</param>
+		/// <param name="ImpersonateUser">Name of the user to impersonate</param>
+		/// <returns>Commit details</returns>
+		public Task<List<ChangeDetails>> GetChangeDetailsAsync(string StreamName, IReadOnlyList<int> ChangeNumbers, string? ImpersonateUser);
+		
+		/// <summary>
+		/// Create a ticket as the specified user
+		/// </summary>
+		/// <param name="ImpersonateUser">Name of the user to impersonate</param>
+		/// <returns>Commit details</returns>
+		public Task<string> CreateTicket(string ImpersonateUser);
+
+		/// <summary>
+		/// Gets the latest changes for a set of depot paths
+		/// </summary>
+		/// <param name="Paths"></param>
+		/// <returns></returns>
+		public Task<List<FileSummary>> FindFilesAsync(IEnumerable<string> Paths);
+
+		/// <summary>
+		/// Gets the contents of a file in the depot
+		/// </summary>
+		/// <param name="Path">Path to read</param>
+		/// <returns>Data for the file</returns>
+		public Task<byte[]> PrintAsync(string Path);
+
+		/// <summary>
+		/// Duplicates a shelved changelist
+		/// </summary>
+		/// <param name="ShelvedChange">The shelved changelist</param>
+		/// <returns>The duplicated changelist</returns>
+		public Task<int> DuplicateShelvedChangeAsync(int ShelvedChange);
+
+		/// <summary>
+		/// Submit a shelved changelist
+		/// </summary>
+		/// <param name="ShelvedChange">The shelved changelist number, created by <see cref="DuplicateShelvedChangeAsync(int)"/></param>
+		/// <param name="OriginalChange">The original changelist number</param>
+		/// <returns>Tuple consisting of the submitted changelist number and message</returns>
+		public Task<(int? Change, string Message)> SubmitShelvedChangeAsync(int ShelvedChange, int OriginalChange);
+
+		/// <summary>
+		/// Deletes a changelist containing shelved files.
+		/// </summary>
+		/// <param name="ShelvedChange">The changelist containing shelved files</param>
+		/// <returns>Async tasy</returns>
+		public Task DeleteShelvedChangeAsync(int ShelvedChange);
+
+		/// <summary>
+		/// Updates a changelist description
+		/// </summary>
+		/// <param name="Change">The change to update</param>
+		/// <param name="Description">The new description</param>
+		/// <returns>Async task</returns>
+		public Task UpdateChangelistDescription(int Change, string Description);
+	}
+
+	/// <summary>
+	/// Extension methods for IPerforceService implementations
+	/// </summary>
+	public static class PerforceServiceExtensions
+	{
+		/// <summary>
+		/// Creates a new change for a template
+		/// </summary>
+		/// <param name="Perforce">The Perforce service instance</param>
+		/// <param name="Stream">Stream containing the template</param>
+		/// <param name="Template">The template being built</param>
+		/// <returns>New changelist number</returns>
+		public static Task<int> CreateNewChangeForTemplateAsync(this IPerforceService Perforce, IStream Stream, ITemplate Template)
+		{
+			Match Match = Regex.Match(Template.SubmitNewChange!, @"^(//[^/]+/[^/]+)/(.+)$");
+			if (Match.Success)
+			{
+				return Perforce.CreateNewChangeAsync(Match.Groups[1].Value, Match.Groups[2].Value);
+			}
+			else
+			{
+				return Perforce.CreateNewChangeAsync(Stream.Name, Template.SubmitNewChange!);
+			}
+		}
+
+		/// <summary>
+		/// Gets the latest change for a particular stream
+		/// </summary>
+		/// <param name="Perforce">The perforce implementation</param>
+		/// <param name="StreamName">The stream to query</param>
+		/// <param name="ChangeNumber">Change number to query</param>
+		/// <param name="ImpersonateUser">Name of the user to impersonate</param>
+		/// <returns>Commit details</returns>
+		public static async Task<ChangeDetails> GetChangeDetailsAsync(this IPerforceService Perforce, string StreamName, int ChangeNumber, string? ImpersonateUser)
+		{
+			List<ChangeDetails> Results = await Perforce.GetChangeDetailsAsync(StreamName, new[] { ChangeNumber }, ImpersonateUser);
+			return Results[0];
+		}
+
+		/// <summary>
+		/// Gets the latest change for a particular stream
+		/// </summary>
+		/// <param name="Perforce">The perforce implementation</param>
+		/// <param name="StreamName">The stream to query</param>
+		/// <param name="MinChange">The minimum changelist number</param>
+		/// <param name="MaxChange">The maximum changelist number</param>
+		/// <param name="Results">Number of results to return</param>
+		/// <param name="ImpersonateUser">Name of the user to impersonate</param>
+		/// <returns>Commit details</returns>
+		public static async Task<List<ChangeDetails>> GetChangeDetailsAsync(this IPerforceService Perforce, string StreamName, int? MinChange, int? MaxChange, int Results, string? ImpersonateUser)
+		{
+			List<ChangeSummary> Changes = await Perforce.GetChangesAsync(StreamName, MinChange, MaxChange, Results, ImpersonateUser);
+			return await Perforce.GetChangeDetailsAsync(StreamName, Changes.ConvertAll(x => x.Number), ImpersonateUser);
+		}
+
+		/// <summary>
+		/// Get the latest submitted change to the stream
+		/// </summary>
+		/// <param name="Perforce">The perforce implementation</param>
+		/// <param name="StreamName">The stream to query</param>
+		/// <param name="ImpersonateUser">Name of the user to impersonate</param>
+		/// <returns>Latest changelist number</returns>
+		public static async Task<int> GetLatestChangeAsync(this IPerforceService Perforce, string StreamName, string? ImpersonateUser)
+		{
+			List<ChangeSummary> Changes = await Perforce.GetChangesAsync(StreamName, null, null, 1, ImpersonateUser);
+			if (Changes.Count == 0)
+			{
+				throw new Exception($"No changes have been submitted to stream {StreamName}");
+			}
+			return Changes[0].Number;
+		}
+	}
+}
diff --git a/Engine/Source/Programs/Horde/HordeServer/Services/Impl/PerforceService.cs b/Engine/Source/Programs/Horde/HordeServer/Services/Impl/PerforceService.cs
new file mode 100644
--- /dev/null
+++ b/Engine/Source/Programs/Horde/HordeServer/Services/Impl/PerforceService.cs
@@ -0,0 +1,1196 @@
+// Copyright Epic Games, Inc. All Rights Reserved.
+
+using HordeServer.Models;
+using System.Linq;
+using Microsoft.Extensions.Logging;
+using System;
+using System.Collections.Generic;
+using System.Net;
+using System.Text;
+using System.Threading.Tasks;
+using System.IO;
+
+using P4 = Perforce.P4;
+using Microsoft.Extensions.Options;
+using System.Runtime.InteropServices;
+using System.Globalization;
+using System.Text.RegularExpressions;
+using System.Diagnostics.CodeAnalysis;
+
+namespace HordeServer.Services
+{
+	/// <summary>
+	/// P4API implementation of the Perforce service
+	/// </summary>
+	public class PerforceService : IPerforceService
+	{
+
+		IOptionsMonitor<ServerSettings> Settings;
+
+		ILogger<PerforceService> Logger;
+
+		P4.Server Server;
+
+		/// <summary>
+		/// Object used for controlling access to the access user tickets
+		/// </summary>
+		static object TicketLock = new object();
+
+		/// <summary>
+		/// Object used for controlling access to the p4 command output
+		/// </summary>
+		static object P4LogLock = new object();
+
+		/// <summary>
+		/// Native -> managed debug logging callback
+		/// </summary>
+		P4.P4CallBacks.LogMessageDelegate LogBridgeDelegate;
+
+
+		Dictionary<string, P4.Credential> UserTickets = new Dictionary<string, P4.Credential>();
+
+		/// <summary>
+		/// Constructor
+		/// </summary>
+		public PerforceService(IOptionsMonitor<ServerSettings> Settings, ILogger<PerforceService> Logger)
+		{
+			this.Settings = Settings;
+			this.Logger = Logger;
+
+			Server = new P4.Server(new P4.ServerAddress(Settings.CurrentValue.P4BridgeServer));
+
+			LogBridgeDelegate = new P4.P4CallBacks.LogMessageDelegate(LogBridgeMessage);
+			P4.P4Debugging.SetBridgeLogFunction(LogBridgeDelegate);
+
+			P4.LogFile.SetLoggingFunction(LogPerforce);
+		}
+
+
+		P4.Repository GetServiceUserConnection()
+		{
+
+			ServerSettings Settings = this.Settings.CurrentValue;
+
+			P4.Repository Repository = new P4.Repository(Server);
+
+			try
+			{
+				P4.Connection Connection = Repository.Connection;
+				Connection.UserName = Settings.P4BridgeServiceUsername;
+
+				P4.Options Options = new P4.Options();
+
+				Options["Ticket"] = Settings.P4BridgeServicePassword;
+
+				// connect to the server
+				if (!Connection.Connect(Options))
+				{
+					throw new Exception("Unable to get P4 server connection");
+				}
+
+			}
+			catch
+			{
+				Repository.Dispose();
+				throw;
+			}
+
+
+			return Repository;
+		}
+
+		P4.Credential GetImpersonateCredential(string ImpersonateUser)
+		{
+
+			if (!CanImpersonate)
+			{
+				throw new Exception($"Service account required to impersonate user {ImpersonateUser}");
+			}
+
+			lock (TicketLock)
+			{
+
+				ServerSettings Settings = this.Settings.CurrentValue;
+
+				// Check if we have a ticket
+				P4.Credential? Credential = null;
+
+				if (UserTickets.TryGetValue(ImpersonateUser, out Credential))
+				{
+					// if the credential expires within the next 15 minutes, refresh
+					TimeSpan Time = new TimeSpan(0, 15, 0);
+					if (Credential.Expires.Subtract(Time) <= DateTime.UtcNow)
+					{
+						UserTickets.Remove(ImpersonateUser);
+						Credential = null;
+					}
+				}
+
+				if (Credential != null)
+				{
+					return Credential;
+				}
+
+				using (P4.Repository Repository = GetServiceUserConnection())
+				{
+					Credential = Repository.Connection.Login(null, new P4.LoginCmdOptions(P4.LoginCmdFlags.AllHosts | P4.LoginCmdFlags.DisplayTicket, null), ImpersonateUser);
+				}
+
+				if (Credential == null)
+				{
+					throw new Exception($"GetImpersonateCredential - Unable to get impersonation credential for {ImpersonateUser} from {Dns.GetHostName()}");
+				}
+
+				UserTickets.Add(ImpersonateUser, Credential);
+
+				return Credential;
+			}
+
+		}
+
+		P4.Repository GetImpersonateConnection(string ImpersonateUser)
+		{
+
+			if (!CanImpersonate)
+			{
+				throw new Exception($"Service account required to impersonate user {ImpersonateUser}");
+			}
+
+			P4.Credential Credential = GetImpersonateCredential(ImpersonateUser);
+
+			// this might be able to be reused and would be threaded internally
+			P4.Repository Repository = new P4.Repository(Server);
+
+			P4.Connection Connection = Repository.Connection;			
+			Connection.UserName = ImpersonateUser;
+
+			P4.Options Options = new P4.Options();
+			Options["Ticket"] = Credential.Ticket;
+
+			// connect to the server
+			if (!Connection.Connect(Options))
+			{
+				throw new Exception($"GetImpersonateConnection - Unable to get impersonated P4 connection for {ImpersonateUser} from {Dns.GetHostName()}");
+
+			}
+
+			return Repository;
+		}
+
+		P4.Repository GetConnection(string? Stream = null, string? Username = null, bool ReadOnly = true, bool CreateChange = false, int? ClientFromChange = null, bool UseClientFromChange = false, bool UsePortFromChange = false, string? ClientId = null, bool NoClient = false)
+		{
+			ServerSettings Settings = this.Settings.CurrentValue;
+
+			P4.Repository? Repository = (Username == null || Settings.P4BridgeServiceUsername == Username || !CanImpersonate) ? GetServiceUserConnection() : GetImpersonateConnection(Username);
+
+			if (!NoClient)
+			{
+				try
+				{
+					if (ClientId != null)
+					{
+						Repository.Connection.SetClient(ClientId);
+					}
+					else
+					{
+						P4.Client Client = GetOrCreateClient(Repository, Stream, Username, ReadOnly, CreateChange, ClientFromChange, UseClientFromChange, UsePortFromChange);
+						Repository.Connection.SetClient(Client.Name);
+
+					}
+
+					if (UseClientFromChange)
+					{
+						string? ClientHost = Repository.Connection.Client.Host;
+						if (!string.IsNullOrEmpty(ClientHost))
+						{
+							Repository.Connection.getP4Server().SetConnectionHost(ClientHost);
+						}
+					}
+
+					/*
+
+					This is disabled until https://jira.it.epicgames.com/servicedesk/customer/portal/1/ITH-144069 is resolved and edge server is sorted (required for things like deleting shelves)
+
+					string? ClientPort = null;
+
+					if (UsePortFromChange)
+					{
+						if (!Repository.Connection.Client.Name.StartsWith("horde-p4bridge-", StringComparison.OrdinalIgnoreCase))
+						{
+							using (P4.P4Command Cmd = new P4.P4Command(Repository, "info", true))
+							{
+								P4.P4CommandResult Results = Cmd.Run();
+
+								if (!Results.Success)
+								{
+									throw new Exception("Unable to get server info");
+								}
+
+								P4.TaggedObject? Tag = Results.TaggedOutput.Find(T => T.ContainsKey("changeServer"));
+
+								if (Tag != null)
+								{
+									if (!Tag.TryGetValue("changeServer", out ClientPort))
+									{
+										throw new Exception("Unable to get change server from tagged server info output");
+									}
+								}
+							}
+						}
+					}
+					*/
+
+				}
+				catch
+				{
+					Repository.Dispose();
+					throw;
+				}
+			}
+
+			if (Repository == null)
+			{
+				throw new Exception($"Unable to get connection for Stream:{Stream} Username:{Username} ReadOnly:{ReadOnly} CreateChange:{CreateChange} ClientFromChange:{ClientFromChange} UseClientFromChange:{UseClientFromChange} UsePortFromChange:{UsePortFromChange}");
+			}
+
+			Repository.Connection.CommandEcho += LogPerforceCommand;
+
+			return Repository;
+
+		}
+
+		string GetClientName(string Stream, bool ReadOnly, bool CreateChange, string? Username = null)
+		{
+			ServerSettings Settings = this.Settings.CurrentValue;
+
+			string ClientName = $"horde-p4bridge-{Dns.GetHostName()}-{Settings.P4BridgeServiceUsername}-{Stream.Replace("/", "+", StringComparison.OrdinalIgnoreCase)}";
+
+			if (!ReadOnly)
+			{
+				ClientName += "-write";
+			}
+
+			if (CreateChange)
+			{
+				ClientName += "-create";
+			}
+
+			if (!string.IsNullOrEmpty(Username))
+			{
+				ClientName += "-" + Username;
+			}
+
+			return ClientName;
+
+		}
+
+		P4.Client GetOrCreateClient(P4.Repository Repository, string? Stream, string? Username = null, bool ReadOnly = true, bool CreateChange = false, int? ClientFromChange = null, bool UseClientFromChange = false, bool UsePortFromChange = false)
+		{
+			ServerSettings Settings = this.Settings.CurrentValue;
+
+			P4.Client? Client = null;
+			P4.Changelist? Changelist = null;
+
+			string? ClientHost = null;
+
+			if (ClientFromChange.HasValue)
+			{
+				Changelist = Repository.GetChangelist(ClientFromChange.Value);
+
+				if (Changelist == null)
+				{
+					throw new Exception($"Unable to get changelist for client {ClientFromChange}");
+				}
+
+				Client = Repository.GetClient(Changelist.ClientId);
+
+				if (Client == null)
+				{
+					throw new Exception($"Unable to get client for id {Changelist.ClientId}");
+				}
+
+
+				Stream = Client.Stream;
+				Username = Client.OwnerName;
+				ClientHost = Client.Host;
+
+
+			}
+
+			if (!UseClientFromChange)
+			{
+				if (Stream == null)
+				{
+					throw new Exception("Stream required for client");
+
+				}
+				string ClientName = GetClientName(Stream, ReadOnly, CreateChange, Username);
+
+				IList<P4.Client>? Clients = Repository.GetClients(new P4.ClientsCmdOptions(P4.ClientsCmdFlags.None, Username, ClientName, 1, Stream));
+
+				if (Clients != null && Clients.Count == 1)
+				{
+					Client = Clients[0];
+					if (!string.IsNullOrEmpty(Client.Host))
+					{
+						Client.Host = string.Empty;
+						Repository.UpdateClient(Client);
+					}
+				}
+				else
+				{
+
+					P4.Client NewClient = new P4.Client();
+
+					NewClient.Name = ClientName;
+					NewClient.Root = RuntimeInformation.IsOSPlatform(OSPlatform.Linux) ? $"/tmp/{ClientName}/" : $"{Path.GetTempPath()}{ClientName}\\";
+					NewClient.Stream = Stream;
+					if (ReadOnly)
+					{
+						NewClient.ClientType = P4.ClientType.@readonly;
+					}
+
+					NewClient.OwnerName = Username ?? Settings.P4BridgeServiceUsername;
+
+					Client = Repository.CreateClient(NewClient);
+
+					if (Client == null)
+					{
+						throw new Exception($"Unable to create client for ${Stream} : {Username ?? Settings.P4BridgeServiceUsername}");
+					}
+				}
+			}
+
+			if (Client == null)
+			{
+				throw new Exception($"Unable to create client for Stream:{Stream} Username:{Username} ReadOnly:{ReadOnly} CreateChange:{CreateChange} ClientFromChange:{ClientFromChange} UseClientFromChange:{UseClientFromChange} UsePortFromChange:{UsePortFromChange}");
+			}
+
+			return Client;
+
+		}
+
+		/// <summary>
+		/// Logs perforce commands 
+		/// </summary>
+		/// <param name="Log">The p4 command log info</param>
+		static void LogPerforceCommand(string Log)
+		{
+			lock (P4LogLock)
+			{
+				Serilog.Log.Information("Perforce: " + Log);
+			}
+			
+		}
+
+		/// <summary>
+		/// Perforce logging funcion
+		/// </summary>
+		/// <param name="LogLevel">The level whether error, warning, message, or debug</param>
+		/// <param name="Source">The log source</param>
+		/// <param name="Message">The log message</param>
+		static void LogPerforce(int LogLevel, String Source, String Message)
+		{
+			lock (P4LogLock)
+			{
+				switch (LogLevel)
+				{
+					case 0:
+					case 1:
+						Serilog.Log.Error("Perforce (Error): {Message} {Source} : ", Message, Source);
+						break;
+					case 2:
+						Serilog.Log.Warning("Perforce (Warning): {Message} {Source} : ", Message, Source);
+						break;
+					case 3:
+						Serilog.Log.Information("Perforce: {Message} {Source} : ", Message, Source);
+						break;
+					default:
+						Serilog.Log.Debug("Perforce (Debug): {Message} {Source} : ", Message, Source);
+						break;
+				};
+			}
+		}
+
+		static void LogBridgeMessage(int LogLevel, String Filename, int Line, String Message)
+		{
+
+			// Note, we do not get log level 4 unless it is defined in native code as it is very, very spammy (P4BridgeServer.cpp)
+			
+			// remove the full path to the source, keep just the file name
+			String fileName = Path.GetFileName(Filename);
+
+			string category = String.Format(CultureInfo.CurrentCulture, "P4Bridge({0}:{1})", fileName, Line);
+
+			P4.LogFile.LogMessage(LogLevel, category, Message);
+		}
+
+
+		/// <inheritdoc/>
+		public async Task<List<ChangeSummary>> GetChangesAsync(string StreamName, int? MinChange, int? MaxChange, int Results, string? ImpersonateUser)
+		{
+
+			return await Task.Run(() =>
+			{
+
+				using (P4.Repository Repository = GetConnection(StreamName, ImpersonateUser))
+				{
+
+					P4.ChangesCmdOptions Options = new P4.ChangesCmdOptions(P4.ChangesCmdFlags.IncludeTime | P4.ChangesCmdFlags.FullDescription, null, Results, P4.ChangeListStatus.Submitted, null);
+
+					string Filter = GetFilter($"//{Repository.Connection.Client.Name}/...", MinChange, MaxChange);
+
+					P4.FileSpec FileSpec = new P4.FileSpec(new P4.DepotPath(Filter), null, null, null);
+
+					IList<P4.Changelist> Changelists = Repository.GetChangelists(Options, FileSpec);
+
+					List<ChangeSummary> Changes = new List<ChangeSummary>();
+
+					if (Changelists != null)
+					{
+						foreach (P4.Changelist Changelist in Changelists)
+						{
+							Changes.Add(new ChangeSummary(Changelist.Id, Changelist.OwnerName, Changelist.Description));
+						}
+					}
+
+					return Changes;
+				}
+			});
+
+		}
+
+		/// <inheritdoc/>
+		public async Task<List<ChangeDetails>> GetChangeDetailsAsync(string StreamName, IReadOnlyList<int> ChangeNumbers, string? ImpersonateUser)
+		{
+			return await Task.Run(() =>
+			{
+				using (P4.Repository Repository = GetConnection(StreamName, ImpersonateUser))
+				{
+					List<ChangeDetails> Results = new List<ChangeDetails>();
+
+					List<string> Args = new List<string> { "-s", "-S" };
+					Args.AddRange(ChangeNumbers.Select(Change => Change.ToString(CultureInfo.InvariantCulture)));
+
+					using (P4.P4Command Command = new P4.P4Command(Repository, "describe", true, Args.ToArray()))
+					{
+						P4.P4CommandResult Result = Command.Run();
+
+						if (!Result.Success)
+						{
+							throw new Exception("Unable to get changes");
+						}
+
+						if (Result.TaggedOutput == null || Result.TaggedOutput.Count <= 0)
+						{
+							return Results;
+						}
+
+						List<P4.Changelist> Changelists = new List<P4.Changelist>() { };
+
+						bool DSTMismatch = false;
+						string Offset = string.Empty;
+
+						if (Server != null && Server.Metadata != null)
+						{
+							Offset = Server.Metadata.DateTimeOffset;
+							DSTMismatch = P4.FormBase.DSTMismatch(Server.Metadata);
+						}
+
+						foreach (P4.TaggedObject TaggedObject in Result.TaggedOutput)
+						{
+							List<string> Files = new List<string>();
+
+							P4.Changelist Change = new P4.Changelist();
+							Change.FromChangeCmdTaggedOutput(TaggedObject, true, Offset, DSTMismatch);
+
+							foreach (P4.FileMetaData DescribeFile in Change.Files)
+							{
+								string? RelativePath;
+								if (TryGetStreamRelativePath(DescribeFile.DepotPath.Path, StreamName, out RelativePath))
+								{
+									Files.Add(RelativePath);
+								}
+							}
+
+							if (Change.ShelvedFiles != null && Change.ShelvedFiles.Count > 0)
+							{
+								foreach (P4.ShelvedFile ShelvedFile in Change.ShelvedFiles)
+								{
+									string? RelativePath;
+									if (TryGetStreamRelativePath(ShelvedFile.Path.ToString(), StreamName, out RelativePath))
+									{
+										Files.Add(RelativePath);
+									}
+								}
+
+							}
+
+							Results.Add(new ChangeDetails(Change.Id, Change.OwnerName, Change.Description, Files));
+						}
+
+					}
+
+					return Results;
+				}
+			});
+		}
+
+		/// <inheritdoc />
+		public async Task<string> CreateTicket(string ImpersonateUser)
+		{
+			return await Task.Run(() =>
+			{
+
+				P4.Credential Credential = GetImpersonateCredential(ImpersonateUser);
+
+				if (Credential == null)
+				{
+					throw new Exception($"Unable to get ticket for user {ImpersonateUser}");
+				}
+
+				return Credential.Ticket;
+			});
+
+		}
+
+		/// <inheritdoc/>
+		public async Task<List<FileSummary>> FindFilesAsync(IEnumerable<string> Paths)
+		{
+			return await Task.Run(() =>
+			{
+				List<FileSummary> Results = new List<FileSummary>();
+
+				using (P4.Repository Repository = GetConnection(NoClient: true))
+				{
+					P4.GetFileMetaDataCmdOptions Options = new P4.GetFileMetaDataCmdOptions(P4.GetFileMetadataCmdFlags.ExcludeClientData, "", "", 0, "", "", "");
+
+					IList<P4.FileMetaData>? Files = Repository.GetFileMetaData(Options, Paths.Select(Path => { P4.FileSpec FileSpec = new P4.FileSpec(new P4.DepotPath(Path)); return FileSpec; }).ToArray());
+
+					if (Files == null)
+					{
+						Files = new List<P4.FileMetaData>();
+					}
+
+					foreach (string Path in Paths)
+					{
+						P4.FileMetaData Meta = Files.FirstOrDefault(File => File.DepotPath.Path == Path);
+
+						if (Meta == null)
+						{
+							Results.Add(new FileSummary(Path, false, 0));
+						}
+						else
+						{
+							Results.Add(new FileSummary(Path, Meta.HeadAction != P4.FileAction.Delete && Meta.HeadAction != P4.FileAction.MoveDelete, Meta.HeadChange));
+						}
+
+					}
+
+					return Results;
+				}
+			});
+		}
+
+		/// <inheritdoc/>
+		public async Task<byte[]> PrintAsync(string DepotPath)
+		{
+			return await Task.Run(() =>
+			{
+
+				if (DepotPath.EndsWith("...", StringComparison.OrdinalIgnoreCase) || DepotPath.EndsWith("*", StringComparison.OrdinalIgnoreCase))
+				{
+					throw new Exception("PrintAsync requires exactly one file to be specified");
+				}
+
+				using (P4.Repository Repository = GetConnection(NoClient: true))
+				{
+					using (P4.P4Command Command = new P4.P4Command(Repository, "print", false, new string[] { "-q", DepotPath }))
+					{
+						P4.P4CommandResult Result = Command.Run();
+
+						if (Result.BinaryOutput != null)
+						{
+							return Result.BinaryOutput;
+						}
+					
+						return Encoding.Default.GetBytes(Result.TextOutput);
+
+					}
+
+				}
+			});
+
+		}
+
+		/// <inheritdoc/>
+		public async Task<int> DuplicateShelvedChangeAsync(int ShelvedChange)
+		{
+			return await Task.Run(() =>
+			{
+				string? ChangeOwner = null;
+
+				// Get the owner of the shelf
+				using (P4.Repository Repository = GetConnection(NoClient: true))
+				{
+					List<string> Args = new List<string> { "-S", ShelvedChange.ToString(CultureInfo.InvariantCulture)};
+
+					using (P4.P4Command Command = new P4.P4Command(Repository, "describe", true, Args.ToArray()))
+					{
+						P4.P4CommandResult Result = Command.Run();
+
+						if (!Result.Success)
+						{
+							throw new Exception($"Unable to get change {ShelvedChange}");
+						}
+
+						if (Result.TaggedOutput == null || Result.TaggedOutput.Count != 1)
+						{
+
+							throw new Exception($"Unable to get tagged output for change: {ShelvedChange}");
+						}
+
+						P4.Changelist Changelist = new P4.Changelist();
+						Changelist.FromChangeCmdTaggedOutput(Result.TaggedOutput[0], true, string.Empty, false);
+						ChangeOwner = Changelist.OwnerName;
+
+					}
+				}
+
+				if (string.IsNullOrEmpty(ChangeOwner))
+				{
+					throw new Exception($"Unable to get owner for shelved change {ShelvedChange}");
+				}
+
+				using (P4.Repository Repository = GetConnection(ReadOnly: false, ClientFromChange: ShelvedChange, UseClientFromChange: false, Username: ChangeOwner))
+				{
+					return ReshelveChange(Repository, ShelvedChange);				
+				}
+
+				throw new Exception($"Unable to duplicate shelve change {ShelvedChange}");
+
+			});
+
+		}
+
+		/// <inheritdoc/>
+		public async Task DeleteShelvedChangeAsync(int ShelvedChange)
+		{
+			await Task.Run(() =>
+			{
+
+				using (P4.Repository Repository = GetConnection(NoClient: true))
+				{
+
+					P4.Changelist? Changelist = Repository.GetChangelist(ShelvedChange, new P4.DescribeCmdOptions(P4.DescribeChangelistCmdFlags.Omit, 0, 0));
+
+					if (Changelist == null)
+					{
+						throw new Exception($"Unable to get shelved changelist for delete: {ShelvedChange}");
+					}
+
+					string ClientId = Changelist.ClientId;
+
+					if (String.IsNullOrEmpty(ClientId))
+					{
+						throw new Exception($"Unable to get shelved changelist client id for delete: {ShelvedChange}");
+					}
+
+					using (P4.Repository ClientRepository = GetConnection(Changelist.Stream, Changelist.OwnerName, ClientId: ClientId, UseClientFromChange: true, UsePortFromChange: true))
+					{
+
+						if (Changelist.Shelved)
+						{
+							IList<P4.FileSpec> Files = ClientRepository.Connection.Client.ShelveFiles(new P4.ShelveFilesCmdOptions(P4.ShelveFilesCmdFlags.Delete, null, ShelvedChange));
+						}
+
+						ClientRepository.DeleteChangelist(Changelist, null);
+					}
+				}
+			});
+		}
+
+		/// <inheritdoc/>
+		public async Task UpdateChangelistDescription(int Change, string Description)
+		{
+			await Task.Run(() =>
+			{
+				try
+				{
+					using (P4.Repository Repository = GetConnection(NoClient: true))
+					{
+						P4.Changelist Changelist = Repository.GetChangelist(Change);
+
+						Repository.Connection.Disconnect();
+
+						// the client must exist for the change list, otherwise will fail (for example, CreateNewChangeAsync deletes the client before returning)
+						using (P4.Repository UpdateRepository = GetConnection(ClientFromChange: Change, UseClientFromChange: true, Username: Changelist.OwnerName))
+						{
+							P4.Changelist UpdatedChangelist = UpdateRepository.GetChangelist(Change);
+							UpdatedChangelist.Description = Description;
+							UpdateRepository.UpdateChangelist(UpdatedChangelist);
+						}
+					}
+				} 
+				catch (Exception Ex)
+				{
+					LogPerforce(1, "", $"Unable to update Changelist for CL {Change} to ${Description}, {Ex.Message}");
+				}
+			});
+		}
+
+		/// <inheritdoc/>
+		public async Task<int> CreateNewChangeAsync(string StreamName, string FilePath)
+		{
+			return await Task.Run(() =>
+			{
+
+				ServerSettings Settings = this.Settings.CurrentValue;
+
+				string Username = CanImpersonate ? "buildmachine" : Settings.P4BridgeServiceUsername!;
+
+				P4.SubmitResults? SubmitResults = null;
+
+				using (P4.Repository Repository = GetConnection(Stream: StreamName, Username: Username, ReadOnly: false, CreateChange: true))
+				{
+					P4.Client Client = Repository.Connection.Client;
+
+					string WorkspaceFilePath = $"//{Client.Name}/{FilePath.TrimStart('/')}";
+					string DiskFilePath = $"{Client.Root + FilePath.TrimStart('/')}";
+
+					P4.FileSpec WorkspaceFileSpec = new P4.FileSpec(new P4.ClientPath(WorkspaceFilePath));
+
+					IList<P4.File> Files = Repository.GetFiles(new P4.FilesCmdOptions(P4.FilesCmdFlags.None, 1), WorkspaceFileSpec);
+
+					P4.Changelist? SubmitChangelist = null;
+					P4.DepotPath? DepotPath = null;
+
+					const int MaxRetries = 10;
+					int Retry = 0;
+
+					for (; ; )
+					{
+						DepotPath = null;
+
+						if (Retry == MaxRetries)
+						{
+							break;
+						}
+
+						try
+						{
+							if (Files == null || Files.Count == 0)
+							{
+								// File does not exist, create it
+								string? DirectoryName = Path.GetDirectoryName(DiskFilePath);
+
+								if (string.IsNullOrEmpty(DirectoryName))
+								{
+									throw new Exception($"Invalid directory name for local client file, disk file path: {DiskFilePath}");
+								}
+
+								// Create the directory
+								if (!Directory.Exists(DirectoryName))
+								{
+									Directory.CreateDirectory(DirectoryName);
+
+									if (!Directory.Exists(DirectoryName))
+									{
+										throw new Exception($"Unable to create directrory: {DirectoryName}");
+									}
+								}
+
+								// Create the file
+								if (!File.Exists(DiskFilePath))
+								{
+									using (FileStream FileStream = File.OpenWrite(DiskFilePath))
+									{
+										FileStream.Close();
+									}
+
+									if (!File.Exists(DiskFilePath))
+									{
+										throw new Exception($"Unable to create local change file: {DiskFilePath}");
+									}
+
+								}
+
+								IList<P4.FileSpec> DepotFiles = Client.AddFiles(new P4.Options(), WorkspaceFileSpec);
+
+								if (DepotFiles == null || DepotFiles.Count != 1)
+								{
+									throw new Exception($"Unable to add local change file,  local: {DiskFilePath} : workspace: {WorkspaceFileSpec}");
+								}
+
+								DepotPath = DepotFiles[0].DepotPath;
+
+							}
+							else
+							{
+								IList<P4.FileSpec> SyncResults = Client.SyncFiles(new P4.SyncFilesCmdOptions(P4.SyncFilesCmdFlags.Force), WorkspaceFileSpec);
+
+								if (SyncResults == null || SyncResults.Count != 1)
+								{
+									throw new Exception($"Unable to sync file, workspace: {WorkspaceFileSpec}");
+								}
+
+								IList<P4.FileSpec> EditResults = Client.EditFiles(new P4.FileSpec[] { new P4.FileSpec(SyncResults[0].DepotPath) }, new P4.Options());
+
+								if (EditResults == null || EditResults.Count != 1)
+								{
+									throw new Exception($"Unable to edit file, workspace: {WorkspaceFileSpec}");
+								}
+
+								DepotPath = EditResults[0].DepotPath;
+
+							}
+
+							if (DepotPath == null || string.IsNullOrEmpty(DepotPath.Path))
+							{
+								throw new Exception($"Unable to get depot path for: {WorkspaceFileSpec}");
+							}
+
+							// create a new change
+							if (SubmitChangelist == null)
+							{
+								P4.Changelist Changelist = new P4.Changelist();
+								Changelist.Description = "New change for Horde job";
+								Changelist.Files.Add(new P4.FileMetaData(new P4.FileSpec(DepotPath)));
+								SubmitChangelist = Repository.CreateChangelist(Changelist);
+
+								if (SubmitChangelist == null)
+								{
+									throw new Exception($"Unable to create a changelist for: {DepotPath}");
+								}
+							}
+
+							SubmitResults = SubmitChangelist.Submit(null);
+
+							if (SubmitResults == null)
+							{
+								throw new Exception($"Unable to submit changelist for: {DepotPath}");
+							}
+
+							break;
+						}
+						catch
+						{
+							Retry++;
+
+							Client.RevertFiles(new P4.Options(), WorkspaceFileSpec);
+
+							continue;
+						}
+
+					}
+
+					string ClientName = Client.Name;
+					try
+					{
+						Repository.DeleteClient(Client, new P4.Options());
+					}
+					catch
+					{
+						Logger.LogError($"Unable to delete client {ClientName}");
+					}
+
+				}
+
+				if (SubmitResults == null)
+				{
+					throw new Exception($"Unable to submit change for {StreamName} {FilePath}");
+				}
+
+				return SubmitResults.ChangeIdAfterSubmit;
+			});
+		}
+
+		/// <inheritdoc/>
+		public async Task<(int? Change, string Message)> SubmitShelvedChangeAsync(int Change, int OriginalChange)
+		{
+			return await Task.Run(() => 
+			{
+			   using (P4.Repository Repository = GetConnection(NoClient: true))
+			   {
+
+				   List<string> Args = new List<string> { "-S", Change.ToString(CultureInfo.InvariantCulture), OriginalChange.ToString(CultureInfo.InvariantCulture) };
+
+				   using (P4.P4Command Command = new P4.P4Command(Repository, "describe", true, Args.ToArray()))
+				   {
+					   P4.P4CommandResult Result = Command.Run();
+
+					   if (!Result.Success)
+					   {
+							return (null, $"Unable to get change {Change}");
+						
+					   }
+
+					   if (Result.TaggedOutput == null || Result.TaggedOutput.Count != 2)
+					   {
+
+							return (null, $"Unable to get tagged output for change: {Change} and original change: {OriginalChange}");
+					   }
+
+					   bool DSTMismatch = false;
+					   string Offset = string.Empty;
+
+					   if (Server != null && Server.Metadata != null)
+					   {
+						   Offset = Server.Metadata.DateTimeOffset;
+						   DSTMismatch = P4.FormBase.DSTMismatch(Server.Metadata);
+					   }
+
+					   P4.Changelist Changelist = new P4.Changelist();
+					   Changelist.FromChangeCmdTaggedOutput(Result.TaggedOutput[0], true, Offset, DSTMismatch);
+
+					   P4.Changelist OriginalChangelist = new P4.Changelist();
+					   OriginalChangelist.FromChangeCmdTaggedOutput(Result.TaggedOutput[1], true, Offset, DSTMismatch);
+
+					   if (OriginalChangelist.ShelvedFiles.Count != Changelist.ShelvedFiles.Count)
+					   {
+							return (null, $"Mismatched number of shelved files for change: {Change} and original change: {OriginalChange}");
+					   }
+
+					   if (OriginalChangelist.ShelvedFiles.Count == 0)
+					   {
+							return (null, $"No shelved file for change: {Change} and original change: {OriginalChange}");
+					   }
+
+						foreach (P4.ShelvedFile ShelvedFile in Changelist.ShelvedFiles)
+					   {
+						   P4.ShelvedFile? Found = OriginalChangelist.ShelvedFiles.FirstOrDefault(Original => Original.Digest == ShelvedFile.Digest && Original.Action == ShelvedFile.Action);
+
+						   if (Found == null)
+						   {
+								return (null, $"Mismatch in shelved file digest or action for {ShelvedFile.Path}");
+						   }
+					   }
+
+					   Repository.Connection.Disconnect();
+
+					   using (P4.Repository SubmitRepository = GetConnection(ReadOnly: false, ClientFromChange: Change, UseClientFromChange: true, Username: Changelist.OwnerName))
+					   {
+							// we might not need a client here, possibly -e below facilitates this, check!
+							using (P4.P4Command SubmitCommand = new P4.P4Command(SubmitRepository, "submit", null, true, new string[] { "-e", Change.ToString(CultureInfo.InvariantCulture) }))
+						   {
+								try
+								{
+									Result = SubmitCommand.Run();
+								}
+								catch (Exception Ex)
+								{
+									return (null, $"Submit command failed: {Ex.Message}");
+								}
+
+							   if (!Result.Success)
+							   {
+									string Message = (Result.ErrorList != null && Result.ErrorList.Count > 0) ? Result.ErrorList[0].ErrorMessage : "Unknown error, no errors in list" ;
+									return (null, $"Unable to submit {Change}, {Message}");
+							   }
+
+							   int? SubmittedChangeId = null;
+
+							   foreach (P4.TaggedObject TaggedObject in Result.TaggedOutput)
+							   {
+								   string? Submitted;
+								   if (TaggedObject.TryGetValue("submittedChange", out Submitted))
+								   {
+									   SubmittedChangeId = int.Parse(Submitted, CultureInfo.InvariantCulture);
+								   }
+							   }
+
+							   if (SubmittedChangeId == null)
+							   {
+								   return (null, $"Submit command succeeded, though unable to parse submitted change number");
+							   }
+
+							   return (SubmittedChangeId, Changelist.Description);
+
+						   }
+					   }
+				   }
+			   }
+
+			   throw new Exception($"Unable to get shelve change: {Change} for original change: {OriginalChange}");
+		   });
+
+		}
+
+		/// <inheritdoc/>		
+		public async Task<PerforceUserInfo?> GetUserInfoAsync(string UserName)
+		{
+			return await Task.Run(() =>
+			{
+
+				using (P4.Repository Repository = GetConnection(NoClient: true))
+				{
+					P4.User User = Repository.GetUser(UserName, new P4.UserCmdOptions(P4.UserCmdFlags.Output));
+
+					if (User == null)
+					{
+						return null;
+					}
+
+					return new PerforceUserInfo { Name = UserName, Email = User.EmailAddress };
+				}
+			});
+		}
+
+		bool CanImpersonate
+		{
+			get
+			{
+				ServerSettings Settings = this.Settings.CurrentValue;
+				return Settings.P4BridgeCanImpersonate;
+			}
+		}
+
+		int ReshelveChange(P4.Repository Repository, int Change)
+		{
+
+			bool EdgeServer = false;
+			string? Value;
+			if (Server.Metadata.RawData.TryGetValue("serverServices", out Value))
+			{
+				if (Value == "edge-server")
+				{
+					EdgeServer = true;
+				}
+			}
+
+			List<string> Arguments = new List<string>();
+
+			// promote shelf if we're on an edge server
+			if (EdgeServer)
+			{
+				Arguments.Add("-p");
+			}
+
+			Arguments.AddRange(new string[] { "-s", Change.ToString(CultureInfo.InvariantCulture), "-f" });
+
+			using (P4.P4Command Cmd = new P4.P4Command(Repository, "reshelve", false, Arguments.ToArray()))
+			{
+				P4.P4CommandResult Results = Cmd.Run();
+
+				if (Results.Success)
+				{
+
+					if (Results.InfoOutput.Count == 0)
+					{
+						Serilog.Log.Logger.Information("Perforce: Unexpected info output when reshelving change");
+						throw new Exception("Unexpected info output when reshelving change");
+					}
+
+					bool Error = true;
+					int ReshelvedChange = 0;
+					string Message = Results.InfoOutput[Results.InfoOutput.Count - 1].Message;
+					Match ChangeMatch = Regex.Match(Message, @"Change (\d+) files shelved");
+
+					if (ChangeMatch.Success && ChangeMatch.Groups.Count == 2)
+					{
+						if (int.TryParse(ChangeMatch.Groups[1].Value, out ReshelvedChange))
+						{
+							Error = false;
+						}
+					}
+
+					if (Error)
+					{
+						Serilog.Log.Logger.Information("Perforce: Unable to parse cl for reshelf: {Message}", Message);
+						throw new Exception($"Unable to parse cl for reshelf: {Message}");
+					}
+
+					return ReshelvedChange;
+				}
+			}
+
+			Serilog.Log.Logger.Information("Perforce: General reshelving failure");
+
+			throw new Exception($"Unable to reshelve CL {Change}");
+
+		}
+
+		/// <inheritdoc/>
+		public virtual async Task<int> GetCodeChangeAsync(string StreamName, int Change)
+		{
+			int MaxChange = Change;
+			for (; ; )
+			{
+				// Query for the changes before this point
+				List<ChangeSummary> Changes = await GetChangesAsync(StreamName, null, MaxChange, 10, null);
+				Serilog.Log.Logger.Information("Finding last code change in {Stream} before {MaxChange}: {NumResults}", StreamName, MaxChange, Changes.Count);
+				if (Changes.Count == 0)
+				{
+					return 0;
+				}
+
+				// Get the details for them
+				List<ChangeDetails> DetailsList = await GetChangeDetailsAsync(StreamName, Changes.ConvertAll(x => x.Number), null);
+				foreach (ChangeDetails Details in DetailsList.OrderByDescending(x => x.Number))
+				{
+					ChangeContentFlags ContentFlags = Details.GetContentFlags();
+					Serilog.Log.Logger.Information("Change {Change} = {Flags}", Details.Number, ContentFlags.ToString());
+					if ((Details.GetContentFlags() & ChangeContentFlags.ContainsCode) != 0)
+					{
+						return Details.Number;
+					}
+				}
+
+				// Loop round again, adjusting our maximum changelist number
+				MaxChange = Changes.Min(x => x.Number) - 1;
+			}
+		}
+
+		/// <summary>
+		/// Gets the wildcard filter for a particular range of changes
+		/// </summary>
+		/// <param name="BasePath">Base path to find files for</param>
+		/// <param name="MinChange">Minimum changelist number to query</param>
+		/// <param name="MaxChange">Maximum changelist number to query</param>
+		/// <returns>Filter string</returns>
+		public static string GetFilter(string BasePath, int? MinChange, int? MaxChange)
+		{
+			StringBuilder Filter = new StringBuilder(BasePath);
+			if (MinChange != null && MaxChange != null)
+			{
+				Filter.Append($"@{MinChange},{MaxChange}");
+			}
+			else if (MinChange != null)
+			{
+				Filter.Append($"@>={MinChange}");
+			}
+			else if (MaxChange != null)
+			{
+				Filter.Append($"@<={MaxChange}");
+			}
+			return Filter.ToString();
+		}
+
+		/// <summary>
+		/// Gets a stream-relative path from a depot path
+		/// </summary>
+		/// <param name="DepotFile">The depot file to check</param>
+		/// <param name="StreamName">Name of the stream</param>
+		/// <param name="RelativePath">Receives the relative path to the file</param>
+		/// <returns>True if the stream-relative path was returned</returns>
+		public static bool TryGetStreamRelativePath(string DepotFile, string StreamName, [NotNullWhen(true)] out string? RelativePath)
+		{
+			if (DepotFile.StartsWith(StreamName, StringComparison.OrdinalIgnoreCase) && DepotFile.Length > StreamName.Length && DepotFile[StreamName.Length] == '/')
+			{
+				RelativePath = DepotFile.Substring(StreamName.Length);
+				return true;
+			}
+
+			Match Match = Regex.Match(DepotFile, "^//[^/]+/[^/]+(/.*)$");
+			if (Match.Success)
+			{
+				RelativePath = Match.Groups[1].Value;
+				return true;
+			}
+
+			RelativePath = null;
+			return false;
+		}
+	}
+
+}
diff --git a/Engine/Source/Programs/Horde/HordeServer/Services/PerforceService.cs b/Engine/Source/Programs/Horde/HordeServer/Services/PerforceService.cs
deleted file mode 100644
diff --git a/Engine/Source/Programs/Horde/HordeServer/Services/PerforceServiceApi.cs b/Engine/Source/Programs/Horde/HordeServer/Services/PerforceServiceApi.cs
deleted file mode 100644
diff --git a/Engine/Source/Programs/Horde/HordeServer/Startup.cs b/Engine/Source/Programs/Horde/HordeServer/Startup.cs
--- a/Engine/Source/Programs/Horde/HordeServer/Startup.cs
+++ b/Engine/Source/Programs/Horde/HordeServer/Startup.cs
@@ -373,7 +373,7 @@ namespace HordeServer
 			Services.AddSingleton<LifetimeService>();
 			Services.AddSingleton<ILogFileService, LogFileService>();
 			Services.AddSingleton<INotificationService, NotificationService>();
-			Services.AddSingleton<IPerforceService, P4APIPerforceService>();
+			Services.AddSingleton<IPerforceService, PerforceService>();
 
 			Services.AddSingleton<PerforceLoadBalancer>();
 			Services.AddSingleton<PoolService>();
