Horde: Allow nullable values for 'starred' and 'investigating' to be returned for UGS data.

diff --git a/Engine/Source/Programs/Horde/HordeServer/Collections/Impl/UgsMetadataCollection.cs b/Engine/Source/Programs/Horde/HordeServer/Collections/Impl/UgsMetadataCollection.cs
--- a/Engine/Source/Programs/Horde/HordeServer/Collections/Impl/UgsMetadataCollection.cs
+++ b/Engine/Source/Programs/Horde/HordeServer/Collections/Impl/UgsMetadataCollection.cs
@@ -42,8 +42,6 @@ namespace HordeServer.Collections.Impl
 
 			long? IUgsUserData.SyncTime => SyncTime;
 			UgsUserVote IUgsUserData.Vote => Vote ?? UgsUserVote.None;
-			bool IUgsUserData.Starred => Starred ?? false;
-			bool IUgsUserData.Investigating => Investigating ?? false;
 
 			public UgsUserDataDocument(string User)
 			{
diff --git a/Engine/Source/Programs/Horde/HordeServer/Models/UgsMetadata.cs b/Engine/Source/Programs/Horde/HordeServer/Models/UgsMetadata.cs
--- a/Engine/Source/Programs/Horde/HordeServer/Models/UgsMetadata.cs
+++ b/Engine/Source/Programs/Horde/HordeServer/Models/UgsMetadata.cs
@@ -33,12 +33,12 @@ namespace HordeServer.Models
 		/// <summary>
 		/// Whether the change is starred
 		/// </summary>
-		bool Starred { get; }
+		bool? Starred { get; }
 
 		/// <summary>
 		/// Whether the change is under investigation
 		/// </summary>
-		bool Investigating { get; }
+		bool? Investigating { get; }
 
 		/// <summary>
 		/// Comment text
