Horde: Fix tests. Override the PerforceService implementation used for the test host.

diff --git a/Engine/Source/Programs/Horde/HordeServerTests/ControllerIntegrationTest.cs b/Engine/Source/Programs/Horde/HordeServerTests/ControllerIntegrationTest.cs
--- a/Engine/Source/Programs/Horde/HordeServerTests/ControllerIntegrationTest.cs
+++ b/Engine/Source/Programs/Horde/HordeServerTests/ControllerIntegrationTest.cs
@@ -9,16 +9,19 @@ using System.Threading.Tasks;
 using HordeServer;
 using HordeServer.Collections.Impl;
 using HordeServer.Services;
+using Microsoft.AspNetCore.Builder;
 using Microsoft.AspNetCore.Hosting;
 using Microsoft.AspNetCore.Mvc.Testing;
+using Microsoft.AspNetCore.TestHost;
 using Microsoft.Extensions.Configuration;
 using Microsoft.Extensions.DependencyInjection;
+using Microsoft.Extensions.Hosting;
 
 namespace HordeServerTests
 {
-    public class TestWebApplicationFactory<TStartup> : WebApplicationFactory<TStartup> where TStartup : class
+	public class TestWebApplicationFactory<TStartup> : WebApplicationFactory<TStartup> where TStartup : class
     {
-        protected override void ConfigureWebHost(IWebHostBuilder builder)
+		protected override void ConfigureWebHost(IWebHostBuilder builder)
         {
             var Dict = new Dictionary<string, string?>
             {
@@ -30,6 +33,7 @@ namespace HordeServerTests
                 {"Horde:OidcClientId", null},
             };
             builder.ConfigureAppConfiguration((hostingContext, config) => { config.AddInMemoryCollection(Dict); });
+			builder.ConfigureTestServices(Collection => Collection.AddSingleton<IPerforceService, Stubs.Services.PerforceServiceStub>());
         }
     }
 
