Horde: Add user id to job queries
#rnx

diff --git a/Engine/Source/Programs/Horde/HordeServer/Collections/IJobCollection.cs b/Engine/Source/Programs/Horde/HordeServer/Collections/IJobCollection.cs
--- a/Engine/Source/Programs/Horde/HordeServer/Collections/IJobCollection.cs
+++ b/Engine/Source/Programs/Horde/HordeServer/Collections/IJobCollection.cs
@@ -95,6 +95,7 @@ namespace HordeServer.Collections
 		/// <param name="MinChange">The minimum changelist number</param>
 		/// <param name="MaxChange">The maximum changelist number</param>
 		/// <param name="PreflightChange">Preflight change to find</param>
+		/// <param name="StartedByUser">User id for which to include jobs</param>
 		/// <param name="PreflightStartedByUser">User for which to include preflight jobs</param>
 		/// <param name="MinCreateTime">The minimum creation time</param>
 		/// <param name="MaxCreateTime">The maximum creation time</param>
@@ -103,7 +104,7 @@ namespace HordeServer.Collections
 		/// <param name="Index">Index of the first result to return</param>
 		/// <param name="Count">Number of results to return</param>
 		/// <returns>List of jobs matching the given criteria</returns>
-		Task<List<IJob>> FindAsync(ObjectId[]? JobIds = null, StreamId? StreamId = null, string? Name = null, TemplateRefId[]? Templates = null, int? MinChange = null, int? MaxChange = null, int? PreflightChange = null, ObjectId? PreflightStartedByUser = null, DateTimeOffset? MinCreateTime = null, DateTimeOffset? MaxCreateTime = null, DateTimeOffset? ModifiedBefore = null, DateTimeOffset? ModifiedAfter = null, int? Index = null, int? Count = null);
+		Task<List<IJob>> FindAsync(ObjectId[]? JobIds = null, StreamId? StreamId = null, string? Name = null, TemplateRefId[]? Templates = null, int? MinChange = null, int? MaxChange = null, int? PreflightChange = null, ObjectId? PreflightStartedByUser = null, ObjectId? StartedByUser = null, DateTimeOffset ? MinCreateTime = null, DateTimeOffset? MaxCreateTime = null, DateTimeOffset? ModifiedBefore = null, DateTimeOffset? ModifiedAfter = null, int? Index = null, int? Count = null);
 
 		/// <summary>
 		/// Updates a new job
diff --git a/Engine/Source/Programs/Horde/HordeServer/Collections/Impl/JobCollection.cs b/Engine/Source/Programs/Horde/HordeServer/Collections/Impl/JobCollection.cs
--- a/Engine/Source/Programs/Horde/HordeServer/Collections/Impl/JobCollection.cs
+++ b/Engine/Source/Programs/Horde/HordeServer/Collections/Impl/JobCollection.cs
@@ -408,7 +408,7 @@ namespace HordeServer.Collections.Impl
 		}
 
 		/// <inheritdoc/>
-		public async Task<List<IJob>> FindAsync(ObjectId[]? JobIds, StreamId? StreamId, string? Name, TemplateRefId[]? Templates, int? MinChange, int? MaxChange, int? PreflightChange, ObjectId? PreflightStartedByUser, DateTimeOffset? MinCreateTime, DateTimeOffset? MaxCreateTime, DateTimeOffset? ModifiedBefore, DateTimeOffset? ModifiedAfter, int? Index, int? Count)
+		public async Task<List<IJob>> FindAsync(ObjectId[]? JobIds, StreamId? StreamId, string? Name, TemplateRefId[]? Templates, int? MinChange, int? MaxChange, int? PreflightChange, ObjectId? PreflightStartedByUser, ObjectId? StartedByUser, DateTimeOffset? MinCreateTime, DateTimeOffset? MaxCreateTime, DateTimeOffset? ModifiedBefore, DateTimeOffset? ModifiedAfter, int? Index, int? Count)
 		{
 			FilterDefinitionBuilder<JobDocument> FilterBuilder = Builders<JobDocument>.Filter;
 
@@ -445,6 +445,10 @@ namespace HordeServer.Collections.Impl
 			{
 				Filter &= FilterBuilder.Or(FilterBuilder.Eq(x => x.PreflightChange, 0), FilterBuilder.Eq(x => x.StartedByUserId, PreflightStartedByUser));
 			}
+			if (StartedByUser != null)
+			{
+				Filter &= FilterBuilder.Eq(x => x.StartedByUserId, StartedByUser);
+			}
 			if (MinCreateTime != null)
 			{
 				Filter &= FilterBuilder.Gte(x => x.CreateTimeUtc!, MinCreateTime.Value.UtcDateTime);
diff --git a/Engine/Source/Programs/Horde/HordeServer/Controllers/JobsController.cs b/Engine/Source/Programs/Horde/HordeServer/Controllers/JobsController.cs
--- a/Engine/Source/Programs/Horde/HordeServer/Controllers/JobsController.cs
+++ b/Engine/Source/Programs/Horde/HordeServer/Controllers/JobsController.cs
@@ -512,6 +512,7 @@ namespace HordeServer.Controllers
 		/// <param name="MaxChange">The maximum changelist number</param>
 		/// <param name="IncludePreflight">Whether to include preflight jobs</param>
 		/// <param name="PreflightChange">The preflighted changelist</param>
+		/// <param name="StartedByUserId">User id for which to include jobs</param>
 		/// <param name="PreflightStartedByUserId">User id for which to include preflight jobs</param>
 		/// <param name="MinCreateTime">Minimum creation time</param>
 		/// <param name="MaxCreateTime">Maximum creation time</param>
@@ -537,6 +538,7 @@ namespace HordeServer.Controllers
 			[FromQuery] bool IncludePreflight = true,
 			[FromQuery] int? PreflightChange = null,
 			[FromQuery] string? PreflightStartedByUserId = null,
+			[FromQuery] string? StartedByUserId = null,
 			[FromQuery] DateTimeOffset? MinCreateTime = null,
 			[FromQuery] DateTimeOffset? MaxCreateTime = null,
 			[FromQuery] DateTimeOffset? ModifiedBefore = null,
@@ -550,7 +552,7 @@ namespace HordeServer.Controllers
 		{
 			ObjectId[]? JobIdValues = (Ids == null) ? (ObjectId[]?)null : Array.ConvertAll(Ids, x => x.ToObjectId());
 			StreamId? StreamIdValue = (StreamId == null)? (StreamId?)null : new StreamId(StreamId);
-
+			
 			TemplateRefId[]? TemplateRefIds = (Templates != null && Templates.Length > 0) ? Templates.Select(x => new TemplateRefId(x)).ToArray() : null;
 
 			if (IncludePreflight == false)
@@ -565,11 +567,18 @@ namespace HordeServer.Controllers
 				PreflightStartedByUserIdValue = new ObjectId(PreflightStartedByUserId);
 			}
 
+			ObjectId? StartedByUserIdValue = null;
+
+			if (StartedByUserId != null)
+			{
+				StartedByUserIdValue = new ObjectId(StartedByUserId);
+			}
+
 			List<IJob> Jobs;
 			using (Scope _ = Tracer.Instance.StartActive("FindJobs"))
 			{
 				Jobs = await JobService.FindJobsAsync(JobIdValues, StreamIdValue, Name, TemplateRefIds, MinChange,
-					MaxChange, PreflightChange, PreflightStartedByUserIdValue, MinCreateTime?.UtcDateTime, MaxCreateTime?.UtcDateTime, Target, State, Outcome,
+					MaxChange, PreflightChange, PreflightStartedByUserIdValue, StartedByUserIdValue, MinCreateTime?.UtcDateTime, MaxCreateTime?.UtcDateTime, Target, State, Outcome,
 					ModifiedBefore, ModifiedAfter, Index, Count);
 			}
 
diff --git a/Engine/Source/Programs/Horde/HordeServer/Services/JobService.cs b/Engine/Source/Programs/Horde/HordeServer/Services/JobService.cs
--- a/Engine/Source/Programs/Horde/HordeServer/Services/JobService.cs
+++ b/Engine/Source/Programs/Horde/HordeServer/Services/JobService.cs
@@ -484,9 +484,10 @@ namespace HordeServer.Services
 		/// <param name="Name">Name of the job</param>
 		/// <param name="Templates">Templates to look for</param>
 		/// <param name="MinChange">The minimum changelist number</param>
-		/// <param name="MaxChange">The maximum changelist number</param>
+		/// <param name="MaxChange">The maximum changelist number</param>		
 		/// <param name="PreflightChange">The preflight change to look for</param>
-		/// <param name="PreflightStartedByUser">User for which to include preflight jobs</param>
+		/// <param name="PreflightStartedByUser">User for which to include preflight jobs</param>		
+		/// <param name="StartedByUser">User for which to include jobs</param>
 		/// <param name="MinCreateTime">The minimum creation time</param>
 		/// <param name="MaxCreateTime">The maximum creation time</param>
 		/// <param name="Target">The target to query</param>
@@ -497,11 +498,11 @@ namespace HordeServer.Services
 		/// <param name="Index">Index of the first result to return</param>
 		/// <param name="Count">Number of results to return</param>
 		/// <returns>List of jobs matching the given criteria</returns>
-		public async Task<List<IJob>> FindJobsAsync(ObjectId[]? JobIds = null, StreamId? StreamId = null, string? Name = null, TemplateRefId[]? Templates = null, int? MinChange = null, int? MaxChange = null, int? PreflightChange = null, ObjectId? PreflightStartedByUser = null, DateTimeOffset? MinCreateTime = null, DateTimeOffset? MaxCreateTime = null, string? Target = null, JobStepState[]? State = null, JobStepOutcome[]? Outcome = null, DateTimeOffset? ModifiedBefore = null, DateTimeOffset? ModifiedAfter = null, int? Index = null, int? Count = null)
+		public async Task<List<IJob>> FindJobsAsync(ObjectId[]? JobIds = null, StreamId? StreamId = null, string? Name = null, TemplateRefId[]? Templates = null, int? MinChange = null, int? MaxChange = null, int? PreflightChange = null, ObjectId ? PreflightStartedByUser = null, ObjectId? StartedByUser = null, DateTimeOffset ? MinCreateTime = null, DateTimeOffset? MaxCreateTime = null, string? Target = null, JobStepState[]? State = null, JobStepOutcome[]? Outcome = null, DateTimeOffset? ModifiedBefore = null, DateTimeOffset? ModifiedAfter = null, int? Index = null, int? Count = null)
 		{
 			if (Target == null && (State == null || State.Length == 0) && (Outcome == null || Outcome.Length == 0))
 			{
-				return await Jobs.FindAsync(JobIds, StreamId, Name, Templates, MinChange, MaxChange, PreflightChange, PreflightStartedByUser, MinCreateTime, MaxCreateTime, ModifiedBefore, ModifiedAfter, Index, Count);
+				return await Jobs.FindAsync(JobIds, StreamId, Name, Templates, MinChange, MaxChange, PreflightChange, PreflightStartedByUser, StartedByUser, MinCreateTime, MaxCreateTime, ModifiedBefore, ModifiedAfter, Index, Count);
 			}
 			else
 			{
@@ -511,7 +512,7 @@ namespace HordeServer.Services
 				int MaxCount = (Count ?? 1);
 				while (Results.Count < MaxCount)
 				{
-					List<IJob> ScanJobs = await Jobs.FindAsync(JobIds, StreamId, Name, Templates, MinChange, MaxChange, PreflightChange, PreflightStartedByUser, MinCreateTime, MaxCreateTime, ModifiedBefore, ModifiedAfter, 0, 5);
+					List<IJob> ScanJobs = await Jobs.FindAsync(JobIds, StreamId, Name, Templates, MinChange, MaxChange, PreflightChange, PreflightStartedByUser, StartedByUser, MinCreateTime, MaxCreateTime, ModifiedBefore, ModifiedAfter, 0, 5);
 					if (ScanJobs.Count == 0)
 					{
 						break;
