Horde: Add user select to job search
#rnx

diff --git a/Engine/Source/Programs/Horde/HordeDashboard/src/backend/Api.ts b/Engine/Source/Programs/Horde/HordeDashboard/src/backend/Api.ts
--- a/Engine/Source/Programs/Horde/HordeDashboard/src/backend/Api.ts
+++ b/Engine/Source/Programs/Horde/HordeDashboard/src/backend/Api.ts
@@ -263,6 +263,7 @@ export type JobQuery = {
 	preflightChange?: number;
 	includePreflight?: boolean,
 	preflightStartedByUserId?: string,
+	startedByUserId?: string,
 	minCreateTime?: string;
 	maxCreateTime?: string;
 	modifiedAfter?: string;
