Horde: Disable reshelving of preflight changes again. Seeing P4 DB corruption.

diff --git a/Engine/Source/Programs/Horde/HordeServer/Services/JobService.cs b/Engine/Source/Programs/Horde/HordeServer/Services/JobService.cs
--- a/Engine/Source/Programs/Horde/HordeServer/Services/JobService.cs
+++ b/Engine/Source/Programs/Horde/HordeServer/Services/JobService.cs
@@ -201,6 +201,13 @@ namespace HordeServer.Services
 			this.Logger = Logger;
 		}
 
+#pragma warning disable CA1801
+		static bool ShouldClonePreflightChange(StreamId StreamId)
+		{
+			return false; //			return StreamId == new StreamId("ue5-main");
+		}
+#pragma warning restore CA1801
+
 		/// <summary>
 		/// Creates a new job
 		/// </summary>
@@ -233,7 +240,7 @@ namespace HordeServer.Services
 			ObjectId JobIdValue = JobId ?? ObjectId.GenerateNewId();
 			using IDisposable Scope = Logger.BeginScope("CreateJobAsync({JobId})", JobIdValue);
 
-			if (PreflightChange != null)
+			if (PreflightChange != null && ShouldClonePreflightChange(StreamId))
 			{
 				ClonedPreflightChange = await CloneShelvedChangeAsync(ClonedPreflightChange ?? PreflightChange.Value);
 			}
@@ -975,7 +982,14 @@ namespace HordeServer.Services
 				int ClonedPreflightChange = Job.ClonedPreflightChange;
 				if (ClonedPreflightChange == 0)
 				{
-					ClonedPreflightChange = await CloneShelvedChangeAsync(Job.PreflightChange);
+					if (ShouldClonePreflightChange(Job.StreamId))
+					{
+						ClonedPreflightChange = await CloneShelvedChangeAsync(Job.PreflightChange);
+					}
+					else
+					{
+						ClonedPreflightChange = Job.PreflightChange;
+					}
 				}
 
 				Logger.LogInformation("Updating description for {ClonedPreflightChange}", ClonedPreflightChange);
@@ -1004,9 +1018,19 @@ namespace HordeServer.Services
 				Logger.LogError(Ex, "Unable to submit shelved change");
 			}
 
-			if (Change != null && Job.ClonedPreflightChange != 0)
+			if (ShouldClonePreflightChange(Job.StreamId))
 			{
-				await DeleteShelvedChangeAsync(Job.PreflightChange);
+				if (Change != null && Job.ClonedPreflightChange != 0)
+				{
+					await DeleteShelvedChangeAsync(Job.PreflightChange);
+				}
+			}
+			else
+			{
+				if (Change != null && Job.PreflightChange != 0)
+				{
+					await DeleteShelvedChangeAsync(Job.PreflightChange);
+				}
 			}
 
 			for (; ; )
