Horde: Unset the "hidden" flag whenever someone logs in. Some users caught in the purge of invalid user records.

diff --git a/Engine/Source/Programs/Horde/HordeServer/Collections/Impl/UserCollectionV2.cs b/Engine/Source/Programs/Horde/HordeServer/Collections/Impl/UserCollectionV2.cs
--- a/Engine/Source/Programs/Horde/HordeServer/Collections/Impl/UserCollectionV2.cs
+++ b/Engine/Source/Programs/Horde/HordeServer/Collections/Impl/UserCollectionV2.cs
@@ -249,7 +249,7 @@ namespace HordeServer.Collections.Impl
 		/// <inheritdoc/>
 		public async Task<IUser> FindOrAddUserByLoginAsync(string Login, string? Name, string? Email)
 		{
-			UpdateDefinition<UserDocument> Update = Builders<UserDocument>.Update.SetOnInsert(x => x.Id, ObjectId.GenerateNewId()).SetOnInsert(x => x.Login, Login);
+			UpdateDefinition<UserDocument> Update = Builders<UserDocument>.Update.SetOnInsert(x => x.Id, ObjectId.GenerateNewId()).SetOnInsert(x => x.Login, Login).Unset(x => x.Hidden);
 
 			if (Name == null)
 			{
