Horde: Add basic support for multiple Perforce servers. PerforceService now selects the appropriate server to use based on a ClusterName parameter passed in.

diff --git a/Engine/Source/Programs/Horde/HordeServer/Collections/Impl/StreamCollection.cs b/Engine/Source/Programs/Horde/HordeServer/Collections/Impl/StreamCollection.cs
--- a/Engine/Source/Programs/Horde/HordeServer/Collections/Impl/StreamCollection.cs
+++ b/Engine/Source/Programs/Horde/HordeServer/Collections/Impl/StreamCollection.cs
@@ -41,6 +41,7 @@ namespace HordeServer.Collections.Impl
 			[BsonRequired]
 			public string Name { get; set; }
 
+			public string? ClusterName { get; set; }
 			public string ConfigPath { get; set; } = String.Empty;
 			public string ConfigRevision { get; set; } = String.Empty;
 
@@ -60,6 +61,7 @@ namespace HordeServer.Collections.Impl
 			public int UpdateIndex { get; set; }
 			public bool Deleted { get; set; }
 
+			string IStream.ClusterName => ClusterName ?? PerforceCluster.DefaultName;
 			IReadOnlyList<StreamTab> IStream.Tabs => Tabs;
 			IReadOnlyDictionary<string, AgentType> IStream.AgentTypes => AgentTypes;
 			IReadOnlyDictionary<string, WorkspaceType> IStream.WorkspaceTypes => WorkspaceTypes;
@@ -204,6 +206,7 @@ namespace HordeServer.Collections.Impl
 		public async Task<IStream?> TryCreateAsync(StreamId Id, ProjectId ProjectId, string ConfigPath, string ConfigRevision, StreamConfig Config, DefaultPreflight? DefaultPreflight, List<StreamTab>? Tabs, Dictionary<string, AgentType>? AgentTypes, Dictionary<string, WorkspaceType>? WorkspaceTypes, Dictionary<TemplateRefId, TemplateRef>? TemplateRefs, Acl? Acl)
 		{
 			StreamDocument NewStream = new StreamDocument(Id, Config.Name, ProjectId);
+			NewStream.ClusterName = Config.ClusterName;
 			NewStream.ConfigPath = ConfigPath;
 			NewStream.ConfigRevision = ConfigRevision;
 			NewStream.Order = Config.Order ?? StreamDocument.DefaultOrder;
@@ -270,6 +273,9 @@ namespace HordeServer.Collections.Impl
 			Stream.ProjectId = ProjectId;
 			Updates.Add(UpdateBuilder.Set(x => x.ProjectId, Stream.ProjectId));
 
+			Stream.ClusterName = Config.ClusterName;
+			Updates.Add(UpdateBuilder.Set(x => x.ClusterName, Stream.ClusterName));
+
 			Stream.ConfigPath = ConfigPath;
 			Updates.Add(UpdateBuilder.Set(x => x.ConfigPath, Stream.ConfigPath));
 
diff --git a/Engine/Source/Programs/Horde/HordeServer/Controllers/DebugController.cs b/Engine/Source/Programs/Horde/HordeServer/Controllers/DebugController.cs
--- a/Engine/Source/Programs/Horde/HordeServer/Controllers/DebugController.cs
+++ b/Engine/Source/Programs/Horde/HordeServer/Controllers/DebugController.cs
@@ -630,7 +630,7 @@ namespace HordeServer.Controllers
 				return BadRequest();
 			}
 
-			return await Perforce.CreateTicket(Username);
+			return await Perforce.CreateTicket(PerforceCluster.DefaultName, Username);
 		}
 
 		/// <summary>
@@ -655,8 +655,8 @@ namespace HordeServer.Controllers
 			// ParallelTest.ForAsync doesn't seem to handle async lambdas
 			IEnumerable<Task> Tasks = Enumerable.Range(0, 16).Select(async (Idx) => {
 
-				int LatestChange = await Perforce.GetLatestChangeAsync(DevBuild, PerforceUser);
-				int CodeChange = await Perforce.GetCodeChangeAsync(DevBuild, LatestChange);				
+				int LatestChange = await Perforce.GetLatestChangeAsync(PerforceCluster.DefaultName, DevBuild, PerforceUser);
+				int CodeChange = await Perforce.GetCodeChangeAsync(PerforceCluster.DefaultName, DevBuild, LatestChange);				
 
 				lock (UpdateLock)
 				{
@@ -668,8 +668,8 @@ namespace HordeServer.Controllers
 
 			Tasks = Enumerable.Range(0, 16).Select(async (Idx) => {
 
-				int LatestChange = await Perforce.GetLatestChangeAsync(DevBuild, PerforceUser);
-				int CodeChange = await Perforce.GetCodeChangeAsync(DevBuild, LatestChange);				
+				int LatestChange = await Perforce.GetLatestChangeAsync(PerforceCluster.DefaultName, DevBuild, PerforceUser);
+				int CodeChange = await Perforce.GetCodeChangeAsync(PerforceCluster.DefaultName, DevBuild, LatestChange);				
 
 				lock (UpdateLock)
 				{
@@ -720,25 +720,25 @@ namespace HordeServer.Controllers
 			Dictionary<string, string> Results = new Dictionary<string, string>();
 
 			// NOTE: This is in Private-Build not DevBuild
-			int NewChangeId = await Perforce.CreateNewChangeAsync("//UE4/Private-Build", "Counter.txt");
+			int NewChangeId = await Perforce.CreateNewChangeAsync(PerforceCluster.DefaultName, "//UE4/Private-Build", "Counter.txt");
 
-			List<ChangeDetails> ChangeDetails = await Perforce.GetChangeDetailsAsync("//UE4/Private-Build", new int[] { NewChangeId }, PerforceUser);
+			List<ChangeDetails> ChangeDetails = await Perforce.GetChangeDetailsAsync(PerforceCluster.DefaultName, "//UE4/Private-Build", new int[] { NewChangeId }, PerforceUser);
 
 			Results.Add("CreateNewChangeAsync", $"{NewChangeId} : {ChangeDetails[0].Description}");
 
-			await Perforce.UpdateChangelistDescription(TestCL, $"Updated Description: New Change CL was {NewChangeId}");
+			await Perforce.UpdateChangelistDescription(PerforceCluster.DefaultName, TestCL, $"Updated Description: New Change CL was {NewChangeId}");
 
-			ChangeDetails = await Perforce.GetChangeDetailsAsync(DevBuild, new int[] { TestCL }, PerforceUser);
+			ChangeDetails = await Perforce.GetChangeDetailsAsync(PerforceCluster.DefaultName, DevBuild, new int[] { TestCL }, PerforceUser);
 
 			Results.Add("UpdateChangelistDescription", $"{TestCL} : {ChangeDetails[0].Description}");
 
-			int LatestChange = await Perforce.GetLatestChangeAsync(DevBuild, PerforceUser);
+			int LatestChange = await Perforce.GetLatestChangeAsync(PerforceCluster.DefaultName, DevBuild, PerforceUser);
 			Results.Add("GetLatestChangeAsync", LatestChange.ToString(CultureInfo.InvariantCulture));
 
-			int CodeChange = await Perforce.GetCodeChangeAsync(DevBuild, LatestChange);
+			int CodeChange = await Perforce.GetCodeChangeAsync(PerforceCluster.DefaultName, DevBuild, LatestChange);
 			Results.Add("GetCodeChangeAsync", CodeChange.ToString(CultureInfo.InvariantCulture));
 
-			PerforceUserInfo? UserInfo = await Perforce.GetUserInfoAsync(PerforceUser);
+			PerforceUserInfo? UserInfo = await Perforce.GetUserInfoAsync(PerforceCluster.DefaultName, PerforceUser);
 
 			if (UserInfo == null)
 			{
@@ -750,14 +750,14 @@ namespace HordeServer.Controllers
 			}
 
 			// changes
-			List<ChangeSummary> Changes = await Perforce.GetChangesAsync(DevBuild, null, LatestChange, 10, PerforceUser);
+			List<ChangeSummary> Changes = await Perforce.GetChangesAsync(PerforceCluster.DefaultName, DevBuild, null, LatestChange, 10, PerforceUser);
 			int Count = 0;
 			foreach (ChangeSummary Summary in Changes)
 			{
 				Results.Add($"GetChangesAsync Result {Count++}", $"{Summary.Number} : {Summary.Author} - {Summary.Description}");
 			}
 
-			ChangeDetails = await Perforce.GetChangeDetailsAsync(DevBuild, Changes.Select(Change => Change.Number).ToArray(), PerforceUser);
+			ChangeDetails = await Perforce.GetChangeDetailsAsync(PerforceCluster.DefaultName, DevBuild, Changes.Select(Change => Change.Number).ToArray(), PerforceUser);
 			Count = 0;
 			foreach (ChangeDetails Details in ChangeDetails)
 			{
@@ -767,7 +767,7 @@ namespace HordeServer.Controllers
 			String Ticket = "Failed (Expected if not running service account)";
 			try
 			{
-				await Perforce.CreateTicket(PerforceUser);
+				await Perforce.CreateTicket(PerforceCluster.DefaultName, PerforceUser);
 				Ticket = "Success";
 			}
 			catch
@@ -776,7 +776,7 @@ namespace HordeServer.Controllers
 			}
 
 			Results.Add($"CreateTicket", $"{Ticket}");
-			List<FileSummary> FileSummaries = await Perforce.FindFilesAsync(ChangeDetails.SelectMany(Details => Details.Files).Select(File => $"{DevBuild}{File}"));
+			List<FileSummary> FileSummaries = await Perforce.FindFilesAsync(PerforceCluster.DefaultName, ChangeDetails.SelectMany(Details => Details.Files).Select(File => $"{DevBuild}{File}"));
 			Count = 0;
 			foreach (FileSummary File in FileSummaries)
 			{
@@ -785,10 +785,10 @@ namespace HordeServer.Controllers
 
 			// print async
 
-			byte[] Bytes = await Perforce.PrintAsync($"{DevBuild}/RunUAT.bat");
+			byte[] Bytes = await Perforce.PrintAsync(PerforceCluster.DefaultName, $"{DevBuild}/RunUAT.bat");
 			Results.Add($"PrintAsync: {DevBuild}/RunUAT.bat ", $"{System.Text.Encoding.Default.GetString(Bytes)}");
 
-			Bytes = await Perforce.PrintAsync($"{DevBuild}/RunUAT.bat@13916380");
+			Bytes = await Perforce.PrintAsync(PerforceCluster.DefaultName, $"{DevBuild}/RunUAT.bat@13916380");
 			Results.Add($"PrintAsync: {DevBuild}/RunUAT.bat@13916380", $"{System.Text.Encoding.Default.GetString(Bytes)}");
 
 			// need to be running as service account for these
@@ -799,7 +799,7 @@ namespace HordeServer.Controllers
 
 			try 
 			{
-				DuplicateCL = await Perforce.DuplicateShelvedChangeAsync(ShelvedChangeOnDevBuild);
+				DuplicateCL = await Perforce.DuplicateShelvedChangeAsync(PerforceCluster.DefaultName, ShelvedChangeOnDevBuild);
 				Results.Add("DuplicateShelvedChangeAsync", $"Duplicated CL {DuplicateCL} from Shelved CL {ShelvedChangeOnDevBuild}");
 			}
 			catch 
@@ -810,15 +810,15 @@ namespace HordeServer.Controllers
 			if (DuplicateCL.HasValue)
 			{
 				// Note change client much exist for the update
-				ChangeDetails = await Perforce.GetChangeDetailsAsync(DevBuild, new int[]{ DuplicateCL.Value}, PerforceUser);
+				ChangeDetails = await Perforce.GetChangeDetailsAsync(PerforceCluster.DefaultName, DevBuild, new int[]{ DuplicateCL.Value}, PerforceUser);
 
 				Results.Add("UpdateChangelistDescription - PreUpdate", $"{DuplicateCL} : {ChangeDetails[0].Description}"); 
 
 				try 
 				{
-					await Perforce.UpdateChangelistDescription(DuplicateCL.Value, "Updated Description from Horde");
+					await Perforce.UpdateChangelistDescription(PerforceCluster.DefaultName, DuplicateCL.Value, "Updated Description from Horde");
 
-					ChangeDetails = await Perforce.GetChangeDetailsAsync(DevBuild, new int[]{ DuplicateCL.Value}, PerforceUser);
+					ChangeDetails = await Perforce.GetChangeDetailsAsync(PerforceCluster.DefaultName, DevBuild, new int[]{ DuplicateCL.Value}, PerforceUser);
 
 					Results.Add("UpdateChangelistDescription - PostUpdate", $"{DuplicateCL} : {ChangeDetails[0].Description}"); 
 				}
@@ -827,15 +827,15 @@ namespace HordeServer.Controllers
 					Results.Add("UpdateChangelistDescription - Faile", $"Client from change {DuplicateCL} must exist"); 
 				}
 
-				await Perforce.DeleteShelvedChangeAsync(DuplicateCL.Value);
+				await Perforce.DeleteShelvedChangeAsync(PerforceCluster.DefaultName, DuplicateCL.Value);
 
-				ChangeDetails = await Perforce.GetChangeDetailsAsync(DevBuild,new int[] {DuplicateCL.Value} , PerforceUser);
+				ChangeDetails = await Perforce.GetChangeDetailsAsync(PerforceCluster.DefaultName, DevBuild, new int[] {DuplicateCL.Value} , PerforceUser);
 
 				Results.Add("GetChangeDetailsAsync - After Delete", $"ChangeDetails.Count: {ChangeDetails.Count}"); 
 
 			}
 
-			byte[] ImageBytes = await Perforce.PrintAsync($"{DevBuild}/Samples/Games/ShooterGame/ShooterGame.png");
+			byte[] ImageBytes = await Perforce.PrintAsync(PerforceCluster.DefaultName, $"{DevBuild}/Samples/Games/ShooterGame/ShooterGame.png");
 
 			Content.AppendLine("<html><body><pre>");
 			foreach (KeyValuePair<string, string> Pair in Results)
diff --git a/Engine/Source/Programs/Horde/HordeServer/Controllers/JobsController.cs b/Engine/Source/Programs/Horde/HordeServer/Controllers/JobsController.cs
--- a/Engine/Source/Programs/Horde/HordeServer/Controllers/JobsController.cs
+++ b/Engine/Source/Programs/Horde/HordeServer/Controllers/JobsController.cs
@@ -197,17 +197,17 @@ namespace HordeServer.Controllers
 			}
 			else
 			{
-				Change = await Perforce.GetLatestChangeAsync(Stream.Name, null);
+				Change = await Perforce.GetLatestChangeAsync(Stream.ClusterName, Stream.Name, null);
 			}
 
 			// And get the matching code changelist
-			int CodeChange = await Perforce.GetCodeChangeAsync(Stream.Name, Change);
+			int CodeChange = await Perforce.GetCodeChangeAsync(Stream.ClusterName, Stream.Name, Change);
 
 			// New properties for the job
 			List<string> Arguments = Create.Arguments ?? Template.GetDefaultArguments();
 
 			// Create the job
-			IJob Job = await JobService.CreateJobAsync(null, Stream.Id, TemplateRefId, Template.Id, Graph, Name, Change, CodeChange, Create.PreflightChange, null, User.GetUserId(), User.GetUserName(), Priority, Create.AutoSubmit, Create.UpdateIssues, TemplateRef.ChainedJobs, TemplateRef.ShowUgsBadges, TemplateRef.ShowUgsAlerts, TemplateRef.NotificationChannel, TemplateRef.NotificationChannelFilter, null, Template.Counters, Arguments);
+			IJob Job = await JobService.CreateJobAsync(null, Stream, TemplateRefId, Template.Id, Graph, Name, Change, CodeChange, Create.PreflightChange, null, User.GetUserId(), User.GetUserName(), Priority, Create.AutoSubmit, Create.UpdateIssues, TemplateRef.ChainedJobs, TemplateRef.ShowUgsBadges, TemplateRef.ShowUgsAlerts, TemplateRef.NotificationChannel, TemplateRef.NotificationChannelFilter, null, Template.Counters, Arguments);
 			await UpdateNotificationsAsync(Job.Id.ToString(), new UpdateNotificationsRequest { Slack = true });
 			return new CreateJobResponse(Job.Id.ToString());
 		}
@@ -226,7 +226,7 @@ namespace HordeServer.Controllers
 			if (Jobs.Count == 0)
 			{
 				Logger.LogInformation("Unable to find successful build of {TemplateRefId} target {Target}. Using latest change instead", TemplateId, Target);
-				return await Perforce.GetLatestChangeAsync(Stream.Name, null);
+				return await Perforce.GetLatestChangeAsync(Stream.ClusterName, Stream.Name, null);
 			}
 			else
 			{
diff --git a/Engine/Source/Programs/Horde/HordeServer/Controllers/StreamsController.cs b/Engine/Source/Programs/Horde/HordeServer/Controllers/StreamsController.cs
--- a/Engine/Source/Programs/Horde/HordeServer/Controllers/StreamsController.cs
+++ b/Engine/Source/Programs/Horde/HordeServer/Controllers/StreamsController.cs
@@ -204,7 +204,7 @@ namespace HordeServer.Controllers
 				return Forbid();
 			}
 
-			List<ChangeSummary> Commits = await PerforceService.GetChangesAsync(Stream.Name, Min, Max, Results, PerforceUser);
+			List<ChangeSummary> Commits = await PerforceService.GetChangesAsync(Stream.ClusterName, Stream.Name, Min, Max, Results, PerforceUser);
 			return Commits.ConvertAll(x => PropertyFilter.Apply(new GetChangeSummaryResponse(x), Filter));
 		}
 
@@ -238,7 +238,7 @@ namespace HordeServer.Controllers
 				return Forbid();
 			}
 
-			ChangeDetails? ChangeDetails = await PerforceService.GetChangeDetailsAsync(Stream.Name, Number, PerforceUser);
+			ChangeDetails? ChangeDetails = await PerforceService.GetChangeDetailsAsync(Stream.ClusterName, Stream.Name, Number, PerforceUser);
 			if(ChangeDetails == null)
 			{
 				return NotFound();
diff --git a/Engine/Source/Programs/Horde/HordeServer/Models/Config.cs b/Engine/Source/Programs/Horde/HordeServer/Models/Config.cs
--- a/Engine/Source/Programs/Horde/HordeServer/Models/Config.cs
+++ b/Engine/Source/Programs/Horde/HordeServer/Models/Config.cs
@@ -140,6 +140,11 @@ namespace HordeServer.Models
 		[Required]
 		public string Name { get; set; } = String.Empty;
 
+		/// <summary>
+		/// The perforce cluster containing the stream
+		/// </summary>
+		public string? ClusterName { get; set; }
+
 		/// <summary>
 		/// Order for this stream
 		/// </summary>
diff --git a/Engine/Source/Programs/Horde/HordeServer/Models/Globals.cs b/Engine/Source/Programs/Horde/HordeServer/Models/Globals.cs
--- a/Engine/Source/Programs/Horde/HordeServer/Models/Globals.cs
+++ b/Engine/Source/Programs/Horde/HordeServer/Models/Globals.cs
@@ -268,6 +268,16 @@ namespace HordeServer.Models
 		[Required]
 		public string Name { get; set; } = null!;
 
+		/// <summary>
+		/// Username for Horde to log in to this server
+		/// </summary>
+		public string ServiceAccount { get; set; } = "buildmachine";
+
+		/// <summary>
+		/// Whether the service account can impersonate other users
+		/// </summary>
+		public bool CanImpersonate { get; set; } = true;
+
 		/// <summary>
 		/// List of servers
 		/// </summary>
diff --git a/Engine/Source/Programs/Horde/HordeServer/Models/Stream.cs b/Engine/Source/Programs/Horde/HordeServer/Models/Stream.cs
--- a/Engine/Source/Programs/Horde/HordeServer/Models/Stream.cs
+++ b/Engine/Source/Programs/Horde/HordeServer/Models/Stream.cs
@@ -497,6 +497,11 @@ namespace HordeServer.Models
 		/// </summary>
 		public string Name { get; }
 
+		/// <summary>
+		/// Name of the perforce cluster
+		/// </summary>
+		public string ClusterName { get; }
+
 		/// <summary>
 		/// Path to the configuration file for this stream
 		/// </summary>
diff --git a/Engine/Source/Programs/Horde/HordeServer/Services/ConfigService.cs b/Engine/Source/Programs/Horde/HordeServer/Services/ConfigService.cs
--- a/Engine/Source/Programs/Horde/HordeServer/Services/ConfigService.cs
+++ b/Engine/Source/Programs/Horde/HordeServer/Services/ConfigService.cs
@@ -17,6 +17,7 @@ using System.IO;
 using System.Linq;
 using System.Text;
 using System.Text.Json;
+using System.Text.RegularExpressions;
 using System.Threading;
 using System.Threading.Tasks;
 
@@ -30,6 +31,9 @@ namespace HordeServer.Services
 	/// </summary>
 	public class ConfigService : ElectedBackgroundService
 	{
+		const string FileScheme = "file";
+		const string PerforceScheme = "p4-cluster";
+
 		/// <summary>
 		/// Config file version number
 		/// </summary>
@@ -97,13 +101,13 @@ namespace HordeServer.Services
 		Dictionary<ProjectId, ProjectConfig> CachedProjectConfigs = new Dictionary<ProjectId, ProjectConfig>();
 		Dictionary<ProjectId, string?> CachedLogoRevisions = new Dictionary<ProjectId, string?>();
 
-		async Task UpdateConfigAsync(string ConfigPath)
+		async Task UpdateConfigAsync(Uri ConfigPath)
 		{
 			// Update the globals singleton
 			GlobalConfig GlobalConfig;
 			for (; ; )
 			{
-				Dictionary<string, string> GlobalRevisions = await FindRevisionsAsync(new[] { ConfigPath });
+				Dictionary<Uri, string> GlobalRevisions = await FindRevisionsAsync(new[] { ConfigPath });
 				if (GlobalRevisions.Count == 0)
 				{
 					throw new Exception($"Invalid config path: {ConfigPath}");
@@ -146,24 +150,24 @@ namespace HordeServer.Services
 			List<IProject> Projects = await ProjectService.GetProjectsAsync();
 
 			// Get the path to all the project configs
-			List<(ProjectConfigRef ProjectRef, string Path)> ProjectConfigs = GlobalConfig.Projects.Select(x => (x, CombinePaths(ConfigPath, x.Path))).ToList();
+			List<(ProjectConfigRef ProjectRef, Uri Path)> ProjectConfigs = GlobalConfig.Projects.Select(x => (x, CombinePaths(ConfigPath, x.Path))).ToList();
 
 			Dictionary<ProjectId, ProjectConfig> PrevCachedProjectConfigs = CachedProjectConfigs;
 			CachedProjectConfigs = new Dictionary<ProjectId, ProjectConfig>();
 
-			List<(ProjectId ProjectId, string Path)> ProjectLogos = new List<(ProjectId ProjectId, string Path)>();
+			List<(ProjectId ProjectId, Uri Path)> ProjectLogos = new List<(ProjectId ProjectId, Uri Path)>();
 
-			List<(ProjectId ProjectId, StreamConfigRef StreamRef, string Path)> StreamConfigs = new List<(ProjectId, StreamConfigRef, string)>();
+			List<(ProjectId ProjectId, StreamConfigRef StreamRef, Uri Path)> StreamConfigs = new List<(ProjectId, StreamConfigRef, Uri)>();
 
 			// List of project ids that were not able to be updated. We will avoid removing any existing project or stream definitions for these.
 			HashSet<ProjectId> SkipProjectIds = new HashSet<ProjectId>();
 
 			// Update any existing projects
-			Dictionary<string, string> ProjectRevisions = await FindRevisionsAsync(ProjectConfigs.Select(x => x.Path));
+			Dictionary<Uri, string> ProjectRevisions = await FindRevisionsAsync(ProjectConfigs.Select(x => x.Path));
 			for (int Idx = 0; Idx < ProjectConfigs.Count; Idx++)
 			{
 				// Make sure we were able to fetch metadata for 
-				(ProjectConfigRef ProjectRef, string ProjectPath) = ProjectConfigs[Idx];
+				(ProjectConfigRef ProjectRef, Uri ProjectPath) = ProjectConfigs[Idx];
 				if (!ProjectRevisions.TryGetValue(ProjectPath, out string? Revision))
 				{
 					Logger.LogWarning("Unable to update project {ProjectId} due to missing revision information", ProjectRef.Id);
@@ -172,7 +176,7 @@ namespace HordeServer.Services
 				}
 
 				IProject? Project = Projects.FirstOrDefault(x => x.Id == ProjectRef.Id);
-				bool Update = (Project == null || Project.ConfigPath != ConfigPath || Project.ConfigRevision != Revision);
+				bool Update = (Project == null || Project.ConfigPath != ProjectPath.ToString() || Project.ConfigRevision != Revision);
 
 				ProjectConfig? ProjectConfig;
 				if (Update || !PrevCachedProjectConfigs.TryGetValue(ProjectRef.Id, out ProjectConfig))
@@ -187,7 +191,7 @@ namespace HordeServer.Services
 					if (Update)
 					{
 						Logger.LogInformation("Updating configuration for project {ProjectId} ({Revision})", ProjectRef.Id, Revision);
-						await ProjectService.Collection.AddOrUpdateAsync(ProjectRef.Id, ProjectPath, Revision, Idx, ProjectConfig);
+						await ProjectService.Collection.AddOrUpdateAsync(ProjectRef.Id, ProjectPath.ToString(), Revision, Idx, ProjectConfig);
 					}
 				}
 
@@ -201,10 +205,10 @@ namespace HordeServer.Services
 			}
 
 			// Get the logo revisions
-			Dictionary<string, string> LogoRevisions = await FindRevisionsAsync(ProjectLogos.Select(x => x.Path));
+			Dictionary<Uri, string> LogoRevisions = await FindRevisionsAsync(ProjectLogos.Select(x => x.Path));
 			for (int Idx = 0; Idx < ProjectLogos.Count; Idx++)
 			{
-				(ProjectId ProjectId, string Path) = ProjectLogos[Idx];
+				(ProjectId ProjectId, Uri Path) = ProjectLogos[Idx];
 				if (LogoRevisions.TryGetValue(Path, out string? Revision))
 				{
 					string? CurrentRevision;
@@ -216,7 +220,7 @@ namespace HordeServer.Services
 					if (Revision != CurrentRevision)
 					{
 						Logger.LogInformation("Updating logo for project {ProjectId} ({Revision})", ProjectId, Revision);
-						await ProjectService.Collection.SetLogoAsync(ProjectId, Path, Revision, GetMimeTypeFromPath(Path), await ReadDataAsync(Path));
+						await ProjectService.Collection.SetLogoAsync(ProjectId, Path.ToString(), Revision, GetMimeTypeFromPath(Path), await ReadDataAsync(Path));
 						CachedLogoRevisions[ProjectId] = Revision;
 					}
 				}
@@ -226,23 +230,23 @@ namespace HordeServer.Services
 			List<IStream> Streams = await StreamService.GetStreamsAsync();
 
 			// Get the revisions for all the stream documents
-			Dictionary<string, string> StreamRevisions = await FindRevisionsAsync(StreamConfigs.Select(x => x.Path));
+			Dictionary<Uri, string> StreamRevisions = await FindRevisionsAsync(StreamConfigs.Select(x => x.Path));
 			for (int Idx = 0; Idx < StreamConfigs.Count; Idx++)
 			{
-				(ProjectId ProjectId, StreamConfigRef StreamRef, string Path) = StreamConfigs[Idx];
-				if (StreamRevisions.TryGetValue(Path, out string? Revision))
+				(ProjectId ProjectId, StreamConfigRef StreamRef, Uri StreamPath) = StreamConfigs[Idx];
+				if (StreamRevisions.TryGetValue(StreamPath, out string? Revision))
 				{
 					IStream? Stream = Streams.FirstOrDefault(x => x.Id == StreamRef.Id);
-					if (Stream == null || Stream.ConfigPath != Path || Stream.ConfigRevision != Revision)
+					if (Stream == null || Stream.ConfigPath != StreamPath.ToString() || Stream.ConfigRevision != Revision)
 					{
 						Logger.LogInformation("Updating configuration for stream {StreamRef} ({Revision})", StreamRef.Id, Revision);
 
-						StreamConfig? StreamConfig = await ReadDataAsync<StreamConfig>(Path);
+						StreamConfig? StreamConfig = await ReadDataAsync<StreamConfig>(StreamPath);
 						if (StreamConfig != null)
 						{
 							for (; ; )
 							{
-								Stream = await StreamService.StreamCollection.TryCreateOrReplaceAsync(StreamRef.Id, Stream, Path, Revision, ProjectId, StreamConfig);
+								Stream = await StreamService.StreamCollection.TryCreateOrReplaceAsync(StreamRef.Id, Stream, StreamPath.ToString(), Revision, ProjectId, StreamConfig);
 
 								if (Stream != null)
 								{
@@ -279,59 +283,64 @@ namespace HordeServer.Services
 
 		static FileExtensionContentTypeProvider ContentTypeProvider = new FileExtensionContentTypeProvider();
 
-		static string GetMimeTypeFromPath(string Path)
+		static string GetMimeTypeFromPath(Uri Path)
 		{
 			string ContentType;
-			if (!ContentTypeProvider.TryGetContentType(Path, out ContentType))
+			if (!ContentTypeProvider.TryGetContentType(Path.AbsolutePath, out ContentType))
 			{
 				ContentType = "application/octet-stream";
 			}
 			return ContentType;
 		}
 
-		static string CombinePaths(string BasePath, string RelativePath)
+		static Uri CombinePaths(Uri BaseUri, string Path)
 		{
-			// Handle Perforce paths
-			if (RelativePath.StartsWith("//", StringComparison.Ordinal))
-			{
-				return RelativePath;
-			}
-			if (BasePath.StartsWith("//", StringComparison.Ordinal))
+			if (Path.StartsWith("//", StringComparison.Ordinal))
 			{
-				return BasePath.Substring(0, BasePath.LastIndexOf('/') + 1) + RelativePath;
+				if (BaseUri.Scheme == PerforceScheme)
+				{
+					return new Uri($"{PerforceScheme}://{BaseUri.Host}{Path}");
+				}
+				else
+				{
+					return new Uri($"{PerforceScheme}://{PerforceCluster.DefaultName}{Path}");
+				}
 			}
-
-			// Handle filesystem paths
-			return Path.GetFullPath(Path.Combine(Path.GetDirectoryName(BasePath)!, RelativePath));
+			return new Uri(BaseUri, Path);
 		}
 
-		async Task<Dictionary<string, string>> FindRevisionsAsync(IEnumerable<string> Paths)
+		async Task<Dictionary<Uri, string>> FindRevisionsAsync(IEnumerable<Uri> Paths)
 		{
-			Dictionary<string, string> Revisions = new Dictionary<string, string>();
+			Dictionary<Uri, string> Revisions = new Dictionary<Uri, string>();
 
 			// Find all the Perforce uris
-			List<string> PerforcePaths = new List<string>();
-			foreach (string Path in Paths)
+			List<Uri> PerforcePaths = new List<Uri>();
+			foreach (Uri Path in Paths)
 			{
-				if (Path.StartsWith("//", StringComparison.Ordinal))
+				if (Path.Scheme == FileScheme)
+				{
+					Revisions[Path] = $"ver={Version},md5={ContentHash.MD5(new FileReference(Path.LocalPath))}";
+				}
+				else if (Path.Scheme == PerforceScheme)
 				{
 					PerforcePaths.Add(Path);
 				}
 				else
 				{
-					Revisions[Path] = $"ver={Version},md5={ContentHash.MD5(new FileReference(Path))}";
+					throw new Exception($"Invalid path format: {Path}");
 				}
 			}
 
 			// Query all the Perforce revisions
-			if (PerforcePaths.Count > 0)
+			foreach (IGrouping<string, Uri> PerforcePath in PerforcePaths.GroupBy(x => x.Host, StringComparer.OrdinalIgnoreCase))
 			{
-				List<FileSummary> Files = await PerforceService.FindFilesAsync(PerforcePaths);
+				List<FileSummary> Files = await PerforceService.FindFilesAsync(PerforcePath.Key, PerforcePath.Select(x => x.AbsolutePath));
 				foreach (FileSummary File in Files)
 				{
+					Uri FileUri = new Uri($"{PerforceScheme}://{PerforcePath.Key}{File.DepotPath}");
 					if (File.Error == null)
 					{
-						Revisions[File.DepotPath] = $"ver={Version},p4={File.DepotPath}@{File.Change}";
+						Revisions[FileUri] = $"ver={Version},chg={File.Change},path={FileUri}";
 					}
 					else
 					{
@@ -343,7 +352,7 @@ namespace HordeServer.Services
 			return Revisions;
 		}
 
-		async Task<T?> ReadDataAsync<T>(string ConfigPath) where T : class
+		async Task<T?> ReadDataAsync<T>(Uri ConfigPath) where T : class
 		{
 			try
 			{
@@ -361,15 +370,16 @@ namespace HordeServer.Services
 			}
 		}
 
-		Task<byte[]> ReadDataAsync(string ConfigPath)
+		Task<byte[]> ReadDataAsync(Uri ConfigPath)
 		{
-			if (ConfigPath.StartsWith("//", StringComparison.Ordinal))
+			switch(ConfigPath.Scheme)
 			{
-				return PerforceService.PrintAsync(ConfigPath);
-			}
-			else
-			{
-				return File.ReadAllBytesAsync(ConfigPath);
+				case FileScheme:
+					return File.ReadAllBytesAsync(ConfigPath.AbsolutePath);
+				case PerforceScheme:
+					return PerforceService.PrintAsync(ConfigPath.Host, ConfigPath.AbsolutePath);
+				default:
+					throw new Exception($"Invalid config path: {ConfigPath}");
 			}
 		}
 
@@ -378,7 +388,8 @@ namespace HordeServer.Services
 		{
 			if (Settings.ConfigPath != null)
 			{
-				await UpdateConfigAsync(Settings.ConfigPath);
+				Uri ConfigUri = CombinePaths(new Uri(FileReference.Combine(Program.DataDir, "_").FullName), Settings.ConfigPath);
+				await UpdateConfigAsync(ConfigUri);
 			}
 			return DateTime.UtcNow + TimeSpan.FromMinutes(1.0);
 		}
diff --git a/Engine/Source/Programs/Horde/HordeServer/Services/IPerforceService.cs b/Engine/Source/Programs/Horde/HordeServer/Services/IPerforceService.cs
--- a/Engine/Source/Programs/Horde/HordeServer/Services/IPerforceService.cs
+++ b/Engine/Source/Programs/Horde/HordeServer/Services/IPerforceService.cs
@@ -45,96 +45,108 @@ namespace HordeServer.Services
 		/// <summary>
 		/// Create a new changelist by submitting the given file
 		/// </summary>
+		/// <param name="ClusterName">Name of the Perforce cluster</param>
 		/// <param name="StreamName">The stream to query</param>
 		/// <param name="Path">Path for the file to submit</param>
 		/// <returns>New changelist number</returns>
-		public Task<int> CreateNewChangeAsync(string StreamName, string Path);
+		public Task<int> CreateNewChangeAsync(string ClusterName, string StreamName, string Path);
 
 		/// <summary>
 		/// Gets the code change corresponding to an actual change submitted to a stream
 		/// </summary>
+		/// <param name="ClusterName">Name of the Perforce cluster</param>
 		/// <param name="StreamName"></param>
 		/// <param name="Change">The changelist number to query</param>
 		/// <returns>Code change for the latest change</returns>
-		public Task<int> GetCodeChangeAsync(string StreamName, int Change);
+		public Task<int> GetCodeChangeAsync(string ClusterName, string StreamName, int Change);
 
 		/// <summary>
 		/// Gets information about the given user
 		/// </summary>
+		/// <param name="ClusterName">Name of the Perforce cluster</param>
 		/// <param name="UserName">The user name</param>
 		/// <returns>User information</returns>
-		public Task<PerforceUserInfo?> GetUserInfoAsync(string UserName);
+		public Task<PerforceUserInfo?> GetUserInfoAsync(string ClusterName, string UserName);
 
 		/// <summary>
 		/// Gets the latest change for a particular stream
 		/// </summary>
+		/// <param name="ClusterName">Name of the Perforce cluster</param>
 		/// <param name="StreamName">The stream to query</param>
 		/// <param name="MinChange">The minimum changelist number</param>
 		/// <param name="MaxChange">The maximum changelist number</param>
 		/// <param name="Results">Number of results to return</param>
 		/// <param name="ImpersonateUser">Name of the user to impersonate</param>
 		/// <returns>Latest changelist number</returns>
-		public Task<List<ChangeSummary>> GetChangesAsync(string StreamName, int? MinChange, int? MaxChange, int Results, string? ImpersonateUser);
+		public Task<List<ChangeSummary>> GetChangesAsync(string ClusterName, string StreamName, int? MinChange, int? MaxChange, int Results, string? ImpersonateUser);
 
 		/// <summary>
 		/// Gets the latest change for a particular stream
 		/// </summary>
+		/// <param name="ClusterName">Name of the Perforce cluster</param>
 		/// <param name="StreamName">The stream to query</param>
 		/// <param name="ChangeNumbers">Change numbers to query</param>
 		/// <param name="ImpersonateUser">Name of the user to impersonate</param>
 		/// <returns>Commit details</returns>
-		public Task<List<ChangeDetails>> GetChangeDetailsAsync(string StreamName, IReadOnlyList<int> ChangeNumbers, string? ImpersonateUser);
-		
+		public Task<List<ChangeDetails>> GetChangeDetailsAsync(string ClusterName, string StreamName, IReadOnlyList<int> ChangeNumbers, string? ImpersonateUser);
+
 		/// <summary>
 		/// Create a ticket as the specified user
 		/// </summary>
+		/// <param name="ClusterName">Name of the Perforce cluster</param>
 		/// <param name="ImpersonateUser">Name of the user to impersonate</param>
 		/// <returns>Commit details</returns>
-		public Task<string> CreateTicket(string ImpersonateUser);
+		public Task<string> CreateTicket(string ClusterName, string ImpersonateUser);
 
 		/// <summary>
 		/// Gets the latest changes for a set of depot paths
 		/// </summary>
+		/// <param name="ClusterName">Name of the Perforce cluster</param>
 		/// <param name="Paths"></param>
 		/// <returns></returns>
-		public Task<List<FileSummary>> FindFilesAsync(IEnumerable<string> Paths);
+		public Task<List<FileSummary>> FindFilesAsync(string ClusterName, IEnumerable<string> Paths);
 
 		/// <summary>
 		/// Gets the contents of a file in the depot
 		/// </summary>
+		/// <param name="ClusterName">Name of the Perforce cluster</param>
 		/// <param name="Path">Path to read</param>
 		/// <returns>Data for the file</returns>
-		public Task<byte[]> PrintAsync(string Path);
+		public Task<byte[]> PrintAsync(string ClusterName, string Path);
 
 		/// <summary>
 		/// Duplicates a shelved changelist
 		/// </summary>
+		/// <param name="ClusterName">Name of the Perforce cluster</param>
 		/// <param name="ShelvedChange">The shelved changelist</param>
 		/// <returns>The duplicated changelist</returns>
-		public Task<int> DuplicateShelvedChangeAsync(int ShelvedChange);
+		public Task<int> DuplicateShelvedChangeAsync(string ClusterName, int ShelvedChange);
 
 		/// <summary>
 		/// Submit a shelved changelist
 		/// </summary>
-		/// <param name="ShelvedChange">The shelved changelist number, created by <see cref="DuplicateShelvedChangeAsync(int)"/></param>
+		/// <param name="ClusterName">Name of the Perforce cluster</param>
+		/// <param name="ShelvedChange">The shelved changelist number, created by <see cref="DuplicateShelvedChangeAsync(string,int)"/></param>
 		/// <param name="OriginalChange">The original changelist number</param>
 		/// <returns>Tuple consisting of the submitted changelist number and message</returns>
-		public Task<(int? Change, string Message)> SubmitShelvedChangeAsync(int ShelvedChange, int OriginalChange);
+		public Task<(int? Change, string Message)> SubmitShelvedChangeAsync(string ClusterName, int ShelvedChange, int OriginalChange);
 
 		/// <summary>
 		/// Deletes a changelist containing shelved files.
 		/// </summary>
+		/// <param name="ClusterName">Name of the Perforce cluster</param>
 		/// <param name="ShelvedChange">The changelist containing shelved files</param>
 		/// <returns>Async tasy</returns>
-		public Task DeleteShelvedChangeAsync(int ShelvedChange);
+		public Task DeleteShelvedChangeAsync(string ClusterName, int ShelvedChange);
 
 		/// <summary>
 		/// Updates a changelist description
 		/// </summary>
+		/// <param name="ClusterName">Name of the Perforce cluster</param>
 		/// <param name="Change">The change to update</param>
 		/// <param name="Description">The new description</param>
 		/// <returns>Async task</returns>
-		public Task UpdateChangelistDescription(int Change, string Description);
+		public Task UpdateChangelistDescription(string ClusterName, int Change, string Description);
 	}
 
 	/// <summary>
@@ -154,11 +166,11 @@ namespace HordeServer.Services
 			Match Match = Regex.Match(Template.SubmitNewChange!, @"^(//[^/]+/[^/]+)/(.+)$");
 			if (Match.Success)
 			{
-				return Perforce.CreateNewChangeAsync(Match.Groups[1].Value, Match.Groups[2].Value);
+				return Perforce.CreateNewChangeAsync(Stream.ClusterName, Match.Groups[1].Value, Match.Groups[2].Value);
 			}
 			else
 			{
-				return Perforce.CreateNewChangeAsync(Stream.Name, Template.SubmitNewChange!);
+				return Perforce.CreateNewChangeAsync(Stream.ClusterName, Stream.Name, Template.SubmitNewChange!);
 			}
 		}
 
@@ -166,13 +178,14 @@ namespace HordeServer.Services
 		/// Gets the latest change for a particular stream
 		/// </summary>
 		/// <param name="Perforce">The perforce implementation</param>
+		/// <param name="ClusterName">Name of the Perforce cluster</param>
 		/// <param name="StreamName">The stream to query</param>
 		/// <param name="ChangeNumber">Change number to query</param>
 		/// <param name="ImpersonateUser">Name of the user to impersonate</param>
 		/// <returns>Commit details</returns>
-		public static async Task<ChangeDetails> GetChangeDetailsAsync(this IPerforceService Perforce, string StreamName, int ChangeNumber, string? ImpersonateUser)
+		public static async Task<ChangeDetails> GetChangeDetailsAsync(this IPerforceService Perforce, string ClusterName, string StreamName, int ChangeNumber, string? ImpersonateUser)
 		{
-			List<ChangeDetails> Results = await Perforce.GetChangeDetailsAsync(StreamName, new[] { ChangeNumber }, ImpersonateUser);
+			List<ChangeDetails> Results = await Perforce.GetChangeDetailsAsync(ClusterName, StreamName, new[] { ChangeNumber }, ImpersonateUser);
 			return Results[0];
 		}
 
@@ -180,28 +193,30 @@ namespace HordeServer.Services
 		/// Gets the latest change for a particular stream
 		/// </summary>
 		/// <param name="Perforce">The perforce implementation</param>
+		/// <param name="ClusterName">Name of the Perforce cluster</param>
 		/// <param name="StreamName">The stream to query</param>
 		/// <param name="MinChange">The minimum changelist number</param>
 		/// <param name="MaxChange">The maximum changelist number</param>
 		/// <param name="Results">Number of results to return</param>
 		/// <param name="ImpersonateUser">Name of the user to impersonate</param>
 		/// <returns>Commit details</returns>
-		public static async Task<List<ChangeDetails>> GetChangeDetailsAsync(this IPerforceService Perforce, string StreamName, int? MinChange, int? MaxChange, int Results, string? ImpersonateUser)
+		public static async Task<List<ChangeDetails>> GetChangeDetailsAsync(this IPerforceService Perforce, string ClusterName, string StreamName, int? MinChange, int? MaxChange, int Results, string? ImpersonateUser)
 		{
-			List<ChangeSummary> Changes = await Perforce.GetChangesAsync(StreamName, MinChange, MaxChange, Results, ImpersonateUser);
-			return await Perforce.GetChangeDetailsAsync(StreamName, Changes.ConvertAll(x => x.Number), ImpersonateUser);
+			List<ChangeSummary> Changes = await Perforce.GetChangesAsync(ClusterName, StreamName, MinChange, MaxChange, Results, ImpersonateUser);
+			return await Perforce.GetChangeDetailsAsync(ClusterName, StreamName, Changes.ConvertAll(x => x.Number), ImpersonateUser);
 		}
 
 		/// <summary>
 		/// Get the latest submitted change to the stream
 		/// </summary>
 		/// <param name="Perforce">The perforce implementation</param>
+		/// <param name="ClusterName">Name of the Perforce cluster</param>
 		/// <param name="StreamName">The stream to query</param>
 		/// <param name="ImpersonateUser">Name of the user to impersonate</param>
 		/// <returns>Latest changelist number</returns>
-		public static async Task<int> GetLatestChangeAsync(this IPerforceService Perforce, string StreamName, string? ImpersonateUser)
+		public static async Task<int> GetLatestChangeAsync(this IPerforceService Perforce, string ClusterName, string StreamName, string? ImpersonateUser)
 		{
-			List<ChangeSummary> Changes = await Perforce.GetChangesAsync(StreamName, null, null, 1, ImpersonateUser);
+			List<ChangeSummary> Changes = await Perforce.GetChangesAsync(ClusterName, StreamName, null, null, 1, ImpersonateUser);
 			if (Changes.Count == 0)
 			{
 				throw new Exception($"No changes have been submitted to stream {StreamName}");
diff --git a/Engine/Source/Programs/Horde/HordeServer/Services/Impl/IssueService.cs b/Engine/Source/Programs/Horde/HordeServer/Services/Impl/IssueService.cs
--- a/Engine/Source/Programs/Horde/HordeServer/Services/Impl/IssueService.cs
+++ b/Engine/Source/Programs/Horde/HordeServer/Services/Impl/IssueService.cs
@@ -151,7 +151,7 @@ namespace HordeServer.Services.Impl
 		/// <summary>
 		/// Collection of stream documents
 		/// </summary>
-		IStreamCollection Streams;
+		StreamService Streams;
 
 		/// <summary>
 		/// Collection of user documents
@@ -215,7 +215,7 @@ namespace HordeServer.Services.Impl
 		/// <param name="Perforce">The perforce service instance</param>
 		/// <param name="LogFileService">Collection of event documents</param>
 		/// <param name="Logger">The logger instance</param>
-		public IssueService(DatabaseService DatabaseService, IIssueCollection IssueCollection, IJobCollection Jobs, IJobStepRefCollection JobStepRefs, IStreamCollection Streams, IUserCollection UserCollection, IPerforceService Perforce, ILogFileService LogFileService, ILogger<IssueService> Logger)
+		public IssueService(DatabaseService DatabaseService, IIssueCollection IssueCollection, IJobCollection Jobs, IJobStepRefCollection JobStepRefs, StreamService Streams, IUserCollection UserCollection, IPerforceService Perforce, ILogFileService LogFileService, ILogger<IssueService> Logger)
 			: base(DatabaseService, ObjectId.Parse("609542152fb0794700a6c3df"), Logger)
 		{
 			Type[] IssueTypes = Assembly.GetExecutingAssembly().GetTypes().Where(x => !x.IsAbstract && typeof(IIssue).IsAssignableFrom(x)).ToArray();
@@ -261,7 +261,7 @@ namespace HordeServer.Services.Impl
 
 			Dictionary<StreamId, HashSet<TemplateRefId>> NewCachedDesktopAlerts = new Dictionary<StreamId, HashSet<TemplateRefId>>();
 
-			List<IStream> CachedStreams = await Streams.FindAllAsync();
+			List<IStream> CachedStreams = await Streams.StreamCollection.FindAllAsync();
 			foreach(IStream CachedStream in CachedStreams)
 			{
 				HashSet<TemplateRefId> Templates = new HashSet<TemplateRefId>(CachedStream.Templates.Where(x => x.Value.ShowUgsAlerts).Select(x => x.Key));
@@ -540,7 +540,7 @@ namespace HordeServer.Services.Impl
 			Scope.Span.SetTag("BatchId", BatchId.ToString());
 			Scope.Span.SetTag("StepId", StepId.ToString());
 
-			IStream? Stream = await Streams.GetAsync(Job.StreamId);
+			IStream? Stream = await Streams.GetStreamAsync(Job.StreamId);
 			if (Stream == null)
 			{
 				throw new Exception($"Invalid stream id '{Job.StreamId}' on job '{Job.Id}'");
@@ -1027,12 +1027,16 @@ namespace HordeServer.Services.Impl
 			{
 				if (!NextFixStreamIds.ContainsKey(Span.StreamId))
 				{
-					bool ContainsFix;
+					bool ContainsFix = false;
 					if (PrevFixStreamIds == null || !PrevFixStreamIds.TryGetValue(Span.StreamId, out ContainsFix))
 					{
-						Logger.LogInformation("Querying fix changelist {FixChange} for issue {IssueId} in {StreamId} (prev={HavePrev})", FixChange, IssueId, Span.StreamId, PrevFixStreamIds != null);
-						List<ChangeSummary> Changes = await Perforce.GetChangesAsync(Span.StreamName, FixChange, FixChange, 1, null);
-						ContainsFix = Changes.Count > 0;
+						IStream? Stream = await Streams.GetCachedStream(Span.StreamId);
+						if (Stream != null)
+						{
+							Logger.LogInformation("Querying fix changelist {FixChange} for issue {IssueId} in {StreamId} (prev={HavePrev})", FixChange, IssueId, Span.StreamId, PrevFixStreamIds != null);
+							List<ChangeSummary> Changes = await Perforce.GetChangesAsync(Stream.ClusterName, Span.StreamName, FixChange, FixChange, 1, null);
+							ContainsFix = Changes.Count > 0;
+						}
 					}
 					NextFixStreamIds[Span.StreamId] = ContainsFix;
 				}
@@ -1065,7 +1069,7 @@ namespace HordeServer.Services.Impl
 				Logger.LogDebug("Querying for changes in {StreamName} between {MinChange} and {MaxChange}", Stream.Name, MinChange, MaxChange);
 
 				// Get the submitted changes before this job
-				List<ChangeDetails> Changes = await PerforceServiceExtensions.GetChangeDetailsAsync(Perforce, Stream.Name, MinChange, MaxChange, MaxChanges, null);
+				List<ChangeDetails> Changes = await PerforceServiceExtensions.GetChangeDetailsAsync(Perforce, Stream.ClusterName, Stream.Name, MinChange, MaxChange, MaxChanges, null);
 				Logger.LogDebug("Found {NumResults} changes", Changes.Count);
 
 				// Get the handler to rank them
diff --git a/Engine/Source/Programs/Horde/HordeServer/Services/Impl/PerforceService.cs b/Engine/Source/Programs/Horde/HordeServer/Services/Impl/PerforceService.cs
--- a/Engine/Source/Programs/Horde/HordeServer/Services/Impl/PerforceService.cs
+++ b/Engine/Source/Programs/Horde/HordeServer/Services/Impl/PerforceService.cs
@@ -16,6 +16,7 @@ using System.Runtime.InteropServices;
 using System.Globalization;
 using System.Text.RegularExpressions;
 using System.Diagnostics.CodeAnalysis;
+using HordeServer.Utilities;
 
 namespace HordeServer.Services
 {
@@ -24,12 +25,23 @@ namespace HordeServer.Services
 	/// </summary>
 	public class PerforceService : IPerforceService
 	{
+		class CachedTicketInfo
+		{
+			public IPerforceServer Server;
+			public string UserName;
+			public P4.Credential Ticket;
 
-		IOptionsMonitor<ServerSettings> Settings;
-
-		ILogger<PerforceService> Logger;
+			public CachedTicketInfo(IPerforceServer Server, string UserName, P4.Credential Ticket)
+			{
+				this.Server = Server;
+				this.UserName = UserName;
+				this.Ticket = Ticket;
+			}
+		}
 
-		P4.Server Server;
+		PerforceLoadBalancer LoadBalancer;
+		LazyCachedValue<Task<Globals>> CachedGlobals;
+		ILogger Logger;
 
 		/// <summary>
 		/// Object used for controlling access to the access user tickets
@@ -46,143 +58,160 @@ namespace HordeServer.Services
 		/// </summary>
 		P4.P4CallBacks.LogMessageDelegate LogBridgeDelegate;
 
-
-		Dictionary<string, P4.Credential> UserTickets = new Dictionary<string, P4.Credential>();
+		Dictionary<string, Dictionary<string, CachedTicketInfo>> ClusterTickets = new Dictionary<string, Dictionary<string, CachedTicketInfo>>(StringComparer.OrdinalIgnoreCase);
 
 		/// <summary>
 		/// Constructor
 		/// </summary>
-		public PerforceService(IOptionsMonitor<ServerSettings> Settings, ILogger<PerforceService> Logger)
+		public PerforceService(PerforceLoadBalancer LoadBalancer, DatabaseService DatabaseService, ILogger<PerforceService> Logger)
 		{
-			this.Settings = Settings;
+			this.LoadBalancer = LoadBalancer;
+			this.CachedGlobals = new LazyCachedValue<Task<Globals>>(() => DatabaseService.GetGlobalsAsync(), TimeSpan.FromSeconds(30.0));
 			this.Logger = Logger;
 
-			Server = new P4.Server(new P4.ServerAddress(Settings.CurrentValue.P4BridgeServer));
-
 			LogBridgeDelegate = new P4.P4CallBacks.LogMessageDelegate(LogBridgeMessage);
 			P4.P4Debugging.SetBridgeLogFunction(LogBridgeDelegate);
 
 			P4.LogFile.SetLoggingFunction(LogPerforce);
 		}
 
-
-		P4.Repository GetServiceUserConnection()
+		async Task<PerforceCluster> GetClusterAsync(string ClusterName)
 		{
+			Globals Globals = await CachedGlobals.GetCached();
 
-			ServerSettings Settings = this.Settings.CurrentValue;
+			PerforceCluster? Cluster = Globals.FindPerforceCluster(ClusterName);
+			if (Cluster == null)
+			{
+				throw new Exception($"Unknown Perforce cluster '{ClusterName}'");
+			}
 
-			P4.Repository Repository = new P4.Repository(Server);
+			return Cluster;
+		}
 
+		async Task<IPerforceServer> SelectServer(PerforceCluster Cluster)
+		{
+			IPerforceServer? Server = await LoadBalancer.SelectServerAsync(Cluster);
+			if (Server == null)
+			{
+				throw new Exception($"Unable to select server from '{Cluster.Name}'");
+			}
+			return Server;
+		}
+
+		static P4.Repository CreateConnection(IPerforceServer Server, string UserName, string Ticket)
+		{
+			P4.Repository Repository = new P4.Repository(new P4.Server(new P4.ServerAddress(Server.ServerAndPort)));
 			try
 			{
 				P4.Connection Connection = Repository.Connection;
-				Connection.UserName = Settings.P4BridgeServiceUsername;
+				Connection.UserName = UserName;
 
 				P4.Options Options = new P4.Options();
-
-				Options["Ticket"] = Settings.P4BridgeServicePassword;
+				Options["Ticket"] = Ticket;
 
 				// connect to the server
 				if (!Connection.Connect(Options))
 				{
 					throw new Exception("Unable to get P4 server connection");
 				}
-
 			}
 			catch
 			{
 				Repository.Dispose();
 				throw;
 			}
-
-
 			return Repository;
 		}
 
-		P4.Credential GetImpersonateCredential(string ImpersonateUser)
+		async Task<P4.Repository> GetServiceUserConnection(PerforceCluster Cluster)
 		{
+			IPerforceServer Server = await SelectServer(Cluster);
+			return GetServiceUserConnection(Cluster, Server);
+		}
 
-			if (!CanImpersonate)
+		static P4.Repository GetServiceUserConnection(PerforceCluster Cluster, IPerforceServer Server)
+		{
+			PerforceCredentials? Credentials = Cluster.Credentials.FirstOrDefault(x => x.UserName.Equals(Cluster.ServiceAccount, StringComparison.OrdinalIgnoreCase));
+			if (Credentials == null)
 			{
-				throw new Exception($"Service account required to impersonate user {ImpersonateUser}");
+				throw new Exception($"No credentials defined for {Cluster.ServiceAccount} on {Cluster.Name}");
 			}
+			return CreateConnection(Server, Credentials.UserName, Credentials.Password);
+		}
 
-			lock (TicketLock)
+		async Task<CachedTicketInfo> GetImpersonateCredential(PerforceCluster Cluster, string ImpersonateUser)
+		{
+			if (!Cluster.CanImpersonate)
 			{
+				throw new Exception($"Service account required to impersonate user {ImpersonateUser}");
+			}
 
-				ServerSettings Settings = this.Settings.CurrentValue;
+			CachedTicketInfo? TicketInfo = null;
 
+			Dictionary<string, CachedTicketInfo>? UserTickets;
+			lock (TicketLock)
+			{
 				// Check if we have a ticket
-				P4.Credential? Credential = null;
-
-				if (UserTickets.TryGetValue(ImpersonateUser, out Credential))
+				if (!ClusterTickets.TryGetValue(Cluster.Name, out UserTickets))
+				{
+					UserTickets = new Dictionary<string, CachedTicketInfo>(StringComparer.OrdinalIgnoreCase);
+					ClusterTickets[Cluster.Name] = UserTickets;
+				}
+				if (UserTickets.TryGetValue(ImpersonateUser, out TicketInfo))
 				{
 					// if the credential expires within the next 15 minutes, refresh
 					TimeSpan Time = new TimeSpan(0, 15, 0);
-					if (Credential.Expires.Subtract(Time) <= DateTime.UtcNow)
+					if (TicketInfo.Ticket.Expires.Subtract(Time) <= DateTime.UtcNow)
 					{
 						UserTickets.Remove(ImpersonateUser);
-						Credential = null;
+						TicketInfo = null;
 					}
 				}
+			}
 
-				if (Credential != null)
-				{
-					return Credential;
-				}
+			if (TicketInfo == null)
+			{
+				IPerforceServer Server = await SelectServer(Cluster);
 
-				using (P4.Repository Repository = GetServiceUserConnection())
+				P4.Credential? Credential;
+				using (P4.Repository Repository = GetServiceUserConnection(Cluster, Server))
 				{
 					Credential = Repository.Connection.Login(null, new P4.LoginCmdOptions(P4.LoginCmdFlags.AllHosts | P4.LoginCmdFlags.DisplayTicket, null), ImpersonateUser);
 				}
-
 				if (Credential == null)
 				{
 					throw new Exception($"GetImpersonateCredential - Unable to get impersonation credential for {ImpersonateUser} from {Dns.GetHostName()}");
 				}
 
-				UserTickets.Add(ImpersonateUser, Credential);
+				TicketInfo = new CachedTicketInfo(Server, ImpersonateUser, Credential);
 
-				return Credential;
+				lock (TicketLock)
+				{
+					UserTickets[ImpersonateUser] = TicketInfo;
+				}
 			}
 
+			return TicketInfo;
 		}
 
-		P4.Repository GetImpersonateConnection(string ImpersonateUser)
+		async Task<P4.Repository> GetImpersonatedConnection(PerforceCluster Cluster, string ImpersonateUser)
 		{
+			CachedTicketInfo TicketInfo = await GetImpersonateCredential(Cluster, ImpersonateUser);
+			return CreateConnection(TicketInfo.Server, TicketInfo.UserName, TicketInfo.Ticket.Ticket);
+		}
 
-			if (!CanImpersonate)
+		async Task<P4.Repository> GetConnection(PerforceCluster Cluster, string? Stream = null, string? Username = null, bool ReadOnly = true, bool CreateChange = false, int? ClientFromChange = null, bool UseClientFromChange = false, bool UsePortFromChange = false, string? ClientId = null, bool NoClient = false)
+		{
+			P4.Repository Repository;
+			if (Username == null || !Cluster.CanImpersonate || Username.Equals(Cluster.ServiceAccount, StringComparison.OrdinalIgnoreCase))
 			{
-				throw new Exception($"Service account required to impersonate user {ImpersonateUser}");
+				Repository = await GetServiceUserConnection(Cluster);
 			}
-
-			P4.Credential Credential = GetImpersonateCredential(ImpersonateUser);
-
-			// this might be able to be reused and would be threaded internally
-			P4.Repository Repository = new P4.Repository(Server);
-
-			P4.Connection Connection = Repository.Connection;			
-			Connection.UserName = ImpersonateUser;
-
-			P4.Options Options = new P4.Options();
-			Options["Ticket"] = Credential.Ticket;
-
-			// connect to the server
-			if (!Connection.Connect(Options))
+			else
 			{
-				throw new Exception($"GetImpersonateConnection - Unable to get impersonated P4 connection for {ImpersonateUser} from {Dns.GetHostName()}");
-
+				Repository = await GetImpersonatedConnection(Cluster, Username);
 			}
 
-			return Repository;
-		}
-
-		P4.Repository GetConnection(string? Stream = null, string? Username = null, bool ReadOnly = true, bool CreateChange = false, int? ClientFromChange = null, bool UseClientFromChange = false, bool UsePortFromChange = false, string? ClientId = null, bool NoClient = false)
-		{
-			ServerSettings Settings = this.Settings.CurrentValue;
-
-			P4.Repository? Repository = (Username == null || Settings.P4BridgeServiceUsername == Username || !CanImpersonate) ? GetServiceUserConnection() : GetImpersonateConnection(Username);
-
 			if (!NoClient)
 			{
 				try
@@ -193,7 +222,7 @@ namespace HordeServer.Services
 					}
 					else
 					{
-						P4.Client Client = GetOrCreateClient(Repository, Stream, Username, ReadOnly, CreateChange, ClientFromChange, UseClientFromChange, UsePortFromChange);
+						P4.Client Client = GetOrCreateClient(Cluster.ServiceAccount, Repository, Stream, Username, ReadOnly, CreateChange, ClientFromChange, UseClientFromChange, UsePortFromChange);
 						Repository.Connection.SetClient(Client.Name);
 
 					}
@@ -259,11 +288,9 @@ namespace HordeServer.Services
 
 		}
 
-		string GetClientName(string Stream, bool ReadOnly, bool CreateChange, string? Username = null)
+		static string GetClientName(string ServiceUserName, string Stream, bool ReadOnly, bool CreateChange, string? Username = null)
 		{
-			ServerSettings Settings = this.Settings.CurrentValue;
-
-			string ClientName = $"horde-p4bridge-{Dns.GetHostName()}-{Settings.P4BridgeServiceUsername}-{Stream.Replace("/", "+", StringComparison.OrdinalIgnoreCase)}";
+			string ClientName = $"horde-p4bridge-{Dns.GetHostName()}-{ServiceUserName}-{Stream.Replace("/", "+", StringComparison.OrdinalIgnoreCase)}";
 
 			if (!ReadOnly)
 			{
@@ -284,10 +311,8 @@ namespace HordeServer.Services
 
 		}
 
-		P4.Client GetOrCreateClient(P4.Repository Repository, string? Stream, string? Username = null, bool ReadOnly = true, bool CreateChange = false, int? ClientFromChange = null, bool UseClientFromChange = false, bool UsePortFromChange = false)
+		static P4.Client GetOrCreateClient(string ServiceUserName, P4.Repository Repository, string? Stream, string? Username = null, bool ReadOnly = true, bool CreateChange = false, int? ClientFromChange = null, bool UseClientFromChange = false, bool UsePortFromChange = false)
 		{
-			ServerSettings Settings = this.Settings.CurrentValue;
-
 			P4.Client? Client = null;
 			P4.Changelist? Changelist = null;
 
@@ -324,7 +349,7 @@ namespace HordeServer.Services
 					throw new Exception("Stream required for client");
 
 				}
-				string ClientName = GetClientName(Stream, ReadOnly, CreateChange, Username);
+				string ClientName = GetClientName(ServiceUserName, Stream, ReadOnly, CreateChange, Username);
 
 				IList<P4.Client>? Clients = Repository.GetClients(new P4.ClientsCmdOptions(P4.ClientsCmdFlags.None, Username, ClientName, 1, Stream));
 
@@ -350,13 +375,13 @@ namespace HordeServer.Services
 						NewClient.ClientType = P4.ClientType.@readonly;
 					}
 
-					NewClient.OwnerName = Username ?? Settings.P4BridgeServiceUsername;
+					NewClient.OwnerName = Username ?? ServiceUserName;
 
 					Client = Repository.CreateClient(NewClient);
 
 					if (Client == null)
 					{
-						throw new Exception($"Unable to create client for ${Stream} : {Username ?? Settings.P4BridgeServiceUsername}");
+						throw new Exception($"Unable to create client for ${Stream} : {Username ?? ServiceUserName}");
 					}
 				}
 			}
@@ -380,7 +405,7 @@ namespace HordeServer.Services
 			{
 				Serilog.Log.Information("Perforce: " + Log);
 			}
-			
+
 		}
 
 		/// <summary>
@@ -416,7 +441,7 @@ namespace HordeServer.Services
 		{
 
 			// Note, we do not get log level 4 unless it is defined in native code as it is very, very spammy (P4BridgeServer.cpp)
-			
+
 			// remove the full path to the source, keep just the file name
 			String fileName = Path.GetFileName(Filename);
 
@@ -427,630 +452,585 @@ namespace HordeServer.Services
 
 
 		/// <inheritdoc/>
-		public async Task<List<ChangeSummary>> GetChangesAsync(string StreamName, int? MinChange, int? MaxChange, int Results, string? ImpersonateUser)
+		public async Task<List<ChangeSummary>> GetChangesAsync(string ClusterName, string StreamName, int? MinChange, int? MaxChange, int Results, string? ImpersonateUser)
 		{
-
-			return await Task.Run(() =>
+			PerforceCluster Cluster = await GetClusterAsync(ClusterName);
+			using (P4.Repository Repository = await GetConnection(Cluster, StreamName, ImpersonateUser))
 			{
 
-				using (P4.Repository Repository = GetConnection(StreamName, ImpersonateUser))
-				{
+				P4.ChangesCmdOptions Options = new P4.ChangesCmdOptions(P4.ChangesCmdFlags.IncludeTime | P4.ChangesCmdFlags.FullDescription, null, Results, P4.ChangeListStatus.Submitted, null);
 
-					P4.ChangesCmdOptions Options = new P4.ChangesCmdOptions(P4.ChangesCmdFlags.IncludeTime | P4.ChangesCmdFlags.FullDescription, null, Results, P4.ChangeListStatus.Submitted, null);
+				string Filter = GetFilter($"//{Repository.Connection.Client.Name}/...", MinChange, MaxChange);
 
-					string Filter = GetFilter($"//{Repository.Connection.Client.Name}/...", MinChange, MaxChange);
+				P4.FileSpec FileSpec = new P4.FileSpec(new P4.DepotPath(Filter), null, null, null);
 
-					P4.FileSpec FileSpec = new P4.FileSpec(new P4.DepotPath(Filter), null, null, null);
+				IList<P4.Changelist> Changelists = Repository.GetChangelists(Options, FileSpec);
 
-					IList<P4.Changelist> Changelists = Repository.GetChangelists(Options, FileSpec);
+				List<ChangeSummary> Changes = new List<ChangeSummary>();
 
-					List<ChangeSummary> Changes = new List<ChangeSummary>();
-
-					if (Changelists != null)
+				if (Changelists != null)
+				{
+					foreach (P4.Changelist Changelist in Changelists)
 					{
-						foreach (P4.Changelist Changelist in Changelists)
-						{
-							Changes.Add(new ChangeSummary(Changelist.Id, Changelist.OwnerName, Changelist.Description));
-						}
+						Changes.Add(new ChangeSummary(Changelist.Id, Changelist.OwnerName, Changelist.Description));
 					}
-
-					return Changes;
 				}
-			});
 
+				return Changes;
+			}
 		}
 
 		/// <inheritdoc/>
-		public async Task<List<ChangeDetails>> GetChangeDetailsAsync(string StreamName, IReadOnlyList<int> ChangeNumbers, string? ImpersonateUser)
+		public async Task<List<ChangeDetails>> GetChangeDetailsAsync(string ClusterName, string StreamName, IReadOnlyList<int> ChangeNumbers, string? ImpersonateUser)
 		{
-			return await Task.Run(() =>
+			PerforceCluster Cluster = await GetClusterAsync(ClusterName);
+			using (P4.Repository Repository = await GetConnection(Cluster, Stream: StreamName, Username: ImpersonateUser))
 			{
-				using (P4.Repository Repository = GetConnection(StreamName, ImpersonateUser))
+				List<ChangeDetails> Results = new List<ChangeDetails>();
+
+				List<string> Args = new List<string> { "-s", "-S" };
+				Args.AddRange(ChangeNumbers.Select(Change => Change.ToString(CultureInfo.InvariantCulture)));
+
+				using (P4.P4Command Command = new P4.P4Command(Repository, "describe", true, Args.ToArray()))
 				{
-					List<ChangeDetails> Results = new List<ChangeDetails>();
+					P4.P4CommandResult Result = Command.Run();
 
-					List<string> Args = new List<string> { "-s", "-S" };
-					Args.AddRange(ChangeNumbers.Select(Change => Change.ToString(CultureInfo.InvariantCulture)));
+					if (!Result.Success)
+					{
+						throw new Exception("Unable to get changes");
+					}
 
-					using (P4.P4Command Command = new P4.P4Command(Repository, "describe", true, Args.ToArray()))
+					if (Result.TaggedOutput == null || Result.TaggedOutput.Count <= 0)
 					{
-						P4.P4CommandResult Result = Command.Run();
+						return Results;
+					}
 
-						if (!Result.Success)
-						{
-							throw new Exception("Unable to get changes");
-						}
+					List<P4.Changelist> Changelists = new List<P4.Changelist>() { };
 
-						if (Result.TaggedOutput == null || Result.TaggedOutput.Count <= 0)
-						{
-							return Results;
-						}
+					bool DSTMismatch = false;
+					string Offset = string.Empty;
+
+					P4.Server Server = Repository.Server;
+					if (Server != null && Server.Metadata != null)
+					{
+						Offset = Server.Metadata.DateTimeOffset;
+						DSTMismatch = P4.FormBase.DSTMismatch(Server.Metadata);
+					}
 
-						List<P4.Changelist> Changelists = new List<P4.Changelist>() { };
+					foreach (P4.TaggedObject TaggedObject in Result.TaggedOutput)
+					{
+						List<string> Files = new List<string>();
 
-						bool DSTMismatch = false;
-						string Offset = string.Empty;
+						P4.Changelist Change = new P4.Changelist();
+						Change.FromChangeCmdTaggedOutput(TaggedObject, true, Offset, DSTMismatch);
 
-						if (Server != null && Server.Metadata != null)
+						foreach (P4.FileMetaData DescribeFile in Change.Files)
 						{
-							Offset = Server.Metadata.DateTimeOffset;
-							DSTMismatch = P4.FormBase.DSTMismatch(Server.Metadata);
+							string? RelativePath;
+							if (TryGetStreamRelativePath(DescribeFile.DepotPath.Path, StreamName, out RelativePath))
+							{
+								Files.Add(RelativePath);
+							}
 						}
 
-						foreach (P4.TaggedObject TaggedObject in Result.TaggedOutput)
+						if (Change.ShelvedFiles != null && Change.ShelvedFiles.Count > 0)
 						{
-							List<string> Files = new List<string>();
-
-							P4.Changelist Change = new P4.Changelist();
-							Change.FromChangeCmdTaggedOutput(TaggedObject, true, Offset, DSTMismatch);
-
-							foreach (P4.FileMetaData DescribeFile in Change.Files)
+							foreach (P4.ShelvedFile ShelvedFile in Change.ShelvedFiles)
 							{
 								string? RelativePath;
-								if (TryGetStreamRelativePath(DescribeFile.DepotPath.Path, StreamName, out RelativePath))
+								if (TryGetStreamRelativePath(ShelvedFile.Path.ToString(), StreamName, out RelativePath))
 								{
 									Files.Add(RelativePath);
 								}
 							}
 
-							if (Change.ShelvedFiles != null && Change.ShelvedFiles.Count > 0)
-							{
-								foreach (P4.ShelvedFile ShelvedFile in Change.ShelvedFiles)
-								{
-									string? RelativePath;
-									if (TryGetStreamRelativePath(ShelvedFile.Path.ToString(), StreamName, out RelativePath))
-									{
-										Files.Add(RelativePath);
-									}
-								}
-
-							}
-
-							Results.Add(new ChangeDetails(Change.Id, Change.OwnerName, Change.Description, Files));
 						}
 
+						Results.Add(new ChangeDetails(Change.Id, Change.OwnerName, Change.Description, Files));
 					}
 
-					return Results;
 				}
-			});
+
+				return Results;
+			}
 		}
 
 		/// <inheritdoc />
-		public async Task<string> CreateTicket(string ImpersonateUser)
+		public async Task<string> CreateTicket(string ClusterName, string ImpersonateUser)
 		{
-			return await Task.Run(() =>
-			{
-
-				P4.Credential Credential = GetImpersonateCredential(ImpersonateUser);
+			PerforceCluster Cluster = await GetClusterAsync(ClusterName);
 
-				if (Credential == null)
-				{
-					throw new Exception($"Unable to get ticket for user {ImpersonateUser}");
-				}
+			CachedTicketInfo Credential = await GetImpersonateCredential(Cluster, ImpersonateUser);
 
-				return Credential.Ticket;
-			});
+			if (Credential == null)
+			{
+				throw new Exception($"Unable to get ticket for user {ImpersonateUser}");
+			}
 
+			return Credential.Ticket.Ticket;
 		}
 
 		/// <inheritdoc/>
-		public async Task<List<FileSummary>> FindFilesAsync(IEnumerable<string> Paths)
+		public async Task<List<FileSummary>> FindFilesAsync(string ClusterName, IEnumerable<string> Paths)
 		{
-			return await Task.Run(() =>
+			List<FileSummary> Results = new List<FileSummary>();
+
+			PerforceCluster Cluster = await GetClusterAsync(ClusterName);
+			using (P4.Repository Repository = await GetConnection(Cluster, NoClient: true))
 			{
-				List<FileSummary> Results = new List<FileSummary>();
+				P4.GetFileMetaDataCmdOptions Options = new P4.GetFileMetaDataCmdOptions(P4.GetFileMetadataCmdFlags.ExcludeClientData, "", "", 0, "", "", "");
+
+				IList<P4.FileMetaData>? Files = Repository.GetFileMetaData(Options, Paths.Select(Path => { P4.FileSpec FileSpec = new P4.FileSpec(new P4.DepotPath(Path)); return FileSpec; }).ToArray());
 
-				using (P4.Repository Repository = GetConnection(NoClient: true))
+				if (Files == null)
 				{
-					P4.GetFileMetaDataCmdOptions Options = new P4.GetFileMetaDataCmdOptions(P4.GetFileMetadataCmdFlags.ExcludeClientData, "", "", 0, "", "", "");
+					Files = new List<P4.FileMetaData>();
+				}
 
-					IList<P4.FileMetaData>? Files = Repository.GetFileMetaData(Options, Paths.Select(Path => { P4.FileSpec FileSpec = new P4.FileSpec(new P4.DepotPath(Path)); return FileSpec; }).ToArray());
+				foreach (string Path in Paths)
+				{
+					P4.FileMetaData Meta = Files.FirstOrDefault(File => File.DepotPath.Path == Path);
 
-					if (Files == null)
+					if (Meta == null)
 					{
-						Files = new List<P4.FileMetaData>();
+						Results.Add(new FileSummary(Path, false, 0));
 					}
-
-					foreach (string Path in Paths)
+					else
 					{
-						P4.FileMetaData Meta = Files.FirstOrDefault(File => File.DepotPath.Path == Path);
-
-						if (Meta == null)
-						{
-							Results.Add(new FileSummary(Path, false, 0));
-						}
-						else
-						{
-							Results.Add(new FileSummary(Path, Meta.HeadAction != P4.FileAction.Delete && Meta.HeadAction != P4.FileAction.MoveDelete, Meta.HeadChange));
-						}
-
+						Results.Add(new FileSummary(Path, Meta.HeadAction != P4.FileAction.Delete && Meta.HeadAction != P4.FileAction.MoveDelete, Meta.HeadChange));
 					}
 
-					return Results;
 				}
-			});
+
+				return Results;
+			}
 		}
 
 		/// <inheritdoc/>
-		public async Task<byte[]> PrintAsync(string DepotPath)
+		public async Task<byte[]> PrintAsync(string ClusterName, string DepotPath)
 		{
-			return await Task.Run(() =>
+			if (DepotPath.EndsWith("...", StringComparison.OrdinalIgnoreCase) || DepotPath.EndsWith("*", StringComparison.OrdinalIgnoreCase))
 			{
+				throw new Exception("PrintAsync requires exactly one file to be specified");
+			}
 
-				if (DepotPath.EndsWith("...", StringComparison.OrdinalIgnoreCase) || DepotPath.EndsWith("*", StringComparison.OrdinalIgnoreCase))
+			PerforceCluster Cluster = await GetClusterAsync(ClusterName);
+			using (P4.Repository Repository = await GetConnection(Cluster, NoClient: true))
+			{
+				using (P4.P4Command Command = new P4.P4Command(Repository, "print", false, new string[] { "-q", DepotPath }))
 				{
-					throw new Exception("PrintAsync requires exactly one file to be specified");
-				}
+					P4.P4CommandResult Result = Command.Run();
 
-				using (P4.Repository Repository = GetConnection(NoClient: true))
-				{
-					using (P4.P4Command Command = new P4.P4Command(Repository, "print", false, new string[] { "-q", DepotPath }))
+					if (Result.BinaryOutput != null)
 					{
-						P4.P4CommandResult Result = Command.Run();
-
-						if (Result.BinaryOutput != null)
-						{
-							return Result.BinaryOutput;
-						}
-					
-						return Encoding.Default.GetBytes(Result.TextOutput);
-
+						return Result.BinaryOutput;
 					}
 
+					return Encoding.Default.GetBytes(Result.TextOutput);
+
 				}
-			});
 
+			}
 		}
 
 		/// <inheritdoc/>
-		public async Task<int> DuplicateShelvedChangeAsync(int ShelvedChange)
+		public async Task<int> DuplicateShelvedChangeAsync(string ClusterName, int ShelvedChange)
 		{
-			return await Task.Run(() =>
+			string? ChangeOwner = null;
+
+			// Get the owner of the shelf
+			PerforceCluster Cluster = await GetClusterAsync(ClusterName);
+			using (P4.Repository Repository = await GetConnection(Cluster, NoClient: true))
 			{
-				string? ChangeOwner = null;
+				List<string> Args = new List<string> { "-S", ShelvedChange.ToString(CultureInfo.InvariantCulture) };
 
-				// Get the owner of the shelf
-				using (P4.Repository Repository = GetConnection(NoClient: true))
+				using (P4.P4Command Command = new P4.P4Command(Repository, "describe", true, Args.ToArray()))
 				{
-					List<string> Args = new List<string> { "-S", ShelvedChange.ToString(CultureInfo.InvariantCulture)};
+					P4.P4CommandResult Result = Command.Run();
 
-					using (P4.P4Command Command = new P4.P4Command(Repository, "describe", true, Args.ToArray()))
+					if (!Result.Success)
 					{
-						P4.P4CommandResult Result = Command.Run();
-
-						if (!Result.Success)
-						{
-							throw new Exception($"Unable to get change {ShelvedChange}");
-						}
-
-						if (Result.TaggedOutput == null || Result.TaggedOutput.Count != 1)
-						{
-
-							throw new Exception($"Unable to get tagged output for change: {ShelvedChange}");
-						}
+						throw new Exception($"Unable to get change {ShelvedChange}");
+					}
 
-						P4.Changelist Changelist = new P4.Changelist();
-						Changelist.FromChangeCmdTaggedOutput(Result.TaggedOutput[0], true, string.Empty, false);
-						ChangeOwner = Changelist.OwnerName;
+					if (Result.TaggedOutput == null || Result.TaggedOutput.Count != 1)
+					{
 
+						throw new Exception($"Unable to get tagged output for change: {ShelvedChange}");
 					}
-				}
 
-				if (string.IsNullOrEmpty(ChangeOwner))
-				{
-					throw new Exception($"Unable to get owner for shelved change {ShelvedChange}");
-				}
+					P4.Changelist Changelist = new P4.Changelist();
+					Changelist.FromChangeCmdTaggedOutput(Result.TaggedOutput[0], true, string.Empty, false);
+					ChangeOwner = Changelist.OwnerName;
 
-				using (P4.Repository Repository = GetConnection(ReadOnly: false, ClientFromChange: ShelvedChange, UseClientFromChange: false, Username: ChangeOwner))
-				{
-					return ReshelveChange(Repository, ShelvedChange);				
 				}
+			}
 
-				throw new Exception($"Unable to duplicate shelve change {ShelvedChange}");
+			if (string.IsNullOrEmpty(ChangeOwner))
+			{
+				throw new Exception($"Unable to get owner for shelved change {ShelvedChange}");
+			}
 
-			});
+			using (P4.Repository Repository = await GetConnection(Cluster, ReadOnly: false, ClientFromChange: ShelvedChange, UseClientFromChange: false, Username: ChangeOwner))
+			{
+				return ReshelveChange(Repository, ShelvedChange);
+			}
 
+			throw new Exception($"Unable to duplicate shelve change {ShelvedChange}");
 		}
 
 		/// <inheritdoc/>
-		public async Task DeleteShelvedChangeAsync(int ShelvedChange)
+		public async Task DeleteShelvedChangeAsync(string ClusterName, int ShelvedChange)
 		{
-			await Task.Run(() =>
+			PerforceCluster Cluster = await GetClusterAsync(ClusterName);
+			using (P4.Repository Repository = await GetConnection(Cluster, NoClient: true))
 			{
 
-				using (P4.Repository Repository = GetConnection(NoClient: true))
+				P4.Changelist? Changelist = Repository.GetChangelist(ShelvedChange, new P4.DescribeCmdOptions(P4.DescribeChangelistCmdFlags.Omit, 0, 0));
+
+				if (Changelist == null)
 				{
+					throw new Exception($"Unable to get shelved changelist for delete: {ShelvedChange}");
+				}
 
-					P4.Changelist? Changelist = Repository.GetChangelist(ShelvedChange, new P4.DescribeCmdOptions(P4.DescribeChangelistCmdFlags.Omit, 0, 0));
+				string ClientId = Changelist.ClientId;
 
-					if (Changelist == null)
-					{
-						throw new Exception($"Unable to get shelved changelist for delete: {ShelvedChange}");
-					}
+				if (String.IsNullOrEmpty(ClientId))
+				{
+					throw new Exception($"Unable to get shelved changelist client id for delete: {ShelvedChange}");
+				}
 
-					string ClientId = Changelist.ClientId;
+				using (P4.Repository ClientRepository = await GetConnection(Cluster, Changelist.Stream, Changelist.OwnerName, ClientId: ClientId, UseClientFromChange: true, UsePortFromChange: true))
+				{
 
-					if (String.IsNullOrEmpty(ClientId))
+					if (Changelist.Shelved)
 					{
-						throw new Exception($"Unable to get shelved changelist client id for delete: {ShelvedChange}");
+						IList<P4.FileSpec> Files = ClientRepository.Connection.Client.ShelveFiles(new P4.ShelveFilesCmdOptions(P4.ShelveFilesCmdFlags.Delete, null, ShelvedChange));
 					}
 
-					using (P4.Repository ClientRepository = GetConnection(Changelist.Stream, Changelist.OwnerName, ClientId: ClientId, UseClientFromChange: true, UsePortFromChange: true))
-					{
-
-						if (Changelist.Shelved)
-						{
-							IList<P4.FileSpec> Files = ClientRepository.Connection.Client.ShelveFiles(new P4.ShelveFilesCmdOptions(P4.ShelveFilesCmdFlags.Delete, null, ShelvedChange));
-						}
-
-						ClientRepository.DeleteChangelist(Changelist, null);
-					}
+					ClientRepository.DeleteChangelist(Changelist, null);
 				}
-			});
+			}
 		}
 
 		/// <inheritdoc/>
-		public async Task UpdateChangelistDescription(int Change, string Description)
+		public async Task UpdateChangelistDescription(string ClusterName, int Change, string Description)
 		{
-			await Task.Run(() =>
+			try
 			{
-				try
+				PerforceCluster Cluster = await GetClusterAsync(ClusterName);
+				using (P4.Repository Repository = await GetConnection(Cluster, NoClient: true))
 				{
-					using (P4.Repository Repository = GetConnection(NoClient: true))
-					{
-						P4.Changelist Changelist = Repository.GetChangelist(Change);
+					P4.Changelist Changelist = Repository.GetChangelist(Change);
 
-						Repository.Connection.Disconnect();
+					Repository.Connection.Disconnect();
 
-						// the client must exist for the change list, otherwise will fail (for example, CreateNewChangeAsync deletes the client before returning)
-						using (P4.Repository UpdateRepository = GetConnection(ClientFromChange: Change, UseClientFromChange: true, Username: Changelist.OwnerName))
-						{
-							P4.Changelist UpdatedChangelist = UpdateRepository.GetChangelist(Change);
-							UpdatedChangelist.Description = Description;
-							UpdateRepository.UpdateChangelist(UpdatedChangelist);
-						}
+					// the client must exist for the change list, otherwise will fail (for example, CreateNewChangeAsync deletes the client before returning)
+					using (P4.Repository UpdateRepository = await GetConnection(Cluster, ClientFromChange: Change, UseClientFromChange: true, Username: Changelist.OwnerName))
+					{
+						P4.Changelist UpdatedChangelist = UpdateRepository.GetChangelist(Change);
+						UpdatedChangelist.Description = Description;
+						UpdateRepository.UpdateChangelist(UpdatedChangelist);
 					}
-				} 
-				catch (Exception Ex)
-				{
-					LogPerforce(1, "", $"Unable to update Changelist for CL {Change} to ${Description}, {Ex.Message}");
 				}
-			});
+			}
+			catch (Exception Ex)
+			{
+				LogPerforce(1, "", $"Unable to update Changelist for CL {Change} to ${Description}, {Ex.Message}");
+			}
 		}
 
 		/// <inheritdoc/>
-		public async Task<int> CreateNewChangeAsync(string StreamName, string FilePath)
+		public async Task<int> CreateNewChangeAsync(string ClusterName, string StreamName, string FilePath)
 		{
-			return await Task.Run(() =>
-			{
-
-				ServerSettings Settings = this.Settings.CurrentValue;
-
-				string Username = CanImpersonate ? "buildmachine" : Settings.P4BridgeServiceUsername!;
+			PerforceCluster Cluster = await GetClusterAsync(ClusterName);
+			P4.SubmitResults? SubmitResults = null;
 
-				P4.SubmitResults? SubmitResults = null;
+			using (P4.Repository Repository = await GetConnection(Cluster, Stream: StreamName, ReadOnly: false, CreateChange: true))
+			{
+				P4.Client Client = Repository.Connection.Client;
 
-				using (P4.Repository Repository = GetConnection(Stream: StreamName, Username: Username, ReadOnly: false, CreateChange: true))
-				{
-					P4.Client Client = Repository.Connection.Client;
+				string WorkspaceFilePath = $"//{Client.Name}/{FilePath.TrimStart('/')}";
+				string DiskFilePath = $"{Client.Root + FilePath.TrimStart('/')}";
 
-					string WorkspaceFilePath = $"//{Client.Name}/{FilePath.TrimStart('/')}";
-					string DiskFilePath = $"{Client.Root + FilePath.TrimStart('/')}";
+				P4.FileSpec WorkspaceFileSpec = new P4.FileSpec(new P4.ClientPath(WorkspaceFilePath));
 
-					P4.FileSpec WorkspaceFileSpec = new P4.FileSpec(new P4.ClientPath(WorkspaceFilePath));
+				IList<P4.File> Files = Repository.GetFiles(new P4.FilesCmdOptions(P4.FilesCmdFlags.None, 1), WorkspaceFileSpec);
 
-					IList<P4.File> Files = Repository.GetFiles(new P4.FilesCmdOptions(P4.FilesCmdFlags.None, 1), WorkspaceFileSpec);
+				P4.Changelist? SubmitChangelist = null;
+				P4.DepotPath? DepotPath = null;
 
-					P4.Changelist? SubmitChangelist = null;
-					P4.DepotPath? DepotPath = null;
+				const int MaxRetries = 10;
+				int Retry = 0;
 
-					const int MaxRetries = 10;
-					int Retry = 0;
+				for (; ; )
+				{
+					DepotPath = null;
 
-					for (; ; )
+					if (Retry == MaxRetries)
 					{
-						DepotPath = null;
+						break;
+					}
 
-						if (Retry == MaxRetries)
+					try
+					{
+						if (Files == null || Files.Count == 0)
 						{
-							break;
-						}
+							// File does not exist, create it
+							string? DirectoryName = Path.GetDirectoryName(DiskFilePath);
 
-						try
-						{
-							if (Files == null || Files.Count == 0)
+							if (string.IsNullOrEmpty(DirectoryName))
 							{
-								// File does not exist, create it
-								string? DirectoryName = Path.GetDirectoryName(DiskFilePath);
+								throw new Exception($"Invalid directory name for local client file, disk file path: {DiskFilePath}");
+							}
 
-								if (string.IsNullOrEmpty(DirectoryName))
-								{
-									throw new Exception($"Invalid directory name for local client file, disk file path: {DiskFilePath}");
-								}
+							// Create the directory
+							if (!Directory.Exists(DirectoryName))
+							{
+								Directory.CreateDirectory(DirectoryName);
 
-								// Create the directory
 								if (!Directory.Exists(DirectoryName))
 								{
-									Directory.CreateDirectory(DirectoryName);
-
-									if (!Directory.Exists(DirectoryName))
-									{
-										throw new Exception($"Unable to create directrory: {DirectoryName}");
-									}
+									throw new Exception($"Unable to create directrory: {DirectoryName}");
 								}
+							}
 
-								// Create the file
-								if (!File.Exists(DiskFilePath))
+							// Create the file
+							if (!File.Exists(DiskFilePath))
+							{
+								using (FileStream FileStream = File.OpenWrite(DiskFilePath))
 								{
-									using (FileStream FileStream = File.OpenWrite(DiskFilePath))
-									{
-										FileStream.Close();
-									}
-
-									if (!File.Exists(DiskFilePath))
-									{
-										throw new Exception($"Unable to create local change file: {DiskFilePath}");
-									}
-
+									FileStream.Close();
 								}
 
-								IList<P4.FileSpec> DepotFiles = Client.AddFiles(new P4.Options(), WorkspaceFileSpec);
-
-								if (DepotFiles == null || DepotFiles.Count != 1)
+								if (!File.Exists(DiskFilePath))
 								{
-									throw new Exception($"Unable to add local change file,  local: {DiskFilePath} : workspace: {WorkspaceFileSpec}");
+									throw new Exception($"Unable to create local change file: {DiskFilePath}");
 								}
 
-								DepotPath = DepotFiles[0].DepotPath;
-
 							}
-							else
-							{
-								IList<P4.FileSpec> SyncResults = Client.SyncFiles(new P4.SyncFilesCmdOptions(P4.SyncFilesCmdFlags.Force), WorkspaceFileSpec);
 
-								if (SyncResults == null || SyncResults.Count != 1)
-								{
-									throw new Exception($"Unable to sync file, workspace: {WorkspaceFileSpec}");
-								}
+							IList<P4.FileSpec> DepotFiles = Client.AddFiles(new P4.Options(), WorkspaceFileSpec);
 
-								IList<P4.FileSpec> EditResults = Client.EditFiles(new P4.FileSpec[] { new P4.FileSpec(SyncResults[0].DepotPath) }, new P4.Options());
+							if (DepotFiles == null || DepotFiles.Count != 1)
+							{
+								throw new Exception($"Unable to add local change file,  local: {DiskFilePath} : workspace: {WorkspaceFileSpec}");
+							}
 
-								if (EditResults == null || EditResults.Count != 1)
-								{
-									throw new Exception($"Unable to edit file, workspace: {WorkspaceFileSpec}");
-								}
+							DepotPath = DepotFiles[0].DepotPath;
 
-								DepotPath = EditResults[0].DepotPath;
+						}
+						else
+						{
+							IList<P4.FileSpec> SyncResults = Client.SyncFiles(new P4.SyncFilesCmdOptions(P4.SyncFilesCmdFlags.Force), WorkspaceFileSpec);
 
+							if (SyncResults == null || SyncResults.Count != 1)
+							{
+								throw new Exception($"Unable to sync file, workspace: {WorkspaceFileSpec}");
 							}
 
-							if (DepotPath == null || string.IsNullOrEmpty(DepotPath.Path))
+							IList<P4.FileSpec> EditResults = Client.EditFiles(new P4.FileSpec[] { new P4.FileSpec(SyncResults[0].DepotPath) }, new P4.Options());
+
+							if (EditResults == null || EditResults.Count != 1)
 							{
-								throw new Exception($"Unable to get depot path for: {WorkspaceFileSpec}");
+								throw new Exception($"Unable to edit file, workspace: {WorkspaceFileSpec}");
 							}
 
-							// create a new change
-							if (SubmitChangelist == null)
-							{
-								P4.Changelist Changelist = new P4.Changelist();
-								Changelist.Description = "New change for Horde job";
-								Changelist.Files.Add(new P4.FileMetaData(new P4.FileSpec(DepotPath)));
-								SubmitChangelist = Repository.CreateChangelist(Changelist);
+							DepotPath = EditResults[0].DepotPath;
 
-								if (SubmitChangelist == null)
-								{
-									throw new Exception($"Unable to create a changelist for: {DepotPath}");
-								}
-							}
+						}
 
-							SubmitResults = SubmitChangelist.Submit(null);
+						if (DepotPath == null || string.IsNullOrEmpty(DepotPath.Path))
+						{
+							throw new Exception($"Unable to get depot path for: {WorkspaceFileSpec}");
+						}
+
+						// create a new change
+						if (SubmitChangelist == null)
+						{
+							P4.Changelist Changelist = new P4.Changelist();
+							Changelist.Description = "New change for Horde job";
+							Changelist.Files.Add(new P4.FileMetaData(new P4.FileSpec(DepotPath)));
+							SubmitChangelist = Repository.CreateChangelist(Changelist);
 
-							if (SubmitResults == null)
+							if (SubmitChangelist == null)
 							{
-								throw new Exception($"Unable to submit changelist for: {DepotPath}");
+								throw new Exception($"Unable to create a changelist for: {DepotPath}");
 							}
-
-							break;
 						}
-						catch
-						{
-							Retry++;
 
-							Client.RevertFiles(new P4.Options(), WorkspaceFileSpec);
+						SubmitResults = SubmitChangelist.Submit(null);
 
-							continue;
+						if (SubmitResults == null)
+						{
+							throw new Exception($"Unable to submit changelist for: {DepotPath}");
 						}
 
-					}
-
-					string ClientName = Client.Name;
-					try
-					{
-						Repository.DeleteClient(Client, new P4.Options());
+						break;
 					}
 					catch
 					{
-						Logger.LogError($"Unable to delete client {ClientName}");
+						Retry++;
+
+						Client.RevertFiles(new P4.Options(), WorkspaceFileSpec);
+
+						continue;
 					}
 
 				}
 
-				if (SubmitResults == null)
+				string ClientName = Client.Name;
+				try
 				{
-					throw new Exception($"Unable to submit change for {StreamName} {FilePath}");
+					Repository.DeleteClient(Client, new P4.Options());
 				}
+				catch
+				{
+					Logger.LogError($"Unable to delete client {ClientName}");
+				}
+
+			}
+
+			if (SubmitResults == null)
+			{
+				throw new Exception($"Unable to submit change for {StreamName} {FilePath}");
+			}
 
-				return SubmitResults.ChangeIdAfterSubmit;
-			});
+			return SubmitResults.ChangeIdAfterSubmit;
 		}
 
 		/// <inheritdoc/>
-		public async Task<(int? Change, string Message)> SubmitShelvedChangeAsync(int Change, int OriginalChange)
+		public async Task<(int? Change, string Message)> SubmitShelvedChangeAsync(string ClusterName, int Change, int OriginalChange)
 		{
-			return await Task.Run(() => 
+			PerforceCluster Cluster = await GetClusterAsync(ClusterName);
+			using (P4.Repository Repository = await GetConnection(Cluster, NoClient: true))
 			{
-			   using (P4.Repository Repository = GetConnection(NoClient: true))
-			   {
-
-				   List<string> Args = new List<string> { "-S", Change.ToString(CultureInfo.InvariantCulture), OriginalChange.ToString(CultureInfo.InvariantCulture) };
-
-				   using (P4.P4Command Command = new P4.P4Command(Repository, "describe", true, Args.ToArray()))
-				   {
-					   P4.P4CommandResult Result = Command.Run();
-
-					   if (!Result.Success)
-					   {
-							return (null, $"Unable to get change {Change}");
-						
-					   }
-
-					   if (Result.TaggedOutput == null || Result.TaggedOutput.Count != 2)
-					   {
-
-							return (null, $"Unable to get tagged output for change: {Change} and original change: {OriginalChange}");
-					   }
-
-					   bool DSTMismatch = false;
-					   string Offset = string.Empty;
-
-					   if (Server != null && Server.Metadata != null)
-					   {
-						   Offset = Server.Metadata.DateTimeOffset;
-						   DSTMismatch = P4.FormBase.DSTMismatch(Server.Metadata);
-					   }
-
-					   P4.Changelist Changelist = new P4.Changelist();
-					   Changelist.FromChangeCmdTaggedOutput(Result.TaggedOutput[0], true, Offset, DSTMismatch);
-
-					   P4.Changelist OriginalChangelist = new P4.Changelist();
-					   OriginalChangelist.FromChangeCmdTaggedOutput(Result.TaggedOutput[1], true, Offset, DSTMismatch);
-
-					   if (OriginalChangelist.ShelvedFiles.Count != Changelist.ShelvedFiles.Count)
-					   {
-							return (null, $"Mismatched number of shelved files for change: {Change} and original change: {OriginalChange}");
-					   }
-
-					   if (OriginalChangelist.ShelvedFiles.Count == 0)
-					   {
-							return (null, $"No shelved file for change: {Change} and original change: {OriginalChange}");
-					   }
-
-						foreach (P4.ShelvedFile ShelvedFile in Changelist.ShelvedFiles)
-					   {
-						   P4.ShelvedFile? Found = OriginalChangelist.ShelvedFiles.FirstOrDefault(Original => Original.Digest == ShelvedFile.Digest && Original.Action == ShelvedFile.Action);
-
-						   if (Found == null)
-						   {
-								return (null, $"Mismatch in shelved file digest or action for {ShelvedFile.Path}");
-						   }
-					   }
-
-					   Repository.Connection.Disconnect();
-
-					   using (P4.Repository SubmitRepository = GetConnection(ReadOnly: false, ClientFromChange: Change, UseClientFromChange: true, Username: Changelist.OwnerName))
-					   {
-							// we might not need a client here, possibly -e below facilitates this, check!
-							using (P4.P4Command SubmitCommand = new P4.P4Command(SubmitRepository, "submit", null, true, new string[] { "-e", Change.ToString(CultureInfo.InvariantCulture) }))
-						   {
-								try
-								{
-									Result = SubmitCommand.Run();
-								}
-								catch (Exception Ex)
-								{
-									return (null, $"Submit command failed: {Ex.Message}");
-								}
 
-							   if (!Result.Success)
-							   {
-									string Message = (Result.ErrorList != null && Result.ErrorList.Count > 0) ? Result.ErrorList[0].ErrorMessage : "Unknown error, no errors in list" ;
-									return (null, $"Unable to submit {Change}, {Message}");
-							   }
+				List<string> Args = new List<string> { "-S", Change.ToString(CultureInfo.InvariantCulture), OriginalChange.ToString(CultureInfo.InvariantCulture) };
 
-							   int? SubmittedChangeId = null;
+				using (P4.P4Command Command = new P4.P4Command(Repository, "describe", true, Args.ToArray()))
+				{
+					P4.P4CommandResult Result = Command.Run();
 
-							   foreach (P4.TaggedObject TaggedObject in Result.TaggedOutput)
-							   {
-								   string? Submitted;
-								   if (TaggedObject.TryGetValue("submittedChange", out Submitted))
-								   {
-									   SubmittedChangeId = int.Parse(Submitted, CultureInfo.InvariantCulture);
-								   }
-							   }
+					if (!Result.Success)
+					{
+						return (null, $"Unable to get change {Change}");
+
+					}
 
-							   if (SubmittedChangeId == null)
-							   {
-								   return (null, $"Submit command succeeded, though unable to parse submitted change number");
-							   }
+					if (Result.TaggedOutput == null || Result.TaggedOutput.Count != 2)
+					{
 
-							   return (SubmittedChangeId, Changelist.Description);
+						return (null, $"Unable to get tagged output for change: {Change} and original change: {OriginalChange}");
+					}
 
-						   }
-					   }
-				   }
-			   }
+					bool DSTMismatch = false;
+					string Offset = string.Empty;
 
-			   throw new Exception($"Unable to get shelve change: {Change} for original change: {OriginalChange}");
-		   });
+					P4.Server Server = Repository.Server;
+					if (Server != null && Server.Metadata != null)
+					{
+						Offset = Server.Metadata.DateTimeOffset;
+						DSTMismatch = P4.FormBase.DSTMismatch(Server.Metadata);
+					}
 
-		}
+					P4.Changelist Changelist = new P4.Changelist();
+					Changelist.FromChangeCmdTaggedOutput(Result.TaggedOutput[0], true, Offset, DSTMismatch);
 
-		/// <inheritdoc/>		
-		public async Task<PerforceUserInfo?> GetUserInfoAsync(string UserName)
-		{
-			return await Task.Run(() =>
-			{
+					P4.Changelist OriginalChangelist = new P4.Changelist();
+					OriginalChangelist.FromChangeCmdTaggedOutput(Result.TaggedOutput[1], true, Offset, DSTMismatch);
 
-				using (P4.Repository Repository = GetConnection(NoClient: true))
-				{
-					P4.User User = Repository.GetUser(UserName, new P4.UserCmdOptions(P4.UserCmdFlags.Output));
+					if (OriginalChangelist.ShelvedFiles.Count != Changelist.ShelvedFiles.Count)
+					{
+						return (null, $"Mismatched number of shelved files for change: {Change} and original change: {OriginalChange}");
+					}
 
-					if (User == null)
+					if (OriginalChangelist.ShelvedFiles.Count == 0)
 					{
-						return null;
+						return (null, $"No shelved file for change: {Change} and original change: {OriginalChange}");
+					}
+
+					foreach (P4.ShelvedFile ShelvedFile in Changelist.ShelvedFiles)
+					{
+						P4.ShelvedFile? Found = OriginalChangelist.ShelvedFiles.FirstOrDefault(Original => Original.Digest == ShelvedFile.Digest && Original.Action == ShelvedFile.Action);
+
+						if (Found == null)
+						{
+							return (null, $"Mismatch in shelved file digest or action for {ShelvedFile.Path}");
+						}
 					}
 
-					return new PerforceUserInfo { Name = UserName, Email = User.EmailAddress };
+					Repository.Connection.Disconnect();
+
+					using (P4.Repository SubmitRepository = await GetConnection(Cluster, ReadOnly: false, ClientFromChange: Change, UseClientFromChange: true, Username: Changelist.OwnerName))
+					{
+						// we might not need a client here, possibly -e below facilitates this, check!
+						using (P4.P4Command SubmitCommand = new P4.P4Command(SubmitRepository, "submit", null, true, new string[] { "-e", Change.ToString(CultureInfo.InvariantCulture) }))
+						{
+							try
+							{
+								Result = SubmitCommand.Run();
+							}
+							catch (Exception Ex)
+							{
+								return (null, $"Submit command failed: {Ex.Message}");
+							}
+
+							if (!Result.Success)
+							{
+								string Message = (Result.ErrorList != null && Result.ErrorList.Count > 0) ? Result.ErrorList[0].ErrorMessage : "Unknown error, no errors in list";
+								return (null, $"Unable to submit {Change}, {Message}");
+							}
+
+							int? SubmittedChangeId = null;
+
+							foreach (P4.TaggedObject TaggedObject in Result.TaggedOutput)
+							{
+								string? Submitted;
+								if (TaggedObject.TryGetValue("submittedChange", out Submitted))
+								{
+									SubmittedChangeId = int.Parse(Submitted, CultureInfo.InvariantCulture);
+								}
+							}
+
+							if (SubmittedChangeId == null)
+							{
+								return (null, $"Submit command succeeded, though unable to parse submitted change number");
+							}
+
+							return (SubmittedChangeId, Changelist.Description);
+
+						}
+					}
 				}
-			});
+			}
+
+			throw new Exception($"Unable to get shelve change: {Change} for original change: {OriginalChange}");
 		}
 
-		bool CanImpersonate
+		/// <inheritdoc/>		
+		public async Task<PerforceUserInfo?> GetUserInfoAsync(string ClusterName, string UserName)
 		{
-			get
+			PerforceCluster Cluster = await GetClusterAsync(ClusterName);
+			using (P4.Repository Repository = await GetConnection(Cluster, NoClient: true))
 			{
-				ServerSettings Settings = this.Settings.CurrentValue;
-				return Settings.P4BridgeCanImpersonate;
+				P4.User User = Repository.GetUser(UserName, new P4.UserCmdOptions(P4.UserCmdFlags.Output));
+
+				if (User == null)
+				{
+					return null;
+				}
+
+				return new PerforceUserInfo { Name = UserName, Email = User.EmailAddress };
 			}
 		}
 
-		int ReshelveChange(P4.Repository Repository, int Change)
+		static int ReshelveChange(P4.Repository Repository, int Change)
 		{
 
 			bool EdgeServer = false;
 			string? Value;
-			if (Server.Metadata.RawData.TryGetValue("serverServices", out Value))
+			if (Repository.Server.Metadata.RawData.TryGetValue("serverServices", out Value))
 			{
 				if (Value == "edge-server")
 				{
@@ -1111,13 +1091,13 @@ namespace HordeServer.Services
 		}
 
 		/// <inheritdoc/>
-		public virtual async Task<int> GetCodeChangeAsync(string StreamName, int Change)
+		public virtual async Task<int> GetCodeChangeAsync(string ClusterName, string StreamName, int Change)
 		{
 			int MaxChange = Change;
 			for (; ; )
 			{
 				// Query for the changes before this point
-				List<ChangeSummary> Changes = await GetChangesAsync(StreamName, null, MaxChange, 10, null);
+				List<ChangeSummary> Changes = await GetChangesAsync(ClusterName, StreamName, null, MaxChange, 10, null);
 				Serilog.Log.Logger.Information("Finding last code change in {Stream} before {MaxChange}: {NumResults}", StreamName, MaxChange, Changes.Count);
 				if (Changes.Count == 0)
 				{
@@ -1125,7 +1105,7 @@ namespace HordeServer.Services
 				}
 
 				// Get the details for them
-				List<ChangeDetails> DetailsList = await GetChangeDetailsAsync(StreamName, Changes.ConvertAll(x => x.Number), null);
+				List<ChangeDetails> DetailsList = await GetChangeDetailsAsync(ClusterName, StreamName, Changes.ConvertAll(x => x.Number), null);
 				foreach (ChangeDetails Details in DetailsList.OrderByDescending(x => x.Number))
 				{
 					ChangeContentFlags ContentFlags = Details.GetContentFlags();
diff --git a/Engine/Source/Programs/Horde/HordeServer/Services/JobService.cs b/Engine/Source/Programs/Horde/HordeServer/Services/JobService.cs
--- a/Engine/Source/Programs/Horde/HordeServer/Services/JobService.cs
+++ b/Engine/Source/Programs/Horde/HordeServer/Services/JobService.cs
@@ -212,7 +212,7 @@ namespace HordeServer.Services
 		/// Creates a new job
 		/// </summary>
 		/// <param name="JobId">A requested job id</param>
-		/// <param name="StreamId">Unique id of the stream that this job belongs to</param>
+		/// <param name="Stream">The stream that this job belongs to</param>
 		/// <param name="TemplateRefId">Name of the template ref</param>
 		/// <param name="TemplateHash">Template for this job</param>
 		/// <param name="Graph">The graph for the new job</param>
@@ -235,14 +235,14 @@ namespace HordeServer.Services
 		/// <param name="Counters">Counters for the job</param>
 		/// <param name="Arguments">Arguments for the job</param>
 		/// <returns>Unique id representing the job</returns>
-		public async Task<IJob> CreateJobAsync(ObjectId? JobId, StreamId StreamId, TemplateRefId TemplateRefId, ContentHash TemplateHash, IGraph Graph, string Name, int Change, int CodeChange, int? PreflightChange, int? ClonedPreflightChange, ObjectId? StartedByUserId, string? StartedByUserName, Priority? Priority, bool? AutoSubmit, bool? UpdateIssues, List<ChainedJobTemplate>? JobTriggers, bool ShowUgsBadges, bool ShowUgsAlerts, string? NotificationChannel, string? NotificationChannelFilter, string? HelixSwarmCallbackUrl, IReadOnlyList<ITemplateCounter> Counters, IReadOnlyList<string> Arguments)
+		public async Task<IJob> CreateJobAsync(ObjectId? JobId, IStream Stream, TemplateRefId TemplateRefId, ContentHash TemplateHash, IGraph Graph, string Name, int Change, int CodeChange, int? PreflightChange, int? ClonedPreflightChange, ObjectId? StartedByUserId, string? StartedByUserName, Priority? Priority, bool? AutoSubmit, bool? UpdateIssues, List<ChainedJobTemplate>? JobTriggers, bool ShowUgsBadges, bool ShowUgsAlerts, string? NotificationChannel, string? NotificationChannelFilter, string? HelixSwarmCallbackUrl, IReadOnlyList<ITemplateCounter> Counters, IReadOnlyList<string> Arguments)
 		{
 			ObjectId JobIdValue = JobId ?? ObjectId.GenerateNewId();
 			using IDisposable Scope = Logger.BeginScope("CreateJobAsync({JobId})", JobIdValue);
 
-			if (PreflightChange != null && ShouldClonePreflightChange(StreamId))
+			if (PreflightChange != null && ShouldClonePreflightChange(Stream.Id))
 			{
-				ClonedPreflightChange = await CloneShelvedChangeAsync(ClonedPreflightChange ?? PreflightChange.Value);
+				ClonedPreflightChange = await CloneShelvedChangeAsync(Stream.ClusterName, ClonedPreflightChange ?? PreflightChange.Value);
 			}
 
 			Logger.LogInformation("Creating job at CL {Change}, code CL {CodeChange}, preflight CL {PreflightChange}, cloned CL {ClonedPreflightChange}", Change, CodeChange, PreflightChange, ClonedPreflightChange);
@@ -252,7 +252,7 @@ namespace HordeServer.Services
 			Properties["CodeChange"] = CodeChange.ToString(CultureInfo.InvariantCulture);
 			Properties["PreflightChange"] = PreflightChange?.ToString(CultureInfo.InvariantCulture) ?? String.Empty;
 			Properties["ClonedPreflightChange"] = ClonedPreflightChange?.ToString(CultureInfo.InvariantCulture) ?? String.Empty;
-			Properties["StreamId"] = StreamId.ToString();
+			Properties["StreamId"] = Stream.Id.ToString();
 			Properties["TemplateId"] = TemplateRefId.ToString();
 			Properties["JobId"] = JobIdValue.ToString();
 
@@ -276,7 +276,7 @@ namespace HordeServer.Services
 
 			Name = StringUtils.ExpandProperties(Name, Properties);
 
-			IJob NewJob = await Jobs.AddAsync(JobIdValue, StreamId, TemplateRefId, TemplateHash, Graph, Name, Change, CodeChange, PreflightChange, ClonedPreflightChange, StartedByUserId, StartedByUserName, Priority, AutoSubmit, UpdateIssues, JobTriggers, ShowUgsBadges, ShowUgsAlerts, NotificationChannel, NotificationChannelFilter, HelixSwarmCallbackUrl, ExpandedArguments);
+			IJob NewJob = await Jobs.AddAsync(JobIdValue, Stream.Id, TemplateRefId, TemplateHash, Graph, Name, Change, CodeChange, PreflightChange, ClonedPreflightChange, StartedByUserId, StartedByUserName, Priority, AutoSubmit, UpdateIssues, JobTriggers, ShowUgsBadges, ShowUgsAlerts, NotificationChannel, NotificationChannelFilter, HelixSwarmCallbackUrl, ExpandedArguments);
 			JobTaskSource.UpdateQueuedJob(NewJob, Graph);
 
 			await JobTaskSource.UpdateUgsBadges(NewJob, Graph, new List<(LabelState, LabelOutcome)>());
@@ -896,7 +896,11 @@ namespace HordeServer.Services
 							}
 							else if (Job.ClonedPreflightChange != 0)
 							{
-								await DeleteShelvedChangeAsync(Job.ClonedPreflightChange);
+								IStream? Stream = await StreamService.GetCachedStream(Job.StreamId);
+								if (Stream != null)
+								{
+									await DeleteShelvedChangeAsync(Stream.ClusterName, Job.ClonedPreflightChange);
+								}
 							}
 						}
 					}
@@ -979,38 +983,51 @@ namespace HordeServer.Services
 			string Message;
 			try
 			{
-				int ClonedPreflightChange = Job.ClonedPreflightChange;
-				if (ClonedPreflightChange == 0)
+				IStream? Stream = await StreamService.GetCachedStream(Job.StreamId);
+				if (Stream != null)
 				{
-					if (ShouldClonePreflightChange(Job.StreamId))
+					int ClonedPreflightChange = Job.ClonedPreflightChange;
+					if (ClonedPreflightChange == 0)
 					{
-						ClonedPreflightChange = await CloneShelvedChangeAsync(Job.PreflightChange);
+						if (ShouldClonePreflightChange(Job.StreamId))
+						{
+							ClonedPreflightChange = await CloneShelvedChangeAsync(Stream.ClusterName, Job.PreflightChange);
+						}
+						else
+						{
+							ClonedPreflightChange = Job.PreflightChange;
+						}
 					}
-					else
+
+					Logger.LogInformation("Updating description for {ClonedPreflightChange}", ClonedPreflightChange);
+
+					ChangeDetails Details = await PerforceService.GetChangeDetailsAsync(Stream.ClusterName, Stream.Name, ClonedPreflightChange, null);
+					await PerforceService.UpdateChangelistDescription(Stream.ClusterName, ClonedPreflightChange, Details.Description.TrimEnd() + $"\n#preflight {Job.Id}");
+
+					Logger.LogInformation("Submitting change {Change} (through {ChangeCopy}) after successful completion of {JobId}", Job.PreflightChange, ClonedPreflightChange, Job.Id);
+					(Change, Message) = await PerforceService.SubmitShelvedChangeAsync(Stream.ClusterName, ClonedPreflightChange, Job.PreflightChange);
+
+					Logger.LogInformation("Attempt to submit {Change} (through {ChangeCopy}): {Message}", Job.PreflightChange, ClonedPreflightChange, Message);
+
+					if (ShouldClonePreflightChange(Job.StreamId))
 					{
-						ClonedPreflightChange = Job.PreflightChange;
+						if (Change != null && Job.ClonedPreflightChange != 0)
+						{
+							await DeleteShelvedChangeAsync(Stream.ClusterName, Job.PreflightChange);
+						}
 					}
-				}
-
-				Logger.LogInformation("Updating description for {ClonedPreflightChange}", ClonedPreflightChange);
-				try
-				{
-					IStream? Stream = await StreamService.GetCachedStream(Job.StreamId);
-					if (Stream != null)
+					else
 					{
-						ChangeDetails Details = await PerforceService.GetChangeDetailsAsync(Stream.Name, ClonedPreflightChange, null);
-						await PerforceService.UpdateChangelistDescription(ClonedPreflightChange, Details.Description.TrimEnd() + $"\n#preflight {Job.Id}");
+						if (Change != null && Job.PreflightChange != 0)
+						{
+							await DeleteShelvedChangeAsync(Stream.ClusterName, Job.PreflightChange);
+						}
 					}
 				}
-				catch (Exception Ex)
+				else
 				{
-					Logger.LogError(Ex, "Unable to update description for changelist {Change}", ClonedPreflightChange);
+					(Change, Message) = ((int?)null, "Stream not found");
 				}
-
-				Logger.LogInformation("Submitting change {Change} (through {ChangeCopy}) after successful completion of {JobId}", Job.PreflightChange, ClonedPreflightChange, Job.Id);
-				(Change, Message) = await PerforceService.SubmitShelvedChangeAsync(ClonedPreflightChange, Job.PreflightChange);
-
-				Logger.LogInformation("Attempt to submit {Change} (through {ChangeCopy}): {Message}", Job.PreflightChange, ClonedPreflightChange, Message);
 			}
 			catch (Exception Ex)
 			{
@@ -1018,21 +1035,6 @@ namespace HordeServer.Services
 				Logger.LogError(Ex, "Unable to submit shelved change");
 			}
 
-			if (ShouldClonePreflightChange(Job.StreamId))
-			{
-				if (Change != null && Job.ClonedPreflightChange != 0)
-				{
-					await DeleteShelvedChangeAsync(Job.PreflightChange);
-				}
-			}
-			else
-			{
-				if (Change != null && Job.PreflightChange != 0)
-				{
-					await DeleteShelvedChangeAsync(Job.PreflightChange);
-				}
-			}
-
 			for (; ; )
 			{
 				if (await Jobs.TryUpdateJobAsync(Job, Graph, AutoSubmitChange: Change, AutoSubmitMessage: Message))
@@ -1053,14 +1055,15 @@ namespace HordeServer.Services
 		/// <summary>
 		/// Clone a shelved changelist for running a preflight
 		/// </summary>
+		/// <param name="ClusterName">Name of the Perforce cluster</param>
 		/// <param name="Change">The changelist to clone</param>
 		/// <returns></returns>
-		private async Task<int> CloneShelvedChangeAsync(int Change)
+		private async Task<int> CloneShelvedChangeAsync(string ClusterName, int Change)
 		{
 			int ClonedChange;
 			try
 			{
-				ClonedChange = await PerforceService.DuplicateShelvedChangeAsync(Change);
+				ClonedChange = await PerforceService.DuplicateShelvedChangeAsync(ClusterName, Change);
 				Logger.LogInformation("CL {Change} was duplicated into {ClonedChange}", Change, ClonedChange);
 			}
 			catch (Exception Ex)
@@ -1074,15 +1077,16 @@ namespace HordeServer.Services
 		/// <summary>
 		/// Deletes a shelved changelist, catching and logging any exceptions
 		/// </summary>
+		/// <param name="ClusterName"></param>
 		/// <param name="Change">The changelist to delete</param>
 		/// <returns>True if the change was deleted successfully, false otherwise</returns>
 		[SuppressMessage("Design", "CA1031:Do not catch general exception types", Justification = "<Pending>")]
-		private async Task<bool> DeleteShelvedChangeAsync(int Change)
+		private async Task<bool> DeleteShelvedChangeAsync(string ClusterName, int Change)
 		{
 			Logger.LogInformation("Removing shelf {Change}", Change);
 			try
 			{
-				await PerforceService.DeleteShelvedChangeAsync(Change);
+				await PerforceService.DeleteShelvedChangeAsync(ClusterName, Change);
 				return true;
 			}
 			catch (Exception Ex)
@@ -1131,7 +1135,7 @@ namespace HordeServer.Services
 					IGraph TriggerGraph = await Graphs.AddAsync(Template);
 					Logger.LogInformation("Creating downstream job {ChainedJobId} from job {JobId}", ChainedJobId, Job.Id);
 
-					await CreateJobAsync(ChainedJobId, Job.StreamId, JobTrigger.TemplateRefId, TemplateRef.Hash, TriggerGraph, TemplateRef.Name, Job.Change, Job.CodeChange, Job.PreflightChange, Job.ClonedPreflightChange, Job.StartedByUserId, Job.StartedByUser, Template.Priority, null, Job.UpdateIssues, TemplateRef.ChainedJobs, false, false, TemplateRef.NotificationChannel, TemplateRef.NotificationChannelFilter, null, Template.Counters, Template.Arguments);
+					await CreateJobAsync(ChainedJobId, Stream, JobTrigger.TemplateRefId, TemplateRef.Hash, TriggerGraph, TemplateRef.Name, Job.Change, Job.CodeChange, Job.PreflightChange, Job.ClonedPreflightChange, Job.StartedByUserId, Job.StartedByUser, Template.Priority, null, Job.UpdateIssues, TemplateRef.ChainedJobs, false, false, TemplateRef.NotificationChannel, TemplateRef.NotificationChannelFilter, null, Template.Counters, Template.Arguments);
 					break;
 				}
 
diff --git a/Engine/Source/Programs/Horde/HordeServer/Services/PerforceLoadBalancer.cs b/Engine/Source/Programs/Horde/HordeServer/Services/PerforceLoadBalancer.cs
--- a/Engine/Source/Programs/Horde/HordeServer/Services/PerforceLoadBalancer.cs
+++ b/Engine/Source/Programs/Horde/HordeServer/Services/PerforceLoadBalancer.cs
@@ -173,17 +173,38 @@ namespace HordeServer.Services
 			return null;
 		}
 
+		/// <summary>
+		/// Select a Perforce server to use by the Horde server
+		/// </summary>
+		/// <param name="Cluster"></param>
+		/// <returns></returns>
+		public Task<IPerforceServer?> SelectServerAsync(PerforceCluster Cluster)
+		{
+			HashSet<string> Properties = new HashSet<string>{ "HordeServer=1" };
+			return SelectServerAsync(Cluster, Properties);
+		}
+
 		/// <summary>
 		/// Select a Perforce server to use
 		/// </summary>
 		/// <param name="Cluster"></param>
 		/// <param name="Agent"></param>
 		/// <returns></returns>
-		public async Task<IPerforceServer?> SelectServerAsync(PerforceCluster Cluster, IAgent Agent)
+		public Task<IPerforceServer?> SelectServerAsync(PerforceCluster Cluster, IAgent Agent)
 		{
 			// Get the set of properties to use for checking server validity
 			HashSet<string>? Properties = Agent.Capabilities.Devices[0].Properties;
+			return SelectServerAsync(Cluster, Properties);
+		}
 
+		/// <summary>
+		/// Select a Perforce server to use
+		/// </summary>
+		/// <param name="Cluster"></param>
+		/// <param name="Properties"></param>
+		/// <returns></returns>
+		public async Task<IPerforceServer?> SelectServerAsync(PerforceCluster Cluster, HashSet<string>? Properties)
+		{
 			// Find all the valid servers for this agent
 			HashSet<string> ValidServers = new HashSet<string>(StringComparer.OrdinalIgnoreCase);
 			if (Properties != null)
diff --git a/Engine/Source/Programs/Horde/HordeServer/Services/ScheduleService.cs b/Engine/Source/Programs/Horde/HordeServer/Services/ScheduleService.cs
--- a/Engine/Source/Programs/Horde/HordeServer/Services/ScheduleService.cs
+++ b/Engine/Source/Programs/Horde/HordeServer/Services/ScheduleService.cs
@@ -281,7 +281,7 @@ namespace HordeServer.Services
 				}
 				else
 				{
-					(Change, CodeChange) = await GetNextChangeAsync(Stream.Name, MinChangeNumber, MaxChangeNumber);
+					(Change, CodeChange) = await GetNextChangeAsync(Stream.ClusterName, Stream.Name, MinChangeNumber, MaxChangeNumber);
 				}
 
 				// Quit if we didn't find anything
@@ -295,7 +295,7 @@ namespace HordeServer.Services
 				}
 
 				// Adjust the changelist for the desired filter
-				if(await ShouldBuildChangeAsync(Stream.Name, Change, FilterFlags))
+				if(await ShouldBuildChangeAsync(Stream.ClusterName, Stream.Name, Change, FilterFlags))
 				{
 					TriggerChanges.Add((Change, CodeChange));
 				}
@@ -340,7 +340,7 @@ namespace HordeServer.Services
 			if (Template.SubmitNewChange != null)
 			{
 				int NewChange = await Perforce.CreateNewChangeForTemplateAsync(Stream, Template);
-				int NewCodeChange = await Perforce.GetCodeChangeAsync(Stream.Name, NewChange);
+				int NewCodeChange = await Perforce.GetCodeChangeAsync(Stream.ClusterName, Stream.Name, NewChange);
 				TriggerChanges = new List<(int, int)> { (NewChange, NewCodeChange) };
 			}
 
@@ -348,7 +348,7 @@ namespace HordeServer.Services
 			foreach ((int Change, int CodeChange) in TriggerChanges.OrderBy(x => x.Change))
 			{
 				List<string> DefaultArguments = Template.GetDefaultArguments();
-				IJob NewJob = await JobService.CreateJobAsync(null, Stream.Id, TemplateRefId, Template.Id, Graph, Template.Name, Change, CodeChange, null, null, null, null, Template.Priority, null, null, TemplateRef.ChainedJobs, TemplateRef.ShowUgsBadges, TemplateRef.ShowUgsAlerts, TemplateRef.NotificationChannel, TemplateRef.NotificationChannelFilter, null, Template.Counters, DefaultArguments);
+				IJob NewJob = await JobService.CreateJobAsync(null, Stream, TemplateRefId, Template.Id, Graph, Template.Name, Change, CodeChange, null, null, null, null, Template.Priority, null, null, TemplateRef.ChainedJobs, TemplateRef.ShowUgsBadges, TemplateRef.ShowUgsAlerts, TemplateRef.NotificationChannel, TemplateRef.NotificationChannelFilter, null, Template.Counters, DefaultArguments);
 				Logger.LogInformation("Started new job for {StreamName} template {TemplateName} at CL {Change} (Code CL {CodeChange}): {JobId}", Stream.Id, TemplateRef.Name, Change, CodeChange, NewJob.Id);
 				await StreamService.UpdateScheduleTriggerAsync(Stream, TemplateRefId, Now, Change, new List<ObjectId> { NewJob.Id }, new List<ObjectId>());
 			}
@@ -357,15 +357,16 @@ namespace HordeServer.Services
 		/// <summary>
 		/// Tests whether a schedule should build a particular change, based on its requested change filters
 		/// </summary>
+		/// <param name="ClusterName"></param>
 		/// <param name="StreamName"></param>
 		/// <param name="Change"></param>
 		/// <param name="FilterFlags"></param>
 		/// <returns></returns>
-		private async Task<bool> ShouldBuildChangeAsync(string StreamName, int Change, ChangeContentFlags? FilterFlags)
+		private async Task<bool> ShouldBuildChangeAsync(string ClusterName, string StreamName, int Change, ChangeContentFlags? FilterFlags)
 		{
 			if (FilterFlags != null && FilterFlags.Value != 0)
 			{
-				ChangeDetails Details = await Perforce.GetChangeDetailsAsync(StreamName, Change, null);
+				ChangeDetails Details = await Perforce.GetChangeDetailsAsync(ClusterName, StreamName, Change, null);
 				if ((Details.GetContentFlags() & FilterFlags.Value) == 0)
 				{
 					Logger.LogDebug("Not building change {Change} ({ChangeFlags}) due to filter flags ({FilterFlags})", Change, Details.GetContentFlags().ToString(), FilterFlags.Value.ToString());
@@ -379,13 +380,13 @@ namespace HordeServer.Services
 		/// Get the next change to build in this stream, between a given range
 		/// </summary>
 		/// <returns>The new change and code change to build, or 0.</returns>
-		private async Task<(int Change, int CodeChange)> GetNextChangeAsync(string StreamName, int? MinChange, int? MaxChange)
+		private async Task<(int Change, int CodeChange)> GetNextChangeAsync(string ClusterName, string StreamName, int? MinChange, int? MaxChange)
 		{
-			List<ChangeSummary> Changes = await Perforce.GetChangesAsync(StreamName, MinChange, MaxChange, 1, null);
+			List<ChangeSummary> Changes = await Perforce.GetChangesAsync(ClusterName, StreamName, MinChange, MaxChange, 1, null);
 			if (Changes.Count > 0)
 			{
 				int Change = Changes[0].Number;
-				int CodeChange = await Perforce.GetCodeChangeAsync(StreamName, Change);
+				int CodeChange = await Perforce.GetCodeChangeAsync(ClusterName, StreamName, Change);
 				return (Change, CodeChange);
 			}
 			return default;
