Horde: Fix tests.

diff --git a/Engine/Source/Programs/Horde/HordeServerTests/Fixture.cs b/Engine/Source/Programs/Horde/HordeServerTests/Fixture.cs
--- a/Engine/Source/Programs/Horde/HordeServerTests/Fixture.cs
+++ b/Engine/Source/Programs/Horde/HordeServerTests/Fixture.cs
@@ -78,7 +78,7 @@ namespace HordeServerTests
 			
 			Job1 = await JobService.CreateJobAsync(
 				JobId: new ObjectId("5f283932841e7fdbcafb6ab5"),
-				StreamId: Stream!.Id,
+				Stream: Stream!,
 				TemplateRefId: TemplateRefId1,
 				TemplateHash: Template.Id,
 				Graph: Graph,
@@ -105,7 +105,7 @@ namespace HordeServerTests
 
 			Job2 = await JobService.CreateJobAsync(
 				JobId: new ObjectId("5f69ea1b68423e921b035106"),
-				StreamId: Stream!.Id,
+				Stream: Stream!,
 				TemplateRefId: new TemplateRefId("template-id-1"),
 				TemplateHash: ContentHash.MD5("made-up-template-hash"),
 				Graph: Fg,
diff --git a/Engine/Source/Programs/Horde/HordeServerTests/JobServiceTests.cs b/Engine/Source/Programs/Horde/HordeServerTests/JobServiceTests.cs
--- a/Engine/Source/Programs/Horde/HordeServerTests/JobServiceTests.cs
+++ b/Engine/Source/Programs/Horde/HordeServerTests/JobServiceTests.cs
@@ -43,7 +43,7 @@ namespace HordeServerTests
 			IStream? Stream = await TestSetup.StreamService.GetStreamAsync(StreamId);
 			Stream = await TestSetup.StreamService.StreamCollection.TryCreateOrReplaceAsync(new StreamId("ue5-main"), Stream, String.Empty, String.Empty, ProjectId, StreamConfig);
 
-			IJob Job = await TestSetup.JobService.CreateJobAsync(null, Stream!.Id, TemplateRefId1, Template.Id, Graph, "Hello", 1234, 1233, 999, null, null, "joe", null, null, null, Stream.Templates[TemplateRefId1].ChainedJobs, true, true, null, null, null, Template.Counters, new List<string>());
+			IJob Job = await TestSetup.JobService.CreateJobAsync(null, Stream!, TemplateRefId1, Template.Id, Graph, "Hello", 1234, 1233, 999, null, null, "joe", null, null, null, Stream!.Templates[TemplateRefId1].ChainedJobs, true, true, null, null, null, Template.Counters, new List<string>());
 			Assert.AreEqual(1, Job.ChainedJobs.Count);
 
 			Assert.IsTrue(await TestSetup.JobService.UpdateBatchAsync(Job, Job.Batches[0].Id, ObjectId.GenerateNewId(), JobStepBatchState.Running));
@@ -94,7 +94,7 @@ namespace HordeServerTests
 		{
 			return TestSetup.JobService.CreateJobAsync(
 				JobId: ObjectId.GenerateNewId(),
-				StreamId: TestSetup.Fixture!.Stream!.Id,
+				Stream: TestSetup.Fixture!.Stream!,
 				TemplateRefId: new TemplateRefId(TemplateRefId),
 				TemplateHash: new ContentHash(Encoding.ASCII.GetBytes(TemplateHash)),
 				Graph: TestSetup.Fixture!.Graph,
@@ -152,7 +152,7 @@ namespace HordeServerTests
 			Assert.IsNotNull(Project);
 
 			StreamId StreamId = new StreamId("ue5-main");
-			await TestSetup.StreamCollection.TryCreateOrReplaceAsync(new StreamId("ue5-main"), null, "", "", Project!.Id, new StreamConfig { Name = "//UE5/Main" });
+			IStream? Stream = await TestSetup.StreamCollection.TryCreateOrReplaceAsync(StreamId, null, "", "", Project!.Id, new StreamConfig { Name = "//UE5/Main" });
 
 			ITemplate Template = await TestSetup.TemplateService.CreateTemplateAsync("Test template", null, false, null, null, new List<TemplateCounter>(), new List<string>(), new List<Parameter>());
 			IGraph Graph = await TestSetup.GraphCollection.AddAsync(Template);
@@ -167,7 +167,7 @@ namespace HordeServerTests
 
 			Graph = await TestSetup.GraphCollection.AppendAsync(Graph, new List<CreateGroupRequest> { GroupA, GroupB });
 
-			IJob Job = await TestSetup.JobService.CreateJobAsync(null, StreamId, new TemplateRefId("temp"), Template.Id, Graph, "Hello", 1234, 1233, 999, null, null, "joe", null, null, null, null, true, true, null, null, null, Template.Counters, new List<string> { "-Target=Pak" });
+			IJob Job = await TestSetup.JobService.CreateJobAsync(null, Stream!, new TemplateRefId("temp"), Template.Id, Graph, "Hello", 1234, 1233, 999, null, null, "joe", null, null, null, null, true, true, null, null, null, Template.Counters, new List<string> { "-Target=Pak" });
 
 			Assert.IsTrue(await TestSetup.JobService.UpdateBatchAsync(Job, Job.Batches[0].Id, ObjectId.GenerateNewId(), JobStepBatchState.Running));
 			Assert.IsNotNull(await TestSetup.JobService.UpdateStepAsync(Job, Job.Batches[0].Id, Job.Batches[0].Steps[0].Id, JobStepState.Running));
diff --git a/Engine/Source/Programs/Horde/HordeServerTests/SchedulerTests.cs b/Engine/Source/Programs/Horde/HordeServerTests/SchedulerTests.cs
--- a/Engine/Source/Programs/Horde/HordeServerTests/SchedulerTests.cs
+++ b/Engine/Source/Programs/Horde/HordeServerTests/SchedulerTests.cs
@@ -355,7 +355,7 @@ namespace HordeServerTests
 			Assert.AreEqual(0, Jobs2.Count);
 
 			// Create a job and fail it
-			IJob Job1 = await TestSetup.JobService.CreateJobAsync(null, StreamId, NewTemplateRefId1, Template.Id, GraphA, "Hello", 1234, 1233, 999, null, null, "joe", null, null, null, null, true, true, null, null, null, Template.Counters, new List<string> { "-Target=TriggerNext" });
+			IJob Job1 = await TestSetup.JobService.CreateJobAsync(null, Stream, NewTemplateRefId1, Template.Id, GraphA, "Hello", 1234, 1233, 999, null, null, "joe", null, null, null, null, true, true, null, null, null, Template.Counters, new List<string> { "-Target=TriggerNext" });
 			Assert.IsTrue(await TestSetup.JobService.UpdateBatchAsync(Job1, Job1.Batches[0].Id, ObjectId.GenerateNewId(), JobStepBatchState.Running));
 			Assert.IsNotNull(await TestSetup.JobService.UpdateStepAsync(Job1, Job1.Batches[0].Id, Job1.Batches[0].Steps[0].Id, JobStepState.Completed, JobStepOutcome.Failure));
 			Assert.IsTrue(await TestSetup.JobService.UpdateBatchAsync(Job1, Job1.Batches[0].Id, ObjectId.GenerateNewId(), JobStepBatchState.Complete));
@@ -368,7 +368,7 @@ namespace HordeServerTests
 			Assert.AreEqual(0, Jobs3.Count);
 
 			// Create a job and make it succeed
-			IJob Job2 = await TestSetup.JobService.CreateJobAsync(null, StreamId, NewTemplateRefId1, Template.Id, GraphA, "Hello", 1234, 1233, 999, null, null, "joe", null, null, null, null, true, true, null, null, null, Template.Counters, new List<string> { "-Target=TriggerNext" });
+			IJob Job2 = await TestSetup.JobService.CreateJobAsync(null, Stream, NewTemplateRefId1, Template.Id, GraphA, "Hello", 1234, 1233, 999, null, null, "joe", null, null, null, null, true, true, null, null, null, Template.Counters, new List<string> { "-Target=TriggerNext" });
 			Assert.IsTrue(await TestSetup.JobService.UpdateBatchAsync(Job2, Job2.Batches[0].Id, ObjectId.GenerateNewId(), JobStepBatchState.Running));
 			Assert.IsNotNull(await TestSetup.JobService.UpdateStepAsync(Job2, Job2.Batches[0].Id, Job2.Batches[0].Steps[0].Id, JobStepState.Completed, JobStepOutcome.Success));
 
diff --git a/Engine/Source/Programs/Horde/HordeServerTests/Stubs/Services/PerforceServiceStub.cs b/Engine/Source/Programs/Horde/HordeServerTests/Stubs/Services/PerforceServiceStub.cs
--- a/Engine/Source/Programs/Horde/HordeServerTests/Stubs/Services/PerforceServiceStub.cs
+++ b/Engine/Source/Programs/Horde/HordeServerTests/Stubs/Services/PerforceServiceStub.cs
@@ -34,7 +34,7 @@ namespace HordeServerTests.Stubs.Services
 			StreamChanges.Add(Number, new ChangeDetails(Number, Author, Description, Files.ToList()));
 		}
 
-		public Task<List<ChangeSummary>> GetChangesAsync(string StreamName, int? MinChange, int? MaxChange, int NumResults, string? ImpersonateUser)
+		public Task<List<ChangeSummary>> GetChangesAsync(string ClusterName, string StreamName, int? MinChange, int? MaxChange, int NumResults, string? ImpersonateUser)
 		{
 			List<ChangeSummary> Results = new List<ChangeSummary>();
 
@@ -61,12 +61,12 @@ namespace HordeServerTests.Stubs.Services
 			return Task.FromResult(Results);
 		}
 
-		public Task<PerforceUserInfo?> GetUserInfoAsync(string UserName)
+		public Task<PerforceUserInfo?> GetUserInfoAsync(string ClusterName, string UserName)
 		{
 			return Task.FromResult<PerforceUserInfo?>(new PerforceUserInfo { Name = UserName, Email = $"[REDACTED_EMAIL]" });
 		}
 
-		public Task<List<ChangeDetails>> GetChangeDetailsAsync(string StreamName, IReadOnlyList<int> ChangeNumbers, string? ImpersonateUser)
+		public Task<List<ChangeDetails>> GetChangeDetailsAsync(string ClusterName, string StreamName, IReadOnlyList<int> ChangeNumbers, string? ImpersonateUser)
 		{
 			List<ChangeDetails> Results = new List<ChangeDetails>();
 			foreach (int ChangeNumber in ChangeNumbers)
@@ -76,12 +76,12 @@ namespace HordeServerTests.Stubs.Services
 			return Task.FromResult(Results);
 		}
 
-		public Task<string> CreateTicket(string ImpersonateUser)
+		public Task<string> CreateTicket(string ClusterName, string ImpersonateUser)
 		{
 			return Task.FromResult("bogus-ticket");
 		}
 
-		public Task<int> GetCodeChangeAsync(string StreamName, int Change)
+		public Task<int> GetCodeChangeAsync(string ClusterName, string StreamName, int Change)
 		{
 			int CodeChange = 0;
 
@@ -101,39 +101,39 @@ namespace HordeServerTests.Stubs.Services
 			return Task.FromResult(CodeChange);
 		}
 
-		public Task<int> CreateNewChangeAsync(string StreamName, string Path)
+		public Task<int> CreateNewChangeAsync(string ClusterName, string StreamName, string Path)
 		{
 			ChangeDetails NewChange = new ChangeDetails(Changes[StreamName].First().Key + 1, "", "", new List<string> { Path });
 			Changes[StreamName].Add(NewChange.Number, NewChange);
 			return Task.FromResult(NewChange.Number);
 		}
 
-		public Task<List<FileSummary>> FindFilesAsync(IEnumerable<string> Paths)
+		public Task<List<FileSummary>> FindFilesAsync(string ClusterName, IEnumerable<string> Paths)
 		{
 			throw new NotImplementedException();
 		}
 
-		public Task<byte[]> PrintAsync(string Path)
+		public Task<byte[]> PrintAsync(string ClusterName, string Path)
 		{
 			throw new NotImplementedException();
 		}
 
-		public Task<(int? Change, string Message)> SubmitShelvedChangeAsync(int ShelvedChange, int OriginalChange)
+		public Task<(int? Change, string Message)> SubmitShelvedChangeAsync(string ClusterName, int ShelvedChange, int OriginalChange)
 		{
 			throw new NotImplementedException();
 		}
 
-		public Task<int> DuplicateShelvedChangeAsync(int ShelvedChange)
+		public Task<int> DuplicateShelvedChangeAsync(string ClusterName, int ShelvedChange)
 		{
 			return Task.FromResult(ShelvedChange);
 		}
 
-		public Task DeleteShelvedChangeAsync(int ShelvedChange)
+		public Task DeleteShelvedChangeAsync(string ClusterName, int ShelvedChange)
 		{
 			return Task.CompletedTask;
 		}
 
-		public Task UpdateChangelistDescription(int Change, string Description)
+		public Task UpdateChangelistDescription(string ClusterName, int Change, string Description)
 		{
 			return Task.CompletedTask;
 		}
