Horde: Fix tests.

diff --git a/Engine/Source/Programs/Horde/HordeServerTests/JobServiceTests.cs b/Engine/Source/Programs/Horde/HordeServerTests/JobServiceTests.cs
--- a/Engine/Source/Programs/Horde/HordeServerTests/JobServiceTests.cs
+++ b/Engine/Source/Programs/Horde/HordeServerTests/JobServiceTests.cs
@@ -152,7 +152,8 @@ namespace HordeServerTests
 			Assert.IsNotNull(Project);
 
 			StreamId StreamId = new StreamId("ue5-main");
-			IStream? Stream = await TestSetup.StreamCollection.TryCreateOrReplaceAsync(StreamId, null, "", "", Project!.Id, new StreamConfig { Name = "//UE5/Main" });
+			IStream? Stream = await TestSetup.StreamCollection.GetAsync(StreamId);
+			Stream = await TestSetup.StreamCollection.TryCreateOrReplaceAsync(StreamId, Stream, "", "", Project!.Id, new StreamConfig { Name = "//UE5/Main" });
 
 			ITemplate Template = await TestSetup.TemplateService.CreateTemplateAsync("Test template", null, false, null, null, new List<TemplateCounter>(), new List<string>(), new List<Parameter>());
 			IGraph Graph = await TestSetup.GraphCollection.AddAsync(Template);
