Horde: Allow reading credentials from server config file, so it can bootstrap itself from Perforce config.

diff --git a/Engine/Source/Programs/Horde/HordeServer/Services/Impl/PerforceService.cs b/Engine/Source/Programs/Horde/HordeServer/Services/Impl/PerforceService.cs
--- a/Engine/Source/Programs/Horde/HordeServer/Services/Impl/PerforceService.cs
+++ b/Engine/Source/Programs/Horde/HordeServer/Services/Impl/PerforceService.cs
@@ -58,15 +58,21 @@ namespace HordeServer.Services
 		/// </summary>
 		P4.P4CallBacks.LogMessageDelegate LogBridgeDelegate;
 
+		/// <summary>
+		/// The server settings
+		/// </summary>
+		ServerSettings Settings;
+
 		Dictionary<string, Dictionary<string, CachedTicketInfo>> ClusterTickets = new Dictionary<string, Dictionary<string, CachedTicketInfo>>(StringComparer.OrdinalIgnoreCase);
 
 		/// <summary>
 		/// Constructor
 		/// </summary>
-		public PerforceService(PerforceLoadBalancer LoadBalancer, DatabaseService DatabaseService, ILogger<PerforceService> Logger)
+		public PerforceService(PerforceLoadBalancer LoadBalancer, DatabaseService DatabaseService, IOptions<ServerSettings> Settings, ILogger<PerforceService> Logger)
 		{
 			this.LoadBalancer = LoadBalancer;
 			this.CachedGlobals = new LazyCachedValue<Task<Globals>>(() => DatabaseService.GetGlobalsAsync(), TimeSpan.FromSeconds(30.0));
+			this.Settings = Settings.Value;
 			this.Logger = Logger;
 
 			LogBridgeDelegate = new P4.P4CallBacks.LogMessageDelegate(LogBridgeMessage);
@@ -129,8 +135,13 @@ namespace HordeServer.Services
 			return GetServiceUserConnection(Cluster, Server);
 		}
 
-		static P4.Repository GetServiceUserConnection(PerforceCluster Cluster, IPerforceServer Server)
+		P4.Repository GetServiceUserConnection(PerforceCluster Cluster, IPerforceServer Server)
 		{
+			if (Cluster.Name == PerforceCluster.DefaultName && Settings.P4BridgeServiceUsername != null && Settings.P4BridgeServicePassword != null)
+			{
+				return CreateConnection(Server, Settings.P4BridgeServiceUsername, Settings.P4BridgeServicePassword);
+			}
+
 			PerforceCredentials? Credentials = Cluster.Credentials.FirstOrDefault(x => x.UserName.Equals(Cluster.ServiceAccount, StringComparison.OrdinalIgnoreCase));
 			if (Credentials == null)
 			{
