Horde: Remove storage debug endpoint.

diff --git a/Engine/Source/Programs/Horde/HordeServer/Controllers/DebugController.cs b/Engine/Source/Programs/Horde/HordeServer/Controllers/DebugController.cs
--- a/Engine/Source/Programs/Horde/HordeServer/Controllers/DebugController.cs
+++ b/Engine/Source/Programs/Horde/HordeServer/Controllers/DebugController.cs
@@ -339,11 +339,6 @@ namespace HordeServer.Controllers
 		/// </summary>
 		ILogFileCollection LogFileCollection;
 
-		/// <summary>
-		/// The storage provider
-		/// </summary>
-		IStorageBackend StorageProvider;
-
 		/// <summary>
 		/// Perforce client
 		/// </summary>
@@ -375,11 +370,10 @@ namespace HordeServer.Controllers
 		/// <param name="StreamCollection">Collection of stream documents</param>
 		/// <param name="GraphCollection">The graph collection</param>
 		/// <param name="LogFileCollection">The log file collection</param>
-		/// <param name="StorageProvider">The storage provider</param>
 		/// <param name="Perforce">Perforce client</param>
 		/// <param name="FleetManager">The default fleet manager</param>
 		/// <param name="Settings">Settings</param>
-		public SecureDebugController(AclService AclService, DatabaseService DatabaseService, JobTaskSource JobTaskSource, ISingletonDocument<Globals> GlobalsSingleton, IPoolCollection PoolCollection, IProjectCollection ProjectCollection, IAgentCollection AgentCollection, ISessionCollection SessionCollection, ILeaseCollection LeaseCollection, ITemplateCollection TemplateCollection, IStreamCollection StreamCollection, IGraphCollection GraphCollection, ILogFileCollection LogFileCollection, IStorageBackend StorageProvider, IPerforceService Perforce, IFleetManager FleetManager, IOptionsMonitor<ServerSettings> Settings)
+		public SecureDebugController(AclService AclService, DatabaseService DatabaseService, JobTaskSource JobTaskSource, ISingletonDocument<Globals> GlobalsSingleton, IPoolCollection PoolCollection, IProjectCollection ProjectCollection, IAgentCollection AgentCollection, ISessionCollection SessionCollection, ILeaseCollection LeaseCollection, ITemplateCollection TemplateCollection, IStreamCollection StreamCollection, IGraphCollection GraphCollection, ILogFileCollection LogFileCollection, IPerforceService Perforce, IFleetManager FleetManager, IOptionsMonitor<ServerSettings> Settings)
 		{
 			this.AclService = AclService;
 			this.DatabaseService = DatabaseService;
@@ -394,7 +388,6 @@ namespace HordeServer.Controllers
 			this.StreamCollection = StreamCollection;
 			this.GraphCollection = GraphCollection;
 			this.LogFileCollection = LogFileCollection;
-			this.StorageProvider = StorageProvider;
 			this.Perforce = Perforce;
 			this.FleetManager = FleetManager;
 			this.Settings = Settings;
@@ -499,33 +492,6 @@ namespace HordeServer.Controllers
 			return new GetGraphResponse(Graph).ApplyFilter(Filter);
 		}
 
-		/// <summary>
-		/// Retrieves data from the configured storage provider
-		/// </summary>
-		/// <param name="Path">Filter for properties to return</param>
-		/// <param name="Inline">Whether to show the data inline</param>
-		/// <returns>Data from the given path</returns>
-		[HttpGet]
-		[Route("/api/v1/debug/storage")]
-		[ProducesResponseType(typeof(List<GetTemplateResponse>), 200)]
-		public async Task<ActionResult<object>> GetRawLogData(string Path, bool Inline = true)
-		{
-			if (!await AclService.AuthorizeAsync(AclAction.AdminRead, User))
-			{
-				return Forbid();
-			}
-
-			ReadOnlyMemory<byte>? Data = await StorageProvider.ReadBytesAsync(Path);
-			if (Data == null)
-			{
-				return NotFound();
-			}
-
-			Func<Stream, ActionContext, Task> CopyTask = async (OutputStream, Context) => await OutputStream.WriteAsync(Data.Value);
-			string MimeType = Inline ? "text/plain" : "application/octet-stream";
-			return new CustomFileCallbackResult(System.IO.Path.GetFileName(Path), MimeType, Inline, CopyTask);
-		}
-
 		/// <summary>
 		/// Query all the job templates.
 		/// </summary>
