HordeAgent: Batch update missing blobs

#rnx

diff --git a/Engine/Source/Programs/Horde/HordeAgent/Commands/Execution/ExecuteCommand.cs b/Engine/Source/Programs/Horde/HordeAgent/Commands/Execution/ExecuteCommand.cs
--- a/Engine/Source/Programs/Horde/HordeAgent/Commands/Execution/ExecuteCommand.cs
+++ b/Engine/Source/Programs/Horde/HordeAgent/Commands/Execution/ExecuteCommand.cs
@@ -278,27 +278,26 @@ namespace HordeAgent.Commands
 			FindMissingBlobsResponse MissingBlobs = await StorageClient.FindMissingBlobsAsync(MissingBlobsRequest);
 			if (MissingBlobs.MissingBlobDigests.Count > 0)
 			{
-				foreach (Digest MissingBlobDigest in MissingBlobs.MissingBlobDigests)
+				BatchUpdateBlobsRequest UpdateRequest = new BatchUpdateBlobsRequest();
+				UpdateRequest.InstanceName = InstanceName;
+				UpdateRequest.Requests.AddRange(MissingBlobs.MissingBlobDigests.Select(x => UploadList.Blobs[x.Hash]));
+				Logger.LogInformation($"Updating {UpdateRequest.Requests.Count} blobs");
+
+				BatchUpdateBlobsResponse UpdateResponse = await StorageClient.BatchUpdateBlobsAsync(UpdateRequest);
+
+				bool UpdateFailed = false;
+				foreach (var Response in UpdateResponse.Responses)
 				{
-					BatchUpdateBlobsRequest UpdateRequest = new BatchUpdateBlobsRequest();
-					UpdateRequest.InstanceName = InstanceName;
-					UpdateRequest.Requests.Add(UploadList.Blobs[MissingBlobDigest.Hash]);
-					BatchUpdateBlobsResponse UpdateResponse = await StorageClient.BatchUpdateBlobsAsync(UpdateRequest);
-						
-					bool UpdateFailed = false;
-					foreach (var Response in UpdateResponse.Responses)
+					if (Response.Status.Code != (int) Google.Rpc.Code.Ok)
 					{
-						if (Response.Status.Code != (int) Google.Rpc.Code.Ok)
-						{
-							Logger.LogError($"Upload failed for {Response.Digest}. Code: {Response.Status.Code} Message: {Response.Status.Message}");
-							UpdateFailed = true;
-						}
+						Logger.LogError($"Upload failed for {Response.Digest}. Code: {Response.Status.Code} Message: {Response.Status.Message}");
+						UpdateFailed = true;
 					}
+				}
 
-					if (UpdateFailed)
-					{
-						throw new Exception("Failed updating blobs in CAS service");
-					}
+				if (UpdateFailed)
+				{
+					throw new Exception("Failed updating blobs in CAS service");
 				}
 			}
 
@@ -350,7 +349,7 @@ namespace HordeAgent.Commands
 
 					Logger.LogDebug("Worker: {Worker}", ActionResult.ExecutionMetadata.Worker);
 					Logger.LogDebug("Cached: {CachedResult}", ExecuteResponse.CachedResult);
-							
+
 					Logger.LogDebug("Output directories:");
 					foreach (OutputDirectory Dir in ActionResult.OutputDirectories)
 					{
