HordeAgent: Smaller update blob batchs
Parallel SetupSandboxAsync execution

#rnx

diff --git a/Engine/Source/Programs/Horde/HordeAgent/Commands/Execution/ExecuteCommand.cs b/Engine/Source/Programs/Horde/HordeAgent/Commands/Execution/ExecuteCommand.cs
--- a/Engine/Source/Programs/Horde/HordeAgent/Commands/Execution/ExecuteCommand.cs
+++ b/Engine/Source/Programs/Horde/HordeAgent/Commands/Execution/ExecuteCommand.cs
@@ -280,24 +280,44 @@ namespace HordeAgent.Commands
 			{
 				BatchUpdateBlobsRequest UpdateRequest = new BatchUpdateBlobsRequest();
 				UpdateRequest.InstanceName = InstanceName;
-				UpdateRequest.Requests.AddRange(MissingBlobs.MissingBlobDigests.Select(x => UploadList.Blobs[x.Hash]));
-				Logger.LogInformation($"Updating {UpdateRequest.Requests.Count} blobs");
+				List<BatchUpdateBlobsRequest.Types.Request> PendingUpdates = MissingBlobs.MissingBlobDigests.Select(x => UploadList.Blobs[x.Hash]).ToList();
+				PendingUpdates.SortBy(x => x.Digest.SizeBytes);
+				Logger.LogInformation($"Found {PendingUpdates.Count} missing blobs ({StringUtils.FormatBytesString(PendingUpdates.Sum(x => x.Digest.SizeBytes))})");
 
-				BatchUpdateBlobsResponse UpdateResponse = await StorageClient.BatchUpdateBlobsAsync(UpdateRequest);
-
-				bool UpdateFailed = false;
-				foreach (var Response in UpdateResponse.Responses)
+				int UpdateFailed = 0;
+				while (PendingUpdates.Count > 0)
 				{
-					if (Response.Status.Code != (int) Google.Rpc.Code.Ok)
+					UpdateRequest.Requests.Clear();
+					long BatchSize = 0;
+					while (PendingUpdates.Count > 0)
+					{
+						BatchUpdateBlobsRequest.Types.Request Update = PendingUpdates[0];
+						if (BatchSize > 0 && BatchSize + PendingUpdates[0].Digest.SizeBytes >= 10 * 1024 * 1024)
+						{
+							break;
+						}
+						UpdateRequest.Requests.Add(PendingUpdates[0]);
+						BatchSize += PendingUpdates[0].Digest.SizeBytes;
+						PendingUpdates.RemoveAt(0);
+					}
+
+					Logger.LogInformation($"Updating {UpdateRequest.Requests.Count} blobs ({StringUtils.FormatBytesString(BatchSize)})");
+
+					BatchUpdateBlobsResponse UpdateResponse = await StorageClient.BatchUpdateBlobsAsync(UpdateRequest);
+
+					foreach (var Response in UpdateResponse.Responses)
 					{
-						Logger.LogError($"Upload failed for {Response.Digest}. Code: {Response.Status.Code} Message: {Response.Status.Message}");
-						UpdateFailed = true;
+						if (Response.Status.Code != (int)Google.Rpc.Code.Ok)
+						{
+							Logger.LogError($"Upload failed for {Response.Digest}. Code: {Response.Status.Code} Message: {Response.Status.Message}");
+							UpdateFailed++;
+						}
 					}
 				}
 
-				if (UpdateFailed)
+				if (UpdateFailed > 0)
 				{
-					throw new Exception("Failed updating blobs in CAS service");
+					throw new Exception($"Failed updating {UpdateFailed} blobs in CAS service");
 				}
 			}
 
diff --git a/Engine/Source/Programs/Horde/HordeAgent/Execution/ActionExecutor.cs b/Engine/Source/Programs/Horde/HordeAgent/Execution/ActionExecutor.cs
--- a/Engine/Source/Programs/Horde/HordeAgent/Execution/ActionExecutor.cs
+++ b/Engine/Source/Programs/Horde/HordeAgent/Execution/ActionExecutor.cs
@@ -248,20 +248,20 @@ namespace HordeAgent
 		{
 			DirectoryReference.CreateDirectory(OutputDir);
 
-			foreach (FileNode FileNode in InputDirectory.Files)
+			await ParallelTask.ForEachAsync(InputDirectory.Files, async FileNode =>
 			{
 				ReadOnlyMemory<byte> Data = await Storage.GetBulkDataAsync(InstanceName, FileNode.Digest);
 				FileReference File = FileReference.Combine(OutputDir, FileNode.Name);
 				Logger.LogInformation("Writing {File} (digest: {Digest}, size: {Size})", File, FileNode.Digest.Hash, FileNode.Digest.SizeBytes);
 				await FileReference.WriteAllBytesAsync(File, Data.ToArray());
-			}
+			});
 
-			foreach (DirectoryNode DirectoryNode in InputDirectory.Directories)
+			await ParallelTask.ForEachAsync(InputDirectory.Directories, async DirectoryNode =>
 			{
 				Directory InputSubDirectory = await Storage.GetProtoMessageAsync<Directory>(InstanceName, DirectoryNode.Digest);
 				DirectoryReference OutputSubDirectory = DirectoryReference.Combine(OutputDir, DirectoryNode.Name);
 				await SetupSandboxAsync(InputSubDirectory, OutputSubDirectory);
-			}
+			});
 		}
 
 		/// <summary>
