Horde: Mark Perforce servers as degraded if their health check does not respond.

diff --git a/Engine/Source/Programs/Horde/HordeServer/Commits/Impl/CommitService.cs b/Engine/Source/Programs/Horde/HordeServer/Commits/Impl/CommitService.cs
--- a/Engine/Source/Programs/Horde/HordeServer/Commits/Impl/CommitService.cs
+++ b/Engine/Source/Programs/Horde/HordeServer/Commits/Impl/CommitService.cs
@@ -71,7 +71,6 @@ namespace HordeServer.Commits.Impl
 			this.StateDocument = new SingletonDocument<ReplicationState>(DatabaseService);
 			this.Logger = Logger;
 			this.Ticker = new ElectedTick(DatabaseService, new ObjectId("60f866c49e7268f71803b6e0"), TickReplicationAsync, TimeSpan.FromSeconds(30.0), Logger);
-			TickReplicationAsync(default).Wait();
 		}
 
 		/// <inheritdoc/>
diff --git a/Engine/Source/Programs/Horde/HordeServer/Controllers/PerforceController.cs b/Engine/Source/Programs/Horde/HordeServer/Controllers/PerforceController.cs
--- a/Engine/Source/Programs/Horde/HordeServer/Controllers/PerforceController.cs
+++ b/Engine/Source/Programs/Horde/HordeServer/Controllers/PerforceController.cs
@@ -57,7 +57,7 @@ namespace HordeServer.Controllers
 			List<object> Responses = new List<object>();
 			foreach (IPerforceServer Server in Servers)
 			{
-				Responses.Add(new { Server.ServerAndPort, Server.BaseServerAndPort, Server.Cluster, Server.NumLeases, Server.Status, Server.Detail });
+				Responses.Add(new { Server.ServerAndPort, Server.BaseServerAndPort, Server.Cluster, Server.NumLeases, Server.Status, Server.Detail, Server.LastUpdateTime });
 			}
 			return Responses;
 		}
diff --git a/Engine/Source/Programs/Horde/HordeServer/Services/PerforceLoadBalancer.cs b/Engine/Source/Programs/Horde/HordeServer/Services/PerforceLoadBalancer.cs
--- a/Engine/Source/Programs/Horde/HordeServer/Services/PerforceLoadBalancer.cs
+++ b/Engine/Source/Programs/Horde/HordeServer/Services/PerforceLoadBalancer.cs
@@ -5,6 +5,7 @@ using HordeCommon.Rpc.Tasks;
 using HordeServer.Collections;
 using HordeServer.Models;
 using HordeServer.Utilities;
+using Microsoft.Extensions.Hosting;
 using Microsoft.Extensions.Logging;
 using MongoDB.Bson;
 using MongoDB.Bson.Serialization.Attributes;
@@ -82,12 +83,17 @@ namespace HordeServer.Services
 		/// Error message related to this server
 		/// </summary>
 		public string? Detail { get; }
+
+		/// <summary>
+		/// Last update time for this server
+		/// </summary>
+		public DateTime? LastUpdateTime { get; }
 	}
 
 	/// <summary>
 	/// Load balancer for Perforce edge servers
 	/// </summary>
-	public class PerforceLoadBalancer : ElectedBackgroundService
+	public sealed class PerforceLoadBalancer : IHostedService, IDisposable
 	{
 		/// <summary>
 		/// Information about a resolved Perforce server
@@ -96,23 +102,27 @@ namespace HordeServer.Services
 		{
 			public string ServerAndPort { get; set; } = String.Empty;
 			public string BaseServerAndPort { get; set; } = String.Empty;
+			public string? HealthCheckUrl { get; set; }
 			public string Cluster { get; set; } = String.Empty;
 			public PerforceServerStatus Status { get; set; }
 			public string? Detail { get; set; }
 			public int NumLeases { get; set; }
+			public DateTime? LastUpdateTime { get; set; }
 
 			[BsonConstructor]
 			private PerforceServerEntry()
 			{
 			}
 
-			public PerforceServerEntry(string ServerAndPort, string BaseServerAndPort, string Cluster, PerforceServerStatus Status, string? Detail)
+			public PerforceServerEntry(string ServerAndPort, string BaseServerAndPort, string? HealthCheckUrl, string Cluster, PerforceServerStatus Status, string? Detail, DateTime? LastUpdateTime)
 			{
 				this.ServerAndPort = ServerAndPort;
 				this.BaseServerAndPort = BaseServerAndPort;
+				this.HealthCheckUrl = HealthCheckUrl;
 				this.Cluster = Cluster;
 				this.Status = Status;
 				this.Detail = Detail;
+				this.LastUpdateTime = LastUpdateTime;
 			}
 		}
 
@@ -130,6 +140,8 @@ namespace HordeServer.Services
 		SingletonDocument<PerforceServerList> ServerListSingleton;
 		Random Random = new Random();
 		ILogger Logger;
+		ElectedTick Ticker;
+		Task InternalTickTask = Task.CompletedTask;
 
 		/// <summary>
 		/// Constructor
@@ -138,12 +150,53 @@ namespace HordeServer.Services
 		/// <param name="LeaseCollection"></param>
 		/// <param name="Logger"></param>
 		public PerforceLoadBalancer(DatabaseService DatabaseService, ILeaseCollection LeaseCollection, ILogger<PerforceLoadBalancer> Logger)
-			: base(DatabaseService, new ObjectId("603fb20536e999974fd6ff6b"), TimeSpan.FromMinutes(1.0), TimeSpan.FromMinutes(2.0), Logger)
 		{
 			this.DatabaseService = DatabaseService;
 			this.LeaseCollection = LeaseCollection;
 			this.ServerListSingleton = new SingletonDocument<PerforceServerList>(DatabaseService);
 			this.Logger = Logger;
+			this.Ticker = new ElectedTick(DatabaseService, new ObjectId("603fb20536e999974fd6ff6b"), TickAsync, TimeSpan.FromMinutes(1.0), Logger);
+		}
+
+		/// <inheritdoc/>
+		public void Dispose()
+		{
+			Ticker.Dispose();
+		}
+
+		/// <inheritdoc/>
+		public Task StartAsync(CancellationToken cancellationToken)
+		{
+			Ticker.Start();
+			return Task.CompletedTask;
+		}
+
+		/// <inheritdoc/>
+		public async Task StopAsync(CancellationToken cancellationToken)
+		{
+			await Ticker.StopAsync();
+			await InternalTickTask;
+		}
+
+		/// <summary>
+		/// Get the current server list
+		/// </summary>
+		/// <returns></returns>
+		async Task<PerforceServerList> GetServerListAsync()
+		{
+			PerforceServerList ServerList = await ServerListSingleton.GetAsync();
+
+			DateTime MinLastUpdateTime = DateTime.UtcNow - TimeSpan.FromMinutes(2.5);
+			foreach (PerforceServerEntry Server in ServerList.Servers)
+			{
+				if (Server.Status == PerforceServerStatus.Healthy && Server.LastUpdateTime != null && Server.LastUpdateTime.Value < MinLastUpdateTime)
+				{
+					Server.Status = PerforceServerStatus.Degraded;
+					Server.Detail = "Server has not responded to health check";
+				}
+			}
+
+			return ServerList;
 		}
 
 		/// <summary>
@@ -152,7 +205,7 @@ namespace HordeServer.Services
 		/// <returns></returns>
 		public async Task<List<IPerforceServer>> GetServersAsync()
 		{
-			PerforceServerList ServerList = await ServerListSingleton.GetAsync();
+			PerforceServerList ServerList = await GetServerListAsync();
 			return ServerList.Servers.ConvertAll<IPerforceServer>(x => x);
 		}
 
@@ -162,7 +215,7 @@ namespace HordeServer.Services
 		/// <returns>The server to use. Null if there is no healthy server available.</returns>
 		public async Task<IPerforceServer?> GetServer(string Cluster)
 		{
-			PerforceServerList ServerList = await ServerListSingleton.GetAsync();
+			PerforceServerList ServerList = await GetServerListAsync();
 
 			List<PerforceServerEntry> Candidates = ServerList.Servers.Where(x => x.Cluster == Cluster && x.Status >= PerforceServerStatus.Healthy).ToList();
 			if(Candidates.Count == 0)
@@ -217,7 +270,7 @@ namespace HordeServer.Services
 			}
 
 			// Find all the matching servers
-			PerforceServerList ServerList = await ServerListSingleton.GetAsync();
+			PerforceServerList ServerList = await GetServerListAsync();
 			List<PerforceServerEntry> Candidates = ServerList.Servers.Where(x => x.Cluster == Cluster.Name && x.Status == PerforceServerStatus.Healthy && ValidServers.Contains(x.BaseServerAndPort)).ToList();
 
 			if (Candidates.Count == 0)
@@ -250,7 +303,30 @@ namespace HordeServer.Services
 		}
 
 		/// <inheritdoc/>
-		protected override async Task<DateTime> TickLeaderAsync(CancellationToken StoppingToken)
+		Task TickAsync(CancellationToken StoppingToken)
+		{
+			if (InternalTickTask.IsCompleted)
+			{
+				InternalTickTask = Task.Run(() => TickInternalGuardedAsync());
+			}
+			return Task.CompletedTask;
+		}
+
+		/// <inheritdoc/>
+		async Task TickInternalGuardedAsync()
+		{
+			try
+			{
+				await TickInternalAsync();
+			}
+			catch (Exception Ex)
+			{
+				Logger.LogError(Ex, "Error updating server state");
+			}
+		}
+
+		/// <inheritdoc/>
+		async Task TickInternalAsync()
 		{
 			Globals Globals = await DatabaseService.GetGlobalsAsync();
 
@@ -274,15 +350,32 @@ namespace HordeServer.Services
 			}
 
 			// Update the number of leases for each entry
-			List<PerforceServerEntry> NewServersList = NewServers.OrderBy(x => x.Cluster).ThenBy(x => x.BaseServerAndPort).ThenBy(x => x.ServerAndPort).ToList();
-			await UpdateLeaseCounts(NewServersList);
+			List<PerforceServerEntry> NewEntries = NewServers.OrderBy(x => x.Cluster).ThenBy(x => x.BaseServerAndPort).ThenBy(x => x.ServerAndPort).ToList();
+			await UpdateLeaseCounts(NewEntries);
+			PerforceServerList List = await ServerListSingleton.UpdateAsync(List => MergeServerList(List, NewEntries));
 
-			// Mark any other servers as deleted
-			PerforceServerList NewServerList = await ServerListSingleton.GetAsync();
-			NewServerList.Servers = NewServers.OrderBy(x => x.Cluster).ThenBy(x => x.BaseServerAndPort).ThenBy(x => x.ServerAndPort).ToList();
-			await ServerListSingleton.TryUpdateAsync(NewServerList);
+			// Now update the health of each entry
+			List<Task> Tasks = new List<Task>();
+			foreach (PerforceServerEntry Entry in List.Servers)
+			{
+				Tasks.Add(Task.Run(() => UpdateHealthAsync(Entry)));
+			}
+			await Task.WhenAll(Tasks);
+		}
 
-			return DateTime.UtcNow.AddMinutes(1.0);
+		static void MergeServerList(PerforceServerList ServerList, List<PerforceServerEntry> NewEntries)
+		{
+			Dictionary<string, PerforceServerEntry> ExistingEntries = ServerList.Servers.ToDictionary(x => x.ServerAndPort, x => x, StringComparer.OrdinalIgnoreCase);
+			foreach (PerforceServerEntry NewEntry in NewEntries)
+			{
+				PerforceServerEntry? ExistingEntry;
+				if (ExistingEntries.TryGetValue(NewEntry.ServerAndPort, out ExistingEntry))
+				{
+					NewEntry.Status = ExistingEntry.Status;
+					NewEntry.Detail = ExistingEntry.Detail;
+					NewEntry.LastUpdateTime = ExistingEntry.LastUpdateTime;
+				}
+			}
 		}
 
 		async Task UpdateLeaseCounts(IEnumerable<PerforceServerEntry> NewServerEntries)
@@ -369,12 +462,15 @@ namespace HordeServer.Services
 				await ResolveServersAsync(InitialHostName, HostNames);
 			}
 
-			List<Task> Tasks = new List<Task>();
 			foreach (string HostName in HostNames)
 			{
-				Tasks.Add(Task.Run(() => UpdateHostAsync(Cluster, Server, HostName, Port, NewServers)));
+				string? HealthCheckUrl = null;
+				if (Server.HealthCheck)
+				{
+					HealthCheckUrl = $"http://{HostName}:5000/healthcheck";
+				}
+				NewServers.Add(new PerforceServerEntry($"{HostName}:{Port}", Server.ServerAndPort, HealthCheckUrl, Cluster.Name, PerforceServerStatus.Unknown, "", DateTime.UtcNow));
 			}
-			await Task.WhenAll(Tasks);
 		}
 
 		async Task ResolveServersAsync(string HostName, List<string> HostNames)
@@ -394,15 +490,22 @@ namespace HordeServer.Services
 			}
 		}
 
-		static async Task UpdateHostAsync(PerforceCluster Cluster, PerforceServer Server, string HostName, int Port, ConcurrentBag<PerforceServerEntry> NewServers)
+		/// <summary>
+		/// Updates 
+		/// </summary>
+		/// <param name="Entry"></param>
+		/// <returns></returns>
+		async Task UpdateHealthAsync(PerforceServerEntry Entry)
 		{
+			DateTime? UpdateTime = null;
 			string Detail = "Health check disabled";
 
 			// Get the health of the server
 			PerforceServerStatus Health = PerforceServerStatus.Healthy;
-			if (Server.HealthCheck)
+			if (Entry.HealthCheckUrl != null)
 			{
-				Uri HealthCheckUrl = new Uri($"http://{HostName}:5000/healthcheck");
+				UpdateTime = DateTime.UtcNow;
+				Uri HealthCheckUrl = new Uri(Entry.HealthCheckUrl);
 				try
 				{
 					(Health, Detail) = await GetServerHealthAsync(HealthCheckUrl);
@@ -414,7 +517,20 @@ namespace HordeServer.Services
 			}
 
 			// Update the server record
-			NewServers.Add(new PerforceServerEntry($"{HostName}:{Port}", Server.ServerAndPort, Cluster.Name, Health, Detail));
+			if (Health != Entry.Status || Detail != Entry.Detail || UpdateTime != Entry.LastUpdateTime)
+			{
+				await ServerListSingleton.UpdateAsync(x => UpdateHealth(Entry, Health, Detail, UpdateTime));
+			}
+		}
+
+		static void UpdateHealth(PerforceServerEntry Entry, PerforceServerStatus Status, string Detail, DateTime? UpdateTime)
+		{
+			if (Entry.LastUpdateTime == null || UpdateTime == null || Entry.LastUpdateTime.Value < UpdateTime)
+			{
+				Entry.Status = Status;
+				Entry.Detail = Detail;
+				Entry.LastUpdateTime = UpdateTime;
+			}
 		}
 
 		static async Task<(PerforceServerStatus, string)> GetServerHealthAsync(Uri HealthCheckUrl)
diff --git a/Engine/Source/Programs/Horde/HordeServer/Utilities/SingletonDocument.cs b/Engine/Source/Programs/Horde/HordeServer/Utilities/SingletonDocument.cs
--- a/Engine/Source/Programs/Horde/HordeServer/Utilities/SingletonDocument.cs
+++ b/Engine/Source/Programs/Horde/HordeServer/Utilities/SingletonDocument.cs
@@ -123,7 +123,7 @@ namespace HordeServer.Utilities
 		/// <param name="Singleton"></param>
 		/// <param name="UpdateAction"></param>
 		/// <returns></returns>
-		public static async Task UpdateAsync<T>(this ISingletonDocument<T> Singleton, Action<T> UpdateAction) where T : SingletonBase, new()
+		public static async Task<T> UpdateAsync<T>(this ISingletonDocument<T> Singleton, Action<T> UpdateAction) where T : SingletonBase, new()
 		{
 			for (; ; )
 			{
@@ -132,7 +132,7 @@ namespace HordeServer.Utilities
 
 				if (await Singleton.TryUpdateAsync(Value))
 				{
-					break;
+					return Value;
 				}
 			}
 		}
