Horde: Fix server list not updating correctly.

diff --git a/Engine/Source/Programs/Horde/HordeServer/Services/PerforceLoadBalancer.cs b/Engine/Source/Programs/Horde/HordeServer/Services/PerforceLoadBalancer.cs
--- a/Engine/Source/Programs/Horde/HordeServer/Services/PerforceLoadBalancer.cs
+++ b/Engine/Source/Programs/Horde/HordeServer/Services/PerforceLoadBalancer.cs
@@ -376,6 +376,7 @@ namespace HordeServer.Services
 					NewEntry.LastUpdateTime = ExistingEntry.LastUpdateTime;
 				}
 			}
+			ServerList.Servers = NewEntries;
 		}
 
 		async Task UpdateLeaseCounts(IEnumerable<PerforceServerEntry> NewServerEntries)
