Horde: optimize report loading
#rnx

diff --git a/Engine/Source/Programs/Horde/HordeDashboard/src/backend/JobDetails.ts b/Engine/Source/Programs/Horde/HordeDashboard/src/backend/JobDetails.ts
--- a/Engine/Source/Programs/Horde/HordeDashboard/src/backend/JobDetails.ts
+++ b/Engine/Source/Programs/Horde/HordeDashboard/src/backend/JobDetails.ts
@@ -73,6 +73,10 @@ export class JobDetails {
                 }
             }
 
+            if (!this.isLogView) {
+                requests.push(this.queryReports());
+            }
+
             if (logId && logId !== this.logId) {
                 requests.push(this.getLogEvents(logId));
             }
@@ -592,13 +596,13 @@ export class JobDetails {
 
             this.issues = [];
 
-            const issueRequests = [ backend.getIssues({ jobId: this.id, stepId: stepId, label: this.labelIdx, count: 50, resolved: false }),
+            const issueRequests = [backend.getIssues({ jobId: this.id, stepId: stepId, label: this.labelIdx, count: 50, resolved: false }),
             backend.getIssues({ jobId: this.id, stepId: stepId, label: this.labelIdx, count: 50, resolved: true })];
-            
+
             Promise.all(issueRequests).then(async (issueResults) => {
 
-                let issues:IssueData[] = [];
-                issueResults.forEach( r => issues.push(...r));
+                let issues: IssueData[] = [];
+                issueResults.forEach(r => issues.push(...r));
 
                 let nstepId = this.stepId;
                 let nlabelIdx = this.labelIdx;
@@ -631,10 +635,10 @@ export class JobDetails {
                 resolve();
             });
         })
-    }    
+    }
 
     getReportData(placement: ReportPlacement, stepId?: string): string | undefined {
-                
+
         if (stepId) {
 
             const step = this.stepById(stepId);
@@ -662,15 +666,24 @@ export class JobDetails {
         }
     }
 
-    private async queryReports():Promise<void> {
+    private async queryReports(): Promise<void> {
 
         const artifacts: string[] = [];
 
-        this.jobdata?.reports?.forEach(r => artifacts.push(r.artifactId));
+        this.jobdata?.reports?.forEach(r => {
+            if (!this.stepId) {
+                artifacts.push(r.artifactId)
+            }
+
+        });
         this.jobdata?.batches?.forEach(b => {
-            b.steps.forEach(s => s.reports?.forEach( r => artifacts.push(r.artifactId)));
+            b.steps.forEach(s => s.reports?.forEach(r => {
+                if (s.id === this.stepId) {
+                    artifacts.push(r.artifactId);
+                }
+            }));
         });
-        
+
 
         for (let i = 0; i < artifacts.length; i++) {
 
@@ -679,8 +692,8 @@ export class JobDetails {
             if (!this.reportData.has(artifactId)) {
 
                 const r = await backend.getArtifactDataById(artifactId) as unknown as string;
-                
-                this.reportData.set(artifactId, r);                                
+
+                this.reportData.set(artifactId, r);
             }
         }
     }
@@ -779,7 +792,7 @@ export class JobDetails {
 
             requests.push(backend.getJobArtifacts(this.id!));
 
-            if (!this.isLogView) {                
+            if (!this.isLogView) {
                 requests.push(backend.getJobTestData(this.id!));
                 requests.push(this.getIssues());
                 requests.push(this.queryReports());
@@ -1005,9 +1018,9 @@ export const getDetailStyle = (state: JobStepState, outcome: JobStepOutcome): De
 
 };
 
-export const getBatchSummaryMarkdown = (jobDetails:JobDetails, batchId: string): string => {
+export const getBatchSummaryMarkdown = (jobDetails: JobDetails, batchId: string): string => {
 
-	const batch = jobDetails.batchById(batchId);
+    const batch = jobDetails.batchById(batchId);
 
     if (!batch) {
         return "";
@@ -1018,30 +1031,30 @@ export const getBatchSummaryMarkdown = (jobDetails:JobDetails, batchId: string):
     const agentType = group?.agentType ?? "";
     const agentPool = jobDetails.stream?.agentTypes[agentType!]?.pool ?? "";
 
-	let summaryText = "";
-	if (batch && batch.agentId) {
-
-		if (batch.state === JobStepBatchState.Running || batch.state === JobStepBatchState.Complete) {
-			let duration = "";
-			if (batch.startTime) {
-				duration = getBatchInitElapsed(batch);
-				summaryText = `Batch started ${getNiceTime(batch.startTime)} and ${batch.finishTime ? "ran" : "is running"} on [${batch.agentId}](?batch=${batchId}&agentId=${encodeURIComponent(batch.agentId)}) with a setup time of ${duration}.`;
-			}
-		}
-	} 
-
-	if (!summaryText) {
-		summaryText = getBatchText({ batch: batch, agentId: batch?.agentId, agentType: agentType, agentPool: agentPool }) ?? "";		
-		if (!summaryText) {
-			summaryText = batch?.agentId ?? (batch?.state ?? "Batch is unassigned");			
-		}
-	}
+    let summaryText = "";
+    if (batch && batch.agentId) {
+
+        if (batch.state === JobStepBatchState.Running || batch.state === JobStepBatchState.Complete) {
+            let duration = "";
+            if (batch.startTime) {
+                duration = getBatchInitElapsed(batch);
+                summaryText = `Batch started ${getNiceTime(batch.startTime)} and ${batch.finishTime ? "ran" : "is running"} on [${batch.agentId}](?batch=${batchId}&agentId=${encodeURIComponent(batch.agentId)}) with a setup time of ${duration}.`;
+            }
+        }
+    }
+
+    if (!summaryText) {
+        summaryText = getBatchText({ batch: batch, agentId: batch?.agentId, agentType: agentType, agentPool: agentPool }) ?? "";
+        if (!summaryText) {
+            summaryText = batch?.agentId ?? (batch?.state ?? "Batch is unassigned");
+        }
+    }
 
     return summaryText;
 
 }
 
-export const getStepSummaryMarkdown = (jobDetails:JobDetails, stepId: string): string => {
+export const getStepSummaryMarkdown = (jobDetails: JobDetails, stepId: string): string => {
 
     const step = jobDetails.stepById(stepId)!;
     const batch = jobDetails.batchByStepId(stepId);
