Horde: Change to information level log output by default.

diff --git a/Engine/Source/Programs/Horde/HordeServer/appsettings.json b/Engine/Source/Programs/Horde/HordeServer/appsettings.json
--- a/Engine/Source/Programs/Horde/HordeServer/appsettings.json
+++ b/Engine/Source/Programs/Horde/HordeServer/appsettings.json
@@ -12,7 +12,7 @@
 	{
 		"MinimumLevel": 
 		{
-			"Default": "Debug",
+			"Default": "Information",
 			"Override": 
 			{
 				"MongoDB": "Warning", // For bundled MongoDB output
