Horde: Disable logging of UpdateSession/QueryServerStateV2 calls.

diff --git a/Engine/Source/Programs/Horde/HordeServer/Program.cs b/Engine/Source/Programs/Horde/HordeServer/Program.cs
--- a/Engine/Source/Programs/Horde/HordeServer/Program.cs
+++ b/Engine/Source/Programs/Horde/HordeServer/Program.cs
@@ -25,6 +25,7 @@ using MongoDB.Driver;
 using Serilog;
 using Serilog.Configuration;
 using Serilog.Events;
+using Serilog.Filters;
 using Serilog.Formatting.Json;
 using Serilog.Sinks.SystemConsole.Themes;
 
@@ -52,6 +53,36 @@ namespace HordeServer
 				return SinkConfig.Console(outputTemplate: "[{Timestamp:HH:mm:ss} {Level:w3}] {Indent}{Message:l}{NewLine}{Exception}", theme: Theme, restrictedToMinimumLevel: LogEventLevel.Debug);
 			}
 		}
+
+		public static LoggerConfiguration WithHordeConfig(this LoggerConfiguration Configuration, ServerSettings Settings)
+		{
+			if (!Settings.LogSessionRequests)
+			{
+				Configuration = Configuration.Filter.ByExcluding(IsSessionLogEvent);
+			}
+			return Configuration;
+		}
+
+		static bool IsSessionLogEvent(LogEvent Event)
+		{
+			if (Event.Level <= LogEventLevel.Information)
+			{
+				LogEventPropertyValue? Value;
+				if (Event.Properties.TryGetValue("RequestPath", out Value))
+				{
+					string RequestPath = Value.ToString();
+					if (RequestPath.Equals("/Horde.HordeRpc/QueryServerStateV2", StringComparison.OrdinalIgnoreCase))
+					{
+						return false;
+					}
+					if (RequestPath.Equals("/Horde.HordeRpc/UpdateSession", StringComparison.OrdinalIgnoreCase))
+					{
+						return false;
+					}
+				}
+			}
+			return true;
+		}
 	}
 
 	class Program
@@ -101,6 +132,7 @@ namespace HordeServer
 			}
 
 			Serilog.Log.Logger = new LoggerConfiguration()
+				.WithHordeConfig(HordeSettings)
 				.Enrich.FromLogContext()
 				.WriteTo.Console(HordeSettings)
 				.WriteTo.File(Path.Combine(LogDir.FullName, "Log.txt"), outputTemplate: "[{Timestamp:HH:mm:ss} {Level:w3}] {Indent}{Message:l}{NewLine}{Exception} [{SourceContext}]", rollingInterval: RollingInterval.Day, rollOnFileSizeLimit: true, fileSizeLimitBytes: 20 * 1024 * 1024, retainedFileCountLimit: 10)
diff --git a/Engine/Source/Programs/Horde/HordeServer/ServerSettings.cs b/Engine/Source/Programs/Horde/HordeServer/ServerSettings.cs
--- a/Engine/Source/Programs/Horde/HordeServer/ServerSettings.cs
+++ b/Engine/Source/Programs/Horde/HordeServer/ServerSettings.cs
@@ -318,6 +318,11 @@ namespace HordeServer
 		/// </summary>
 		public bool LogJsonToStdOut { get; set; } = false;
 
+		/// <summary>
+		/// Whether to log requests to the UpdateSession and QueryServerState RPC endpoints
+		/// </summary>
+		public bool LogSessionRequests { get; set; } = false;
+
 		/// <summary>
 		/// Which fleet manager service to use
 		/// </summary>
diff --git a/Engine/Source/Programs/Horde/HordeServer/Services/RpcService.cs b/Engine/Source/Programs/Horde/HordeServer/Services/RpcService.cs
--- a/Engine/Source/Programs/Horde/HordeServer/Services/RpcService.cs
+++ b/Engine/Source/Programs/Horde/HordeServer/Services/RpcService.cs
@@ -216,7 +216,7 @@ namespace HordeServer.Services
 			if (await Reader.MoveNext())
 			{
 				QueryServerStateRequest Request = Reader.Current;
-				Logger.LogInformation("Start server query for client {Name}", Request.Name);
+				Logger.LogDebug("Start server query for client {Name}", Request.Name);
 
 				try
 				{
