Horde: Fix inverted check for log level.

diff --git a/Engine/Source/Programs/Horde/HordeServer/Program.cs b/Engine/Source/Programs/Horde/HordeServer/Program.cs
--- a/Engine/Source/Programs/Horde/HordeServer/Program.cs
+++ b/Engine/Source/Programs/Horde/HordeServer/Program.cs
@@ -73,15 +73,15 @@ namespace HordeServer
 					string RequestPath = Value.ToString();
 					if (RequestPath.Equals("/Horde.HordeRpc/QueryServerStateV2", StringComparison.OrdinalIgnoreCase))
 					{
-						return false;
+						return true;
 					}
 					if (RequestPath.Equals("/Horde.HordeRpc/UpdateSession", StringComparison.OrdinalIgnoreCase))
 					{
-						return false;
+						return true;
 					}
 				}
 			}
-			return true;
+			return false;
 		}
 	}
 
