Horde: Fix Perforce server state not being updated.

diff --git a/Engine/Source/Programs/Horde/HordeServer/Services/PerforceLoadBalancer.cs b/Engine/Source/Programs/Horde/HordeServer/Services/PerforceLoadBalancer.cs
--- a/Engine/Source/Programs/Horde/HordeServer/Services/PerforceLoadBalancer.cs
+++ b/Engine/Source/Programs/Horde/HordeServer/Services/PerforceLoadBalancer.cs
@@ -520,17 +520,21 @@ namespace HordeServer.Services
 			// Update the server record
 			if (Health != Entry.Status || Detail != Entry.Detail || UpdateTime != Entry.LastUpdateTime)
 			{
-				await ServerListSingleton.UpdateAsync(x => UpdateHealth(Entry, Health, Detail, UpdateTime));
+				await ServerListSingleton.UpdateAsync(x => UpdateHealth(x, Entry.ServerAndPort, Health, Detail, UpdateTime));
 			}
 		}
 
-		static void UpdateHealth(PerforceServerEntry Entry, PerforceServerStatus Status, string Detail, DateTime? UpdateTime)
+		static void UpdateHealth(PerforceServerList ServerList, string ServerAndPort, PerforceServerStatus Status, string Detail, DateTime? UpdateTime)
 		{
-			if (Entry.LastUpdateTime == null || UpdateTime == null || Entry.LastUpdateTime.Value < UpdateTime)
+			PerforceServerEntry? Entry = ServerList.Servers.FirstOrDefault(x => x.ServerAndPort.Equals(ServerAndPort, StringComparison.Ordinal));
+			if (Entry != null)
 			{
-				Entry.Status = Status;
-				Entry.Detail = Detail;
-				Entry.LastUpdateTime = UpdateTime;
+				if (Entry.LastUpdateTime == null || UpdateTime == null || Entry.LastUpdateTime.Value < UpdateTime.Value)
+				{
+					Entry.Status = Status;
+					Entry.Detail = Detail;
+					Entry.LastUpdateTime = UpdateTime;
+				}
 			}
 		}
 
