Horde: Fix agent trying to use root data directory for storage. Since log files go here (and are kept open), any attempt to conform this directory will fail.

diff --git a/Engine/Source/Programs/Horde/HordeAgent/AgentSettings.cs b/Engine/Source/Programs/Horde/HordeAgent/AgentSettings.cs
--- a/Engine/Source/Programs/Horde/HordeAgent/AgentSettings.cs
+++ b/Engine/Source/Programs/Horde/HordeAgent/AgentSettings.cs
@@ -1,5 +1,6 @@
 // Copyright Epic Games, Inc. All Rights Reserved.
 
+using EpicGames.Core;
 using HordeCommon.Rpc;
 using System;
 using System.Collections.Generic;
@@ -179,7 +180,7 @@ namespace HordeAgent
 		/// <summary>
 		/// Working directory
 		/// </summary>
-		public string? WorkingDir { get; set; } = Program.DataDir.FullName;
+		public string? WorkingDir { get; set; } = DirectoryReference.Combine(Program.DataDir, "Data").FullName;
 
 		/// <summary>
 		/// Whether to mount the specified list of network shares
