Horde: Filter out additional session logging.

diff --git a/Engine/Source/Programs/Horde/HordeServer/Program.cs b/Engine/Source/Programs/Horde/HordeServer/Program.cs
--- a/Engine/Source/Programs/Horde/HordeServer/Program.cs
+++ b/Engine/Source/Programs/Horde/HordeServer/Program.cs
@@ -67,17 +67,32 @@ namespace HordeServer
 		{
 			if (Event.Level <= LogEventLevel.Information)
 			{
-				LogEventPropertyValue? Value;
-				if (Event.Properties.TryGetValue("RequestPath", out Value))
+				LogEventPropertyValue? PropertyValue;
+				if (Event.Properties.TryGetValue("RequestPath", out PropertyValue))
 				{
-					string RequestPath = Value.ToString();
-					if (RequestPath.Equals("/Horde.HordeRpc/QueryServerStateV2", StringComparison.OrdinalIgnoreCase))
+					ScalarValue? ScalarPropertyValue = PropertyValue as ScalarValue;
+					if (ScalarPropertyValue != null)
 					{
-						return true;
-					}
-					if (RequestPath.Equals("/Horde.HordeRpc/UpdateSession", StringComparison.OrdinalIgnoreCase))
-					{
-						return true;
+						string? RequestPath = ScalarPropertyValue.Value as string;
+						if (RequestPath != null)
+						{
+							if (RequestPath.Equals("/Horde.HordeRpc/QueryServerStateV2", StringComparison.OrdinalIgnoreCase))
+							{
+								return true;
+							}
+							if (RequestPath.Equals("/Horde.HordeRpc/UpdateSession", StringComparison.OrdinalIgnoreCase))
+							{
+								return true;
+							}
+							if (RequestPath.Equals("/Horde.HordeRpc/CreateEvents", StringComparison.OrdinalIgnoreCase))
+							{
+								return true;
+							}
+							if (RequestPath.Equals("/Horde.HordeRpc/WriteOutput", StringComparison.OrdinalIgnoreCase))
+							{
+								return true;
+							}
+						}
 					}
 				}
 			}
diff --git a/Engine/Source/Programs/Horde/HordeServer/appsettings.json b/Engine/Source/Programs/Horde/HordeServer/appsettings.json
--- a/Engine/Source/Programs/Horde/HordeServer/appsettings.json
+++ b/Engine/Source/Programs/Horde/HordeServer/appsettings.json
@@ -20,10 +20,7 @@
 				"Microsoft": "Warning",
 				"Microsoft.AspNetCore": "Warning",
 				"Microsoft.Hosting": "Information",
-				"HordeServer.Authentication.OktaHandler": "Warning",
-				"HordeServer.Authentication.OktaViaJwtHandler": "Warning",
-				"HordeServer.Authentication.HordeJwtBearerHandler": "Warning",
-				"HordeServer.Authentication.ServiceAccountAuthHandler": "Warning",
+				"HordeServer.Authentication": "Warning",
 				"System.Net.Http.HttpClient": "Warning",
 				"Grpc": "Warning"
 			}
