Horde: Make the commit index unique.

diff --git a/Engine/Source/Programs/Horde/HordeServer/Commits/Impl/CommitCollection.cs b/Engine/Source/Programs/Horde/HordeServer/Commits/Impl/CommitCollection.cs
--- a/Engine/Source/Programs/Horde/HordeServer/Commits/Impl/CommitCollection.cs
+++ b/Engine/Source/Programs/Horde/HordeServer/Commits/Impl/CommitCollection.cs
@@ -87,7 +87,7 @@ namespace HordeServer.Commits.Impl
 
 			if (!DatabaseService.ReadOnlyMode)
 			{
-				Commits.Indexes.CreateOne(new CreateIndexModel<Commit>(Builders<Commit>.IndexKeys.Ascending(x => x.StreamId).Descending(x => x.Change)));
+				Commits.Indexes.CreateOne(new CreateIndexModel<Commit>(Builders<Commit>.IndexKeys.Ascending(x => x.StreamId).Descending(x => x.Change), new CreateIndexOptions { Unique = true }));
 			}
 		}
 
