Horde: Re-enable Slack notifications for config update failures.

diff --git a/Engine/Source/Programs/Horde/HordeServer/Notifications/INotificationService.cs b/Engine/Source/Programs/Horde/HordeServer/Notifications/INotificationService.cs
--- a/Engine/Source/Programs/Horde/HordeServer/Notifications/INotificationService.cs
+++ b/Engine/Source/Programs/Horde/HordeServer/Notifications/INotificationService.cs
@@ -56,18 +56,12 @@ namespace HordeServer.Notifications
 		/// <summary>
 		/// Notify slack channel about a stream update failure
 		/// </summary>
-		/// <param name="Stream">The stream</param>
 		/// <param name="ErrorMessage">Error message passed back</param>
+		/// <param name="FileName"></param>
 		/// <param name="Change">Latest change number of the file</param>
 		/// <param name="Author">Author of the change, if found from p4 service</param>
 		/// <param name="Description">Description of the change, if found from p4 service</param>
-		void NotifyUpdateStreamFailure(IStream Stream, string ErrorMessage, int Change, string? Author, string? Description);
-
-		/// <summary>
-		/// Notify slack channel about a stream update failure in relation to a file
-		/// </summary>
-		/// <param name="File">The file</param>
-		void NotifyUpdateStreamFailure(FileSummary File);
+		void NotifyConfigUpdateFailure(string ErrorMessage, string FileName, int? Change = null, string? Author = null, string? Description = null);
 
 		/// <summary>
 		/// Sends a notification to a user regarding a build health issue.
diff --git a/Engine/Source/Programs/Horde/HordeServer/Notifications/INotificationSink.cs b/Engine/Source/Programs/Horde/HordeServer/Notifications/INotificationSink.cs
--- a/Engine/Source/Programs/Horde/HordeServer/Notifications/INotificationSink.cs
+++ b/Engine/Source/Programs/Horde/HordeServer/Notifications/INotificationSink.cs
@@ -73,19 +73,12 @@ namespace HordeServer.Notifications
 		/// <summary>
 		/// Notification that a stream has failed to update
 		/// </summary>
-		/// <param name="Stream"></param>
 		/// <param name="ErrorMessage"></param>
+		/// <param name="FileName"></param>
 		/// <param name="Change"></param>
 		/// <param name="Author"></param>
 		/// <param name="Description"></param>
 		/// <returns></returns>
-		Task NotifyStreamUpdateFailedAsync(IStream Stream, string ErrorMessage, int Change, IUser? Author = null, string? Description = null);
-
-		/// <summary>
-		/// Notification that a stream has failed to update
-		/// </summary>
-		/// <param name="File"></param>
-		/// <returns></returns>
-		Task NotifyStreamUpdateFailedAsync(FileSummary File);
+		Task NotifyConfigUpdateFailureAsync(string ErrorMessage, string FileName, int? Change = null, IUser? Author = null, string? Description = null);
 	}
 }
diff --git a/Engine/Source/Programs/Horde/HordeServer/Notifications/Impl/NotificationService.cs b/Engine/Source/Programs/Horde/HordeServer/Notifications/Impl/NotificationService.cs
--- a/Engine/Source/Programs/Horde/HordeServer/Notifications/Impl/NotificationService.cs
+++ b/Engine/Source/Programs/Horde/HordeServer/Notifications/Impl/NotificationService.cs
@@ -218,20 +218,15 @@ namespace HordeServer.Notifications.Impl
 		}
 
 		/// <inheritdoc/>
-		public void NotifyUpdateStreamFailure(IStream Stream, string ErrorMessage, int Change, string? Author = null, string? Description = null)
+		public void NotifyConfigUpdateFailure(string ErrorMessage, string FileName, int? Change = null, string? Author = null, string? Description = null)
 		{
-			EnqueueTask(() => NotifyUpdateStreamFailureAsync(Stream, ErrorMessage, Change, Author, Description));
+			EnqueueTask(() => NotifyConfigUpdateFailureAsync(ErrorMessage, FileName, Change, Author, Description));
 		}
 
-		async Task NotifyUpdateStreamFailureAsync(IStream Stream, string ErrorMessage, int Change, string? Author = null, string? Description = null)
+		async Task NotifyConfigUpdateFailureAsync(string ErrorMessage, string FileName, int? Change = null, string? Author = null, string? Description = null)
 		{
 			IUser? User = (Author != null) ? await UserCollection.FindUserByLoginAsync(Author) : null;
-			EnqueueTasks(Sink => Sink.NotifyStreamUpdateFailedAsync(Stream, ErrorMessage, Change, User, Description));
-		}
-		/// <inheritdoc/>
-		public void NotifyUpdateStreamFailure(FileSummary File)
-		{
-			EnqueueTasks(Sink => Sink.NotifyStreamUpdateFailedAsync(File));
+			EnqueueTasks(Sink => Sink.NotifyConfigUpdateFailureAsync(ErrorMessage, FileName, Change, User, Description));
 		}
 
 		/// <summary>
diff --git a/Engine/Source/Programs/Horde/HordeServer/Notifications/Impl/SlackNotificationSink.cs b/Engine/Source/Programs/Horde/HordeServer/Notifications/Impl/SlackNotificationSink.cs
--- a/Engine/Source/Programs/Horde/HordeServer/Notifications/Impl/SlackNotificationSink.cs
+++ b/Engine/Source/Programs/Horde/HordeServer/Notifications/Impl/SlackNotificationSink.cs
@@ -874,9 +874,9 @@ namespace HordeServer.Notifications.Impl
 		#region Stream updates
 
 		/// <inheritdoc/>
-		public async Task NotifyStreamUpdateFailedAsync(IStream Stream, string ErrorMessage, int Change, IUser? Author = null, string? Description = null)
+		public async Task NotifyConfigUpdateFailureAsync(string ErrorMessage, string FileName, int? Change = null, IUser? Author = null, string? Description = null)
 		{
-			Logger.LogInformation("Sending stream update failure notification for {Stream} ({StreamId})", Stream.Name, Stream.Id);
+			Logger.LogInformation("Sending config update failure notification for {FileName}", FileName);
 
 			string? SlackUserId = null;
 			if (Author != null)
@@ -886,40 +886,39 @@ namespace HordeServer.Notifications.Impl
 
 			if (SlackUserId != null)
 			{
-				await SendStreamUpdateFailureMessageAsync(SlackUserId, Stream, ErrorMessage, Change, SlackUserId, Description);
+				await SendConfigUpdateFailureMessageAsync(SlackUserId, ErrorMessage, FileName, Change, SlackUserId, Description);
 			}
 			if (Settings.UpdateStreamsNotificationChannel != null)
 			{
-				await SendStreamUpdateFailureMessageAsync($"#{Settings.UpdateStreamsNotificationChannel}", Stream, ErrorMessage, Change, SlackUserId, Description);
+				await SendConfigUpdateFailureMessageAsync($"#{Settings.UpdateStreamsNotificationChannel}", ErrorMessage, FileName, Change, SlackUserId, Description);
 			}
 		}
 
-		private Task SendStreamUpdateFailureMessageAsync(string Recipient, IStream Stream, string ErrorMessage, int Change, string? Author = null, string? Description = null)
+		private Task SendConfigUpdateFailureMessageAsync(string Recipient, string ErrorMessage, string FileName, int? Change = null, string? Author = null, string? Description = null)
 		{
 			Color OutcomeColor = BlockKitAttachmentColors.Error;
 			BlockKitAttachment Attachment = new BlockKitAttachment();
 			Attachment.Color = OutcomeColor;
-			Attachment.FallbackText = $"{Stream.Name} ({Stream.Id}) - Update Failure";
+			Attachment.FallbackText = $"Update Failure: {FileName}";
 
-			Attachment.Blocks.Add(new HeaderBlock($"Stream Update Failure :rip:", false, true));
+			Attachment.Blocks.Add(new HeaderBlock($"Config Update Failure :rip:", false, true));
 
-			Attachment.Blocks.Add(new SectionBlock($"Horde was unable to update {Stream.Name} ({Stream.Id})"));
+			Attachment.Blocks.Add(new SectionBlock($"Horde was unable to update {FileName}"));
 			Attachment.Blocks.Add(new SectionBlock($"```{ErrorMessage}```"));
-			if (Author != null)
-			{
-				Attachment.Blocks.Add(new SectionBlock($"Possibly due to CL: {Change} by <@{Author}>"));
-			}
-			else
+			if (Change != null)
 			{
-				Attachment.Blocks.Add(new SectionBlock($"Possibly due to CL: {Change} - (Could not determine author from P4 user)"));
-			}
-			if (Description != null)
-			{
-				Attachment.Blocks.Add(new SectionBlock($"```{Description}```"));
-			}
-			else
-			{
-				Attachment.Blocks.Add(new SectionBlock($"```Could not determine description of change from P4```"));
+				if (Author != null)
+				{
+					Attachment.Blocks.Add(new SectionBlock($"Possibly due to CL: {Change.Value} by <@{Author}>"));
+				}
+				else
+				{
+					Attachment.Blocks.Add(new SectionBlock($"Possibly due to CL: {Change.Value} - (Could not determine author from P4 user)"));
+				}
+				if (Description != null)
+				{
+					Attachment.Blocks.Add(new SectionBlock($"```{Description}```"));
+				}
 			}
 
 			return SendMessageAsync(Recipient, Attachments: new[] { Attachment });
diff --git a/Engine/Source/Programs/Horde/HordeServer/Services/ConfigService.cs b/Engine/Source/Programs/Horde/HordeServer/Services/ConfigService.cs
--- a/Engine/Source/Programs/Horde/HordeServer/Services/ConfigService.cs
+++ b/Engine/Source/Programs/Horde/HordeServer/Services/ConfigService.cs
@@ -344,7 +344,7 @@ namespace HordeServer.Services
 					}
 					else
 					{
-						NotificationService.NotifyUpdateStreamFailure(File);
+						NotificationService.NotifyConfigUpdateFailure(File.Error, File.DepotPath);
 					}
 				}
 			}
@@ -366,6 +366,29 @@ namespace HordeServer.Services
 			catch (Exception Ex)
 			{
 				Logger.LogError(Ex, "Unable to read data from {ConfigPath}: {Message}", ConfigPath, Ex.Message);
+
+				string FileName = ConfigPath.AbsolutePath;
+				int Change = -1;
+				string? Author = null;
+				string? Description = null;
+
+				if (ConfigPath.Scheme == PerforceScheme)
+				{
+					try
+					{
+						List<FileSummary> Files = await PerforceService.FindFilesAsync(ConfigPath.Host, new[] { FileName });
+						Change = Files[0].Change;
+
+						ChangeDetails Details = await PerforceService.GetChangeDetailsAsync(ConfigPath.Host, Change);
+						(Author, Description) = (Details.Author, Details.Description); 
+					}
+					catch (Exception Ex2)
+					{
+						Logger.LogError(Ex2, "Unable to identify change that last modified {ConfigPath} from Perforce", ConfigPath);
+					}
+				}
+
+				NotificationService.NotifyConfigUpdateFailure(Ex.Message, FileName, Change, Author, Description);
 				return null;
 			}
 		}
