Horde: Make it possible to clear issue owner
#rnx

diff --git a/Engine/Source/Programs/Horde/HordeServer/Api/Issues.cs b/Engine/Source/Programs/Horde/HordeServer/Api/Issues.cs
--- a/Engine/Source/Programs/Horde/HordeServer/Api/Issues.cs
+++ b/Engine/Source/Programs/Horde/HordeServer/Api/Issues.cs
@@ -521,7 +521,7 @@ namespace HordeServer.Api
 		public string? Summary { get; set; }
 
 		/// <summary>
-		/// New user id for owner of the issue
+		/// New user id for owner of the issue, can be cleared by passing empty string
 		/// </summary>
 		public string? OwnerId { get; set; }
 
diff --git a/Engine/Source/Programs/Horde/HordeServer/Controllers/IssuesController.cs b/Engine/Source/Programs/Horde/HordeServer/Controllers/IssuesController.cs
--- a/Engine/Source/Programs/Horde/HordeServer/Controllers/IssuesController.cs
+++ b/Engine/Source/Programs/Horde/HordeServer/Controllers/IssuesController.cs
@@ -399,7 +399,7 @@ namespace HordeServer.Controllers
 			ObjectId? NewOwnerId = null;
 			if (Request.OwnerId != null)
 			{
-				NewOwnerId = new ObjectId(Request.OwnerId);
+				NewOwnerId = Request.OwnerId.Length == 0 ? ObjectId.Empty : new ObjectId(Request.OwnerId);
 			}
 
 			ObjectId? NewNominatedById = null;
