Horde: Improve issue resolved text to include who resolved it
#rnx

diff --git a/Engine/Source/Programs/Horde/HordeDashboard/src/base/utilities/issueUtils.ts b/Engine/Source/Programs/Horde/HordeDashboard/src/base/utilities/issueUtils.ts
--- a/Engine/Source/Programs/Horde/HordeDashboard/src/base/utilities/issueUtils.ts
+++ b/Engine/Source/Programs/Horde/HordeDashboard/src/base/utilities/issueUtils.ts
@@ -7,30 +7,29 @@ import { displayTimeZone } from "./timeUtils";
 
 export function getIssueStatus(issue: GetIssueResponse, showResolveTime?:boolean): string {
 
-	let text = "";	
+	let text = "";		
 
 	if (issue.resolvedAt) {
 
-		if (!showResolveTime) {
+		let resolvedText = "Resolved";
 
-			if (!issue.fixChange) {
-				return `Resolved`;
-			}
-			return `Resolved in CL ${issue.fixChange}`;			
+		if (issue.fixChange) {
+			resolvedText += ` in CL ${issue.fixChange}`;
 		}
 
-		const displayTime = moment(issue.resolvedAt).tz(displayTimeZone());
+		resolvedText += ` by ${issue.resolvedBy ?? "Horde" }`;
 
-		const format = dashboard.display24HourClock ? "MMM Do, HH:mm z" : "MMM Do, h:mma z";
+		if (!showResolveTime) {
+			return resolvedText;
+		}
 
-		let text = `Resolved at ${displayTime.format(format)}`;
+		const displayTime = moment(issue.resolvedAt).tz(displayTimeZone());
 
-		if (issue.fixChange) {
-			text += ` in CL ${issue.fixChange}`;
-		}
+		const format = dashboard.display24HourClock ? "MMM Do, HH:mm z" : "MMM Do, h:mma z";
 
-		return text;
+		resolvedText += ` on ${displayTime.format(format)}`;
 
+		return resolvedText;
 	}
 
 	if (!issue.ownerId) {
