Horde: Follow up submit for issue assignee changes
#rnx

diff --git a/Engine/Source/Programs/Horde/HordeDashboard/src/backend/Api.ts b/Engine/Source/Programs/Horde/HordeDashboard/src/backend/Api.ts
--- a/Engine/Source/Programs/Horde/HordeDashboard/src/backend/Api.ts
+++ b/Engine/Source/Programs/Horde/HordeDashboard/src/backend/Api.ts
@@ -2587,6 +2587,12 @@ export type GetIssueResponse = {
 	/**Time at which the issue was resolved */
 	resolvedAt?: Date | string;
 
+	/** Name of the user that resolved the issue */
+	resolvedBy?: string;
+
+	/** User id of the person that resolved the issue */
+	resolvedById?: string;
+
 	/**  List of stream paths affected by this issue */
 	streams: string[];
 
@@ -2612,8 +2618,8 @@ export type UpdateIssueRequest = {
 	/** Summary of the issue */
 	summary?: string;
 
-	/**New owner of the issue */
-	ownerId?: string | null;
+	/**New owner of the issue, pass empty string to clear current owner  */
+	ownerId?: string;
 
 	/**User than nominates the new owner */
 	nominatedById?: string | null;
