Horde: Do not return 200 from health checks when server is shutting down. This code was disabled in cbe78b02 on GitHub, but not clear on the reasons.

diff --git a/Engine/Source/Programs/Horde/HordeServer/Controllers/LifetimeController.cs b/Engine/Source/Programs/Horde/HordeServer/Controllers/LifetimeController.cs
--- a/Engine/Source/Programs/Horde/HordeServer/Controllers/LifetimeController.cs
+++ b/Engine/Source/Programs/Horde/HordeServer/Controllers/LifetimeController.cs
@@ -38,10 +38,10 @@ namespace HordeServer.Controllers
 		public Task<ActionResult> K8sReadinessProbe()
 		{
 			// Disabled for now
-			bool IsRunning = true;//!LifetimeService.IsStopping;
+			bool IsRunning = !LifetimeService.IsStopping;
 			bool IsMongoDBHealthy = true; //await LifetimeService.IsMongoDbConnectionHealthy();
 			bool IsRedisHealthy = true; //await LifetimeService.IsRedisConnectionHealthy();
-			int StatusCode = IsRunning && IsMongoDBHealthy && IsRedisHealthy ? 200 : 500;
+			int StatusCode = IsRunning && IsMongoDBHealthy && IsRedisHealthy ? 200 : 503;
 
 			string Content = $"IsRunning={IsRunning}\n";
 			Content += $"IsMongoDBHealthy={IsMongoDBHealthy}\n";
