Horde: Return degraded status from the health check endpoint for 15 seconds before starting to message to agents. This should allow the load balancer to start routing traffic away from the server before agents start to reconnect.

diff --git a/Engine/Source/Programs/Horde/HordeServer/Controllers/LifetimeController.cs b/Engine/Source/Programs/Horde/HordeServer/Controllers/LifetimeController.cs
--- a/Engine/Source/Programs/Horde/HordeServer/Controllers/LifetimeController.cs
+++ b/Engine/Source/Programs/Horde/HordeServer/Controllers/LifetimeController.cs
@@ -38,7 +38,7 @@ namespace HordeServer.Controllers
 		public Task<ActionResult> K8sReadinessProbe()
 		{
 			// Disabled for now
-			bool IsRunning = !LifetimeService.IsStopping;
+			bool IsRunning = !LifetimeService.IsPreStopping;
 			bool IsMongoDBHealthy = true; //await LifetimeService.IsMongoDbConnectionHealthy();
 			bool IsRedisHealthy = true; //await LifetimeService.IsRedisConnectionHealthy();
 			int StatusCode = IsRunning && IsMongoDBHealthy && IsRedisHealthy ? 200 : 503;
diff --git a/Engine/Source/Programs/Horde/HordeServer/Services/LifetimeService.cs b/Engine/Source/Programs/Horde/HordeServer/Services/LifetimeService.cs
--- a/Engine/Source/Programs/Horde/HordeServer/Services/LifetimeService.cs
+++ b/Engine/Source/Programs/Horde/HordeServer/Services/LifetimeService.cs
@@ -13,6 +13,7 @@ using Grpc.Core;
 using Grpc.Core.Interceptors;
 using HordeServer.Utilities;
 using StackExchange.Redis;
+using System.Diagnostics;
 
 namespace HordeServer.Services
 {
@@ -26,11 +27,21 @@ namespace HordeServer.Services
 		/// </summary>
 		ILogger<LifetimeService> Logger;
 
+		/// <summary>
+		/// Task which waits for a period of time and sets the stopping task
+		/// </summary>
+		Task WaitAndTriggerStoppingTask = Task.CompletedTask;
+
 		/// <summary>
 		/// Task source for the server stopping
 		/// </summary>
 		TaskCompletionSource<bool> StoppingTaskCompletionSource;
 
+		/// <summary>
+		/// Task source for the server stopping
+		/// </summary>
+		TaskCompletionSource<bool> PreStoppingTaskCompletionSource;
+
 		/// <summary>
 		/// Registration token for the stopping event
 		/// </summary>
@@ -79,6 +90,7 @@ namespace HordeServer.Services
 			this.RequestTrackerService = RequestTrackerService;
 			this.Logger = Logger;
 			this.StoppingTaskCompletionSource = new TaskCompletionSource<bool>();
+			this.PreStoppingTaskCompletionSource = new TaskCompletionSource<bool>();
 			this.Registration = Lifetime.ApplicationStopping.Register(ApplicationStopping);
 		}
 
@@ -93,10 +105,11 @@ namespace HordeServer.Services
 		/// </summary>
 		void ApplicationStopping()
 		{
-			if (StoppingTaskCompletionSource.TrySetResult(true))
+			if (PreStoppingTaskCompletionSource.TrySetResult(true))
 			{
-				IsStopping = true;
-				Logger.LogInformation("Server is stopping");
+				Logger.LogInformation("Setting pre-stop event");
+				IsPreStopping = true;
+				WaitAndTriggerStoppingTask = Task.Run(() => ExecStoppingTask());
 				/*
 				Logger.LogInformation("App is stopping. Waiting an initial {InitialDelay} secs before waiting on any requests...", (int)InitialStoppingDelay.TotalSeconds);
 				Thread.Sleep(InitialStoppingDelay);
@@ -122,13 +135,24 @@ namespace HordeServer.Services
 			}
 		}
 
+		async Task ExecStoppingTask()
+		{
+			await Task.Delay(TimeSpan.FromSeconds(15.0));
+
+			Logger.LogInformation("Setting stopping event");
+			IsStopping = true;
+			StoppingTaskCompletionSource.TrySetResult(true);
+		}
+
 		/// <summary>
 		/// Returns true if the server is stopping
 		/// </summary>
-		public bool IsStopping
-		{
-			get; private set;
-		}
+		public bool IsStopping { get; private set; }
+
+		/// <summary>
+		/// Returns true if the server is stopping, but may not be removed from the load balancer yet
+		/// </summary>
+		public bool IsPreStopping { get; private set; }
 
 		/// <summary>
 		/// Gets an awaitable task for the server stopping
