Horde: Wrap stream/project updates in try/catch handler that reports notifications.

diff --git a/Engine/Source/Programs/Horde/HordeServer/Collections/IStreamCollection.cs b/Engine/Source/Programs/Horde/HordeServer/Collections/IStreamCollection.cs
--- a/Engine/Source/Programs/Horde/HordeServer/Collections/IStreamCollection.cs
+++ b/Engine/Source/Programs/Horde/HordeServer/Collections/IStreamCollection.cs
@@ -110,4 +110,31 @@ namespace HordeServer.Collections
 		/// <returns>Async task</returns>
 		Task DeleteAsync(StreamId StreamId);
 	}
+
+	static class StreamCollectionExtensions
+	{
+		/// <summary>
+		/// Creates or replaces a stream configuration
+		/// </summary>
+		/// <param name="StreamCollection">The stream collection</param>
+		/// <param name="Id">Unique id for the new stream</param>
+		/// <param name="Stream">The current stream value. If not-null, this will attempt to replace the existing instance.</param>
+		/// <param name="ConfigPath">Path to the config file</param>
+		/// <param name="Revision">The config file revision</param>
+		/// <param name="ProjectId">The project id</param>
+		/// <param name="Config">The stream configuration</param>
+		/// <returns></returns>
+		public static async Task<IStream> CreateOrReplaceAsync(this IStreamCollection StreamCollection, StreamId Id, IStream? Stream, string ConfigPath, string Revision, ProjectId ProjectId, StreamConfig Config)
+		{
+			for (; ; )
+			{
+				Stream = await StreamCollection.TryCreateOrReplaceAsync(Id, Stream, ConfigPath, Revision, ProjectId, Config);
+				if (Stream != null)
+				{
+					return Stream;
+				}
+				Stream = await StreamCollection.GetAsync(Id);
+			}
+		}
+	}
 }
diff --git a/Engine/Source/Programs/Horde/HordeServer/Services/ConfigService.cs b/Engine/Source/Programs/Horde/HordeServer/Services/ConfigService.cs
--- a/Engine/Source/Programs/Horde/HordeServer/Services/ConfigService.cs
+++ b/Engine/Source/Programs/Horde/HordeServer/Services/ConfigService.cs
@@ -117,13 +117,16 @@ namespace HordeServer.Services
 				if (CachedGlobalConfig == null || Revision != CachedGlobalConfigRevision)
 				{
 					Logger.LogInformation("Caching global config from {Revision}", Revision);
-					CachedGlobalConfig = await ReadDataAsync<GlobalConfig>(ConfigPath);
-					if (CachedGlobalConfig == null)
+					try
+					{
+						CachedGlobalConfig = await ReadDataAsync<GlobalConfig>(ConfigPath);
+						CachedGlobalConfigRevision = Revision;
+					}
+					catch (Exception Ex)
 					{
-						Logger.LogWarning("Unable to update any projects or streams due to invalid global config.");
+						await SendFailureNotificationAsync(Ex, ConfigPath);
 						return;
 					}
-					CachedGlobalConfigRevision = Revision;
 				}
 				GlobalConfig = CachedGlobalConfig;
 
@@ -182,16 +185,20 @@ namespace HordeServer.Services
 				if (Update || !PrevCachedProjectConfigs.TryGetValue(ProjectRef.Id, out ProjectConfig))
 				{
 					Logger.LogInformation("Caching configuration for project {ProjectId} ({Revision})", ProjectRef.Id, Revision);
-					ProjectConfig = await ReadDataAsync<ProjectConfig>(ProjectPath);
-					if (ProjectConfig == null)
+					try
 					{
-						SkipProjectIds.Add(ProjectRef.Id);
-						continue;
+						ProjectConfig = await ReadDataAsync<ProjectConfig>(ProjectPath);
+						if (Update)
+						{
+							Logger.LogInformation("Updating configuration for project {ProjectId} ({Revision})", ProjectRef.Id, Revision);
+							await ProjectService.Collection.AddOrUpdateAsync(ProjectRef.Id, ProjectPath.ToString(), Revision, Idx, ProjectConfig);
+						}
 					}
-					if (Update)
+					catch (Exception Ex)
 					{
-						Logger.LogInformation("Updating configuration for project {ProjectId} ({Revision})", ProjectRef.Id, Revision);
-						await ProjectService.Collection.AddOrUpdateAsync(ProjectRef.Id, ProjectPath.ToString(), Revision, Idx, ProjectConfig);
+						await SendFailureNotificationAsync(Ex, ProjectPath);
+						SkipProjectIds.Add(ProjectRef.Id);
+						continue;
 					}
 				}
 
@@ -220,8 +227,16 @@ namespace HordeServer.Services
 					if (Revision != CurrentRevision)
 					{
 						Logger.LogInformation("Updating logo for project {ProjectId} ({Revision})", ProjectId, Revision);
-						await ProjectService.Collection.SetLogoAsync(ProjectId, Path.ToString(), Revision, GetMimeTypeFromPath(Path), await ReadDataAsync(Path));
-						CachedLogoRevisions[ProjectId] = Revision;
+						try
+						{
+							await ProjectService.Collection.SetLogoAsync(ProjectId, Path.ToString(), Revision, GetMimeTypeFromPath(Path), await ReadDataAsync(Path));
+							CachedLogoRevisions[ProjectId] = Revision;
+						}
+						catch (Exception Ex)
+						{
+							await SendFailureNotificationAsync(Ex, Path);
+							continue;
+						}
 					}
 				}
 			}
@@ -240,21 +255,15 @@ namespace HordeServer.Services
 					if (Stream == null || Stream.ConfigPath != StreamPath.ToString() || Stream.ConfigRevision != Revision)
 					{
 						Logger.LogInformation("Updating configuration for stream {StreamRef} ({Revision})", StreamRef.Id, Revision);
-
-						StreamConfig? StreamConfig = await ReadDataAsync<StreamConfig>(StreamPath);
-						if (StreamConfig != null)
+						try
 						{
-							for (; ; )
-							{
-								Stream = await StreamService.StreamCollection.TryCreateOrReplaceAsync(StreamRef.Id, Stream, StreamPath.ToString(), Revision, ProjectId, StreamConfig);
-
-								if (Stream != null)
-								{
-									break;
-								}
-
-								Stream = await StreamService.GetStreamAsync(StreamRef.Id);
-							}
+							StreamConfig StreamConfig = await ReadDataAsync<StreamConfig>(StreamPath);
+							Stream = await StreamService.StreamCollection.TryCreateOrReplaceAsync(StreamRef.Id, Stream, StreamPath.ToString(), Revision, ProjectId, StreamConfig);
+						}
+						catch (Exception Ex)
+						{
+							await SendFailureNotificationAsync(Ex, StreamPath);
+							continue;
 						}
 					}
 				}
@@ -352,45 +361,14 @@ namespace HordeServer.Services
 			return Revisions;
 		}
 
-		async Task<T?> ReadDataAsync<T>(Uri ConfigPath) where T : class
+		async Task<T> ReadDataAsync<T>(Uri ConfigPath) where T : class
 		{
-			try
-			{
-				byte[] Data = await ReadDataAsync(ConfigPath);
-
-				JsonSerializerOptions Options = new JsonSerializerOptions();
-				Startup.ConfigureJsonSerializer(Options);
-
-				return JsonSerializer.Deserialize<T>(Data, Options);
-			}
-			catch (Exception Ex)
-			{
-				Logger.LogError(Ex, "Unable to read data from {ConfigPath}: {Message}", ConfigPath, Ex.Message);
+			byte[] Data = await ReadDataAsync(ConfigPath);
 
-				string FileName = ConfigPath.AbsolutePath;
-				int Change = -1;
-				string? Author = null;
-				string? Description = null;
+			JsonSerializerOptions Options = new JsonSerializerOptions();
+			Startup.ConfigureJsonSerializer(Options);
 
-				if (ConfigPath.Scheme == PerforceScheme)
-				{
-					try
-					{
-						List<FileSummary> Files = await PerforceService.FindFilesAsync(ConfigPath.Host, new[] { FileName });
-						Change = Files[0].Change;
-
-						ChangeDetails Details = await PerforceService.GetChangeDetailsAsync(ConfigPath.Host, Change);
-						(Author, Description) = (Details.Author, Details.Description); 
-					}
-					catch (Exception Ex2)
-					{
-						Logger.LogError(Ex2, "Unable to identify change that last modified {ConfigPath} from Perforce", ConfigPath);
-					}
-				}
-
-				NotificationService.NotifyConfigUpdateFailure(Ex.Message, FileName, Change, Author, Description);
-				return null;
-			}
+			return JsonSerializer.Deserialize<T>(Data, Options);
 		}
 
 		Task<byte[]> ReadDataAsync(Uri ConfigPath)
@@ -406,6 +384,34 @@ namespace HordeServer.Services
 			}
 		}
 
+		async Task SendFailureNotificationAsync(Exception Ex, Uri ConfigPath)
+		{
+			Logger.LogError(Ex, "Unable to read data from {ConfigPath}: {Message}", ConfigPath, Ex.Message);
+
+			string FileName = ConfigPath.AbsolutePath;
+			int Change = -1;
+			string? Author = null;
+			string? Description = null;
+
+			if (ConfigPath.Scheme == PerforceScheme)
+			{
+				try
+				{
+					List<FileSummary> Files = await PerforceService.FindFilesAsync(ConfigPath.Host, new[] { FileName });
+					Change = Files[0].Change;
+
+					ChangeDetails Details = await PerforceService.GetChangeDetailsAsync(ConfigPath.Host, Change);
+					(Author, Description) = (Details.Author, Details.Description);
+				}
+				catch (Exception Ex2)
+				{
+					Logger.LogError(Ex2, "Unable to identify change that last modified {ConfigPath} from Perforce", ConfigPath);
+				}
+			}
+
+			NotificationService.NotifyConfigUpdateFailure(Ex.Message, FileName, Change, Author, Description);
+		}
+
 		/// <inheritdoc/>
 		protected override async Task<DateTime> TickLeaderAsync(CancellationToken StoppingToken)
 		{
