Horde: Allow replacing an existing lease. Should fix some errors where a lease has to be reassigned due to a server migration.

diff --git a/Engine/Source/Programs/Horde/HordeServer/Collections/Impl/LeaseCollection.cs b/Engine/Source/Programs/Horde/HordeServer/Collections/Impl/LeaseCollection.cs
--- a/Engine/Source/Programs/Horde/HordeServer/Collections/Impl/LeaseCollection.cs
+++ b/Engine/Source/Programs/Horde/HordeServer/Collections/Impl/LeaseCollection.cs
@@ -88,7 +88,7 @@ namespace HordeServer.Collections.Impl
 		public async Task<ILease> AddAsync(ObjectId Id, string Name, AgentId AgentId, ObjectId SessionId, StreamId? StreamId, PoolId? PoolId, ObjectId? LogId, DateTime StartTime, byte[] Payload)
 		{
 			LeaseDocument Lease = new LeaseDocument(Id, Name, AgentId, SessionId, StreamId, PoolId, LogId, StartTime, Payload);
-			await Leases.InsertOneAsync(Lease);
+			await Leases.ReplaceOneAsync(x => x.Id == Id, Lease, new ReplaceOptions { IsUpsert = true });
 			return Lease;
 		}
 
