Horde: Fix ManagementObject instances not being disposed of.

diff --git a/Engine/Source/Programs/Horde/HordeAgent/Services/WorkerService.cs b/Engine/Source/Programs/Horde/HordeAgent/Services/WorkerService.cs
--- a/Engine/Source/Programs/Horde/HordeAgent/Services/WorkerService.cs
+++ b/Engine/Source/Programs/Horde/HordeAgent/Services/WorkerService.cs
@@ -1330,18 +1330,21 @@ namespace HordeAgent.Services
 				{
 					foreach (ManagementObject Row in Searcher.Get())
 					{
-						Dictionary<string, object> Properties = GetWmiProperties(Row);
-
-						object? Name;
-						if(Properties.TryGetValue("Caption", out Name))
+						using (Row)
 						{
-							PrimaryDevice.Properties.Add($"OSDistribution={Name}");
-						}
+							Dictionary<string, object> Properties = GetWmiProperties(Row);
 
-						object? Version;
-						if(Properties.TryGetValue("Version", out Version))
-						{
-							PrimaryDevice.Properties.Add($"OSKernelVersion={Version}");
+							object? Name;
+							if (Properties.TryGetValue("Caption", out Name))
+							{
+								PrimaryDevice.Properties.Add($"OSDistribution={Name}");
+							}
+
+							object? Version;
+							if (Properties.TryGetValue("Version", out Version))
+							{
+								PrimaryDevice.Properties.Add($"OSKernelVersion={Version}");
+							}
 						}
 					}
 				}
@@ -1355,27 +1358,30 @@ namespace HordeAgent.Services
 
 					foreach (ManagementObject Row in Searcher.Get())
 					{
-						Dictionary<string, object> Properties = GetWmiProperties(Row);
-
-						object? NameObject;
-						if(Properties.TryGetValue("Name", out NameObject))
+						using (Row)
 						{
-							string Name = NameObject.ToString() ?? String.Empty;
-							int Count;
-							NameToCount.TryGetValue(Name, out Count);
-							NameToCount[Name] = Count + 1;
-						}
+							Dictionary<string, object> Properties = GetWmiProperties(Row);
 
-						object? NumPhysicalCores;
-						if((Properties.TryGetValue("NumberOfEnabledCore", out NumPhysicalCores) && NumPhysicalCores is uint) || (Properties.TryGetValue("NumberOfCores", out NumPhysicalCores) && NumPhysicalCores is uint))
-						{
-							TotalPhysicalCores += (int)(uint)NumPhysicalCores;
-						}
+							object? NameObject;
+							if (Properties.TryGetValue("Name", out NameObject))
+							{
+								string Name = NameObject.ToString() ?? String.Empty;
+								int Count;
+								NameToCount.TryGetValue(Name, out Count);
+								NameToCount[Name] = Count + 1;
+							}
 
-						object? NumLogicalCores;
-						if(Properties.TryGetValue("NumberOfLogicalProcessors", out NumLogicalCores) && NumLogicalCores is uint)
-						{
-							TotalLogicalCores += (int)(uint)NumLogicalCores;
+							object? NumPhysicalCores;
+							if ((Properties.TryGetValue("NumberOfEnabledCore", out NumPhysicalCores) && NumPhysicalCores is uint) || (Properties.TryGetValue("NumberOfCores", out NumPhysicalCores) && NumPhysicalCores is uint))
+							{
+								TotalPhysicalCores += (int)(uint)NumPhysicalCores;
+							}
+
+							object? NumLogicalCores;
+							if (Properties.TryGetValue("NumberOfLogicalProcessors", out NumLogicalCores) && NumLogicalCores is uint)
+							{
+								TotalLogicalCores += (int)(uint)NumLogicalCores;
+							}
 						}
 					}
 
@@ -1401,10 +1407,13 @@ namespace HordeAgent.Services
 					ulong TotalCapacity = 0;
 					foreach (ManagementObject Row in Searcher.Get())
 					{
-						object? Capacity = Row.GetPropertyValue("Capacity");
-						if (Capacity is ulong)
+						using (Row)
 						{
-							TotalCapacity += (ulong)Capacity;
+							object? Capacity = Row.GetPropertyValue("Capacity");
+							if (Capacity is ulong)
+							{
+								TotalCapacity += (ulong)Capacity;
+							}
 						}
 					}
 
@@ -1420,15 +1429,18 @@ namespace HordeAgent.Services
 					int Index = 0;
 					foreach (ManagementObject Row in Searcher.Get())
 					{
-						WmiProperties Properties = new WmiProperties(Row);
-						if (Properties.TryGetValue("Name", out string? Name) && Properties.TryGetValue("DriverVersion", out string? DriverVersion))
+						using (Row)
 						{
-							DeviceCapabilities Device = new DeviceCapabilities();
-							Device.Handle = $"GPU-{++Index}";
-							Device.Properties.Add($"Name={Name}");
-							Device.Properties.Add($"Type=GPU");
-							Device.Properties.Add($"DriverVersion={DriverVersion}");
-							OtherDevices.Add(Device);
+							WmiProperties Properties = new WmiProperties(Row);
+							if (Properties.TryGetValue("Name", out string? Name) && Properties.TryGetValue("DriverVersion", out string? DriverVersion))
+							{
+								DeviceCapabilities Device = new DeviceCapabilities();
+								Device.Handle = $"GPU-{++Index}";
+								Device.Properties.Add($"Name={Name}");
+								Device.Properties.Add($"Type=GPU");
+								Device.Properties.Add($"DriverVersion={DriverVersion}");
+								OtherDevices.Add(Device);
+							}
 						}
 					}
 				}
