Use Microsoft.CodeAnalysis.CSharp.Workspaces for code analysis

#rnx

diff --git a/Engine/Source/Programs/DotNETCommon/BuildUtilities/BuildUtilities.csproj b/Engine/Source/Programs/DotNETCommon/BuildUtilities/BuildUtilities.csproj
--- a/Engine/Source/Programs/DotNETCommon/BuildUtilities/BuildUtilities.csproj
+++ b/Engine/Source/Programs/DotNETCommon/BuildUtilities/BuildUtilities.csproj
@@ -30,6 +30,7 @@
     <ProjectReference Include="..\..\Shared\EpicGames.Core\EpicGames.Core.csproj" />
   </ItemGroup>
   <ItemGroup>
+    <PackageReference Include="Microsoft.CodeAnalysis.CSharp.Workspaces" Version="3.10.0" />
     <PackageReference Include="Microsoft.CSharp" Version="4.7.0" />
     <PackageReference Include="System.Data.DataSetExtensions" Version="4.5.0" />
   </ItemGroup>
diff --git a/Engine/Source/Programs/Horde/HordeAgent/HordeAgent.csproj b/Engine/Source/Programs/Horde/HordeAgent/HordeAgent.csproj
--- a/Engine/Source/Programs/Horde/HordeAgent/HordeAgent.csproj
+++ b/Engine/Source/Programs/Horde/HordeAgent/HordeAgent.csproj
@@ -29,6 +29,7 @@
       <PrivateAssets>all</PrivateAssets>
       <IncludeAssets>runtime; build; native; contentfiles; analyzers; buildtransitive</IncludeAssets>
     </PackageReference>
+    <PackageReference Include="Microsoft.CodeAnalysis.CSharp.Workspaces" Version="3.10.0" />
     <PackageReference Include="Microsoft.Extensions.Hosting" Version="3.1.5" />
     <PackageReference Include="Microsoft.Extensions.Hosting.Abstractions" Version="3.1.5" />
     <PackageReference Include="Microsoft.Extensions.Hosting.WindowsServices" Version="3.1.5" />
@@ -117,13 +118,13 @@
 
   <Target Name="PublishSymbols" AfterTargets="AfterBuild" Condition="'$(FullSymbolsDir)' != ''">
     <ItemGroup>
-      <OutputFiles Include="$(OutDir)**"/>
-      <BinaryFilesBase Include="$(OutDir)\*.dll"/>
-      <BinaryFiles Include="@(BinaryFilesBase)" Condition="Exists('%(RelativeDir)%(FileName).pdb')"/>
+      <OutputFiles Include="$(OutDir)**" />
+      <BinaryFilesBase Include="$(OutDir)\*.dll" />
+      <BinaryFiles Include="@(BinaryFilesBase)" Condition="Exists('%(RelativeDir)%(FileName).pdb')" />
     </ItemGroup>
 
     <MakeDir Directories="$(FullSymbolsDir)" />
-    <Copy SourceFiles="@(Files)" DestinationFolder="$(FullSymbolsDir)"/>
+    <Copy SourceFiles="@(Files)" DestinationFolder="$(FullSymbolsDir)" />
     <Exec Command="&quot;$(Pdb2PdbPath)&quot; &quot;%(BinaryFiles.Identity)&quot; /out &quot;$(FullSymbolsDir)\%(BinaryFiles.FileName).pdb&quot; /nowarn 21" IgnoreExitCode="false" />
   </Target>
   
diff --git a/Engine/Source/Programs/Horde/HordeAgentTests/HordeAgentTests.csproj b/Engine/Source/Programs/Horde/HordeAgentTests/HordeAgentTests.csproj
--- a/Engine/Source/Programs/Horde/HordeAgentTests/HordeAgentTests.csproj
+++ b/Engine/Source/Programs/Horde/HordeAgentTests/HordeAgentTests.csproj
@@ -8,6 +8,8 @@
 
   <ItemGroup>
 
+    <PackageReference Include="Microsoft.CodeAnalysis.CSharp.Workspaces" Version="3.10.0" />
+
     <PackageReference Include="Microsoft.Extensions.Logging.Abstractions" Version="5.0.0" />
     <PackageReference Include="Microsoft.NET.Test.Sdk" Version="16.5.0" />
     <PackageReference Include="Moq" Version="4.14.5" />
diff --git a/Engine/Source/Programs/Horde/HordeCommon/HordeCommon.csproj b/Engine/Source/Programs/Horde/HordeCommon/HordeCommon.csproj
--- a/Engine/Source/Programs/Horde/HordeCommon/HordeCommon.csproj
+++ b/Engine/Source/Programs/Horde/HordeCommon/HordeCommon.csproj
@@ -12,6 +12,7 @@
       <PrivateAssets>all</PrivateAssets>
       <IncludeAssets>runtime; build; native; contentfiles; analyzers; buildtransitive</IncludeAssets>
     </PackageReference>
+    <PackageReference Include="Microsoft.CodeAnalysis.CSharp.Workspaces" Version="3.10.0" />
     <PackageReference Include="Microsoft.Extensions.Logging" Version="5.0.0" />
   </ItemGroup>
 
diff --git a/Engine/Source/Programs/Horde/HordeServer/Collections/IIssueCollection.cs b/Engine/Source/Programs/Horde/HordeServer/Collections/IIssueCollection.cs
--- a/Engine/Source/Programs/Horde/HordeServer/Collections/IIssueCollection.cs
+++ b/Engine/Source/Programs/Horde/HordeServer/Collections/IIssueCollection.cs
@@ -127,7 +127,9 @@ namespace HordeServer.Collections
 		/// <returns>The hashcode value</returns>
 		public override int GetHashCode()
 		{
+#pragma warning disable RS1024 // Compare symbols correctly
 			int Result = string.GetHashCode(Type, StringComparison.Ordinal);
+#pragma warning restore RS1024 // Compare symbols correctly
 			foreach (string Key in Keys)
 			{
 				Result = HashCode.Combine(Result, Key);
diff --git a/Engine/Source/Programs/Horde/HordeServer/HordeServer.csproj b/Engine/Source/Programs/Horde/HordeServer/HordeServer.csproj
--- a/Engine/Source/Programs/Horde/HordeServer/HordeServer.csproj
+++ b/Engine/Source/Programs/Horde/HordeServer/HordeServer.csproj
@@ -29,10 +29,7 @@
     <PackageReference Include="Microsoft.AspNetCore.Authentication.JwtBearer" Version="3.1.5" />
     <PackageReference Include="Microsoft.AspNetCore.Authentication.OpenIdConnect" Version="3.1.5" />
     <PackageReference Include="Microsoft.AspNetCore.DataProtection.StackExchangeRedis" Version="3.1.7" />
-    <PackageReference Include="Microsoft.CodeAnalysis.FxCopAnalyzers" Version="3.0.0">
-      <PrivateAssets>all</PrivateAssets>
-      <IncludeAssets>runtime; build; native; contentfiles; analyzers; buildtransitive</IncludeAssets>
-    </PackageReference>
+    <PackageReference Include="Microsoft.CodeAnalysis.CSharp.Workspaces" Version="3.10.0" />
     <PackageReference Include="Microsoft.Extensions.Caching.Memory" Version="3.1.6" />
     <PackageReference Include="MongoDB.Driver" Version="2.10.4" />
     <PackageReference Include="Polly" Version="7.2.1" />
diff --git a/Engine/Source/Programs/Horde/HordeServerTests/HordeServerTests.csproj b/Engine/Source/Programs/Horde/HordeServerTests/HordeServerTests.csproj
--- a/Engine/Source/Programs/Horde/HordeServerTests/HordeServerTests.csproj
+++ b/Engine/Source/Programs/Horde/HordeServerTests/HordeServerTests.csproj
@@ -18,6 +18,7 @@
 
   <ItemGroup>
     <PackageReference Include="Microsoft.AspNetCore.Mvc.Testing" Version="3.1.6" />
+    <PackageReference Include="Microsoft.CodeAnalysis.CSharp.Workspaces" Version="3.10.0" />
     <PackageReference Include="Microsoft.NET.Test.Sdk" Version="16.5.0" />
     <PackageReference Include="Moq" Version="4.14.5" />
     <PackageReference Include="MSTest.TestAdapter" Version="2.1.0" />
diff --git a/Engine/Source/Programs/Shared/EpicGames.Core/BinaryArchiveWriter.cs b/Engine/Source/Programs/Shared/EpicGames.Core/BinaryArchiveWriter.cs
--- a/Engine/Source/Programs/Shared/EpicGames.Core/BinaryArchiveWriter.cs
+++ b/Engine/Source/Programs/Shared/EpicGames.Core/BinaryArchiveWriter.cs
@@ -29,7 +29,9 @@ namespace EpicGames.Core
 
 			int IEqualityComparer<object>.GetHashCode(object X)
 			{
+#pragma warning disable RS1024 // Compare symbols correctly ([REDACTED_PROFILE_URL]/roslyn/issues/49386)
 				return RuntimeHelpers.GetHashCode(X);
+#pragma warning restore RS1024 // Compare symbols correctly
 			}
 		}
 
diff --git a/Engine/Source/Programs/Shared/EpicGames.Core/StringView.cs b/Engine/Source/Programs/Shared/EpicGames.Core/StringView.cs
--- a/Engine/Source/Programs/Shared/EpicGames.Core/StringView.cs
+++ b/Engine/Source/Programs/Shared/EpicGames.Core/StringView.cs
@@ -95,7 +95,9 @@ namespace EpicGames.Core
 		/// <inheritdoc/>
 		public override int GetHashCode()
 		{
-			return String.GetHashCode(Memory.Span);
+#pragma warning disable RS1024 // Compare symbols correctly
+			return string.GetHashCode(Memory.Span);
+#pragma warning restore RS1024 // Compare symbols correctly
 		}
 
 		/// <inheritdoc/>
@@ -161,7 +163,9 @@ namespace EpicGames.Core
 		/// <inheritdoc/>
 		public int GetHashCode(StringView obj)
 		{
-			return String.GetHashCode(obj.Span);
+#pragma warning disable RS1024 // Compare symbols correctly
+			return string.GetHashCode(obj.Span);
+#pragma warning restore RS1024 // Compare symbols correctly
 		}
 
 		/// <inheritdoc/>
diff --git a/Engine/Source/Programs/Shared/EpicGames.Horde.Common/EpicGames.Horde.Common.csproj b/Engine/Source/Programs/Shared/EpicGames.Horde.Common/EpicGames.Horde.Common.csproj
--- a/Engine/Source/Programs/Shared/EpicGames.Horde.Common/EpicGames.Horde.Common.csproj
+++ b/Engine/Source/Programs/Shared/EpicGames.Horde.Common/EpicGames.Horde.Common.csproj
@@ -33,6 +33,7 @@
       <PrivateAssets>all</PrivateAssets>
       <IncludeAssets>runtime; build; native; contentfiles; analyzers; buildtransitive</IncludeAssets>
     </PackageReference>
+    <PackageReference Include="Microsoft.CodeAnalysis.CSharp.Workspaces" Version="3.10.0" />
     <PackageReference Include="Microsoft.Extensions.Logging" Version="5.0.0" />
   </ItemGroup>
   <ItemGroup>
diff --git a/Engine/Source/Programs/Shared/EpicGames.IoHash/EpicGames.IoHash.csproj b/Engine/Source/Programs/Shared/EpicGames.IoHash/EpicGames.IoHash.csproj
--- a/Engine/Source/Programs/Shared/EpicGames.IoHash/EpicGames.IoHash.csproj
+++ b/Engine/Source/Programs/Shared/EpicGames.IoHash/EpicGames.IoHash.csproj
@@ -28,6 +28,7 @@
   </ItemGroup>
   <ItemGroup>
     <PackageReference Include="Blake3" Version="0.3.0" />
+    <PackageReference Include="Microsoft.CodeAnalysis.CSharp.Workspaces" Version="3.10.0" />
     <PackageReference Include="System.Memory" Version="4.5.4" />
   </ItemGroup>
   <ItemGroup>
diff --git a/Engine/Source/Programs/Shared/EpicGames.Perforce.Managed/EpicGames.Perforce.Managed.csproj b/Engine/Source/Programs/Shared/EpicGames.Perforce.Managed/EpicGames.Perforce.Managed.csproj
--- a/Engine/Source/Programs/Shared/EpicGames.Perforce.Managed/EpicGames.Perforce.Managed.csproj
+++ b/Engine/Source/Programs/Shared/EpicGames.Perforce.Managed/EpicGames.Perforce.Managed.csproj
@@ -31,6 +31,7 @@
   </ItemGroup>
   <ItemGroup>
     <PackageReference Include="LibGit2Sharp" Version="0.26.2" />
+    <PackageReference Include="Microsoft.CodeAnalysis.CSharp.Workspaces" Version="3.10.0" />
   </ItemGroup>
 
 </Project>
diff --git a/Engine/Source/Programs/Shared/EpicGames.Perforce/EpicGames.Perforce.csproj b/Engine/Source/Programs/Shared/EpicGames.Perforce/EpicGames.Perforce.csproj
--- a/Engine/Source/Programs/Shared/EpicGames.Perforce/EpicGames.Perforce.csproj
+++ b/Engine/Source/Programs/Shared/EpicGames.Perforce/EpicGames.Perforce.csproj
@@ -26,6 +26,9 @@
       <Link>Properties\MetaData.cs</Link>
     </Compile>
   </ItemGroup>
+  <ItemGroup>
+    <PackageReference Include="Microsoft.CodeAnalysis.CSharp.Workspaces" Version="3.10.0" />
+  </ItemGroup>
 
   <ItemGroup>
     <ProjectReference Include="..\EpicGames.Core\EpicGames.Core.csproj" />
diff --git a/Engine/Source/Programs/Shared/EpicGames.Serialization.Tests/EpicGames.Serialization.Tests.csproj b/Engine/Source/Programs/Shared/EpicGames.Serialization.Tests/EpicGames.Serialization.Tests.csproj
--- a/Engine/Source/Programs/Shared/EpicGames.Serialization.Tests/EpicGames.Serialization.Tests.csproj
+++ b/Engine/Source/Programs/Shared/EpicGames.Serialization.Tests/EpicGames.Serialization.Tests.csproj
@@ -21,6 +21,7 @@
   </ItemGroup>
 
   <ItemGroup>
+    <PackageReference Include="Microsoft.CodeAnalysis.CSharp.Workspaces" Version="3.10.0" />
     <PackageReference Include="Microsoft.NET.Test.Sdk" Version="16.5.0" />
     <PackageReference Include="MSTest.TestAdapter" Version="2.1.0" />
     <PackageReference Include="MSTest.TestFramework" Version="2.1.0" />
diff --git a/Engine/Source/Programs/Shared/EpicGames.Serialization/EpicGames.Serialization.csproj b/Engine/Source/Programs/Shared/EpicGames.Serialization/EpicGames.Serialization.csproj
--- a/Engine/Source/Programs/Shared/EpicGames.Serialization/EpicGames.Serialization.csproj
+++ b/Engine/Source/Programs/Shared/EpicGames.Serialization/EpicGames.Serialization.csproj
@@ -26,6 +26,9 @@
       <Link>Properties\MetaData.cs</Link>
     </Compile>
   </ItemGroup>
+  <ItemGroup>
+    <PackageReference Include="Microsoft.CodeAnalysis.CSharp.Workspaces" Version="3.10.0" />
+  </ItemGroup>
   <ItemGroup>
     <ProjectReference Include="..\EpicGames.IoHash\EpicGames.IoHash.csproj" />
   </ItemGroup>
diff --git a/Engine/Source/Programs/UnrealBuildTool/Configuration/TargetDescriptor.cs b/Engine/Source/Programs/UnrealBuildTool/Configuration/TargetDescriptor.cs
--- a/Engine/Source/Programs/UnrealBuildTool/Configuration/TargetDescriptor.cs
+++ b/Engine/Source/Programs/UnrealBuildTool/Configuration/TargetDescriptor.cs
@@ -449,11 +449,7 @@ namespace UnrealBuildTool
 
 		public override int GetHashCode()
 		{
-			return String.GetHashCode(ProjectFile?.FullName) + 
-				Name.GetHashCode() + 
-				Platform.GetHashCode() + 
-				Configuration.GetHashCode() + 
-				Architecture.GetHashCode();
+			return HashCode.Combine(ProjectFile?.FullName, Name, Platform, Configuration, Architecture, AdditionalArguments);
 		}
 
 		public override bool Equals(object? Obj) 
@@ -466,7 +462,8 @@ namespace UnrealBuildTool
 					Name == Other.Name &&
 					Platform == Other.Platform &&
 					Configuration == Other.Configuration &&
-					Architecture == Other.Architecture;
+					Architecture == Other.Architecture &&
+					AdditionalArguments == Other.AdditionalArguments;
 			}
 			return false;
 		}
diff --git a/Engine/Source/Programs/UnrealBuildTool/UnrealBuildTool.csproj b/Engine/Source/Programs/UnrealBuildTool/UnrealBuildTool.csproj
--- a/Engine/Source/Programs/UnrealBuildTool/UnrealBuildTool.csproj
+++ b/Engine/Source/Programs/UnrealBuildTool/UnrealBuildTool.csproj
@@ -101,7 +101,7 @@
     <ProjectReference Include="..\Shared\EpicGames.Core\EpicGames.Core.csproj" />
   </ItemGroup>
   <ItemGroup>
-    <PackageReference Include="Microsoft.CodeAnalysis.CSharp" Version="3.7.0" />
+    <PackageReference Include="Microsoft.CodeAnalysis.CSharp.Workspaces" Version="3.10.0" />
     <PackageReference Include="Microsoft.Win32.Registry" Version="4.7.0" />
     <PackageReference Include="System.Management" Version="4.7.0" />
     <PackageReference Include="System.Security.Cryptography.Csp" Version="4.3.0" />
