Horde: Fix cached project config not being updated when file is modified in Perforce.

diff --git a/Engine/Source/Programs/Horde/HordeServer/Services/ConfigService.cs b/Engine/Source/Programs/Horde/HordeServer/Services/ConfigService.cs
--- a/Engine/Source/Programs/Horde/HordeServer/Services/ConfigService.cs
+++ b/Engine/Source/Programs/Horde/HordeServer/Services/ConfigService.cs
@@ -98,7 +98,7 @@ namespace HordeServer.Services
 
 		GlobalConfig? CachedGlobalConfig;
 		string? CachedGlobalConfigRevision;
-		Dictionary<ProjectId, ProjectConfig> CachedProjectConfigs = new Dictionary<ProjectId, ProjectConfig>();
+		Dictionary<ProjectId, (ProjectConfig Config, string Revision)> CachedProjectConfigs = new Dictionary<ProjectId, (ProjectConfig, string)>();
 		Dictionary<ProjectId, string?> CachedLogoRevisions = new Dictionary<ProjectId, string?>();
 
 		async Task UpdateConfigAsync(Uri ConfigPath)
@@ -155,8 +155,8 @@ namespace HordeServer.Services
 			// Get the path to all the project configs
 			List<(ProjectConfigRef ProjectRef, Uri Path)> ProjectConfigs = GlobalConfig.Projects.Select(x => (x, CombinePaths(ConfigPath, x.Path))).ToList();
 
-			Dictionary<ProjectId, ProjectConfig> PrevCachedProjectConfigs = CachedProjectConfigs;
-			CachedProjectConfigs = new Dictionary<ProjectId, ProjectConfig>();
+			Dictionary<ProjectId, (ProjectConfig Config, string Revision)> PrevCachedProjectConfigs = CachedProjectConfigs;
+			CachedProjectConfigs = new Dictionary<ProjectId, (ProjectConfig, string)>();
 
 			List<(ProjectId ProjectId, Uri Path)> ProjectLogos = new List<(ProjectId ProjectId, Uri Path)>();
 
@@ -182,7 +182,11 @@ namespace HordeServer.Services
 				bool Update = (Project == null || Project.ConfigPath != ProjectPath.ToString() || Project.ConfigRevision != Revision);
 
 				ProjectConfig? ProjectConfig;
-				if (Update || !PrevCachedProjectConfigs.TryGetValue(ProjectRef.Id, out ProjectConfig))
+				if (!Update && PrevCachedProjectConfigs.TryGetValue(ProjectRef.Id, out (ProjectConfig Config, string Revision) Result) && Result.Revision == Revision)
+				{
+					ProjectConfig = Result.Config;
+				}
+				else
 				{
 					Logger.LogInformation("Caching configuration for project {ProjectId} ({Revision})", ProjectRef.Id, Revision);
 					try
@@ -207,7 +211,7 @@ namespace HordeServer.Services
 					ProjectLogos.Add((ProjectRef.Id, CombinePaths(ProjectPath, ProjectConfig.Logo)));
 				}
 
-				CachedProjectConfigs[ProjectRef.Id] = ProjectConfig;
+				CachedProjectConfigs[ProjectRef.Id] = (ProjectConfig, Revision);
 				StreamConfigs.AddRange(ProjectConfig.Streams.Select(x => (ProjectRef.Id, x, CombinePaths(ProjectPath, x.Path))));
 			}
 
