Horde: Fix an instance of the Deleted flag not being cleared on a stream.

diff --git a/Engine/Source/Programs/Horde/HordeServer/Services/ConfigService.cs b/Engine/Source/Programs/Horde/HordeServer/Services/ConfigService.cs
--- a/Engine/Source/Programs/Horde/HordeServer/Services/ConfigService.cs
+++ b/Engine/Source/Programs/Horde/HordeServer/Services/ConfigService.cs
@@ -2,6 +2,7 @@
 
 using EpicGames.Core;
 using HordeServer.Api;
+using HordeServer.Collections;
 using HordeServer.Controllers;
 using HordeServer.Models;
 using HordeServer.Notifications;
@@ -262,7 +263,7 @@ namespace HordeServer.Services
 						try
 						{
 							StreamConfig StreamConfig = await ReadDataAsync<StreamConfig>(StreamPath);
-							Stream = await StreamService.StreamCollection.TryCreateOrReplaceAsync(StreamRef.Id, Stream, StreamPath.ToString(), Revision, ProjectId, StreamConfig);
+							Stream = await StreamService.StreamCollection.CreateOrReplaceAsync(StreamRef.Id, Stream, StreamPath.ToString(), Revision, ProjectId, StreamConfig);
 						}
 						catch (Exception Ex)
 						{
