Horde: Add a Redis message queue type, which allows posting/receiving messages from different channels or automatically expiring them after a given time.

diff --git a/Engine/Source/Programs/Horde/HordeServer/Utilities/RedisExtensions.cs b/Engine/Source/Programs/Horde/HordeServer/Utilities/RedisExtensions.cs
new file mode 100644
--- /dev/null
+++ b/Engine/Source/Programs/Horde/HordeServer/Utilities/RedisExtensions.cs
@@ -0,0 +1,36 @@
+// Copyright Epic Games, Inc. All Rights Reserved.
+
+using StackExchange.Redis;
+using System;
+using System.Collections.Generic;
+using System.Linq;
+using System.Threading.Tasks;
+
+namespace HordeServer.Utilities
+{
+	/// <summary>
+	/// Extension methods for Redis
+	/// </summary>
+	static class RedisExtensions
+	{
+		/// <summary>
+		/// Convert a RedisValue to a RedisKey
+		/// </summary>
+		/// <param name="Value"></param>
+		/// <returns></returns>
+		public static RedisKey AsKey(this RedisValue Value)
+		{
+			return (byte[])Value;
+		}
+
+		/// <summary>
+		/// Convert a RedisKey to a RedisValue
+		/// </summary>
+		/// <param name="Key"></param>
+		/// <returns></returns>
+		public static RedisValue AsValue(this RedisKey Key)
+		{
+			return (byte[])Key;
+		}
+	}
+}
diff --git a/Engine/Source/Programs/Horde/HordeServer/Utilities/RedisMessageQueue.cs b/Engine/Source/Programs/Horde/HordeServer/Utilities/RedisMessageQueue.cs
new file mode 100644
--- /dev/null
+++ b/Engine/Source/Programs/Horde/HordeServer/Utilities/RedisMessageQueue.cs
@@ -0,0 +1,202 @@
+// Copyright Epic Games, Inc. All Rights Reserved.
+
+using StackExchange.Redis;
+using System;
+using System.Collections.Concurrent;
+using System.Collections.Generic;
+using System.Linq;
+using System.Threading.Tasks;
+
+namespace HordeServer.Utilities
+{
+	/// <summary>
+	/// Implements a multi-channel message queue using Redis
+	/// </summary>
+	class RedisMessageQueue<TChannel, TMessage>
+	{
+		/// <summary>
+		/// The Redis database
+		/// </summary>
+		IDatabase Redis;
+
+		/// <summary>
+		/// Prefix for all keys used by the queue
+		/// </summary>
+		RedisKey KeyPrefix;
+
+		/// <summary>
+		/// Key for the index, storing the name of each queue
+		/// </summary>
+		RedisKey IndexKey;
+
+		/// <summary>
+		/// Cache tracking when each channel was last refreshed in the index 
+		/// </summary>
+		Dictionary<RedisKey, double> ChannelNameToScore = new Dictionary<RedisKey, double>();
+
+		/// <summary>
+		/// Time after which entries should be expired
+		/// </summary>
+		public TimeSpan ExpireTime { get; set; } = TimeSpan.FromSeconds(120);
+
+		/// <summary>
+		/// Age at which channel index entries should be updated
+		/// </summary>
+		TimeSpan IndexRefreshTime => ExpireTime * 0.5;
+
+		/// <summary>
+		/// How often to clean up expired entries from the index cache
+		/// </summary>
+		TimeSpan IndexShrinkTime = TimeSpan.FromSeconds(10.0);
+
+		/// <summary>
+		/// Gets the known channel names
+		/// </summary>
+		public IEnumerable<RedisKey> ChannelNames => ChannelNameToScore.Keys;
+
+		/// <summary>
+		/// Last time that the index  time that the ChannelNameToScore was synchronized with the database
+		/// </summary>
+		DateTime LastShrinkTime = DateTime.MinValue;
+
+		/// <summary>
+		/// Constructor
+		/// </summary>
+		/// <param name="Redis">The Redis database instance</param>
+		/// <param name="KeyPrefix">Prefix for keys to use in the database</param>
+		public RedisMessageQueue(IDatabase Redis, RedisKey KeyPrefix)
+		{
+			this.Redis = Redis;
+			this.KeyPrefix = KeyPrefix;
+			this.IndexKey = KeyPrefix.Append("index");
+		}
+
+		/// <summary>
+		/// Add an item to a channel
+		/// </summary>
+		/// <param name="Channel"></param>
+		/// <param name="Message">Message to post</param>
+		/// <returns></returns>
+		public Task PostAsync(TChannel Channel, TMessage Message)
+		{
+			return PostAsync(Channel, new[] { Message });
+		}
+
+		/// <summary>
+		/// Add an item to a channel
+		/// </summary>
+		/// <param name="Channel"></param>
+		/// <param name="Messages"></param>
+		/// <returns></returns>
+		public async Task PostAsync(TChannel Channel, TMessage[] Messages)
+		{
+			// Periodically remove expired entries from the cached index
+			DateTime UtcNow = DateTime.UtcNow;
+			if (UtcNow > LastShrinkTime + IndexShrinkTime)
+			{
+				Shrink();
+			}
+
+			// Get the keys and message values
+			byte[] RelativeChannelKeyBytes = (byte[])RedisSerializer.Serialize(Channel);
+			RedisKey RelativeChannelKey = RelativeChannelKeyBytes;
+			RedisKey FullChannelKey = KeyPrefix.Append(RelativeChannelKey);
+			SortedSetEntry[] MessageEntries = Array.ConvertAll(Messages, x => new SortedSetEntry(RedisSerializer.Serialize(x), UtcNow.Ticks));
+
+			// Minimum score to avoid updating the index
+			double MinScore = (UtcNow - IndexRefreshTime).Ticks;
+
+			// Refresh this queue time if necessary
+			double Score;
+			if (ChannelNameToScore.TryGetValue(RelativeChannelKey, out Score) && Score > MinScore)
+			{
+				await Redis.SortedSetAddAsync(FullChannelKey, MessageEntries, CommandFlags.None);
+			}
+			else
+			{
+				IBatch Batch = Redis.CreateBatch();
+
+				List<Task> Tasks = new List<Task>();
+				Tasks.Add(Batch.SortedSetAddAsync(IndexKey, RelativeChannelKeyBytes, UtcNow.Ticks, CommandFlags.None));
+				Tasks.Add(Batch.SortedSetAddAsync(FullChannelKey, MessageEntries, CommandFlags.None));
+
+				Batch.Execute();
+
+				await Task.WhenAll(Tasks);
+
+				lock (ChannelNameToScore)
+				{
+					ChannelNameToScore[RelativeChannelKey] = UtcNow.Ticks;
+				}
+			}
+		}
+
+		/// <summary>
+		/// Receives a set of items on a channel
+		/// </summary>
+		/// <param name="Channel"></param>
+		/// <param name="MaxItems"></param>
+		/// <returns></returns>
+		public async Task<List<TMessage>> DequeueAsync(TChannel Channel, int MaxItems)
+		{
+			RedisKey ChannelKey = KeyPrefix.Append(RedisSerializer.Serialize(Channel).AsKey());
+
+			double MinScore = (DateTime.UtcNow - ExpireTime).Ticks;
+
+			List<TMessage> Messages = new List<TMessage>();
+			while (Messages.Count < MaxItems)
+			{
+				int Count = MaxItems - Messages.Count;
+
+				SortedSetEntry[] Entries = await Redis.SortedSetPopAsync(ChannelKey, Count);
+				Messages.AddRange(Entries.Where(x => x.Score >= MinScore).Select(x => RedisSerializer.Deserialize<TMessage>(x.Element)));
+
+				if (Entries.Length < Count)
+				{
+					break;
+				}
+			}
+			return Messages;
+		}
+
+		/// <summary>
+		/// Updates cached values used by the queue
+		/// </summary>
+		/// <returns></returns>
+		public void Shrink()
+		{
+			DateTime UtcNow = DateTime.UtcNow;
+			lock (ChannelNameToScore)
+			{
+				double MinScore = (UtcNow - ExpireTime).Ticks;
+
+				RedisKey[] RemoveKeys = ChannelNameToScore.Where(x => x.Value < MinScore).Select(x => x.Key).ToArray();
+				foreach (RedisKey RemoveKey in RemoveKeys)
+				{
+					ChannelNameToScore.Remove(RemoveKey);
+				}
+			}
+			LastShrinkTime = UtcNow;
+		}
+
+		/// <summary>
+		/// Removes items which have expired from the queue
+		/// </summary>
+		/// <returns></returns>
+		public async Task RemovedExpiredAsync()
+		{
+			double MinScore = (DateTime.UtcNow - ExpireTime).Ticks;
+
+			RedisValue[] Values = await Redis.SortedSetRangeByScoreAsync(IndexKey, Double.NegativeInfinity, MinScore);
+			if (Values.Length > 0)
+			{
+				foreach (RedisValue Value in Values)
+				{
+					RedisKey QueueKey = KeyPrefix.Append(Value.AsKey());
+					await Redis.SortedSetRemoveRangeByScoreAsync(QueueKey, Double.NegativeInfinity, MinScore);
+				}
+				await Redis.SortedSetRemoveRangeByScoreAsync(IndexKey, Double.NegativeInfinity, MinScore);
+			}
+		}
+	}
+}
diff --git a/Engine/Source/Programs/Horde/HordeServer/Utilities/RedisSerializer.cs b/Engine/Source/Programs/Horde/HordeServer/Utilities/RedisSerializer.cs
new file mode 100644
--- /dev/null
+++ b/Engine/Source/Programs/Horde/HordeServer/Utilities/RedisSerializer.cs
@@ -0,0 +1,158 @@
+// Copyright Epic Games, Inc. All Rights Reserved.
+
+using EpicGames.Serialization;
+using StackExchange.Redis;
+using System;
+using System.Collections.Generic;
+using System.ComponentModel;
+using System.Linq;
+using System.Reflection;
+using System.Threading.Tasks;
+
+namespace HordeServer.Utilities
+{
+	/// <summary>
+	/// Attribute specifying the converter type to use for a class
+	/// </summary>
+	public class RedisConverterAttribute : Attribute
+	{
+		/// <summary>
+		/// Type of the converter to use
+		/// </summary>
+		public Type ConverterType { get; }
+
+		/// <summary>
+		/// Constructor
+		/// </summary>
+		/// <param name="ConverterType">The converter type</param>
+		public RedisConverterAttribute(Type ConverterType)
+		{
+			this.ConverterType = ConverterType;
+		}
+	}
+
+	/// <summary>
+	/// Converter to and from RedisValue types
+	/// </summary>
+	public interface IRedisConverter<T>
+	{
+		/// <summary>
+		/// Serailize an object to a RedisValue
+		/// </summary>
+		/// <param name="Value"></param>
+		/// <returns></returns>
+		RedisValue ToRedisValue(T Value);
+
+		/// <summary>
+		/// Deserialize an object from a RedisValue
+		/// </summary>
+		/// <param name="Value"></param>
+		/// <returns></returns>
+		T FromRedisValue(RedisValue Value);
+	}
+
+	/// <summary>
+	/// Redis serializer that uses compact binary to serialize objects
+	/// </summary>
+	/// <typeparam name="T"></typeparam>
+	public sealed class RedisCbConverter<T> : IRedisConverter<T>
+	{
+		/// <inheritdoc/>
+		public RedisValue ToRedisValue(T Value)
+		{
+			return CbSerializer.Serialize(Value).GetView();
+		}
+
+		/// <inheritdoc/>
+		public T FromRedisValue(RedisValue Value)
+		{
+			return CbSerializer.Deserialize<T>(new CbField((byte[])Value));
+		}
+	}
+
+	/// <summary>
+	/// Handles serialization of types to RedisValue instances
+	/// </summary>
+	public static class RedisSerializer
+	{
+		class RedisStringConverter<T> : IRedisConverter<T>
+		{
+			TypeConverter TypeConverter;
+
+			public RedisStringConverter(TypeConverter TypeConverter)
+			{
+				this.TypeConverter = TypeConverter;
+			}
+
+			public RedisValue ToRedisValue(T Value) => (string)TypeConverter.ConvertTo(Value, typeof(string));
+			public T FromRedisValue(RedisValue Value) => (T)TypeConverter.ConvertFrom((string)Value);
+		}
+
+		/// <summary>
+		/// Static class for caching converter lookups
+		/// </summary>
+		/// <typeparam name="T"></typeparam>
+		class CachedConverter<T>
+		{
+			public static IRedisConverter<T> Converter = CreateConverter();
+
+			static IRedisConverter<T> CreateConverter()
+			{
+				Type Type = typeof(T);
+
+				// Check for a custom converter
+				RedisConverterAttribute? Attribute = Type.GetCustomAttribute<RedisConverterAttribute>();
+				if (Attribute != null)
+				{
+					Type ConverterType = Attribute.ConverterType;
+					if (ConverterType.IsGenericTypeDefinition)
+					{
+						ConverterType = ConverterType.MakeGenericType(Type);
+					}
+					return (IRedisConverter<T>)Activator.CreateInstance(ConverterType)!;
+				}
+
+				// Check if there's a regular converter we can use to convert to/from a string
+				TypeConverter? Converter = TypeDescriptor.GetConverter(Type);
+				if (Converter != null && Converter.CanConvertFrom(typeof(string)) && Converter.CanConvertTo(typeof(string)))
+				{
+					return new RedisStringConverter<T>(Converter);
+				}
+
+				throw new Exception($"Unable to find Redis converter for {Type.Name}");
+			}
+		}
+
+		/// <summary>
+		/// Gets the converter for a particular type
+		/// </summary>
+		/// <typeparam name="T"></typeparam>
+		/// <returns></returns>
+		public static IRedisConverter<T> GetConverter<T>()
+		{
+			return CachedConverter<T>.Converter;
+		}
+
+		/// <summary>
+		/// Serialize an object to a <see cref="RedisValue"/>
+		/// </summary>
+		/// <typeparam name="T"></typeparam>
+		/// <param name="Value"></param>
+		/// <returns></returns>
+		public static RedisValue Serialize<T>(T Value)
+		{
+			return CachedConverter<T>.Converter.ToRedisValue(Value);
+		}
+
+		/// <summary>
+		/// Deserialize a <see cref="RedisValue"/>
+		/// </summary>
+		/// <typeparam name="T"></typeparam>
+		/// <param name="Value"></param>
+		/// <returns></returns>
+		public static T Deserialize<T>(RedisValue Value)
+		{
+			return CachedConverter<T>.Converter.FromRedisValue(Value);
+		}
+	}
+}
