Horde: Implement new gRPC interfaces for the blob store and ref table. Currently exist alongside the older implementations.

diff --git a/Engine/Source/Programs/Horde/HordeServer/Startup.cs b/Engine/Source/Programs/Horde/HordeServer/Startup.cs
--- a/Engine/Source/Programs/Horde/HordeServer/Startup.cs
+++ b/Engine/Source/Programs/Horde/HordeServer/Startup.cs
@@ -770,7 +770,9 @@ namespace HordeServer
 				Endpoints.MapGrpcService<ExecutionService>();
 
 				Endpoints.MapGrpcService<BlobRpc>();
+				Endpoints.MapGrpcService<BlobStoreRpc>();
 				Endpoints.MapGrpcService<RefRpc>();
+				Endpoints.MapGrpcService<RefTableRpc>();
 
 				Endpoints.MapGrpcReflectionService();
 
diff --git a/Engine/Source/Programs/Horde/HordeServer/Storage/Collections/BlobStoreRpc.cs b/Engine/Source/Programs/Horde/HordeServer/Storage/Collections/BlobStoreRpc.cs
new file mode 100644
--- /dev/null
+++ b/Engine/Source/Programs/Horde/HordeServer/Storage/Collections/BlobStoreRpc.cs
@@ -0,0 +1,81 @@
+// Copyright Epic Games, Inc. All Rights Reserved.
+
+using EpicGames.Core;
+using EpicGames.Horde.Common;
+using EpicGames.Horde.Storage;
+using Google.Protobuf;
+using Grpc.Core;
+using HordeServer.Utilities;
+using Microsoft.Extensions.Options;
+using System;
+using System.Collections.Generic;
+using System.IO;
+using System.Linq;
+using System.Threading.Tasks;
+
+namespace HordeServer.Storage.Collections
+{
+	using NamespaceId = StringId<INamespace>;
+
+	class BlobStoreRpc : BlobStore.BlobStoreBase
+	{
+		IBlobCollection BlobCollection;
+
+		/// <summary>
+		/// Constructor
+		/// </summary>
+		/// <param name="BlobCollection"></param>
+		public BlobStoreRpc(IBlobCollection BlobCollection)
+		{
+			this.BlobCollection = BlobCollection;
+		}
+
+		/// <inheritdoc/>
+		public override async Task<AddBlobsResponse> Add(AddBlobsRequest Request, ServerCallContext Context)
+		{
+			NamespaceId NamespaceId = new NamespaceId(Request.NamespaceId);
+			foreach (AddBlobRequest Blob in Request.Blobs)
+			{
+				await BlobCollection.WriteBytesAsync(NamespaceId, Blob.Hash, Blob.Data.ToByteArray());
+			}
+			return new AddBlobsResponse();
+		}
+
+		/// <inheritdoc/>
+		public override async Task<GetBlobsResponse> Get(GetBlobsRequest Request, ServerCallContext Context)
+		{
+			NamespaceId NamespaceId = new NamespaceId(Request.NamespaceId);
+
+			List<Task<byte[]?>> Tasks = new List<Task<byte[]?>>();
+			foreach (IoHash Hash in Request.Blobs)
+			{
+				Tasks.Add(Task.Run(() => BlobCollection.ReadBytesAsync(NamespaceId, Hash)));
+			}
+
+			GetBlobsResponse Response = new GetBlobsResponse();
+			for (int Idx = 0; Idx < Request.Blobs.Count; Idx++)
+			{
+				byte[]? Data = await Tasks[Idx];
+
+				GetBlobResponse Blob = new GetBlobResponse();
+				Blob.Data = ByteString.CopyFrom(Data);
+				Blob.Exists = Data != null;
+				Response.Blobs.Add(Blob);
+			}
+			return Response;
+		}
+
+		/// <inheritdoc/>
+		public override async Task<FindMissingBlobsResponse> FindMissing(FindMissingBlobsRequest Request, ServerCallContext Context)
+		{
+			NamespaceId NamespaceId = new NamespaceId(Request.NamespaceId);
+			IEnumerable<IoHash> Hashes = Request.Blobs.Select(x => (IoHash)x);
+
+			List<IoHash> ExistingHashes = await BlobCollection.ExistsAsync(NamespaceId, Hashes);
+
+			FindMissingBlobsResponse Response = new FindMissingBlobsResponse();
+			Response.Blobs.Add(Hashes.Except(ExistingHashes).Select(x => (IoHashWrapper)x));
+			return Response;
+		}
+	}
+}
diff --git a/Engine/Source/Programs/Horde/HordeServer/Storage/Collections/RefTableRpc.cs b/Engine/Source/Programs/Horde/HordeServer/Storage/Collections/RefTableRpc.cs
new file mode 100644
--- /dev/null
+++ b/Engine/Source/Programs/Horde/HordeServer/Storage/Collections/RefTableRpc.cs
@@ -0,0 +1,88 @@
+// Copyright Epic Games, Inc. All Rights Reserved.
+
+using System;
+using System.Collections.Generic;
+using System.IO;
+using System.Linq;
+using System.Threading.Tasks;
+using EpicGames.Core;
+using EpicGames.Horde.Common;
+using EpicGames.Horde.Storage;
+using EpicGames.Serialization;
+using Google.Protobuf;
+using Google.Protobuf.WellKnownTypes;
+using Grpc.Core;
+using HordeServer.Utilities;
+
+namespace HordeServer.Storage.Collections
+{
+	using NamespaceId = StringId<INamespace>;
+	using BucketId = StringId<IBucket>;
+
+	class RefTableRpc : RefTable.RefTableBase
+	{
+		IRefCollection RefCollection;
+
+		/// <summary>
+		/// Constructor
+		/// </summary>
+		/// <param name="RefCollection"></param>
+		public RefTableRpc(IRefCollection RefCollection)
+		{
+			this.RefCollection = RefCollection;
+		}
+
+		/// <inheritdoc/>
+		public override async Task<SetRefResponse> Set(SetRefRequest Request, ServerCallContext Context)
+		{
+			NamespaceId NamespaceId = new NamespaceId(Request.NamespaceId);
+			BucketId BucketId = new BucketId(Request.BucketId);
+			List<IoHash> MissingHashes = await RefCollection.SetAsync(NamespaceId, BucketId, Request.Name, Request.Object);
+
+			return new SetRefResponse(MissingHashes);
+		}
+
+		/// <inheritdoc/>
+		public override async Task<GetRefResponse> Get(GetRefRequest Request, ServerCallContext Context)
+		{
+			NamespaceId NamespaceId = new NamespaceId(Request.NamespaceId);
+			BucketId BucketId = new BucketId(Request.BucketId);
+
+			IRef? Ref = await RefCollection.GetAsync(NamespaceId, BucketId, Request.Name);
+			if(Ref == null)
+			{
+				throw new RpcException(new Status(StatusCode.NotFound, $"Unable to find ref {NamespaceId}/{BucketId}/{Request.Name}"));
+			}
+
+			return new GetRefResponse(Ref.Value, Ref.LastAccessTime, Ref.Finalized);
+		}
+
+		/// <inheritdoc/>
+		public override async Task<TouchResponse> Touch(TouchRequest Request, ServerCallContext Context)
+		{
+			NamespaceId NamespaceId = new NamespaceId(Request.NamespaceId);
+			BucketId BucketId = new BucketId(Request.BucketId);
+
+			if (!await RefCollection.TouchAsync(NamespaceId, BucketId, Request.Name))
+			{
+				throw new RpcException(new Status(StatusCode.NotFound, $"Unable to find ref {NamespaceId}/{BucketId}/{Request.Name}"));
+			}
+
+			return new TouchResponse();
+		}
+
+		/// <inheritdoc/>
+		public override async Task<DeleteResponse> Delete(DeleteRequest Request, ServerCallContext Context)
+		{
+			NamespaceId NamespaceId = new NamespaceId(Request.NamespaceId);
+			BucketId BucketId = new BucketId(Request.BucketId);
+
+			if (!await RefCollection.DeleteAsync(NamespaceId, BucketId, Request.Name))
+			{
+				throw new RpcException(new Status(StatusCode.NotFound, $"Unable to find ref {NamespaceId}/{BucketId}/{Request.Name}"));
+			}
+
+			return new DeleteResponse();
+		}
+	}
+}
diff --git a/Engine/Source/Programs/Shared/EpicGames.Horde/EpicGames.Horde.csproj b/Engine/Source/Programs/Shared/EpicGames.Horde/EpicGames.Horde.csproj
--- a/Engine/Source/Programs/Shared/EpicGames.Horde/EpicGames.Horde.csproj
+++ b/Engine/Source/Programs/Shared/EpicGames.Horde/EpicGames.Horde.csproj
@@ -44,7 +44,7 @@
   <ItemGroup>
     <Protobuf Include="Protos\horde\common\cbobject_wrapper.proto" ProtoRoot="Protos" />
     <Protobuf Include="Protos\horde\common\iohash_wrapper.proto" ProtoRoot="Protos" />
-    <Protobuf Include="Protos\horde\storage\blob_store.proto" ProtoRoot="Protos" />
-    <Protobuf Include="Protos\horde\storage\ref_table.proto" ProtoRoot="Protos" />
+    <Protobuf Include="Protos\horde\storage\blob_store_messages.proto" ProtoRoot="Protos" />
+    <Protobuf Include="Protos\horde\storage\ref_table_messages.proto" ProtoRoot="Protos" />
   </ItemGroup>
 </Project>
\ No newline at end of file
