Horde: Rearrange public Horde classes and projects.

* EpicGames.Horde contains all types and functionality common to server, agent, and clients of Horde.
* EpicGames.Horde.Client contains client interfaces for calling server functionality (but not server stubs for implementing those interfaces). Together with EpicGames.Horde, this defines the public interface for Horde. Note that types are still defined within the EpicGames.Horde namespace.
* HordeCommon contains functionality common to the server and agent (eg. session/lease messaging, task data structures, etc...), but which is not exposed to clients of Horde.

EpicGames.Horde.Common contains legacy REAPI functionality and will be removed once CB-REAPI is up and running.

diff --git a/Engine/Source/Programs/Horde/HordeCommon/HordeCommon.csproj b/Engine/Source/Programs/Horde/HordeCommon/HordeCommon.csproj
--- a/Engine/Source/Programs/Horde/HordeCommon/HordeCommon.csproj
+++ b/Engine/Source/Programs/Horde/HordeCommon/HordeCommon.csproj
@@ -16,27 +16,23 @@
   </ItemGroup>
 
   <ItemGroup>
-    <ProjectReference Include="..\..\Shared\EpicGames.Horde.Common\EpicGames.Horde.Common.csproj" />
-  </ItemGroup>
-
-  <ItemGroup>
-    <Protobuf Include="..\..\Shared\EpicGames.Horde.Common\Protos\horde\horde_rpc_messages.proto" ProtoRoot="..\..\Shared\EpicGames.Horde.Common\Protos" />
-    <Protobuf Include="..\..\Shared\EpicGames.Horde.Common\Protos\horde\common\agent_status.proto" ProtoRoot="..\..\Shared\EpicGames.Horde.Common\Protos" />
-    <Protobuf Include="..\..\Shared\EpicGames.Horde.Common\Protos\horde\common\event_severity.proto" ProtoRoot="..\..\Shared\EpicGames.Horde.Common\Protos" />
-    <Protobuf Include="..\..\Shared\EpicGames.Horde.Common\Protos\horde\common\jobstep_outcome.proto" ProtoRoot="..\..\Shared\EpicGames.Horde.Common\Protos" />
-    <Protobuf Include="..\..\Shared\EpicGames.Horde.Common\Protos\horde\common\jobstep_state.proto" ProtoRoot="..\..\Shared\EpicGames.Horde.Common\Protos" />
-    <Protobuf Include="..\..\Shared\EpicGames.Horde.Common\Protos\horde\common\label_change.proto" ProtoRoot="..\..\Shared\EpicGames.Horde.Common\Protos" />
-    <Protobuf Include="..\..\Shared\EpicGames.Horde.Common\Protos\horde\common\lease_outcome.proto" ProtoRoot="..\..\Shared\EpicGames.Horde.Common\Protos" />
-    <Protobuf Include="..\..\Shared\EpicGames.Horde.Common\Protos\horde\common\lease_state.proto" ProtoRoot="..\..\Shared\EpicGames.Horde.Common\Protos" />
-    <Protobuf Include="..\..\Shared\EpicGames.Horde.Common\Protos\horde\common\priority.proto" ProtoRoot="..\..\Shared\EpicGames.Horde.Common\Protos" />
-    <Protobuf Include="..\..\Shared\EpicGames.Horde.Common\Protos\horde\messages\agent_capabilities.proto" ProtoRoot="..\..\Shared\EpicGames.Horde.Common\Protos" />
-    <Protobuf Include="..\..\Shared\EpicGames.Horde.Common\Protos\horde\messages\agent_requirements.proto" ProtoRoot="..\..\Shared\EpicGames.Horde.Common\Protos" />
-    <Protobuf Include="..\..\Shared\EpicGames.Horde.Common\Protos\horde\messages\agent_workspace.proto" ProtoRoot="..\..\Shared\EpicGames.Horde.Common\Protos" />
-    <Protobuf Include="..\..\Shared\EpicGames.Horde.Common\Protos\horde\messages\lease.proto" ProtoRoot="..\..\Shared\EpicGames.Horde.Common\Protos" />
-    <Protobuf Include="..\..\Shared\EpicGames.Horde.Common\Protos\horde\tasks\conform_task.proto" ProtoRoot="..\..\Shared\EpicGames.Horde.Common\Protos" />
-    <Protobuf Include="..\..\Shared\EpicGames.Horde.Common\Protos\horde\tasks\job_task.proto" ProtoRoot="..\..\Shared\EpicGames.Horde.Common\Protos" />
-    <Protobuf Include="..\..\Shared\EpicGames.Horde.Common\Protos\horde\tasks\restart_task.proto" ProtoRoot="..\..\Shared\EpicGames.Horde.Common\Protos" />
-    <Protobuf Include="..\..\Shared\EpicGames.Horde.Common\Protos\horde\tasks\shutdown_task.proto" ProtoRoot="..\..\Shared\EpicGames.Horde.Common\Protos" />
-    <Protobuf Include="..\..\Shared\EpicGames.Horde.Common\Protos\horde\tasks\upgrade_task.proto" ProtoRoot="..\..\Shared\EpicGames.Horde.Common\Protos" />
+    <Protobuf Include="Protos\horde\horde_rpc_messages.proto" ProtoRoot="Protos" />
+    <Protobuf Include="Protos\horde\common\agent_status.proto" ProtoRoot="Protos" />
+    <Protobuf Include="Protos\horde\common\event_severity.proto" ProtoRoot="Protos" />
+    <Protobuf Include="Protos\horde\common\jobstep_outcome.proto" ProtoRoot="Protos" />
+    <Protobuf Include="Protos\horde\common\jobstep_state.proto" ProtoRoot="Protos" />
+    <Protobuf Include="Protos\horde\common\label_change.proto" ProtoRoot="Protos" />
+    <Protobuf Include="Protos\horde\common\lease_outcome.proto" ProtoRoot="Protos" />
+    <Protobuf Include="Protos\horde\common\lease_state.proto" ProtoRoot="Protos" />
+    <Protobuf Include="Protos\horde\common\priority.proto" ProtoRoot="Protos" />
+    <Protobuf Include="Protos\horde\messages\agent_capabilities.proto" ProtoRoot="Protos" />
+    <Protobuf Include="Protos\horde\messages\agent_requirements.proto" ProtoRoot="Protos" />
+    <Protobuf Include="Protos\horde\messages\agent_workspace.proto" ProtoRoot="Protos" />
+    <Protobuf Include="Protos\horde\messages\lease.proto" ProtoRoot="Protos" />
+    <Protobuf Include="Protos\horde\tasks\conform_task.proto" ProtoRoot="Protos" />
+    <Protobuf Include="Protos\horde\tasks\job_task.proto" ProtoRoot="Protos" />
+    <Protobuf Include="Protos\horde\tasks\restart_task.proto" ProtoRoot="Protos" />
+    <Protobuf Include="Protos\horde\tasks\shutdown_task.proto" ProtoRoot="Protos" />
+    <Protobuf Include="Protos\horde\tasks\upgrade_task.proto" ProtoRoot="Protos" />
   </ItemGroup>
 </Project>
diff --git a/Engine/Source/Programs/Horde/HordeServer/HordeServer.csproj b/Engine/Source/Programs/Horde/HordeServer/HordeServer.csproj
--- a/Engine/Source/Programs/Horde/HordeServer/HordeServer.csproj
+++ b/Engine/Source/Programs/Horde/HordeServer/HordeServer.csproj
@@ -62,10 +62,10 @@
     <Protobuf Include="..\..\Shared\EpicGames.Horde\Protos\horde\storage\blob_store.proto" ProtoRoot="..\..\Shared\EpicGames.Horde\Protos" GrpcServices="Server" />
     <Protobuf Include="..\..\Shared\EpicGames.Horde.Common\Protos\google\bytestream\bytestream.proto" ProtoRoot="..\..\Shared\EpicGames.Horde.Common\Protos" GrpcServices="Server" />
     <Protobuf Include="..\..\Shared\EpicGames.Horde.Common\Protos\horde\action_rpc.proto" ProtoRoot="..\..\Shared\EpicGames.Horde.Common\Protos" GrpcServices="Server" />
-    <Protobuf Include="..\..\Shared\EpicGames.Horde.Common\Protos\horde\horde_rpc.proto" ProtoRoot="..\..\Shared\EpicGames.Horde.Common\Protos" GrpcServices="Server" />
     <Protobuf Include="..\..\Shared\EpicGames.Horde.Common\Protos\horde\tasks\action_task.proto" ProtoRoot="..\..\Shared\EpicGames.Horde.Common\Protos" GrpcServices="Server" />
     <Protobuf Include="..\..\Shared\EpicGames.Horde.Common\Protos\horde\storage\ref_service.proto" ProtoRoot="..\..\Shared\EpicGames.Horde.Common\Protos" GrpcServices="Server" />
     <Protobuf Include="..\..\Shared\EpicGames.Horde.Common\Protos\horde\storage\blob_service.proto" ProtoRoot="..\..\Shared\EpicGames.Horde.Common\Protos" GrpcServices="Server" />
+    <Protobuf Include="..\HordeCommon\Protos\horde\horde_rpc.proto" ProtoRoot="..\HordeCommon\Protos" GrpcServices="Server" />
   </ItemGroup>
 
   <ItemGroup>
diff --git a/Engine/Source/Programs/Shared/EpicGames.Horde.Client/EpicGames.Horde.Client.csproj b/Engine/Source/Programs/Shared/EpicGames.Horde.Client/EpicGames.Horde.Client.csproj
new file mode 100644
--- /dev/null
+++ b/Engine/Source/Programs/Shared/EpicGames.Horde.Client/EpicGames.Horde.Client.csproj
@@ -0,0 +1,57 @@
+<Project Sdk="Microsoft.NET.Sdk">
+  <PropertyGroup>
+    <TargetFramework>netcoreapp3.1</TargetFramework>
+    <Configuration Condition=" '$(Configuration)' == '' ">Development</Configuration>
+    <OutputType>Library</OutputType>
+    <Deterministic>true</Deterministic>
+    <RootNamespace>EpicGames.Horde</RootNamespace>
+    <GenerateAssemblyInfo>false</GenerateAssemblyInfo>
+    <GenerateTargetFrameworkAttribute>false</GenerateTargetFrameworkAttribute>
+    <Configurations>Debug;Release;Development</Configurations>
+    <Nullable>enable</Nullable>
+  </PropertyGroup>
+  <PropertyGroup Condition=" '$(Configuration)|$(Platform)' == 'Debug|AnyCPU' ">
+    <DocumentationFile></DocumentationFile>
+  </PropertyGroup>
+  <PropertyGroup Condition="'$(Configuration)|$(Platform)' == 'Development|AnyCPU'">
+    <DocumentationFile></DocumentationFile>
+    <Optimize>false</Optimize>
+  </PropertyGroup>
+  <PropertyGroup Condition="'$(Configuration)|$(Platform)' == 'Release|AnyCPU'">
+    <DocumentationFile></DocumentationFile>
+    <Optimize>true</Optimize>
+  </PropertyGroup>
+  <ItemGroup>
+    <Compile Include="..\..\DotNETCommon\MetaData.cs">
+      <Link>Properties\MetaData.cs</Link>
+    </Compile>
+  </ItemGroup>
+  <ItemGroup>
+    <PackageReference Include="Google.Protobuf" Version="3.12.3" />
+    <PackageReference Include="Grpc.Core.Api" Version="2.35.0" />
+    <PackageReference Include="Grpc.Tools" Version="2.30.0">
+      <PrivateAssets>all</PrivateAssets>
+      <IncludeAssets>runtime; build; native; contentfiles; analyzers; buildtransitive</IncludeAssets>
+    </PackageReference>
+    <PackageReference Include="Microsoft.CodeAnalysis.CSharp.Workspaces" Version="3.10.0" />
+    <PackageReference Include="Microsoft.Extensions.Logging" Version="5.0.0" />
+  </ItemGroup>
+  <ItemGroup>
+    <ProjectReference Include="..\EpicGames.Core\EpicGames.Core.csproj" />
+    <ProjectReference Include="..\EpicGames.Horde\EpicGames.Horde.csproj" />
+    <ProjectReference Include="..\EpicGames.IoHash\EpicGames.IoHash.csproj" />
+    <ProjectReference Include="..\EpicGames.Serialization\EpicGames.Serialization.csproj" />
+  </ItemGroup>
+  <ItemGroup>
+    <Content Include="..\EpicGames.Horde\Protos\horde\storage\blob_store.proto">
+      <Link>Protos\horde\storage\blob_store.proto</Link>
+    </Content>
+    <Content Include="..\EpicGames.Horde\Protos\horde\storage\ref_table.proto">
+      <Link>Protos\horde\storage\ref_table.proto</Link>
+    </Content>
+  </ItemGroup>
+  <ItemGroup>
+    <Protobuf Include="..\EpicGames.Horde\Protos\horde\storage\blob_store.proto" ProtoRoot="..\EpicGames.Horde\Protos" GrpcServices="Client" />
+    <Protobuf Include="..\EpicGames.Horde\Protos\horde\storage\ref_table.proto" ProtoRoot="..\EpicGames.Horde\Protos" GrpcServices="Client" />
+  </ItemGroup>
+</Project>
\ No newline at end of file
diff --git a/Engine/Source/Programs/Shared/EpicGames.Horde/EpicGames.Horde.csproj b/Engine/Source/Programs/Shared/EpicGames.Horde/EpicGames.Horde.csproj
--- a/Engine/Source/Programs/Shared/EpicGames.Horde/EpicGames.Horde.csproj
+++ b/Engine/Source/Programs/Shared/EpicGames.Horde/EpicGames.Horde.csproj
@@ -1,4 +1,4 @@
-<Project Sdk="Microsoft.NET.Sdk">
+﻿<Project Sdk="Microsoft.NET.Sdk">
   <PropertyGroup>
     <TargetFramework>netcoreapp3.1</TargetFramework>
     <Configuration Condition=" '$(Configuration)' == '' ">Development</Configuration>
