Horde: Remove unused action_lease protobuf.

diff --git a/Engine/Source/Programs/Horde/HordeAgent/HordeAgent.csproj b/Engine/Source/Programs/Horde/HordeAgent/HordeAgent.csproj
--- a/Engine/Source/Programs/Horde/HordeAgent/HordeAgent.csproj
+++ b/Engine/Source/Programs/Horde/HordeAgent/HordeAgent.csproj
@@ -92,7 +92,7 @@
     <Protobuf Include="..\..\Shared\EpicGames.Horde.Common\Protos\google\devtools\remoteworkers\v1test2\bots.proto" ProtoRoot="..\..\Shared\EpicGames.Horde.Common\Protos" GrpcServices="Client" />
     <Protobuf Include="..\..\Shared\EpicGames.Horde.Common\Protos\google\devtools\remoteworkers\v1test2\worker.proto" ProtoRoot="..\..\Shared\EpicGames.Horde.Common\Protos" GrpcServices="Client" />
     <Protobuf Include="..\..\Shared\EpicGames.Horde.Common\Protos\horde\action_rpc.proto" ProtoRoot="..\..\Shared\EpicGames.Horde.Common\Protos" GrpcServices="Client" />
-    <Protobuf Include="..\..\Shared\EpicGames.Horde.Common\Protos\horde\horde_rpc.proto" ProtoRoot="..\..\Shared\EpicGames.Horde.Common\Protos" GrpcServices="Client" />
+    <Protobuf Include="..\HordeCommon\Protos\horde\horde_rpc.proto" ProtoRoot="..\HordeCommon\Protos" GrpcServices="Client" />
     <Protobuf Include="..\..\Shared\EpicGames.Horde.Common\Protos\horde\tasks\action_task.proto" ProtoRoot="..\..\Shared\EpicGames.Horde.Common\Protos" GrpcServices="Client" />
   </ItemGroup>
 
@@ -117,13 +117,13 @@
 
   <Target Name="PublishSymbols" AfterTargets="AfterBuild" Condition="'$(FullSymbolsDir)' != ''">
     <ItemGroup>
-      <OutputFiles Include="$(OutDir)**"/>
-      <BinaryFilesBase Include="$(OutDir)\*.dll"/>
-      <BinaryFiles Include="@(BinaryFilesBase)" Condition="Exists('%(RelativeDir)%(FileName).pdb')"/>
+      <OutputFiles Include="$(OutDir)**" />
+      <BinaryFilesBase Include="$(OutDir)\*.dll" />
+      <BinaryFiles Include="@(BinaryFilesBase)" Condition="Exists('%(RelativeDir)%(FileName).pdb')" />
     </ItemGroup>
 
     <MakeDir Directories="$(FullSymbolsDir)" />
-    <Copy SourceFiles="@(Files)" DestinationFolder="$(FullSymbolsDir)"/>
+    <Copy SourceFiles="@(Files)" DestinationFolder="$(FullSymbolsDir)" />
     <Exec Command="&quot;$(Pdb2PdbPath)&quot; &quot;%(BinaryFiles.Identity)&quot; /out &quot;$(FullSymbolsDir)\%(BinaryFiles.FileName).pdb&quot; /nowarn 21" IgnoreExitCode="false" />
   </Target>
   
