Fix a couple of warnings when building the Horde server.

diff --git a/Engine/Source/Programs/Horde/HordeServer/Collections/IIssueCollection.cs b/Engine/Source/Programs/Horde/HordeServer/Collections/IIssueCollection.cs
--- a/Engine/Source/Programs/Horde/HordeServer/Collections/IIssueCollection.cs
+++ b/Engine/Source/Programs/Horde/HordeServer/Collections/IIssueCollection.cs
@@ -127,7 +127,7 @@ namespace HordeServer.Collections
 		/// <returns>The hashcode value</returns>
 		public override int GetHashCode()
 		{
-			int Result = string.GetHashCode(Type, StringComparison.Ordinal);
+			int Result = StringComparer.Ordinal.GetHashCode(Type);
 			foreach (string Key in Keys)
 			{
 				Result = HashCode.Combine(Result, Key);
