Horde: Excluding PDB files from agent archives. These already exist in the symbol store.

diff --git a/Engine/Source/Programs/Horde/HordeAgent/Commands/Utilities/PublishCommand.cs b/Engine/Source/Programs/Horde/HordeAgent/Commands/Utilities/PublishCommand.cs
--- a/Engine/Source/Programs/Horde/HordeAgent/Commands/Utilities/PublishCommand.cs
+++ b/Engine/Source/Programs/Horde/HordeAgent/Commands/Utilities/PublishCommand.cs
@@ -66,7 +66,11 @@ namespace HordeAgent.Commands.Utilities
 					foreach (FileReference InputFile in DirectoryReference.EnumerateFiles(InputDir, "*", SearchOption.AllDirectories))
 					{
 						string RelativePath = InputFile.MakeRelativeTo(InputDir).Replace(Path.DirectorySeparatorChar, '/');
-						if (!RelativePath.Equals("Log.json", StringComparison.OrdinalIgnoreCase) && !RelativePath.Equals("Log.txt", StringComparison.OrdinalIgnoreCase))
+						if (RelativePath.EndsWith(".pdb", StringComparison.OrdinalIgnoreCase))
+						{
+							Logger.LogInformation("  Skipping {RelativePath}", RelativePath);
+						}
+						else if (!RelativePath.Equals("Log.json", StringComparison.OrdinalIgnoreCase) && !RelativePath.Equals("Log.txt", StringComparison.OrdinalIgnoreCase))
 						{
 							Logger.LogInformation("  Adding {RelativePath}", RelativePath);
 
