Horde: Device management updates
#rnx

diff --git a/Engine/Source/Programs/Horde/HordeDashboard/src/backend/Api.ts b/Engine/Source/Programs/Horde/HordeDashboard/src/backend/Api.ts
--- a/Engine/Source/Programs/Horde/HordeDashboard/src/backend/Api.ts
+++ b/Engine/Source/Programs/Horde/HordeDashboard/src/backend/Api.ts
@@ -2785,12 +2785,24 @@ export type GetDevicePlatformResponse = {
 }
 
 
+export enum DevicePoolType {
+
+	/** Used in CIS jobs */
+	Automation = "Automation",
+
+	/** User devices which can be checked out and in */
+	Shared = "Shared"
+}
+
 /** Device pool response object */
 export type GetDevicePoolResponse = {
 
 	/**  Id of  the device pool */
 	id: string;
 
+	/** The type of the pool */
+	poolType: DevicePoolType;
+
 	/**  Name of the device pool */
 	name: string;
 
@@ -2825,6 +2837,22 @@ export type CreateDeviceResponse = {
 
 }
 
+/// Get response object which describes a device
+export type GetDeviceUtilizationResponse = {
+	/// The job id which utilized device
+	jobId?: string;
+
+	/// The job's step id
+	stepId?: string;
+
+	/// The time device was reserved
+	reservationStartUtc: Date | string;
+
+	/// The time device was freed
+	reservationFinishUtc?: Date | string;
+}
+
+
 /** Get response object which describes a device */
 export type GetDeviceResponse = {
 
@@ -2858,8 +2886,29 @@ export type GetDeviceResponse = {
 	/** The UTC time when a device was set for maintenance */
 	maintenanceTime?: Date | string;
 
+	/// The user id that has the device checked out, pass "" to clear checkout
+	checkedOutByUserId?: string;
+
+	/// The last time the device was checked out
+	checkOutTime?: Date | string;
+
+	/** Id of the user which last modified this device */
+	modifiedByUser?: string;
+
+	/** Utilization info */
+	utilization?: GetDeviceUtilizationResponse[];
+
 }
 
+/** Device checkout request object */
+export type CheckoutDeviceRequest = {
+
+	/** Whether to checkout or in the device */
+	checkout: boolean;
+}
+
+
+
 /** Request object for updating a device */
 export type UpdateDeviceRequest = {
 
diff --git a/Engine/Source/Programs/Horde/HordeDashboard/src/backend/Backend.ts b/Engine/Source/Programs/Horde/HordeDashboard/src/backend/Backend.ts
--- a/Engine/Source/Programs/Horde/HordeDashboard/src/backend/Backend.ts
+++ b/Engine/Source/Programs/Horde/HordeDashboard/src/backend/Backend.ts
@@ -1131,6 +1131,17 @@ export class Backend {
         });        
     }
 
+    checkoutDevice(deviceId: string, checkout:boolean): Promise<void> {
+        return new Promise<void>((resolve, reject) => {
+            this.backend.put(`/api/v2/devices/${deviceId}/checkout`, {checkout: checkout}).then(() => {
+                resolve();
+            }).catch(reason => {
+                reject(reason);
+            });
+        });        
+    }
+
+
     modifyDevice(deviceId: string, request: UpdateDeviceRequest): Promise<void> {
         return new Promise<void>((resolve, reject) => {
             this.backend.put(`/api/v2/devices/${deviceId}`, request).then((value) => {
diff --git a/Engine/Source/Programs/Horde/HordeServer/Api/Devices.cs b/Engine/Source/Programs/Horde/HordeServer/Api/Devices.cs
--- a/Engine/Source/Programs/Horde/HordeServer/Api/Devices.cs
+++ b/Engine/Source/Programs/Horde/HordeServer/Api/Devices.cs
@@ -4,6 +4,8 @@ using System;
 using System.Collections.Generic;
 using System.ComponentModel.DataAnnotations;
 using System.Diagnostics.CodeAnalysis;
+using System.Linq;
+using HordeServer.Models;
 
 namespace HordeServer.Api
 {
@@ -95,6 +97,12 @@ namespace HordeServer.Api
 		[Required]
 		public string Name { get; set; } = null!;
 
+		/// <summary>
+		/// The name for the new pool
+		/// </summary>
+		[Required]
+		public DevicePoolType PoolType { get; set; }
+
 	}
 
 	/// <summary>
@@ -131,14 +139,20 @@ namespace HordeServer.Api
 		/// </summary>
 		public string Name { get; set; }
 
+		/// <summary>
+		/// Type of the device pool
+		/// </summary>
+		public DevicePoolType PoolType { get; set; }
+
 		/// <summary>
 		/// Constructor
 		/// </summary>
-		public GetDevicePoolResponse(string Id, string Name)
+		public GetDevicePoolResponse(string Id, string Name, DevicePoolType PoolType)
 		{
 			this.Id = Id;
 			this.Name = Name;
-		}
+            this.PoolType = PoolType;
+        }
 	}
 
 	// Devices 
@@ -204,6 +218,47 @@ namespace HordeServer.Api
 			this.Id = Id;
 		}
 	}
+
+	/// <summary>
+	/// Get response object which describes a device
+	/// </summary>
+	public class GetDeviceUtilizationResponse
+	{
+		/// <summary>
+		/// The job id which utilized device
+		/// </summary>
+		public string? JobId { get; set;}
+
+		/// <summary>
+		/// The job's step id
+		/// </summary>
+		public string? StepId { get; set;}
+
+		/// <summary>
+		/// The time device was reserved
+		/// </summary>
+		public DateTime ReservationStartUtc { get; set;}
+
+		/// <summary>
+		/// The time device was freed
+		/// </summary>
+		public DateTime? ReservationFinishUtc { get; set;}
+
+		/// <summary>
+		/// Constructor
+		/// </summary>
+		/// <param name="Telemetry"></param>
+		public GetDeviceUtilizationResponse(DeviceUtilizationTelemetry Telemetry)
+		{
+			this.JobId = Telemetry.JobId;
+			this.StepId = Telemetry.StepId;
+			this.ReservationStartUtc = Telemetry.ReservationStartUtc;
+			this.ReservationFinishUtc = Telemetry.ReservationFinishUtc;
+		}
+
+	}
+
+
 	/// <summary>
 	/// Get response object which describes a device
 	/// </summary>
@@ -259,11 +314,31 @@ namespace HordeServer.Api
 		/// </summary>
 		public DateTime? MaintenanceTime { get; set; }
 
+		/// <summary>
+		/// The user id that has the device checked out
+		/// </summary>
+		public string? CheckedOutByUserId { get; set; }
+
+		/// <summary>
+		/// The last time the device was checked out
+		/// </summary>
+		public DateTime? CheckOutTime { get; set; }
 
 		/// <summary>
-		/// Device response constructor
+		/// The last user to modifiy the device
 		/// </summary>
-		public GetDeviceResponse(string Id, string PlatformId, string PoolId, string Name, bool Enabled, string? Address, string? ModelId, string? Notes, DateTime? ProblemTime, DateTime? MaintenanceTime)
+		public string? ModifiedByUser { get; set; }
+
+		/// <summary>
+		///  Device Utilization data
+		/// </summary>
+		public List<GetDeviceUtilizationResponse>? Utilization { get; set;}
+
+
+	/// <summary>
+	/// Device response constructor
+	/// </summary>
+	public GetDeviceResponse(string Id, string PlatformId, string PoolId, string Name, bool Enabled, string? Address, string? ModelId, string? ModifiedByUser, string? Notes, DateTime? ProblemTime, DateTime? MaintenanceTime, List<DeviceUtilizationTelemetry>? Utilization, string? CheckedOutByUser = null, DateTime? CheckOutTime = null)
 		{
 			this.Id = Id;
 			this.Name = Name;
@@ -272,10 +347,14 @@ namespace HordeServer.Api
 			this.Enabled = Enabled;
 			this.Address = Address;
 			this.ModelId = ModelId;
-			this.Notes = Notes;
+            this.ModifiedByUser = ModifiedByUser;
+            this.Notes = Notes;
 			this.ProblemTime = ProblemTime;
 			this.MaintenanceTime = MaintenanceTime;
-		}
+			this.Utilization = Utilization?.Select(U => new GetDeviceUtilizationResponse(U)).ToList();
+            this.CheckedOutByUserId = CheckedOutByUser;
+            this.CheckOutTime = CheckOutTime;
+        }
 	}
 
 	/// <summary>
@@ -322,9 +401,20 @@ namespace HordeServer.Api
 		/// Whether to set or clear any device problem state
 		/// </summary>
 		public bool? Problem { get; set; }
-
 	}
 
+	/// <summary>
+	/// Device checkout request object
+	/// </summary>
+	public class CheckoutDeviceRequest
+	{
+		/// <summary>
+        /// Whether to checkout or in the device
+        /// </summary>
+        public bool Checkout { get; set; }
+    }
+
+
 	/// <summary>
 	/// Device reservation request object
 	/// </summary>
@@ -478,6 +568,17 @@ namespace HordeServer.Api
 		/// The PoolId of reservation request 
 		/// </summary>
 		public string? PoolId { get; set; } = null;
+
+		/// <summary>
+		/// The JobId of reservation request 
+		/// </summary>
+		public string? JobId { get; set; } = null;
+
+		/// <summary>
+		/// The StepId of reservation request 
+		/// </summary>
+		public string? StepId { get; set; } = null;
+
 	}
 
 	/// <summary>
@@ -515,6 +616,17 @@ namespace HordeServer.Api
 		[Required]
 		public string Duration { get; set; } = null!;
 
+		/// <summary>
+		/// The JobId of reservation request 
+		/// </summary>
+		public string? JobId { get; set; } = null;
+
+		/// <summary>
+		/// The StepId of reservation request 
+		/// </summary>
+		public string? StepId { get; set; } = null;
+
+
 		/// <summary>
 		/// The legacy guid of the reservation
 		/// </summary>
diff --git a/Engine/Source/Programs/Horde/HordeServer/Collections/IDeviceCollection.cs b/Engine/Source/Programs/Horde/HordeServer/Collections/IDeviceCollection.cs
--- a/Engine/Source/Programs/Horde/HordeServer/Collections/IDeviceCollection.cs
+++ b/Engine/Source/Programs/Horde/HordeServer/Collections/IDeviceCollection.cs
@@ -83,7 +83,8 @@ namespace HordeServer.Collections
 		/// </summary>
 		/// <param name="Id">The id of the new pool</param>
 		/// <param name="Name">The friendly name of the new pool</param>
-		Task<IDevicePool?> TryAddPoolAsync(DevicePoolId Id, string Name);
+        /// <param name="PoolType">The pool type</param>
+		Task<IDevicePool?> TryAddPoolAsync(DevicePoolId Id, string Name, DevicePoolType PoolType);
 
 		/// <summary>
 		/// Get a pool by id
@@ -123,7 +124,8 @@ namespace HordeServer.Collections
 		/// <param name="Enabled">Whather the device is enabled by default</param>
 		/// <param name="Address">The network address or hostname the device can be reached at</param>
 		/// <param name="ModelId">The vendor model id of the device</param>
-		Task<IDevice?> TryAddDeviceAsync(DeviceId Id, string Name, DevicePlatformId PlatformId, DevicePoolId PoolId, bool? Enabled, string? Address, string? ModelId);
+        /// <param name="UserId">The user adding the device</param>
+		Task<IDevice?> TryAddDeviceAsync(DeviceId Id, string Name, DevicePlatformId PlatformId, DevicePoolId PoolId, bool? Enabled, string? Address, string? ModelId, ObjectId? UserId);
 
 		/// <summary>
 		/// Update a device
@@ -137,7 +139,8 @@ namespace HordeServer.Collections
 		/// <param name="NewEnabled">Whether the device is enabled or not</param>
 		/// <param name="NewProblem">Whether to set or clear problem state</param>
 		/// <param name="NewMaintenance">Whether to set or clear maintenance state</param>
-		Task UpdateDeviceAsync(DeviceId DeviceId, DevicePoolId? NewPoolId, string? NewName, string? NewAddress, string? NewModelId, string? NewNotes, bool? NewEnabled, bool? NewProblem, bool? NewMaintenance);
+        /// <param name="ModifiedByUserId">The user who is updating the device</param>
+		Task UpdateDeviceAsync(DeviceId DeviceId, DevicePoolId? NewPoolId, string? NewName, string? NewAddress, string? NewModelId, string? NewNotes, bool? NewEnabled, bool? NewProblem, bool? NewMaintenance, ObjectId? ModifiedByUserId = null);
 
 		/// <summary>
 		/// Delete a device from the collection
@@ -145,16 +148,26 @@ namespace HordeServer.Collections
 		/// <param name="DeviceId">The id of the device to delete</param>
 		Task<bool> DeleteDeviceAsync(DeviceId DeviceId);
 
-		// RESERVATIONS
-
 		/// <summary>
-		/// Create a new reseveration in the pool with the specified devices
-		/// </summary>
-		/// <param name="PoolId">The pool of devices to use for the new reservation</param>
-		/// <param name="Request">The requested devices for the reservation</param>
-		/// <param name="Hostname">The hostname of the machine making the reservation</param>
-		/// <param name="ReservationDetails">The details of the reservation</param>
-		Task<IDeviceReservation?> TryAddReservationAsync(DevicePoolId PoolId, List<DeviceRequestData> Request, string? Hostname, string? ReservationDetails);
+        /// Checkout or checkin the specified device
+        /// </summary>
+        /// <param name="DeviceId"></param>
+        /// <param name="CheckedOutByUserId"></param>
+        /// <returns></returns>
+        Task CheckoutDeviceAsync(DeviceId DeviceId, ObjectId? CheckedOutByUserId);
+
+        // RESERVATIONS
+
+        /// <summary>
+        /// Create a new reseveration in the pool with the specified devices
+        /// </summary>
+        /// <param name="PoolId">The pool of devices to use for the new reservation</param>
+        /// <param name="Request">The requested devices for the reservation</param>
+        /// <param name="Hostname">The hostname of the machine making the reservation</param>
+        /// <param name="ReservationDetails">The details of the reservation</param>
+        /// <param name="JobId">The Job Id associated with the job</param>
+        /// <param name="StepId">The Step Id associated with the job</param>
+        Task<IDeviceReservation?> TryAddReservationAsync(DevicePoolId PoolId, List<DeviceRequestData> Request, string? Hostname, string? ReservationDetails, string? JobId, string? StepId);
 
 		/// <summary>
 		/// Gets a reservation by guid for legacy clients
diff --git a/Engine/Source/Programs/Horde/HordeServer/Collections/Impl/DeviceCollection.cs b/Engine/Source/Programs/Horde/HordeServer/Collections/Impl/DeviceCollection.cs
--- a/Engine/Source/Programs/Horde/HordeServer/Collections/Impl/DeviceCollection.cs
+++ b/Engine/Source/Programs/Horde/HordeServer/Collections/Impl/DeviceCollection.cs
@@ -63,6 +63,9 @@ namespace HordeServer.Collections.Impl
 			public DevicePoolId Id { get; set; }
 
 			[BsonRequired]
+            public DevicePoolType PoolType { get; set; }
+
+            [BsonRequired]
 			public string Name { get; set; } = null!;
 
 			[BsonIgnoreIfNull]
@@ -71,14 +74,14 @@ namespace HordeServer.Collections.Impl
 			[BsonConstructor]
 			private DevicePoolDocument()
 			{
+            }
 
-			}
-
-			public DevicePoolDocument(DevicePoolId Id, string Name)
+			public DevicePoolDocument(DevicePoolId Id, string Name, DevicePoolType PoolType)
 			{
 				this.Id = Id;
 				this.Name = Name;
-			}
+                this.PoolType = PoolType;
+            }
 
 
 		}
@@ -98,21 +101,21 @@ namespace HordeServer.Collections.Impl
 			public DevicePoolId PoolId { get; set; }
 
 			[BsonIgnoreIfNull]
-			public ObjectId? JobId { get; }
-
-			[BsonIgnoreIfNull]
-			public SubResourceId? StepId { get; }
-
-			/// <summary>
-			/// Reservations held by a user, requires a token like download code
-			/// </summary>
-			[BsonIgnoreIfNull]
-			public ObjectId? UserId { get; }
+			public string? JobId { get; set; }
 
-			/// <summary>
-			/// The hostname of the machine which has made the reservation
-			/// </summary>
 			[BsonIgnoreIfNull]
+			public string? StepId { get; set; }
+
+            /// <summary>
+            /// Reservations held by a user, requires a token like download code
+            /// </summary>
+            [BsonIgnoreIfNull]
+			public ObjectId? UserId { get; set; }
+
+            /// <summary>
+            /// The hostname of the machine which has made the reservation
+            /// </summary>
+            [BsonIgnoreIfNull]
 			public string? Hostname { get; set; }
 
 			/// <summary>
@@ -139,7 +142,7 @@ namespace HordeServer.Collections.Impl
 
 			}
 
-			public DeviceReservationDocument(ObjectId Id, DevicePoolId PoolId, List<DeviceId> Devices, DateTime CreateTimeUtc, string? Hostname, string? ReservationDetails)
+			public DeviceReservationDocument(ObjectId Id, DevicePoolId PoolId, List<DeviceId> Devices, DateTime CreateTimeUtc, string? Hostname, string? ReservationDetails, string? JobId, string? StepId)
 			{
 				this.Id = Id;
 				this.PoolId = PoolId;
@@ -148,8 +151,10 @@ namespace HordeServer.Collections.Impl
 				this.UpdateTimeUtc = CreateTimeUtc;
 				this.Hostname = Hostname;
 				this.ReservationDetails = ReservationDetails;
+                this.JobId = JobId;
+                this.StepId = StepId;
 
-				this.LegacyGuid = Guid.NewGuid().ToString();
+                this.LegacyGuid = Guid.NewGuid().ToString();
 
 			}
 
@@ -178,6 +183,12 @@ namespace HordeServer.Collections.Impl
 			[BsonIgnoreIfNull]
 			public string? Address { get; set; }
 
+			[BsonIgnoreIfNull]
+			public string? CheckedOutByUser { get; set; }
+
+			[BsonIgnoreIfNull]
+			public DateTime? CheckOutTime { get; set; }
+
 			[BsonIgnoreIfNull]
 			public DateTime? ProblemTimeUtc { get; set; }
 
@@ -196,6 +207,9 @@ namespace HordeServer.Collections.Impl
 			[BsonIgnoreIfNull]
 			public string? Notes { get; set; }
 
+			[BsonIgnoreIfNull]
+			public List<DeviceUtilizationTelemetry>? Utilization { get; set; }
+
 			[BsonIgnoreIfNull]
 			public Acl? Acl { get; set; }
 
@@ -205,7 +219,7 @@ namespace HordeServer.Collections.Impl
 
 			}
 
-			public DeviceDocument(DeviceId Id, DevicePlatformId PlatformId, DevicePoolId PoolId, string Name, bool Enabled, string? Address, string? ModelId)
+			public DeviceDocument(DeviceId Id, DevicePlatformId PlatformId, DevicePoolId PoolId, string Name, bool Enabled, string? Address, string? ModelId, ObjectId? UserId)
 			{
 				this.Id = Id;
 				this.PlatformId = PlatformId;
@@ -214,8 +228,8 @@ namespace HordeServer.Collections.Impl
 				this.Enabled = Enabled;
 				this.Address = Address;
 				this.ModelId = ModelId;
-
-			}
+                this.ModifiedByUser = UserId?.ToString();
+            }
 
 		}
 
@@ -246,9 +260,9 @@ namespace HordeServer.Collections.Impl
 		}
 
 		/// <inheritdoc/>
-		public async Task<IDevice?> TryAddDeviceAsync(DeviceId Id, string Name, DevicePlatformId PlatformId, DevicePoolId PoolId, bool? Enabled, string? Address, string? ModelId)
+		public async Task<IDevice?> TryAddDeviceAsync(DeviceId Id, string Name, DevicePlatformId PlatformId, DevicePoolId PoolId, bool? Enabled, string? Address, string? ModelId, ObjectId? UserId)
 		{
-			DeviceDocument NewDevice = new DeviceDocument(Id, PlatformId, PoolId, Name, Enabled ?? true, Address, ModelId);
+			DeviceDocument NewDevice = new DeviceDocument(Id, PlatformId, PoolId, Name, Enabled ?? true, Address, ModelId, UserId);
 			await Devices.InsertOneAsync(NewDevice);
 			return NewDevice;
 		}
@@ -306,9 +320,9 @@ namespace HordeServer.Collections.Impl
 		}
 
 		/// <inheritdoc/>
-		public async Task<IDevicePool?> TryAddPoolAsync(DevicePoolId Id, string Name)
+		public async Task<IDevicePool?> TryAddPoolAsync(DevicePoolId Id, string Name, DevicePoolType PoolType)
 		{
-			DevicePoolDocument NewPool = new DevicePoolDocument(Id, Name);
+			DevicePoolDocument NewPool = new DevicePoolDocument(Id, Name, PoolType);
 
 			try
 			{
@@ -372,7 +386,27 @@ namespace HordeServer.Collections.Impl
 		}
 
 		/// <inheritdoc/>
-		public async Task UpdateDeviceAsync(DeviceId DeviceId, DevicePoolId? NewPoolId, string? NewName, string? NewAddress, string? NewModelId, string? NewNotes, bool? NewEnabled, bool? NewProblem, bool? NewMaintenance)
+		public async Task CheckoutDeviceAsync(DeviceId DeviceId, ObjectId? CheckedOutByUserId)
+		{
+			UpdateDefinitionBuilder<DeviceDocument> UpdateBuilder = Builders<DeviceDocument>.Update;
+
+			List<UpdateDefinition<DeviceDocument>> Updates = new List<UpdateDefinition<DeviceDocument>>();
+
+            string? UserId = CheckedOutByUserId?.ToString();
+
+            Updates.Add(UpdateBuilder.Set(x => x.CheckedOutByUser, string.IsNullOrEmpty(UserId) ? null : UserId));
+
+			if (CheckedOutByUserId != null)
+			{
+				Updates.Add(UpdateBuilder.Set(x => x.CheckOutTime, DateTime.UtcNow));
+			}
+
+			await Devices.FindOneAndUpdateAsync<DeviceDocument>(x => x.Id == DeviceId, UpdateBuilder.Combine(Updates));
+		}
+
+
+		/// <inheritdoc/>
+		public async Task UpdateDeviceAsync(DeviceId DeviceId, DevicePoolId? NewPoolId, string? NewName, string? NewAddress, string? NewModelId, string? NewNotes, bool? NewEnabled, bool? NewProblem, bool? NewMaintenance, ObjectId? ModifiedByUserId = null)
 		{
 			UpdateDefinitionBuilder<DeviceDocument> UpdateBuilder = Builders<DeviceDocument>.Update;
 
@@ -380,6 +414,11 @@ namespace HordeServer.Collections.Impl
 
 			DateTime UtcNow = DateTime.UtcNow;
 
+			if (ModifiedByUserId != null)
+			{
+				Updates.Add(UpdateBuilder.Set(x => x.ModifiedByUser, ModifiedByUserId.ToString()));
+			}
+
 			if (NewPoolId != null)
 			{
 				Updates.Add(UpdateBuilder.Set(x => x.PoolId, NewPoolId.Value));
@@ -459,7 +498,7 @@ namespace HordeServer.Collections.Impl
 		}
 
 		/// <inheritdoc/>
-		public async Task<IDeviceReservation?> TryAddReservationAsync(DevicePoolId PoolId, List<DeviceRequestData> Request, string? Hostname, string? ReservationDetails)
+		public async Task<IDeviceReservation?> TryAddReservationAsync(DevicePoolId PoolId, List<DeviceRequestData> Request, string? Hostname, string? ReservationDetails, string? JobId, string? StepId)
 		{
 
 			if (Request.Count == 0)
@@ -467,27 +506,34 @@ namespace HordeServer.Collections.Impl
 				return null;
 			}
 
+			DevicePoolDocument? Pool = await Pools.Find<DevicePoolDocument>(x => x.Id == PoolId).FirstOrDefaultAsync();
 
+			if (Pool == null || Pool.PoolType != DevicePoolType.Automation)
+			{
+                return null;
+            }
+			
 			HashSet<DeviceId> Allocated = new HashSet<DeviceId>();
 
 			List<DeviceReservationDocument> PoolReservations = await Reservations.Find(x => x.PoolId == PoolId).ToListAsync();
 
-			DateTime UtcNow = DateTime.UtcNow;
+			DateTime ReservationTimeUtc = DateTime.UtcNow;
 
 			// Get available devices
 			List<DeviceDocument> PoolDevices = await Devices.Find(x =>
 				x.PoolId == PoolId &&
-				x.Enabled &&				
+				x.Enabled &&
 				x.MaintenanceTimeUtc == null).ToListAsync();
 
 			// filter out problem devices
-			PoolDevices = PoolDevices.FindAll(x => (x.ProblemTimeUtc == null || ((UtcNow - x.ProblemTimeUtc).Value.TotalMinutes > 30)));			
+			PoolDevices = PoolDevices.FindAll(x => (x.ProblemTimeUtc == null || ((ReservationTimeUtc - x.ProblemTimeUtc).Value.TotalMinutes > 30)));
 
 			// filter out currently reserved devices
 			PoolDevices = PoolDevices.FindAll(x => PoolReservations.FirstOrDefault(p => p.Devices.Contains(x.Id)) == null);
 
 			// sort to use last reserved first to cycle devices
-			PoolDevices.Sort((A, B) => {
+			PoolDevices.Sort((A, B) =>
+			{
 				DateTime? ATime = A.ReservationTimeUtc;
 				DateTime? BTime = B.ReservationTimeUtc;
 
@@ -542,12 +588,39 @@ namespace HordeServer.Collections.Impl
 				Allocated.Add(Device.Id);
 			}
 
-			// update reservation time for allocated devices
-			await Devices.UpdateManyAsync(Builders<DeviceDocument>.Filter.In(x => x.Id, Allocated), Builders<DeviceDocument>.Update.Set(x => x.ReservationTimeUtc, UtcNow));
+			// update reservation time and utilization for allocated devices 
+			foreach (DeviceId Id in Allocated)
+			{
+				DeviceDocument Device = PoolDevices.First((Device) => Device.Id == Id);
+
+				List<DeviceUtilizationTelemetry>? Utilization = Device.Utilization;
+
+				if (Utilization == null)
+				{
+					Utilization = new List<DeviceUtilizationTelemetry>();
+				}
+
+				// keep up to 100, maintaining order
+				if (Utilization.Count > 99)
+				{
+					Utilization.Take(99);
+				}
+
+				Utilization.Insert(0, new DeviceUtilizationTelemetry(ReservationTimeUtc) { JobId = JobId, StepId = StepId });
+
+				UpdateDefinitionBuilder<DeviceDocument> DeviceBuilder = Builders<DeviceDocument>.Update;
+				List<UpdateDefinition<DeviceDocument>> DeviceUpdates = new List<UpdateDefinition<DeviceDocument>>();
+
+				DeviceUpdates.Add(DeviceBuilder.Set(x => x.ReservationTimeUtc, ReservationTimeUtc));
+				DeviceUpdates.Add(DeviceBuilder.Set(x => x.Utilization, Utilization));
+
+				await Devices.FindOneAndUpdateAsync<DeviceDocument>(x => x.Id == Id, DeviceBuilder.Combine(DeviceUpdates));
+			}
 
 			// Create new reservation
-			DeviceReservationDocument NewReservation = new DeviceReservationDocument(ObjectId.GenerateNewId(), PoolId, Allocated.ToList(), DateTime.UtcNow, Hostname, ReservationDetails);
+			DeviceReservationDocument NewReservation = new DeviceReservationDocument(ObjectId.GenerateNewId(), PoolId, Allocated.ToList(), ReservationTimeUtc, Hostname, ReservationDetails, JobId, StepId);
 			await Reservations.InsertOneAsync(NewReservation);
+
 			return NewReservation;
 
 		}
@@ -573,16 +646,16 @@ namespace HordeServer.Collections.Impl
 		public async Task<bool> ExpireReservationsAsync()
 		{
 			List<IDeviceReservation> Reserves = await Reservations.Find(A => true).ToListAsync<DeviceReservationDocument, IDeviceReservation>();
-		
+
 			DateTime UtcNow = DateTime.UtcNow;
 
-			Reserves = Reserves.FindAll(R => ( UtcNow - R.UpdateTimeUtc ).TotalMinutes > 10).ToList();			
+			Reserves = Reserves.FindAll(R => (UtcNow - R.UpdateTimeUtc).TotalMinutes > 10).ToList();
 
 			if (Reserves.Count > 0)
 			{
 				DeleteResult Result = await Reservations.DeleteManyAsync(Builders<DeviceReservationDocument>.Filter.In(x => x.Id, Reserves.Select(y => y.Id)));
 				return (Result.DeletedCount > 0);
-			}			
+			}
 
 			return false;
 
diff --git a/Engine/Source/Programs/Horde/HordeServer/Controllers/DevicesController.cs b/Engine/Source/Programs/Horde/HordeServer/Controllers/DevicesController.cs
--- a/Engine/Source/Programs/Horde/HordeServer/Controllers/DevicesController.cs
+++ b/Engine/Source/Programs/Horde/HordeServer/Controllers/DevicesController.cs
@@ -30,6 +30,12 @@ namespace HordeServer.Controllers
 	public class DevicesController : ControllerBase
 	{
 
+		/// <summary>
+		/// The user collection instance
+		/// </summary>
+		IUserCollection UserCollection { get; set; }
+
+
 		/// <summary>
 		/// Singleton instance of the device service
 		/// </summary>
@@ -43,9 +49,10 @@ namespace HordeServer.Controllers
 		/// <summary>
 		/// Constructor
 		/// </summary>
-		public DevicesController(DeviceService DeviceService, ILogger<DevicesController> Logger)
+		public DevicesController(DeviceService DeviceService, IUserCollection UserCollection, ILogger<DevicesController> Logger)
 		{
-			this.DeviceService = DeviceService;
+            this.UserCollection = UserCollection;
+            this.DeviceService = DeviceService;
 			this.Logger = Logger;
 		}
 
@@ -65,6 +72,12 @@ namespace HordeServer.Controllers
 				return Forbid();
 			}
 
+			IUser? InternalUser = await UserCollection.GetUserAsync(User);
+			if (InternalUser == null)
+			{
+				return NotFound();
+			}
+
 			IDevicePlatform? Platform = await DeviceService.GetPlatformAsync(new DevicePlatformId(DeviceRequest.PlatformId!));
 
 			if (Platform == null)
@@ -95,7 +108,7 @@ namespace HordeServer.Controllers
 			string Name = DeviceRequest.Name.Trim();
 			string? Address = DeviceRequest.Address?.Trim();
 
-			IDevice? Device = await DeviceService.TryCreateDeviceAsync(DeviceId.Sanitize(Name), Name, Platform.Id, Pool.Id, DeviceRequest.Enabled, Address, ModelId);
+			IDevice? Device = await DeviceService.TryCreateDeviceAsync(DeviceId.Sanitize(Name), Name, Platform.Id, Pool.Id, DeviceRequest.Enabled, Address, ModelId, InternalUser.Id);
 
 			if (Device == null)
 			{
@@ -126,7 +139,7 @@ namespace HordeServer.Controllers
 
 			foreach (IDevice Device in Devices)
 			{
-				Responses.Add(new GetDeviceResponse(Device.Id.ToString(), Device.PlatformId.ToString(), Device.PoolId.ToString(), Device.Name, Device.Enabled, Device.Address, Device.ModelId?.ToString(), Device.Notes, Device.ProblemTimeUtc, Device.MaintenanceTimeUtc));
+				Responses.Add(new GetDeviceResponse(Device.Id.ToString(), Device.PlatformId.ToString(), Device.PoolId.ToString(), Device.Name, Device.Enabled, Device.Address, Device.ModelId?.ToString(), Device.ModifiedByUser, Device.Notes, Device.ProblemTimeUtc, Device.MaintenanceTimeUtc, Device.Utilization, Device.CheckedOutByUser, Device.CheckOutTime));
 			}
 
 			return Responses;
@@ -156,7 +169,7 @@ namespace HordeServer.Controllers
 				return BadRequest($"Unable to find device with id {DeviceId}");
 			}
 
-			return new GetDeviceResponse(Device.Id.ToString(), Device.PlatformId.ToString(), Device.PoolId.ToString(), Device.Name, Device.Enabled, Device.Address, Device.ModelId?.ToString(), Device.Notes, Device.ProblemTimeUtc, Device.MaintenanceTimeUtc);
+			return new GetDeviceResponse(Device.Id.ToString(), Device.PlatformId.ToString(), Device.PoolId.ToString(), Device.Name, Device.Enabled, Device.Address, Device.ModelId?.ToString(), Device.ModifiedByUser?.ToString(), Device.Notes, Device.ProblemTimeUtc, Device.MaintenanceTimeUtc, Device.Utilization, Device.CheckedOutByUser, Device.CheckOutTime);
 		}
 
 		/// <summary>
@@ -173,6 +186,12 @@ namespace HordeServer.Controllers
 				return Forbid();
 			}
 
+			IUser? InternalUser = await UserCollection.GetUserAsync(User);
+			if (InternalUser == null)
+			{
+				return NotFound();
+			}
+
 			DeviceId DeviceIdValue = new DeviceId(DeviceId);
 
 			IDevice? Device = await DeviceService.GetDeviceAsync(DeviceIdValue);
@@ -212,11 +231,59 @@ namespace HordeServer.Controllers
 			string? Name = Update.Name?.Trim();
 			string? Address = Update.Address?.Trim();
 
-			await DeviceService.UpdateDeviceAsync(DeviceIdValue, PoolIdValue, Name, Address, ModelIdValue, Update.Notes, Update.Enabled, Update.Problem, Update.Maintenance);
+            await DeviceService.UpdateDeviceAsync(DeviceIdValue, PoolIdValue, Name, Address, ModelIdValue, Update.Notes, Update.Enabled, Update.Problem, Update.Maintenance, InternalUser.Id);
 
 			return Ok();
 		}
 
+		/// <summary>
+		/// Checkout a device
+		/// </summary>
+		[HttpPut]
+		[Authorize]
+		[Route("/api/v2/devices/{DeviceId}/checkout")]
+		[ProducesResponseType(typeof(List<GetDeviceResponse>), 200)]
+		public async Task<ActionResult> CheckoutDeviceAsync(string DeviceId, [FromBody] CheckoutDeviceRequest Request)
+		{
+			if (!await DeviceService.AuthorizeAsync(AclAction.DeviceWrite, User))
+			{
+				return Forbid();
+			}
+
+			IUser? InternalUser = await UserCollection.GetUserAsync(User);
+			if (InternalUser == null)
+			{
+				return NotFound();
+			}
+
+			DeviceId DeviceIdValue = new DeviceId(DeviceId);
+
+			IDevice? Device = await DeviceService.GetDeviceAsync(DeviceIdValue);
+
+			if (Device == null)
+			{
+				return BadRequest($"Device with id ${DeviceId} does not exist");
+			}
+
+			if (Request.Checkout)
+			{
+				if (!string.IsNullOrEmpty(Device.CheckedOutByUser))
+				{
+					return BadRequest($"Already checked out by user {Device.CheckedOutByUser}");
+				}
+
+                await DeviceService.CheckoutDeviceAsync(DeviceIdValue, InternalUser.Id);
+
+            }
+			else
+			{
+				await DeviceService.CheckoutDeviceAsync(DeviceIdValue, null);
+			}
+
+			return Ok();
+		}
+
+
 		/// <summary>
 		/// Delete a specific device
 		/// </summary>
@@ -337,7 +404,7 @@ namespace HordeServer.Controllers
 
 			string Name = Request.Name.Trim();
 
-			IDevicePool? Pool = await DeviceService.TryCreatePoolAsync(DevicePoolId.Sanitize(Name), Name);
+			IDevicePool? Pool = await DeviceService.TryCreatePoolAsync(DevicePoolId.Sanitize(Name), Name, Request.PoolType);
 
 			if (Pool == null)
 			{
@@ -371,7 +438,7 @@ namespace HordeServer.Controllers
 			foreach (IDevicePool Pool in Pools)
 			{
 				// @todo: ACL per platform
-				Responses.Add(new GetDevicePoolResponse(Pool.Id.ToString(), Pool.Name));
+				Responses.Add(new GetDevicePoolResponse(Pool.Id.ToString(), Pool.Name, Pool.PoolType));
 			}
 
 			return Responses;
@@ -462,7 +529,7 @@ namespace HordeServer.Controllers
 
 			foreach (IDevice Device in Devices)
 			{
-				Response.Devices.Add(new GetDeviceResponse(Device.Id.ToString(), Device.PlatformId.ToString(), Device.PoolId.ToString(), Device.Name, Device.Enabled, Device.Address, Device.ModelId?.ToString(), Device.Notes, Device.ProblemTimeUtc, Device.MaintenanceTimeUtc));
+				Response.Devices.Add(new GetDeviceResponse(Device.Id.ToString(), Device.PlatformId.ToString(), Device.PoolId.ToString(), Device.Name, Device.Enabled, Device.Address, Device.ModelId?.ToString(), Device.ModifiedByUser, Device.Notes, Device.ProblemTimeUtc, Device.MaintenanceTimeUtc, Device.Utilization));
 			}
 
 			return Response;
@@ -609,10 +676,14 @@ namespace HordeServer.Controllers
 			List<IDevicePool> Pools = await DeviceService.GetPoolsAsync();
 			List<IDevicePlatform> Platforms = await DeviceService.GetPlatformsAsync();
 
-			//@todo: pull this fromm request
-			string PoolId = "fnue5";
+			string? PoolId = Request.PoolId;
+
+			if (PoolId == null)
+			{
+                PoolId = "fnue5";
+            }			
 
-			DevicePoolId PoolIdValue = new DevicePoolId(PoolId);
+			DevicePoolId PoolIdValue = DevicePoolId.Sanitize(PoolId);
 			IDevicePool? Pool = Pools.FirstOrDefault(x => x.Id == PoolIdValue);
 			if (Pool == null)
 			{
@@ -678,7 +749,7 @@ namespace HordeServer.Controllers
 
 			}
 
-			IDeviceReservation? Reservation = await DeviceService.TryCreateReservationAsync(PoolIdValue, RequestedDevices, Request.Hostname, Request.ReservationDetails);
+			IDeviceReservation? Reservation = await DeviceService.TryCreateReservationAsync(PoolIdValue, RequestedDevices, Request.Hostname, Request.ReservationDetails, Request.JobId, Request.StepId);
 
 			if (Reservation == null)
 			{
@@ -705,24 +776,24 @@ namespace HordeServer.Controllers
 		/// </summary>
 		[HttpPut]
 		[Route("/api/v1/reservations/{ReservationGuid}")]
-		public async Task<ActionResult> UpdateReservationV1Async(string ReservationGuid, [FromBody] string Duration)
+		public async Task<ActionResult> UpdateReservationV1Async(string ReservationGuid /* [FromBody] string Duration */)
 		{
 			IDeviceReservation? Reservation = await DeviceService.TryGetReservationFromLegacyGuidAsync(ReservationGuid);
 
 			if (Reservation == null)
 			{
-				return BadRequest($"Unable to find reservation for guid {ReservationGuid}");
+                Logger.LogError($"Unable to find reservation for legacy guid {ReservationGuid}");
+                return BadRequest();
 			}
 
 			bool Updated = await DeviceService.TryUpdateReservationAsync(Reservation.Id);
 
 			if (!Updated)
 			{
-				return BadRequest("Failed to update reservation");
+				Logger.LogError($"Unable to find reservation for reservation {Reservation.Id}");
+				return BadRequest();
 			}
 
-			Logger.LogDebug($"Ignoring update duration of {Duration}");
-
 			return Ok();
 		}
 
diff --git a/Engine/Source/Programs/Horde/HordeServer/Models/Device.cs b/Engine/Source/Programs/Horde/HordeServer/Models/Device.cs
--- a/Engine/Source/Programs/Horde/HordeServer/Models/Device.cs
+++ b/Engine/Source/Programs/Horde/HordeServer/Models/Device.cs
@@ -2,6 +2,8 @@
 
 using System;
 using MongoDB.Bson;
+using MongoDB.Bson.Serialization.Attributes;
+
 using System.Collections.Generic;
 
 using DeviceId = HordeServer.Utilities.StringId<HordeServer.Models.IDevice>;
@@ -29,12 +31,12 @@ namespace HordeServer.Models
 		/// <summary>
 		/// JobID holding reservation
 		/// </summary>
-		public ObjectId? JobId { get; }
+		public string? JobId { get; }
 
 		/// <summary>
 		/// Job step id holding reservation
 		/// </summary>
-		public SubResourceId? StepId { get; }
+		public string? StepId { get; }
 
 		/// <summary>
 		/// Reservations held by a user, requires a token
@@ -95,11 +97,27 @@ namespace HordeServer.Models
 
 	}
 
-
-	/// <summary>
-	/// A logical pool of devices
-	/// </summary>
-	public interface IDevicePool
+    /// <summary>
+    /// The type of device pool
+    /// </summary>
+    public enum DevicePoolType
+    {
+        /// <summary>
+        /// Available to CIS jobs
+        /// </summary>
+        Automation,
+
+        /// <summary>
+        /// Shared by users with remote checking and checkouts
+        /// </summary>
+        Shared
+
+    }
+
+    /// <summary>
+    /// A logical pool of devices
+    /// </summary>
+    public interface IDevicePool
 	{
 		/// <summary>
 		/// Unique identifier of pool
@@ -107,9 +125,14 @@ namespace HordeServer.Models
 		public DevicePoolId Id { get; }
 
 		/// <summary>
-		/// Friendly name of the pool
+		/// The type of pool 
 		/// </summary>
-		string Name { get; }
+		public DevicePoolType PoolType { get; }
+
+        /// <summary>
+        /// Friendly name of the pool
+        /// </summary>
+        string Name { get; }
 
 		/// <summary>
 		/// Acl access to the pool
@@ -117,6 +140,51 @@ namespace HordeServer.Models
 		public Acl? Acl { get; }
 	}
 
+	/// <summary>
+	/// A device utilization snapshot
+	/// </summary>
+	public class DeviceUtilizationTelemetry
+	{
+		/// <summary>
+		/// The job id which utilized device
+		/// </summary>
+		public string? JobId { get; set;}
+
+		/// <summary>
+		/// The job's step id
+		/// </summary>
+		public string? StepId { get; set;}
+
+		/// <summary>
+		/// The time device was reserved
+		/// </summary>
+		public DateTime ReservationStartUtc { get; set;}
+
+		/// <summary>
+		/// The time device was freed
+		/// </summary>
+		public DateTime? ReservationFinishUtc { get; set;}
+
+		/// <summary>
+		/// Private constructor
+		/// </summary>
+		[BsonConstructor]
+		private DeviceUtilizationTelemetry()
+		{
+		}
+
+		/// <summary>
+		/// Constructor
+		/// </summary>
+		/// <param name="ReservationStartUtc"></param>
+		public DeviceUtilizationTelemetry( DateTime ReservationStartUtc )
+		{
+			this.ReservationStartUtc = ReservationStartUtc;
+		}
+
+	}
+
+
 	/// <summary>
 	/// A physical device
 	/// </summary>
@@ -159,10 +227,21 @@ namespace HordeServer.Models
 		public bool Enabled { get; }
 
 		/// <summary>
-		/// Name of the user that last modified this device
+		/// Id of the user that last modified this device
 		/// </summary>
 		public string? ModifiedByUser { get; }
 
+		/// <summary>
+		/// Id of the user that has this device checked out
+		/// </summary>
+		public string? CheckedOutByUser { get; }
+
+		/// <summary>
+		/// The last time this device was checked out
+		/// </summary>
+		public DateTime? CheckOutTime { get; }
+
+
 		/// <summary>
 		/// The last time a problem was reported
 		/// </summary>
@@ -178,6 +257,11 @@ namespace HordeServer.Models
 		/// </summary>
 		public string? Notes { get; }
 
+		/// <summary>
+		/// Device job utilization history 
+		/// </summary>
+		public List<DeviceUtilizationTelemetry>? Utilization { get; set; }
+
 		/// <summary>
 		/// ACL for modifying this device
 		/// </summary>
diff --git a/Engine/Source/Programs/Horde/HordeServer/Services/DeviceService.cs b/Engine/Source/Programs/Horde/HordeServer/Services/DeviceService.cs
--- a/Engine/Source/Programs/Horde/HordeServer/Services/DeviceService.cs
+++ b/Engine/Source/Programs/Horde/HordeServer/Services/DeviceService.cs
@@ -108,9 +108,9 @@ namespace HordeServer.Services
 		/// <summary>
 		/// Create a new device pool
 		/// </summary>
-		public Task<IDevicePool?> TryCreatePoolAsync(DevicePoolId Id, string Name)
+		public Task<IDevicePool?> TryCreatePoolAsync(DevicePoolId Id, string Name, DevicePoolType PoolType)
 		{
-			return Devices.TryAddPoolAsync(Id, Name);
+			return Devices.TryAddPoolAsync(Id, Name, PoolType);
 		}
 
 		/// <summary>
@@ -163,26 +163,35 @@ namespace HordeServer.Services
 		/// <param name="Enabled">Whether the device is enabled</param>
 		/// <param name="Address">Address or hostname of device</param>
 		/// <param name="ModelId">Vendor model id</param>
+        /// <param name="UserId">User adding the device</param>
 		/// <returns></returns>
-		public Task<IDevice?> TryCreateDeviceAsync(DeviceId Id, string Name, DevicePlatformId PlatformId, DevicePoolId PoolId, bool? Enabled, string? Address, string? ModelId)
+		public Task<IDevice?> TryCreateDeviceAsync(DeviceId Id, string Name, DevicePlatformId PlatformId, DevicePoolId PoolId, bool? Enabled, string? Address, string? ModelId, ObjectId? UserId = null)
 		{
-			return Devices.TryAddDeviceAsync(Id, Name, PlatformId, PoolId, Enabled, Address, ModelId);
+			return Devices.TryAddDeviceAsync(Id, Name, PlatformId, PoolId, Enabled, Address, ModelId, UserId);
 		}
 
 		/// <summary>
 		/// Update a device
 		/// </summary>
-		public Task UpdateDeviceAsync(DeviceId DeviceId, DevicePoolId? NewPoolId = null, string? NewName = null, string? NewAddress = null, string? NewModelId = null, string? NewNotes = null, bool? NewEnabled = null, bool? NewProblem = null, bool? NewMaintenance = null)
+		public Task UpdateDeviceAsync(DeviceId DeviceId, DevicePoolId? NewPoolId = null, string? NewName = null, string? NewAddress = null, string? NewModelId = null, string? NewNotes = null, bool? NewEnabled = null, bool? NewProblem = null, bool? NewMaintenance = null, ObjectId? ModifiedByUserId = null)
 		{
-			return Devices.UpdateDeviceAsync(DeviceId, NewPoolId, NewName, NewAddress, NewModelId, NewNotes, NewEnabled, NewProblem, NewMaintenance);
+			return Devices.UpdateDeviceAsync(DeviceId, NewPoolId, NewName, NewAddress, NewModelId, NewNotes, NewEnabled, NewProblem, NewMaintenance, ModifiedByUserId);
 		}
 
+		/// <summary>
+		/// Checkout a device
+		/// </summary>
+		public Task CheckoutDeviceAsync(DeviceId DeviceId, ObjectId? UserId)
+		{
+            return Devices.CheckoutDeviceAsync(DeviceId, UserId);
+        }
+
 		/// <summary>
 		/// Try to create a reservation satisfying the specified device platforms and models
 		/// </summary>
-		public Task<IDeviceReservation?> TryCreateReservationAsync(DevicePoolId Pool, List<DeviceRequestData> Request, string? Hostname = null, string? ReservationDetails = null)
+		public Task<IDeviceReservation?> TryCreateReservationAsync(DevicePoolId Pool, List<DeviceRequestData> Request, string? Hostname = null, string? ReservationDetails = null, string? JobId = null, string? StepId = null)
 		{
-			return Devices.TryAddReservationAsync(Pool, Request, Hostname, ReservationDetails);
+			return Devices.TryAddReservationAsync(Pool, Request, Hostname, ReservationDetails, JobId, StepId);
 		}
 
 		/// <summary>
