Horde: Fix pool check
#rnx

diff --git a/Engine/Source/Programs/Horde/HordeServer/Controllers/DevicesController.cs b/Engine/Source/Programs/Horde/HordeServer/Controllers/DevicesController.cs
--- a/Engine/Source/Programs/Horde/HordeServer/Controllers/DevicesController.cs
+++ b/Engine/Source/Programs/Horde/HordeServer/Controllers/DevicesController.cs
@@ -678,7 +678,7 @@ namespace HordeServer.Controllers
 
 			string? PoolId = Request.PoolId;
 
-			if (PoolId == null)
+			if (string.IsNullOrEmpty(PoolId))
 			{
                 PoolId = "fnue5";
             }			
