Horde: Prevent exceptions being swallowed during REAPI sandbox setup

* Improve test coverage for ActionExecutor by testing missing digest scenarios
* Raise a dedicated exception for missing digests for easier detection, primarily in tests

Previous use of ParallelTask.ForEachAsync prevented exceptions from being propagated. This caused create process exceptions to occur when the binary for an action could not be found.

diff --git a/Engine/Source/Programs/Horde/HordeAgent/Execution/ActionExecutor.cs b/Engine/Source/Programs/Horde/HordeAgent/Execution/ActionExecutor.cs
--- a/Engine/Source/Programs/Horde/HordeAgent/Execution/ActionExecutor.cs
+++ b/Engine/Source/Programs/Horde/HordeAgent/Execution/ActionExecutor.cs
@@ -248,20 +248,22 @@ namespace HordeAgent
 		{
 			DirectoryReference.CreateDirectory(OutputDir);
 
-			await ParallelTask.ForEachAsync(InputDirectory.Files, async FileNode =>
+			async Task DownloadFile(FileNode FileNode)
 			{
 				ReadOnlyMemory<byte> Data = await Storage.GetBulkDataAsync(InstanceName, FileNode.Digest);
 				FileReference File = FileReference.Combine(OutputDir, FileNode.Name);
 				Logger.LogInformation("Writing {File} (digest: {Digest}, size: {Size})", File, FileNode.Digest.Hash, FileNode.Digest.SizeBytes);
 				await FileReference.WriteAllBytesAsync(File, Data.ToArray());
-			});
-
-			await ParallelTask.ForEachAsync(InputDirectory.Directories, async DirectoryNode =>
+			}
+			await Task.WhenAll(InputDirectory.Files.Select(x => Task.Run(() => DownloadFile(x))));
+			
+			async Task DownloadDir(DirectoryNode DirectoryNode)
 			{
 				Directory InputSubDirectory = await Storage.GetProtoMessageAsync<Directory>(InstanceName, DirectoryNode.Digest);
 				DirectoryReference OutputSubDirectory = DirectoryReference.Combine(OutputDir, DirectoryNode.Name);
 				await SetupSandboxAsync(InputSubDirectory, OutputSubDirectory);
-			});
+			}
+			await Task.WhenAll(InputDirectory.Directories.Select(x => Task.Run(() => DownloadDir(x))));
 		}
 
 		/// <summary>
diff --git a/Engine/Source/Programs/Horde/HordeAgentTests/ActionExecutorTest.cs b/Engine/Source/Programs/Horde/HordeAgentTests/ActionExecutorTest.cs
--- a/Engine/Source/Programs/Horde/HordeAgentTests/ActionExecutorTest.cs
+++ b/Engine/Source/Programs/Horde/HordeAgentTests/ActionExecutorTest.cs
@@ -2,10 +2,17 @@
 
 using System;
 using System.IO;
+using System.Threading;
+using System.Threading.Tasks;
 using EpicGames.Core;
+using EpicGames.Horde.Common.RemoteExecution;
 using Grpc.Core;
 using HordeAgent;
+using HordeCommon.Rpc.Tasks;
+using Microsoft.Extensions.Logging;
+using Microsoft.Extensions.Logging.Console;
 using Microsoft.VisualStudio.TestTools.UnitTesting;
+using Directory = System.IO.Directory;
 
 namespace HordeAgentTests
 {
@@ -13,6 +20,12 @@ namespace HordeAgentTests
 	public class ActionExecutorTest
 	{
 		private readonly string TempDir;
+		private readonly ILogger Logger = CreateLogger();
+		private readonly string AgentName = "myAgentName";
+		private readonly string InstanceName = "myInstanceName";
+		private readonly FakeCasClient CasClient = new FakeCasClient();
+		private readonly FakeActionRpcClient ActionRpcClient = new FakeActionRpcClient();
+
 
 		public ActionExecutorTest()
 		{
@@ -56,6 +69,59 @@ namespace HordeAgentTests
 			Assert.AreEqual("1/2/3/waldo", Resolved[4].RelativePath);
 		}
 
+		[TestMethod]
+		public async Task ExecuteAction()
+		{
+			ActionExecutor Executor = new ActionExecutor(AgentName, InstanceName, CasClient, null!, ActionRpcClient, Logger);
+			ActionTask ActionTask = TaskTest.CreateActionTask(InstanceName, CasClient);
+			await Executor.ExecuteActionAsync("myLeaseId1", ActionTask, new DirectoryReference(TempDir), DateTimeOffset.UtcNow, CancellationToken.None);
+		}
+
+		[TestMethod]
+		public async Task ExecuteActionMissingBinary()
+		{
+			ActionExecutor Executor = new ActionExecutor(AgentName, InstanceName, CasClient, null!, ActionRpcClient, Logger);
+			ActionTask ActionTask = TaskTest.CreateActionTask(InstanceName, CasClient, UploadBin: false);
+			await Assert.ThrowsExceptionAsync<DigestNotFoundException>(() => Executor.ExecuteActionAsync("myLeaseId1", ActionTask,
+				new DirectoryReference(TempDir), DateTimeOffset.UtcNow, CancellationToken.None));
+		}
+		
+		[TestMethod]
+		public async Task ExecuteActionMissingDirectory()
+		{
+			ActionExecutor Executor = new ActionExecutor(AgentName, InstanceName, CasClient, null!, ActionRpcClient, Logger);
+			ActionTask ActionTask = TaskTest.CreateActionTask(InstanceName, CasClient, UploadDir: false);
+			await Assert.ThrowsExceptionAsync<DigestNotFoundException>(() => Executor.ExecuteActionAsync("myLeaseId1", ActionTask,
+				new DirectoryReference(TempDir), DateTimeOffset.UtcNow, CancellationToken.None));
+		}
+		
+		[TestMethod]
+		public async Task ExecuteActionMissingAction()
+		{
+			ActionExecutor Executor = new ActionExecutor(AgentName, InstanceName, CasClient, null!, ActionRpcClient, Logger);
+			ActionTask ActionTask = TaskTest.CreateActionTask(InstanceName, CasClient, UploadAction: false);
+			await Assert.ThrowsExceptionAsync<DigestNotFoundException>(() => Executor.ExecuteActionAsync("myLeaseId1", ActionTask,
+				new DirectoryReference(TempDir), DateTimeOffset.UtcNow, CancellationToken.None));
+		}
+		
+		[TestMethod]
+		public async Task ExecuteActionMissingCommand()
+		{
+			ActionExecutor Executor = new ActionExecutor(AgentName, InstanceName, CasClient, null!, ActionRpcClient, Logger);
+			ActionTask ActionTask = TaskTest.CreateActionTask(InstanceName, CasClient, UploadCommand: false);
+			await Assert.ThrowsExceptionAsync<DigestNotFoundException>(() => Executor.ExecuteActionAsync("myLeaseId1", ActionTask,
+				new DirectoryReference(TempDir), DateTimeOffset.UtcNow, CancellationToken.None));
+		}
+		
+		private static ILogger CreateLogger()
+		{
+			LoggerFactory LoggerFactory = new LoggerFactory();
+			ConsoleLoggerOptions LoggerOptions = new ConsoleLoggerOptions();
+			TestOptionsMonitor<ConsoleLoggerOptions> LoggerOptionsMon = new TestOptionsMonitor<ConsoleLoggerOptions>(LoggerOptions);
+			LoggerFactory.AddProvider(new ConsoleLoggerProvider(LoggerOptionsMon));
+			return LoggerFactory.CreateLogger<ActionExecutor>();
+		}
+
 		[TestCleanup]
 		public void Cleanup()
 		{
diff --git a/Engine/Source/Programs/Horde/HordeAgentTests/TaskTest.cs b/Engine/Source/Programs/Horde/HordeAgentTests/TaskTest.cs
--- a/Engine/Source/Programs/Horde/HordeAgentTests/TaskTest.cs
+++ b/Engine/Source/Programs/Horde/HordeAgentTests/TaskTest.cs
@@ -88,7 +88,7 @@ namespace HordeAgentTests
 		{
 			if (!RuntimeInformation.IsOSPlatform(OSPlatform.Windows)) return;
 			
-			ActionTask ActionTask = CreateActionTask();
+			ActionTask ActionTask = CreateActionTask(InstanceName, Cas);
 			Assert.IsNull(ActionRpc.ActionResultRequest);
 			LeaseOutcome Outcome = await Ws.HandleLeasePayloadAsync(RpcConnection, "my-agent-id", CreateLeaseInfo(ActionTask));
 			Assert.AreEqual(LeaseOutcome.Success, Outcome);
@@ -102,14 +102,14 @@ namespace HordeAgentTests
 		{
 			if (!RuntimeInformation.IsOSPlatform(OSPlatform.Windows)) return;
 			
-			ActionTask ActionTask = CreateActionTask();
+			ActionTask ActionTask = CreateActionTask(InstanceName, Cas);
 			LeaseOutcome Outcome = await Ws.HandleLeasePayloadAsync(RpcConnection, "my-agent-id", CreateLeaseInfo(ActionTask));
 			Assert.AreEqual(LeaseOutcome.Failed, Outcome);
 			Assert.IsNotNull(ActionRpc.ActionResultRequest);
 			Assert.IsNull(ActionRpc.ActionResultRequest!.Error);
 		}
 		
-		private string GetHash(byte[] data)
+		private static string GetHash(byte[] data)
 		{
 			using System.Security.Cryptography.MD5 md5 = System.Security.Cryptography.MD5.Create();
 			byte[] hashBytes = md5.ComputeHash(data);
@@ -121,20 +121,34 @@ namespace HordeAgentTests
 			return sb.ToString();
 		}
 
-		private Build.Bazel.Remote.Execution.V2.Digest GetDigest(byte[] data)
+		private static Build.Bazel.Remote.Execution.V2.Digest GetDigest(byte[] data)
 		{
 			return new Build.Bazel.Remote.Execution.V2.Digest {Hash = GetHash(data), SizeBytes = data.Length};
 		}
 
-		private ActionTask CreateActionTask()
+		/// <summary>
+		/// Create an action task with the remote exec test binary
+		/// Will populate the CAS as well. Flags for uploading allow tests to verify different failure scenarios.
+		/// </summary>
+		/// <param name="InstanceName">Instance name</param>
+		/// <param name="Cas">The content-addressable store being used</param>
+		/// <param name="UploadBin">True to upload the remote exec test binary to CAS</param>
+		/// <param name="UploadDir">True to upload directory to CAS</param>
+		/// <param name="UploadCommand">True to upload the command to CAS</param>
+		/// <param name="UploadAction">True to upload action to CAS</param>
+		/// <returns>An ActionTask populated in CAS</returns>
+		public static ActionTask CreateActionTask(string InstanceName, FakeCasClient Cas, bool UploadBin = true, bool UploadDir = true, bool UploadCommand = true, bool UploadAction = true)
 		{
 			string BinName = "remote-exec-test-bin.exe";
 			string AssemblyPath = Path.GetDirectoryName(System.Reflection.Assembly.GetExecutingAssembly().Location)!;
 			string BinPath = Path.Combine(AssemblyPath, "remote-exec-test-bin", BinName);
-			
+
 			byte[] ExeFileData = File.ReadAllBytes(BinPath);
-			Cas.SetBlob(InstanceName, GetHash(ExeFileData), ByteString.CopyFrom(ExeFileData));
-				
+			if (UploadBin)
+			{
+				Cas.SetBlob(InstanceName, GetHash(ExeFileData), ByteString.CopyFrom(ExeFileData));
+			}
+
 			Build.Bazel.Remote.Execution.V2.FileNode ExeFile = new Build.Bazel.Remote.Execution.V2.FileNode();
 			ExeFile.Digest = GetDigest(ExeFileData);
 			ExeFile.Name = BinName;
@@ -143,21 +157,29 @@ namespace HordeAgentTests
 			Build.Bazel.Remote.Execution.V2.Directory Dir = new Build.Bazel.Remote.Execution.V2.Directory();
 			Dir.Files.Add(ExeFile);
 			byte[] DirData = Dir.ToByteArray();
-			Cas.SetBlob(InstanceName, GetHash(DirData), ByteString.CopyFrom(DirData));
-			
+			if (UploadDir)
+			{
+				Cas.SetBlob(InstanceName, GetHash(DirData), ByteString.CopyFrom(DirData));
+			}
+
 			Build.Bazel.Remote.Execution.V2.Command Command = new Build.Bazel.Remote.Execution.V2.Command();
 			Command.Arguments.Add(ExeFile.Name);
 			byte[] CommandData = Command.ToByteArray();
-			Cas.SetBlob(InstanceName, GetHash(CommandData), ByteString.CopyFrom(CommandData));
-			
-			
+			if (UploadCommand)
+			{
+				Cas.SetBlob(InstanceName, GetHash(CommandData), ByteString.CopyFrom(CommandData));
+			}
+
 			Build.Bazel.Remote.Execution.V2.Action Action = new Build.Bazel.Remote.Execution.V2.Action();
 			Action.CommandDigest = GetDigest(CommandData);
 			Action.InputRootDigest = GetDigest(DirData);
 			Action.DoNotCache = true;
 			
 			byte[] ActionData = Action.ToByteArray();
-			Cas.SetBlob(InstanceName, GetHash(ActionData), ByteString.CopyFrom(ActionData));
+			if (UploadAction)
+			{
+				Cas.SetBlob(InstanceName, GetHash(ActionData), ByteString.CopyFrom(ActionData));
+			}
 
 			ActionTask ActionTask = new ActionTask();
 			ActionTask.Digest = GetDigest(ActionData);
diff --git a/Engine/Source/Programs/Shared/EpicGames.Horde.Common/RemoteExecution/StorageExtensions.cs b/Engine/Source/Programs/Shared/EpicGames.Horde.Common/RemoteExecution/StorageExtensions.cs
--- a/Engine/Source/Programs/Shared/EpicGames.Horde.Common/RemoteExecution/StorageExtensions.cs
+++ b/Engine/Source/Programs/Shared/EpicGames.Horde.Common/RemoteExecution/StorageExtensions.cs
@@ -4,13 +4,21 @@ using Build.Bazel.Remote.Execution.V2;
 using EpicGames.Core;
 using Google.Protobuf;
 using System;
+using System.Net.Http;
 using System.Security.Cryptography;
 using System.Threading.Tasks;
+using Google.Rpc;
 using ContentAddressableStorageClient = Build.Bazel.Remote.Execution.V2.ContentAddressableStorage.ContentAddressableStorageClient;
 using Digest = Build.Bazel.Remote.Execution.V2.Digest;
 
 namespace EpicGames.Horde.Common.RemoteExecution
 {
+	public class DigestNotFoundException : Exception
+	{
+		public DigestNotFoundException(string InstanceName, Digest Digest, Status Status)
+			: base($"Failed getting digest {Digest}. Code={Status.Code} Message={Status.Message} InstanceName={InstanceName}")
+		{}
+	}
 	public static class StorageExtensions
 	{
 		public static async Task<byte[]> GetBulkDataAsync(this ContentAddressableStorageClient Storage, string InstanceName, Digest Digest)
@@ -19,8 +27,15 @@ namespace EpicGames.Horde.Common.RemoteExecution
 			Request.InstanceName = InstanceName;
 			Request.Digests.Add(Digest);
 
-			BatchReadBlobsResponse Response = await Storage.BatchReadBlobsAsync(Request);
-			return Response.Responses[0].Data.ToByteArray();
+			BatchReadBlobsResponse BatchResponse = await Storage.BatchReadBlobsAsync(Request);
+			BatchReadBlobsResponse.Types.Response Res = BatchResponse.Responses[0];
+			
+			if (Res.Status.Code != (int)Google.Rpc.Code.Ok)
+			{
+				throw new DigestNotFoundException(InstanceName, Digest, Res.Status);
+			}
+
+			return Res.Data.ToByteArray();
 		}
 
 		public static async Task<Digest> PutBulkDataAsync(this ContentAddressableStorageClient Storage, string InstanceName, byte[] Data)
@@ -47,7 +62,7 @@ namespace EpicGames.Horde.Common.RemoteExecution
 			BatchReadBlobsResponse Response = await Storage.BatchReadBlobsAsync(Request);
 			if (Response.Responses[0].Status.Code != (int)Google.Rpc.Code.Ok)
 			{
-				throw new Exception($"Failed getting digest {Digest}. Code={Response.Responses[0].Status.Code} Message={Response.Responses[0].Status.Message}");
+				throw new DigestNotFoundException(InstanceName, Digest, Response.Responses[0].Status);
 			}
 
 			T Item = new T();
