Horde: Adding device service Slack notifications
#rnx

diff --git a/Engine/Source/Programs/Horde/HordeServer/Collections/IDeviceCollection.cs b/Engine/Source/Programs/Horde/HordeServer/Collections/IDeviceCollection.cs
--- a/Engine/Source/Programs/Horde/HordeServer/Collections/IDeviceCollection.cs
+++ b/Engine/Source/Programs/Horde/HordeServer/Collections/IDeviceCollection.cs
@@ -175,10 +175,16 @@ namespace HordeServer.Collections
 		/// <param name="LegacyGuid">YThe legacy guid of the reservation</param>
 		Task<IDeviceReservation?> TryGetReservationFromLegacyGuidAsync(string LegacyGuid);
 
-		/// <summary>
-		/// Get a list of all reservations
-		/// </summary>
-		Task<List<IDeviceReservation>> FindAllReservationsAsync();
+        /// <summary>
+        /// Gets a reservation by device id
+        /// </summary>
+        /// <param name="Id">A device contained in reservation</param>
+        Task<IDeviceReservation?> TryGetDeviceReservationAsync(DeviceId Id);
+
+        /// <summary>
+        /// Get a list of all reservations
+        /// </summary>
+        Task<List<IDeviceReservation>> FindAllReservationsAsync();
 
 		/// <summary>
 		/// Updates a reservation to the current time, for expiration
diff --git a/Engine/Source/Programs/Horde/HordeServer/Collections/Impl/DeviceCollection.cs b/Engine/Source/Programs/Horde/HordeServer/Collections/Impl/DeviceCollection.cs
--- a/Engine/Source/Programs/Horde/HordeServer/Collections/Impl/DeviceCollection.cs
+++ b/Engine/Source/Programs/Horde/HordeServer/Collections/Impl/DeviceCollection.cs
@@ -526,7 +526,7 @@ namespace HordeServer.Collections.Impl
 				x.MaintenanceTimeUtc == null).ToListAsync();
 
 			// filter out problem devices
-			PoolDevices = PoolDevices.FindAll(x => (x.ProblemTimeUtc == null || ((ReservationTimeUtc - x.ProblemTimeUtc).Value.TotalMinutes > 30)));
+			PoolDevices = PoolDevices.FindAll(x => (x.ProblemTimeUtc == null || ((ReservationTimeUtc - x.ProblemTimeUtc).Value.TotalMinutes > 120)));
 
 			// filter out currently reserved devices
 			PoolDevices = PoolDevices.FindAll(x => PoolReservations.FirstOrDefault(p => p.Devices.Contains(x.Id)) == null);
@@ -672,9 +672,15 @@ namespace HordeServer.Collections.Impl
 		/// <inheritdoc/>
 		public async Task<IDeviceReservation?> TryGetReservationFromLegacyGuidAsync(string LegacyGuid)
 		{
-			return await Reservations.Find<DeviceReservationDocument>(x => x.LegacyGuid == LegacyGuid).FirstOrDefaultAsync();
+			return await Reservations.Find<DeviceReservationDocument>(R => R.LegacyGuid == LegacyGuid).FirstOrDefaultAsync();
 		}
 
+		/// <inheritdoc/>
+		public async Task<IDeviceReservation?> TryGetDeviceReservationAsync(DeviceId Id)
+		{
+            return await Reservations.Find<DeviceReservationDocument>(R => R.Devices.FirstOrDefault(D => D == Id) != null).FirstOrDefaultAsync();
+        }
+
 	}
 
 }
diff --git a/Engine/Source/Programs/Horde/HordeServer/Controllers/DevicesController.cs b/Engine/Source/Programs/Horde/HordeServer/Controllers/DevicesController.cs
--- a/Engine/Source/Programs/Horde/HordeServer/Controllers/DevicesController.cs
+++ b/Engine/Source/Programs/Horde/HordeServer/Controllers/DevicesController.cs
@@ -53,7 +53,7 @@ namespace HordeServer.Controllers
 		{
             this.UserCollection = UserCollection;
             this.DeviceService = DeviceService;
-			this.Logger = Logger;
+            this.Logger = Logger;
 		}
 
 		// DEVICES
@@ -676,18 +676,26 @@ namespace HordeServer.Controllers
 			List<IDevicePool> Pools = await DeviceService.GetPoolsAsync();
 			List<IDevicePlatform> Platforms = await DeviceService.GetPlatformsAsync();
 
-			string? PoolId = Request.PoolId;
+            string Message;
+            string? PoolId = Request.PoolId;
 
 			if (string.IsNullOrEmpty(PoolId))
 			{
-                PoolId = "fnue5";
+                Message = $"No pool specified for device reservation request";				
+                Logger.LogError(Message);
+				await DeviceService.NotifyDeviceServiceAsync(Message, null, Request.JobId, Request.StepId);
+
+                return BadRequest(Message);
             }			
 
 			DevicePoolId PoolIdValue = DevicePoolId.Sanitize(PoolId);
 			IDevicePool? Pool = Pools.FirstOrDefault(x => x.Id == PoolIdValue);
 			if (Pool == null)
 			{
-				return BadRequest($"Unknown pool {PoolId} on device reservation request");
+                Message = $"Unknown pool {PoolId} on device reservation request";
+				Logger.LogError(Message);
+				await DeviceService.NotifyDeviceServiceAsync(Message, null, Request.JobId, Request.StepId);
+                return BadRequest(Message);
 			}
 
 			List<DeviceRequestData> RequestedDevices = new List<DeviceRequestData>();
@@ -713,14 +721,20 @@ namespace HordeServer.Controllers
 
 				if (!PlatformMap.TryGetValue(PlatformName, out PlatformId))
 				{
-					return BadRequest($"Unknown platform for device type {PlatformName} on device reservation request");
+                    Message = $"Unknown platform for device type {PlatformName} on device reservation request";
+					Logger.LogError(Message);
+					await DeviceService.NotifyDeviceServiceAsync(Message, null, Request.JobId, Request.StepId);
+                    return BadRequest(Message);
 				}
 
 				IDevicePlatform? Platform = Platforms.FirstOrDefault(x => x.Id == PlatformId);
 
 				if (Platform == null)
 				{
-					return BadRequest($"Unknown platform {PlatformId} on device reservation request");
+                    Message = $"Unknown platform {PlatformId} on device reservation request";
+					Logger.LogError(Message);
+					await DeviceService.NotifyDeviceServiceAsync(Message, null, Request.JobId, Request.StepId);
+                    return BadRequest(Message);
 				}
 
 				List<string> IncludeModels = new List<string>();
@@ -884,16 +898,28 @@ namespace HordeServer.Controllers
 
 			if (Device == null)
 			{
+				Logger.LogError($"Device error reported for unknown device {DeviceName}");
 				return BadRequest($"Unknown device {DeviceName}");
-			}
+			}			
 
 			await DeviceService.UpdateDeviceAsync(Device.Id, NewProblem: true);
 
+            IDeviceReservation? Reservation = await DeviceService.TryGetDeviceReservation(Device.Id);
 
-			return Ok();
+            string? JobId = null;
+            string? StepId = null;
+            if (Reservation != null) 
+			{
+				JobId = !string.IsNullOrEmpty(Reservation.JobId) ? Reservation.JobId : null;
+                StepId = !string.IsNullOrEmpty(Reservation.StepId) ? Reservation.StepId : null;
+            }
 
-		}
+			Logger.LogError($"Device error reported for {DeviceName}, Job {JobId}, StepId {StepId}");
+
+            await DeviceService.NotifyDeviceServiceAsync("Device reported an error", Device.Id, JobId, StepId);
 
+			return Ok();
 
+		}
 	}
 }
\ No newline at end of file
diff --git a/Engine/Source/Programs/Horde/HordeServer/Notifications/INotificationService.cs b/Engine/Source/Programs/Horde/HordeServer/Notifications/INotificationService.cs
--- a/Engine/Source/Programs/Horde/HordeServer/Notifications/INotificationService.cs
+++ b/Engine/Source/Programs/Horde/HordeServer/Notifications/INotificationService.cs
@@ -68,5 +68,17 @@ namespace HordeServer.Notifications
 		/// </summary>
 		/// <param name="Issue">The new issue that was created</param>
 		void NotifyIssueUpdated(IIssue Issue);
-	}
+
+		/// <summary>
+        /// Send a device service notification
+        /// </summary>
+        /// <param name="Message">The message to send</param>
+        /// <param name="Device">The device</param>
+        /// <param name="Pool">The pool</param>
+        /// <param name="Stream"></param>
+        /// <param name="Job"></param>
+        /// <param name="Step"></param>
+        /// <param name="Node"></param>
+        void NotifyDeviceService(string Message, IDevice? Device = null, IDevicePool? Pool = null, IStream? Stream = null, IJob? Job = null, IJobStep? Step = null, INode? Node = null);
+    }
 }
diff --git a/Engine/Source/Programs/Horde/HordeServer/Notifications/INotificationSink.cs b/Engine/Source/Programs/Horde/HordeServer/Notifications/INotificationSink.cs
--- a/Engine/Source/Programs/Horde/HordeServer/Notifications/INotificationSink.cs
+++ b/Engine/Source/Programs/Horde/HordeServer/Notifications/INotificationSink.cs
@@ -80,5 +80,18 @@ namespace HordeServer.Notifications
 		/// <param name="Description"></param>
 		/// <returns></returns>
 		Task NotifyConfigUpdateFailureAsync(string ErrorMessage, string FileName, int? Change = null, IUser? Author = null, string? Description = null);
-	}
+		
+		/// <summary>
+		/// Notification for device service
+		/// </summary>
+		/// <param name="Message"></param>
+		/// <param name="Device"></param>
+		/// <param name="Pool"></param>
+		/// <param name="Stream"></param>
+		/// <param name="Job"></param>
+        /// <param name="Step"></param>
+        /// <param name="Node"></param>
+		/// <returns></returns>
+        Task NotifyDeviceServiceAsync(string Message, IDevice? Device = null, IDevicePool? Pool = null, IStream? Stream = null, IJob? Job = null, IJobStep? Step = null, INode? Node = null);
+    }
 }
diff --git a/Engine/Source/Programs/Horde/HordeServer/Notifications/Impl/NotificationService.cs b/Engine/Source/Programs/Horde/HordeServer/Notifications/Impl/NotificationService.cs
--- a/Engine/Source/Programs/Horde/HordeServer/Notifications/Impl/NotificationService.cs
+++ b/Engine/Source/Programs/Horde/HordeServer/Notifications/Impl/NotificationService.cs
@@ -229,6 +229,12 @@ namespace HordeServer.Notifications.Impl
 			EnqueueTasks(Sink => Sink.NotifyConfigUpdateFailureAsync(ErrorMessage, FileName, Change, User, Description));
 		}
 
+		/// <inheritdoc/>
+		public void NotifyDeviceService(string Message, IDevice? Device = null, IDevicePool? Pool = null, IStream? Stream = null, IJob? Job = null, IJobStep? Step = null, INode? Node = null)
+		{
+			EnqueueTasks(Sink => Sink.NotifyDeviceServiceAsync(Message, Device, Pool, Stream, Job, Step, Node));
+		}
+
 		/// <summary>
 		/// Enqueues an async task
 		/// </summary>
diff --git a/Engine/Source/Programs/Horde/HordeServer/Notifications/Impl/SlackNotificationSink.cs b/Engine/Source/Programs/Horde/HordeServer/Notifications/Impl/SlackNotificationSink.cs
--- a/Engine/Source/Programs/Horde/HordeServer/Notifications/Impl/SlackNotificationSink.cs
+++ b/Engine/Source/Programs/Horde/HordeServer/Notifications/Impl/SlackNotificationSink.cs
@@ -961,6 +961,56 @@ namespace HordeServer.Notifications.Impl
 
 		#endregion
 
+		#region Device notifications
+
+		/// <inheritdoc/>
+		public async Task NotifyDeviceServiceAsync(string Message, IDevice? Device = null, IDevicePool? Pool = null, IStream? Stream = null, IJob? Job = null, IJobStep? Step = null, INode? Node = null)
+		{
+			Logger.LogDebug("Sending device service failure notification for {DeviceName} in pool {PoolName}", Device?.Name, Pool?.Name);
+			if (Settings.DeviceServiceNotificationChannel != null)
+			{
+				await SendDeviceServiceMessage($"#{Settings.DeviceServiceNotificationChannel}", Message, Device, Pool, Stream, Job, Step, Node);
+			}
+		}
+
+		/// <summary>
+		/// Creates a Slack message about a completed step job.
+		/// </summary>
+		/// <param name="Recipient"></param>
+        /// <param name="Message"></param>
+        /// <param name="Device"></param>
+        /// <param name="Pool"></param>
+		/// <param name="Stream"></param>
+		/// <param name="Job">The job that contains the step that completed.</param>
+		/// <param name="Step">The job step that completed.</param>
+		/// <param name="Node">The node for the job step.</param>
+		private Task SendDeviceServiceMessage(string Recipient, string Message, IDevice? Device = null, IDevicePool? Pool = null, IStream? Stream = null, IJob? Job = null, IJobStep? Step = null, INode? Node = null)
+		{
+			BlockKitAttachment Attachment = new BlockKitAttachment();
+
+			Attachment.FallbackText = $"{Message}";
+
+			if (Device != null && Pool != null)
+			{
+                Attachment.FallbackText += $" - Device: {Device.Name} Pool: {Pool.Name}";
+            }
+
+			if (Stream != null && Job != null && Step != null && Node != null)
+			{
+				Uri JobStepLink = new Uri($"{Settings.DashboardUrl}job/{Job.Id}?step={Step.Id}");
+				Uri JobStepLogLink = new Uri($"{Settings.DashboardUrl}log/{Step.LogId}");
+				
+				Attachment.FallbackText += $" - {Stream.Name} - {GetJobChangeText(Job)} - {Job.Name} - {Node.Name}";
+				Attachment.Blocks.Add(new SectionBlock($"*<{JobStepLink}|{Stream.Name} - {GetJobChangeText(Job)} - {Job.Name} - {Node.Name}>*"));
+				Attachment.Blocks.Add(new SectionBlock($"<{JobStepLogLink}|View Job Step Log>"));
+			}
+
+			return SendMessageAsync(Recipient, Attachments: new[] { Attachment });
+		}
+		
+		#endregion
+
+
 		const int MaxJobStepEvents = 5;
 
 		static string GetJobChangeText(IJob Job)
diff --git a/Engine/Source/Programs/Horde/HordeServer/ServerSettings.cs b/Engine/Source/Programs/Horde/HordeServer/ServerSettings.cs
--- a/Engine/Source/Programs/Horde/HordeServer/ServerSettings.cs
+++ b/Engine/Source/Programs/Horde/HordeServer/ServerSettings.cs
@@ -353,6 +353,11 @@ namespace HordeServer
 		/// </summary>
 		public string? UpdateStreamsNotificationChannel { get; set; }
 
+		/// <summary>
+		/// Channel to send device notifications to
+		/// </summary>
+		public string? DeviceServiceNotificationChannel { get; set; }
+
 		/// <summary>
 		/// URI to the SmtpServer to use for sending email notifications
 		/// </summary>
diff --git a/Engine/Source/Programs/Horde/HordeServer/Services/DeviceService.cs b/Engine/Source/Programs/Horde/HordeServer/Services/DeviceService.cs
--- a/Engine/Source/Programs/Horde/HordeServer/Services/DeviceService.cs
+++ b/Engine/Source/Programs/Horde/HordeServer/Services/DeviceService.cs
@@ -18,6 +18,8 @@ using DevicePlatformId = HordeServer.Utilities.StringId<HordeServer.Models.IDevi
 using DevicePoolId = HordeServer.Utilities.StringId<HordeServer.Models.IDevicePool>;
 using System.Security.Claims;
 using System.Diagnostics.CodeAnalysis;
+using HordeServer.Notifications;
+using System.Linq;
 
 namespace HordeServer.Services
 {
@@ -31,6 +33,22 @@ namespace HordeServer.Services
 		/// </summary>
 		AclService AclService;
 
+		/// <summary>
+		/// Instance of the notification service
+		/// </summary>
+		INotificationService NotificationService;
+
+
+		/// <summary>
+		/// Singleton instance of the job service
+		/// </summary>
+		JobService JobService;
+
+		/// <summary>
+		/// Singleton instance of the stream service
+		/// </summary>
+		StreamService StreamService;
+
 		/// <summary>
 		/// Log output writer
 		/// </summary>
@@ -44,12 +62,15 @@ namespace HordeServer.Services
 		/// <summary>
 		/// Device service constructor
 		/// </summary>
-		public DeviceService(IDeviceCollection Devices, AclService AclService, ILogger<DeviceService> Logger)
+		public DeviceService(IDeviceCollection Devices, JobService JobService, StreamService StreamService, AclService AclService, INotificationService NotificationService, ILogger<DeviceService> Logger)
 			: base(TimeSpan.FromMinutes(1.0), Logger)
 		{
 			this.Devices = Devices;
-			this.AclService = AclService;
-			this.Logger = Logger;
+            this.JobService = JobService;
+            this.StreamService = StreamService;
+            this.AclService = AclService;
+            this.NotificationService = NotificationService;
+            this.Logger = Logger;
 		}
 
 		/// <summary>
@@ -218,6 +239,14 @@ namespace HordeServer.Services
 			return Devices.TryGetReservationFromLegacyGuidAsync(LegacyGuid);
 		}
 
+		/// <summary>
+		/// Get a reservation from a device id
+		/// </summary>
+		public Task<IDeviceReservation?> TryGetDeviceReservation(DeviceId DeviceId)
+		{
+			return Devices.TryGetDeviceReservationAsync(DeviceId);
+		}
+
 		/// <summary>
 		/// Get a list of existing device reservations
 		/// </summary>
@@ -226,6 +255,65 @@ namespace HordeServer.Services
 			return Devices.FindAllReservationsAsync();
 		}
 
+		/// <summary>
+		/// 
+		/// </summary>
+		/// <param name="Message"></param>
+		/// <param name="DeviceId"></param>
+		/// <param name="JobId"></param>
+		/// <param name="StepId"></param>
+		public async Task NotifyDeviceServiceAsync(string Message, DeviceId? DeviceId = null, string? JobId = null, string? StepId = null)
+		{
+			try 
+			{
+				IDevice? Device = null;
+				IDevicePool? Pool = null;
+				IJob? Job = null;
+				IJobStep? Step = null;
+				INode? Node = null;
+				IStream? Stream = null;
+
+				if (DeviceId.HasValue)
+				{
+					Device = await GetDeviceAsync(DeviceId.Value);
+					Pool = await GetPoolAsync(Device!.PoolId);
+				}
+
+				if (JobId != null)
+				{
+					Job = await JobService.GetJobAsync(new ObjectId(JobId));
+
+					if (Job != null)
+					{
+						Stream = await StreamService.GetStreamAsync(Job.StreamId);
+
+						if (StepId != null)
+						{
+							IGraph Graph = await JobService.GetGraphAsync(Job)!;
+
+							SubResourceId StepIdValue = SubResourceId.Parse(StepId);
+							IJobStepBatch Batch = Job.Batches.FirstOrDefault(B => B.Steps.FirstOrDefault(S => S.Id == StepIdValue) != null);
+							if (Batch != null)
+							{
+								Step = Batch.Steps.FirstOrDefault(S => S.Id == StepIdValue)!;
+								INodeGroup Group = Graph.Groups[Batch.GroupIdx];
+								Node = Group.Nodes[Step.NodeIdx];
+							}
+						}
+					}
+				}
+
+				NotificationService.NotifyDeviceService(Message, Device, Pool, Stream, Job, Step, Node);
+
+			}
+			catch (Exception Ex)
+			{
+                Logger.LogError($"Error on device notification {Ex.Message}");
+            }
+        }
+
+
+
 		/// <summary>
 		/// Authorize device action
 		/// </summary>
