Horde: Set default pool if none specified during migration
#rnx

diff --git a/Engine/Source/Programs/Horde/HordeServer/Controllers/DevicesController.cs b/Engine/Source/Programs/Horde/HordeServer/Controllers/DevicesController.cs
--- a/Engine/Source/Programs/Horde/HordeServer/Controllers/DevicesController.cs
+++ b/Engine/Source/Programs/Horde/HordeServer/Controllers/DevicesController.cs
@@ -681,11 +681,11 @@ namespace HordeServer.Controllers
 
 			if (string.IsNullOrEmpty(PoolId))
 			{
-                Message = $"No pool specified for device reservation request";				
-                Logger.LogError(Message);
+                Message = $"No pool specified for device reservation request, defaulting to UE5FN (will soon be an error)";
+                Logger.LogError(Message + $" JobId: {Request.JobId}, StepId: {Request.StepId}");
 				await DeviceService.NotifyDeviceServiceAsync(Message, null, Request.JobId, Request.StepId);
-
-                return BadRequest(Message);
+				PoolId = "ue5fn";
+                //return BadRequest(Message);
             }			
 
 			DevicePoolId PoolIdValue = DevicePoolId.Sanitize(PoolId);
