Horde: Add logging when downloading files for an REAPI action

diff --git a/Engine/Source/Programs/Horde/HordeAgent/Execution/ActionExecutor.cs b/Engine/Source/Programs/Horde/HordeAgent/Execution/ActionExecutor.cs
--- a/Engine/Source/Programs/Horde/HordeAgent/Execution/ActionExecutor.cs
+++ b/Engine/Source/Programs/Horde/HordeAgent/Execution/ActionExecutor.cs
@@ -250,8 +250,9 @@ namespace HordeAgent
 
 			async Task DownloadFile(FileNode FileNode)
 			{
-				ReadOnlyMemory<byte> Data = await Storage.GetBulkDataAsync(InstanceName, FileNode.Digest);
 				FileReference File = FileReference.Combine(OutputDir, FileNode.Name);
+				Logger.LogInformation("Downloading {File} (digest: {Digest}, size: {Size})", File, FileNode.Digest.Hash, FileNode.Digest.SizeBytes);
+				ReadOnlyMemory<byte> Data = await Storage.GetBulkDataAsync(InstanceName, FileNode.Digest);
 				Logger.LogInformation("Writing {File} (digest: {Digest}, size: {Size})", File, FileNode.Digest.Hash, FileNode.Digest.SizeBytes);
 				await FileReference.WriteAllBytesAsync(File, Data.ToArray());
 			}
