Horde: Adding slack notification header
#rnx

diff --git a/Engine/Source/Programs/Horde/HordeServer/Controllers/DevicesController.cs b/Engine/Source/Programs/Horde/HordeServer/Controllers/DevicesController.cs
--- a/Engine/Source/Programs/Horde/HordeServer/Controllers/DevicesController.cs
+++ b/Engine/Source/Programs/Horde/HordeServer/Controllers/DevicesController.cs
@@ -684,7 +684,7 @@ namespace HordeServer.Controllers
                 Message = $"No pool specified for device reservation request, defaulting to UE5FN (will soon be an error)";
                 Logger.LogError(Message + $" JobId: {Request.JobId}, StepId: {Request.StepId}");
 				await DeviceService.NotifyDeviceServiceAsync(Message, null, Request.JobId, Request.StepId);
-				PoolId = "ue5fn";
+				PoolId = "fnue5";
                 //return BadRequest(Message);
             }			
 
diff --git a/Engine/Source/Programs/Horde/HordeServer/Notifications/Impl/SlackNotificationSink.cs b/Engine/Source/Programs/Horde/HordeServer/Notifications/Impl/SlackNotificationSink.cs
--- a/Engine/Source/Programs/Horde/HordeServer/Notifications/Impl/SlackNotificationSink.cs
+++ b/Engine/Source/Programs/Horde/HordeServer/Notifications/Impl/SlackNotificationSink.cs
@@ -1001,6 +1001,7 @@ namespace HordeServer.Notifications.Impl
 				Uri JobStepLogLink = new Uri($"{Settings.DashboardUrl}log/{Step.LogId}");
 				
 				Attachment.FallbackText += $" - {Stream.Name} - {GetJobChangeText(Job)} - {Job.Name} - {Node.Name}";
+				Attachment.Blocks.Add(new HeaderBlock($"{Message}", false, false));
 				Attachment.Blocks.Add(new SectionBlock($"*<{JobStepLink}|{Stream.Name} - {GetJobChangeText(Job)} - {Job.Name} - {Node.Name}>*"));
 				Attachment.Blocks.Add(new SectionBlock($"<{JobStepLogLink}|View Job Step Log>"));
 			}
