Horde: Fix mongo device query
#rnx

diff --git a/Engine/Source/Programs/Horde/HordeServer/Collections/Impl/DeviceCollection.cs b/Engine/Source/Programs/Horde/HordeServer/Collections/Impl/DeviceCollection.cs
--- a/Engine/Source/Programs/Horde/HordeServer/Collections/Impl/DeviceCollection.cs
+++ b/Engine/Source/Programs/Horde/HordeServer/Collections/Impl/DeviceCollection.cs
@@ -526,7 +526,7 @@ namespace HordeServer.Collections.Impl
 				x.MaintenanceTimeUtc == null).ToListAsync();
 
 			// filter out problem devices
-			PoolDevices = PoolDevices.FindAll(x => (x.ProblemTimeUtc == null || ((ReservationTimeUtc - x.ProblemTimeUtc).Value.TotalMinutes > 120)));
+			PoolDevices = PoolDevices.FindAll(x => (x.ProblemTimeUtc == null || ((ReservationTimeUtc - x.ProblemTimeUtc).Value.TotalMinutes > 30)));
 
 			// filter out currently reserved devices
 			PoolDevices = PoolDevices.FindAll(x => PoolReservations.FirstOrDefault(p => p.Devices.Contains(x.Id)) == null);
@@ -678,7 +678,8 @@ namespace HordeServer.Collections.Impl
 		/// <inheritdoc/>
 		public async Task<IDeviceReservation?> TryGetDeviceReservationAsync(DeviceId Id)
 		{
-            return await Reservations.Find<DeviceReservationDocument>(R => R.Devices.FirstOrDefault(D => D == Id) != null).FirstOrDefaultAsync();
+			List<DeviceReservationDocument> Results = await Reservations.Find(x => true).ToListAsync();
+            return Results.FirstOrDefault(R => R.Devices.FirstOrDefault(D => D == Id) != null);
         }
 
 	}
diff --git a/Engine/Source/Programs/Horde/HordeServer/Controllers/DevicesController.cs b/Engine/Source/Programs/Horde/HordeServer/Controllers/DevicesController.cs
--- a/Engine/Source/Programs/Horde/HordeServer/Controllers/DevicesController.cs
+++ b/Engine/Source/Programs/Horde/HordeServer/Controllers/DevicesController.cs
@@ -914,9 +914,11 @@ namespace HordeServer.Controllers
                 StepId = !string.IsNullOrEmpty(Reservation.StepId) ? Reservation.StepId : null;
             }
 
-			Logger.LogError($"Device error reported for {DeviceName}, Job {JobId}, StepId {StepId}");
+            string Message = $"Device reported an error, {Device.Name} : {Device.PoolId.ToString().ToUpperInvariant()}";
 
-            await DeviceService.NotifyDeviceServiceAsync("Device reported an error", Device.Id, JobId, StepId);
+            Logger.LogError(Message);
+
+            await DeviceService.NotifyDeviceServiceAsync(Message, Device.Id, JobId, StepId);
 
 			return Ok();
 
