Horde: Another reservation lookup fix
#rnx

diff --git a/Engine/Source/Programs/Horde/HordeServer/Collections/Impl/DeviceCollection.cs b/Engine/Source/Programs/Horde/HordeServer/Collections/Impl/DeviceCollection.cs
--- a/Engine/Source/Programs/Horde/HordeServer/Collections/Impl/DeviceCollection.cs
+++ b/Engine/Source/Programs/Horde/HordeServer/Collections/Impl/DeviceCollection.cs
@@ -679,7 +679,7 @@ namespace HordeServer.Collections.Impl
 		public async Task<IDeviceReservation?> TryGetDeviceReservationAsync(DeviceId Id)
 		{
 			List<DeviceReservationDocument> Results = await Reservations.Find(x => true).ToListAsync();
-            return Results.FirstOrDefault(R => R.Devices.FirstOrDefault(D => D == Id) != null);
+			return Results.FirstOrDefault(R => R.Devices.Contains(Id));
         }
 
 	}
diff --git a/Engine/Source/Programs/Horde/HordeServer/Controllers/DevicesController.cs b/Engine/Source/Programs/Horde/HordeServer/Controllers/DevicesController.cs
--- a/Engine/Source/Programs/Horde/HordeServer/Controllers/DevicesController.cs
+++ b/Engine/Source/Programs/Horde/HordeServer/Controllers/DevicesController.cs
@@ -681,7 +681,7 @@ namespace HordeServer.Controllers
 
 			if (string.IsNullOrEmpty(PoolId))
 			{
-                Message = $"No pool specified for device reservation request, defaulting to UE5FN (will soon be an error)";
+                Message = $"No pool specified for device reservation request, defaulting to FNUE5 (will soon be an error)";
                 Logger.LogError(Message + $" JobId: {Request.JobId}, StepId: {Request.StepId}");
 				await DeviceService.NotifyDeviceServiceAsync(Message, null, Request.JobId, Request.StepId);
 				PoolId = "fnue5";
