Horde: Add Slack notification section in case of no job information
#rnx

diff --git a/Engine/Source/Programs/Horde/HordeServer/Notifications/Impl/SlackNotificationSink.cs b/Engine/Source/Programs/Horde/HordeServer/Notifications/Impl/SlackNotificationSink.cs
--- a/Engine/Source/Programs/Horde/HordeServer/Notifications/Impl/SlackNotificationSink.cs
+++ b/Engine/Source/Programs/Horde/HordeServer/Notifications/Impl/SlackNotificationSink.cs
@@ -995,16 +995,22 @@ namespace HordeServer.Notifications.Impl
                 Attachment.FallbackText += $" - Device: {Device.Name} Pool: {Pool.Name}";
             }
 
+			Attachment.Blocks.Add(new HeaderBlock($"{Message}", false, false));
+
 			if (Stream != null && Job != null && Step != null && Node != null)
 			{
 				Uri JobStepLink = new Uri($"{Settings.DashboardUrl}job/{Job.Id}?step={Step.Id}");
 				Uri JobStepLogLink = new Uri($"{Settings.DashboardUrl}log/{Step.LogId}");
 				
-				Attachment.FallbackText += $" - {Stream.Name} - {GetJobChangeText(Job)} - {Job.Name} - {Node.Name}";
-				Attachment.Blocks.Add(new HeaderBlock($"{Message}", false, false));
+				Attachment.FallbackText += $" - {Stream.Name} - {GetJobChangeText(Job)} - {Job.Name} - {Node.Name}";				
 				Attachment.Blocks.Add(new SectionBlock($"*<{JobStepLink}|{Stream.Name} - {GetJobChangeText(Job)} - {Job.Name} - {Node.Name}>*"));
 				Attachment.Blocks.Add(new SectionBlock($"<{JobStepLogLink}|View Job Step Log>"));
 			}
+			else
+			{				
+				Attachment.FallbackText += " - No job information (Gauntlet might need to be updated in stream)";				
+				Attachment.Blocks.Add(new SectionBlock("*No job information (Gauntlet might need to be updated in stream)*"));
+			}
 
 			return SendMessageAsync(Recipient, Attachments: new[] { Attachment });
 		}
