Horde: Show legacy reservation details in slack notifications for streams which haven't been updated
#rnx

diff --git a/Engine/Source/Programs/Horde/HordeServer/Controllers/DevicesController.cs b/Engine/Source/Programs/Horde/HordeServer/Controllers/DevicesController.cs
--- a/Engine/Source/Programs/Horde/HordeServer/Controllers/DevicesController.cs
+++ b/Engine/Source/Programs/Horde/HordeServer/Controllers/DevicesController.cs
@@ -679,9 +679,23 @@ namespace HordeServer.Controllers
             string Message;
             string? PoolId = Request.PoolId;
 
+			// @todo: Remove this once all streams are updated to provide jobid
+            string Details = "";
+            if ((String.IsNullOrEmpty(Request.JobId) || String.IsNullOrEmpty(Request.StepId)))
+			{
+ 				if (!string.IsNullOrEmpty(Request.ReservationDetails))
+				 {
+					Details = $" - {Request.ReservationDetails}";
+				 }
+				 else
+				 {
+					Details = $" - No Reservation Information Available";
+				 }							 
+            }
+
 			if (string.IsNullOrEmpty(PoolId))
 			{
-                Message = $"No pool specified for device reservation request, defaulting to FNUE5 (will soon be an error)";
+                Message = $"No pool specified for device reservation request, defaulting to FNUE5 (will soon be an error)" + Details;
                 Logger.LogError(Message + $" JobId: {Request.JobId}, StepId: {Request.StepId}");
 				await DeviceService.NotifyDeviceServiceAsync(Message, null, Request.JobId, Request.StepId);
 				PoolId = "fnue5";
@@ -692,7 +706,7 @@ namespace HordeServer.Controllers
 			IDevicePool? Pool = Pools.FirstOrDefault(x => x.Id == PoolIdValue);
 			if (Pool == null)
 			{
-                Message = $"Unknown pool {PoolId} on device reservation request";
+                Message = $"Unknown pool {PoolId} on device reservation request" + Details;
 				Logger.LogError(Message);
 				await DeviceService.NotifyDeviceServiceAsync(Message, null, Request.JobId, Request.StepId);
                 return BadRequest(Message);
@@ -721,7 +735,7 @@ namespace HordeServer.Controllers
 
 				if (!PlatformMap.TryGetValue(PlatformName, out PlatformId))
 				{
-                    Message = $"Unknown platform for device type {PlatformName} on device reservation request";
+                    Message = $"Unknown platform for device type {PlatformName} on device reservation request" + Details;
 					Logger.LogError(Message);
 					await DeviceService.NotifyDeviceServiceAsync(Message, null, Request.JobId, Request.StepId);
                     return BadRequest(Message);
@@ -731,7 +745,7 @@ namespace HordeServer.Controllers
 
 				if (Platform == null)
 				{
-                    Message = $"Unknown platform {PlatformId} on device reservation request";
+                    Message = $"Unknown platform {PlatformId} on device reservation request" + Details;
 					Logger.LogError(Message);
 					await DeviceService.NotifyDeviceServiceAsync(Message, null, Request.JobId, Request.StepId);
                     return BadRequest(Message);
@@ -904,6 +918,8 @@ namespace HordeServer.Controllers
 
 			await DeviceService.UpdateDeviceAsync(Device.Id, NewProblem: true);
 
+			string Message = $"Device reported an error, {Device.Name} : {Device.PoolId.ToString().ToUpperInvariant()}";
+
             IDeviceReservation? Reservation = await DeviceService.TryGetDeviceReservation(Device.Id);
 
             string? JobId = null;
@@ -912,9 +928,19 @@ namespace HordeServer.Controllers
 			{
 				JobId = !string.IsNullOrEmpty(Reservation.JobId) ? Reservation.JobId : null;
                 StepId = !string.IsNullOrEmpty(Reservation.StepId) ? Reservation.StepId : null;
-            }
 
-            string Message = $"Device reported an error, {Device.Name} : {Device.PoolId.ToString().ToUpperInvariant()}";
+				if ((JobId == null || StepId == null))
+				{
+					if (!string.IsNullOrEmpty(Reservation.ReservationDetails))
+					{
+						Message += $" - {Reservation.ReservationDetails}";
+					}
+					else
+					{
+						Message += $" - No Reservation Information Available";
+					}                    
+                }
+            }
 
             Logger.LogError(Message);
 
