Horde: Include host in slack notifications
#rnx

diff --git a/Engine/Source/Programs/Horde/HordeServer/Controllers/DevicesController.cs b/Engine/Source/Programs/Horde/HordeServer/Controllers/DevicesController.cs
--- a/Engine/Source/Programs/Horde/HordeServer/Controllers/DevicesController.cs
+++ b/Engine/Source/Programs/Horde/HordeServer/Controllers/DevicesController.cs
@@ -689,7 +689,7 @@ namespace HordeServer.Controllers
 				 }
 				 else
 				 {
-					Details = $" - No Reservation Information Available";
+					Details = $" - No reservation information available from host {Request.Hostname}";
 				 }							 
             }
 
@@ -937,7 +937,7 @@ namespace HordeServer.Controllers
 					}
 					else
 					{
-						Message += $" - No Reservation Information Available";
+						Message += $" - No reservation information available from host {Reservation.Hostname}";
 					}                    
                 }
             }
