Horde: Rename ITaskSource.AbortTaskAsync to ITaskSource.OnCompletedLeaseAsync to better reflect its semantics. Remove code to abort steps in a succesful job lease.

diff --git a/Engine/Source/Programs/Horde/HordeServer/Models/Lease.cs b/Engine/Source/Programs/Horde/HordeServer/Models/Lease.cs
--- a/Engine/Source/Programs/Horde/HordeServer/Models/Lease.cs
+++ b/Engine/Source/Programs/Horde/HordeServer/Models/Lease.cs
@@ -13,6 +13,8 @@ using System.Threading.Tasks;
 
 using PoolId = HordeServer.Utilities.StringId<HordeServer.Models.IPool>;
 using StreamId = HordeServer.Utilities.StringId<HordeServer.Models.IStream>;
+using Google.Protobuf.WellKnownTypes;
+using Google.Protobuf;
 
 namespace HordeServer.Models
 {
@@ -76,4 +78,31 @@ namespace HordeServer.Models
 		/// </summary>
 		public LeaseOutcome Outcome { get; }
 	}
+
+	/// <summary>
+	/// Extension methods for leases
+	/// </summary>
+	public static class LeaseExtensions
+	{
+		/// <summary>
+		/// Gets the task from a lease, encoded as an Any protobuf object
+		/// </summary>
+		/// <param name="Lease">The lease to query</param>
+		/// <returns>The task definition encoded as a protobuf Any object</returns>
+		public static Any GetTask(this ILease Lease)
+		{
+			return Any.Parser.ParseFrom(Lease.Payload.ToArray());
+		}
+
+		/// <summary>
+		/// Gets a typed task object from a lease
+		/// </summary>
+		/// <typeparam name="T">Type of the protobuf message to return</typeparam>
+		/// <param name="Lease">The lease to query</param>
+		/// <returns>The task definition</returns>
+		public static T GetTask<T>(this ILease Lease) where T : IMessage<T>, new()
+		{
+			return GetTask(Lease).Unpack<T>();
+		}
+	}
 }
diff --git a/Engine/Source/Programs/Horde/HordeServer/Services/AgentService.cs b/Engine/Source/Programs/Horde/HordeServer/Services/AgentService.cs
--- a/Engine/Source/Programs/Horde/HordeServer/Services/AgentService.cs
+++ b/Engine/Source/Programs/Horde/HordeServer/Services/AgentService.cs
@@ -843,7 +843,7 @@ namespace HordeServer.Services
 				{
 					if (Any.Is(TaskSource.Descriptor))
 					{
-						await TaskSource.AbortTaskAsync(Agent, Lease.Id, Any);
+						await TaskSource.OnLeaseCompletedAsync(Agent, Lease.Id, Any, Outcome);
 						break;
 					}
 				}
diff --git a/Engine/Source/Programs/Horde/HordeServer/Tasks/ITaskSource.cs b/Engine/Source/Programs/Horde/HordeServer/Tasks/ITaskSource.cs
--- a/Engine/Source/Programs/Horde/HordeServer/Tasks/ITaskSource.cs
+++ b/Engine/Source/Programs/Horde/HordeServer/Tasks/ITaskSource.cs
@@ -3,6 +3,7 @@
 using Google.Protobuf;
 using Google.Protobuf.Reflection;
 using Google.Protobuf.WellKnownTypes;
+using HordeCommon;
 using HordeServer.Models;
 using HordeServer.Services;
 using MongoDB.Bson;
@@ -103,13 +104,14 @@ namespace HordeServer.Tasks
 		Task<ITaskListener?> SubscribeAsync(IAgent Agent);
 
 		/// <summary>
-		/// Force a lease as complete
+		/// Notification that a task has completed
 		/// </summary>
 		/// <param name="Agent">The agent that was allocated to the lease</param>
 		/// <param name="LeaseId">The lease id</param>
 		/// <param name="Payload">The lease payload</param>
+		/// <param name="Outcome">Outcome of the lease</param>
 		/// <returns>Async task</returns>
-		Task AbortTaskAsync(IAgent Agent, ObjectId LeaseId, Any Payload);
+		Task OnLeaseCompletedAsync(IAgent Agent, ObjectId LeaseId, Any Payload, LeaseOutcome Outcome);
 	}
 
 	/// <summary>
@@ -125,19 +127,13 @@ namespace HordeServer.Tasks
 		public abstract Task<ITaskListener?> SubscribeAsync(IAgent Agent);
 
 		/// <inheritdoc/>
-		public Task AbortTaskAsync(IAgent Agent, ObjectId LeaseId, Any Payload)
+		public Task OnLeaseCompletedAsync(IAgent Agent, ObjectId LeaseId, Any Task, LeaseOutcome Outcome)
 		{
-			return AbortTaskAsync(Agent, LeaseId, Payload.Unpack<T>());
+			return OnLeaseCompletedAsync(Agent, LeaseId, Task.Unpack<T>(), Outcome);
 		}
 
-		/// <summary>
-		/// Cancels the given task
-		/// </summary>
-		/// <param name="Agent">Agent executing the task</param>
-		/// <param name="LeaseId">The lease id</param>
-		/// <param name="Payload">The task type</param>
-		/// <returns>Async task</returns>
-		protected virtual Task AbortTaskAsync(IAgent Agent, ObjectId LeaseId, T Payload)
+		/// <inheritdoc cref="OnLeaseCompletedAsync(IAgent, ObjectId, Any, LeaseOutcome)"/>
+		public virtual Task OnLeaseCompletedAsync(IAgent Agent, ObjectId LeaseId, T TaskInput, LeaseOutcome Outcome)
 		{
 			return Task.CompletedTask;
 		}
diff --git a/Engine/Source/Programs/Horde/HordeServer/Tasks/Impl/ActionTaskSource.cs b/Engine/Source/Programs/Horde/HordeServer/Tasks/Impl/ActionTaskSource.cs
--- a/Engine/Source/Programs/Horde/HordeServer/Tasks/Impl/ActionTaskSource.cs
+++ b/Engine/Source/Programs/Horde/HordeServer/Tasks/Impl/ActionTaskSource.cs
@@ -90,7 +90,7 @@ namespace HordeServer.Tasks.Impl
 	/// <summary>
 	/// Dispatches remote actions. Does not implement any cross-pod communication to satisfy leases; only agents connected to this server instance will be stored.
 	/// </summary>
-	public class ActionTaskSource : ITaskSource
+	public class ActionTaskSource : TaskSourceBase<ActionTask>
 	{
 		/// <summary>
 		/// Number of buckets to create for storing operations. Operations are garbage collected from random buckets.
@@ -345,9 +345,6 @@ namespace HordeServer.Tasks.Impl
 		ILogger Logger;
 		RemoteExecSettings RemoteExecSettings;
 
-		/// <inheritdoc/>
-		public MessageDescriptor Descriptor => ActionTask.Descriptor;
-
 		/// <summary>
 		/// Constructor
 		/// </summary>
@@ -650,7 +647,7 @@ namespace HordeServer.Tasks.Impl
 		}
 
 		/// <inheritdoc/>
-		public Task<ITaskListener?> SubscribeAsync(IAgent Agent)
+		public override Task<ITaskListener?> SubscribeAsync(IAgent Agent)
 		{
 			if (Agent.Leases.Count >= RemoteExecSettings.MaxConcurrentLeasesPerAgent)
 			{
@@ -660,12 +657,6 @@ namespace HordeServer.Tasks.Impl
 			
 			return Task.FromResult<ITaskListener?>(new Subscription(this, Agent));
 		}
-
-		/// <inheritdoc/>
-		public Task AbortTaskAsync(IAgent Agent, ObjectId LeaseId, Any Payload)
-		{
-			return Task.CompletedTask;
-		}
 		
 		internal static bool CanBeScheduledOnAgent(IAgent Agent, List<string>? PoolFilter)
 		{
diff --git a/Engine/Source/Programs/Horde/HordeServer/Tasks/Impl/ClusteredActionTaskSource.cs b/Engine/Source/Programs/Horde/HordeServer/Tasks/Impl/ClusteredActionTaskSource.cs
--- a/Engine/Source/Programs/Horde/HordeServer/Tasks/Impl/ClusteredActionTaskSource.cs
+++ b/Engine/Source/Programs/Horde/HordeServer/Tasks/Impl/ClusteredActionTaskSource.cs
@@ -26,7 +26,7 @@ namespace HordeServer.Tasks.Impl
 	/// <summary>
 	/// Dispatches remote actions across all agents that are connected to a Horde server in the cluster
 	/// </summary>
-	public class ClusteredActionTaskSource : ITaskSource
+	public class ClusteredActionTaskSource : TaskSourceBase<ActionTask>
 	{
 		/// <summary>
 		/// Tracks an agent waiting for work
@@ -146,9 +146,6 @@ namespace HordeServer.Tasks.Impl
 		private readonly string KeyCompletedOps = "re-completed-ops";
 		private readonly string ChannelOps = "re-ops";
 
-		/// <inheritdoc/>
-		public MessageDescriptor Descriptor => ActionTask.Descriptor;
-
 		/// <summary>
 		/// Constructor
 		/// </summary>
@@ -489,16 +486,9 @@ namespace HordeServer.Tasks.Impl
 		}
 
 		/// <inheritdoc/>
-		public Task<ITaskListener?> SubscribeAsync(IAgent Agent)
+		public override Task<ITaskListener?> SubscribeAsync(IAgent Agent)
 		{
 			return Task.FromResult<ITaskListener?>(new AgentSubscription(this, Agent));
 		}
-
-		/// <inheritdoc/>
-		public Task AbortTaskAsync(IAgent Agent, ObjectId LeaseId, Any Payload)
-		{
-			// Remotely executed tasks cannot be aborted (for the time being)
-			return Task.CompletedTask;
-		}
 	}
 }
diff --git a/Engine/Source/Programs/Horde/HordeServer/Tasks/Impl/ConformTaskSource.cs b/Engine/Source/Programs/Horde/HordeServer/Tasks/Impl/ConformTaskSource.cs
--- a/Engine/Source/Programs/Horde/HordeServer/Tasks/Impl/ConformTaskSource.cs
+++ b/Engine/Source/Programs/Horde/HordeServer/Tasks/Impl/ConformTaskSource.cs
@@ -211,7 +211,7 @@ namespace HordeServer.Tasks.Impl
 		}
 
 		/// <inheritdoc/>
-		protected override Task AbortTaskAsync(IAgent Agent, ObjectId LeaseId, ConformTask Task)
+		public override Task OnLeaseCompletedAsync(IAgent Agent, ObjectId LeaseId, ConformTask Task, LeaseOutcome Outcome)
 		{
 			return ReleaseConformLeaseAsync(LeaseId);
 		}
diff --git a/Engine/Source/Programs/Horde/HordeServer/Tasks/Impl/JobTaskSource.cs b/Engine/Source/Programs/Horde/HordeServer/Tasks/Impl/JobTaskSource.cs
--- a/Engine/Source/Programs/Horde/HordeServer/Tasks/Impl/JobTaskSource.cs
+++ b/Engine/Source/Programs/Horde/HordeServer/Tasks/Impl/JobTaskSource.cs
@@ -1009,48 +1009,51 @@ namespace HordeServer.Tasks.Impl
 		}
 
 		/// <inheritdoc/>
-		public async Task AbortTaskAsync(IAgent Agent, ObjectId LeaseId, Any Any)
+		public async Task OnLeaseCompletedAsync(IAgent Agent, ObjectId LeaseId, Any Any, LeaseOutcome Outcome)
 		{
-			AgentId AgentId = Agent.Id;
-			ExecuteJobTask Task = Any.Unpack<ExecuteJobTask>();
-			ObjectId JobId = Task.JobId.ToObjectId();
-			SubResourceId BatchId = Task.BatchId.ToSubResourceId();
-
-			// Update the batch
-			for (; ; )
+			if (Outcome != LeaseOutcome.Success)
 			{
-				IJob? Job = await Jobs.GetAsync(JobId);
-				if (Job == null)
-				{
-					break;
-				}
+				AgentId AgentId = Agent.Id;
+				ExecuteJobTask Task = Any.Unpack<ExecuteJobTask>();
+				ObjectId JobId = Task.JobId.ToObjectId();
+				SubResourceId BatchId = Task.BatchId.ToSubResourceId();
 
-				int BatchIdx = Job.Batches.FindIndex(x => x.Id == BatchId);
-				if (BatchIdx == -1)
-				{
-					break;
-				}
-
-				IJobStepBatch Batch = Job.Batches[BatchIdx];
-				if (Batch.AgentId != AgentId)
+				// Update the batch
+				for (; ; )
 				{
-					break;
-				}
+					IJob? Job = await Jobs.GetAsync(JobId);
+					if (Job == null)
+					{
+						break;
+					}
 
-				int RunningStepIdx = Batch.Steps.FindIndex(x => x.State == JobStepState.Running);
+					int BatchIdx = Job.Batches.FindIndex(x => x.Id == BatchId);
+					if (BatchIdx == -1)
+					{
+						break;
+					}
 
-				IGraph Graph = await Graphs.GetAsync(Job.GraphHash);
-				if (await Jobs.TryFailBatchAsync(Job, BatchIdx, Graph))
-				{
-					if (Batch.Error != JobStepBatchError.None)
+					IJobStepBatch Batch = Job.Batches[BatchIdx];
+					if (Batch.AgentId != AgentId)
 					{
-						Logger.LogInformation("Failed job {JobId}, batch {BatchId} with error {Error}", Job.Id, Batch.Id, Batch.Error);
+						break;
 					}
-					if (RunningStepIdx != -1)
+
+					int RunningStepIdx = Batch.Steps.FindIndex(x => x.State == JobStepState.Running);
+
+					IGraph Graph = await Graphs.GetAsync(Job.GraphHash);
+					if (await Jobs.TryFailBatchAsync(Job, BatchIdx, Graph))
 					{
-						await JobStepRefs.UpdateAsync(Job, Batch, Batch.Steps[RunningStepIdx], Graph);
+						if (Batch.Error != JobStepBatchError.None)
+						{
+							Logger.LogInformation("Failed job {JobId}, batch {BatchId} with error {Error}", Job.Id, Batch.Id, Batch.Error);
+						}
+						if (RunningStepIdx != -1)
+						{
+							await JobStepRefs.UpdateAsync(Job, Batch, Batch.Steps[RunningStepIdx], Graph);
+						}
+						break;
 					}
-					break;
 				}
 			}
 		}
diff --git a/Engine/Source/Programs/Horde/HordeServer/Tasks/Impl/UpgradeTaskSource.cs b/Engine/Source/Programs/Horde/HordeServer/Tasks/Impl/UpgradeTaskSource.cs
--- a/Engine/Source/Programs/Horde/HordeServer/Tasks/Impl/UpgradeTaskSource.cs
+++ b/Engine/Source/Programs/Horde/HordeServer/Tasks/Impl/UpgradeTaskSource.cs
@@ -18,14 +18,12 @@ using AgentSoftwareChannelName = HordeServer.Utilities.StringId<HordeServer.Serv
 
 namespace HordeServer.Tasks.Impl
 {
-	class UpgradeTaskSource : ITaskSource
+	class UpgradeTaskSource : TaskSourceBase<UpgradeTask>
 	{
 		AgentSoftwareService AgentSoftwareService;
 		ILogFileService LogService;
 		IClock Clock;
 
-		public MessageDescriptor Descriptor => UpgradeTask.Descriptor;
-
 		public UpgradeTaskSource(AgentSoftwareService AgentSoftwareService, ILogFileService LogService, IClock Clock)
 		{
 			this.AgentSoftwareService = AgentSoftwareService;
@@ -33,7 +31,7 @@ namespace HordeServer.Tasks.Impl
 			this.Clock = Clock;
 		}
 
-		public async Task<ITaskListener?> SubscribeAsync(IAgent Agent)
+		public async override Task<ITaskListener?> SubscribeAsync(IAgent Agent)
 		{
 			string? RequiredVersion = await GetRequiredSoftwareVersion(Agent);
 			if (RequiredVersion != null && Agent.Version != RequiredVersion)
@@ -56,11 +54,6 @@ namespace HordeServer.Tasks.Impl
 			return null;
 		}
 
-		public Task AbortTaskAsync(IAgent Agent, ObjectId LeaseId, Any Payload)
-		{
-			return Task.CompletedTask;
-		}
-
 		/// <summary>
 		/// Determines the client software version that should be installed on an agent
 		/// </summary>
