Horde: Add an output return path for leases, allowing data to be passed back to task sources for subsequent processing.

diff --git a/Engine/Source/Programs/Horde/HordeAgent/Services/WorkerService.cs b/Engine/Source/Programs/Horde/HordeAgent/Services/WorkerService.cs
--- a/Engine/Source/Programs/Horde/HordeAgent/Services/WorkerService.cs
+++ b/Engine/Source/Programs/Horde/HordeAgent/Services/WorkerService.cs
@@ -52,6 +52,7 @@ using Amazon.EC2;
 using Scope = Datadog.Trace.Scope;
 using Amazon.EC2.Model;
 using Status = Grpc.Core.Status;
+using Google.Protobuf;
 
 [assembly: InternalsVisibleTo("HordeAgentTests")]
 
@@ -93,6 +94,56 @@ namespace HordeAgent.Services
 			}
 		}
 
+		/// <summary>
+		/// Result from executing a lease
+		/// </summary>
+		internal class LeaseResult
+		{
+			/// <summary>
+			/// Outcome of the lease (whether it completed/failed due to an internal error, etc...)
+			/// </summary>
+			public LeaseOutcome Outcome { get; }
+
+			/// <summary>
+			/// Output from executing the task
+			/// </summary>
+			public byte[]? Output { get; }
+
+			/// <summary>
+			/// Static instance of a cancelled result
+			/// </summary>
+			public static LeaseResult Cancelled { get; } = new LeaseResult(LeaseOutcome.Cancelled);
+
+			/// <summary>
+			/// Static instance of a failed result
+			/// </summary>
+			public static LeaseResult Failed { get; } = new LeaseResult(LeaseOutcome.Failed);
+
+			/// <summary>
+			/// Static instance of a succesful result without a payload
+			/// </summary>
+			public static LeaseResult Success { get; } = new LeaseResult(LeaseOutcome.Success);
+
+			/// <summary>
+			/// Constructor for 
+			/// </summary>
+			/// <param name="Outcome"></param>
+			private LeaseResult(LeaseOutcome Outcome)
+			{
+				this.Outcome = Outcome;
+			}
+
+			/// <summary>
+			/// Constructor for successful results
+			/// </summary>
+			/// <param name="Output"></param>
+			public LeaseResult(byte[]? Output)
+			{
+				this.Outcome = LeaseOutcome.Success;
+				this.Output = Output;
+			}
+		}
+
 		/// <summary>
 		/// List of processes that should be terminated before running a job
 		/// </summary>
@@ -596,10 +647,10 @@ namespace HordeAgent.Services
 			Logger.LogInformation("Handling lease {LeaseId}", LeaseInfo.Lease.Id);
 
 			// Get the lease outcome
-			LeaseOutcome Outcome = LeaseOutcome.Failed;
+			LeaseResult Result = LeaseResult.Failed;
 			try
 			{
-				Outcome = await HandleLeasePayloadAsync(RpcConnection, AgentId, LeaseInfo);
+				Result = await HandleLeasePayloadAsync(RpcConnection, AgentId, LeaseInfo);
 			}
 			catch (Exception Ex)
 			{
@@ -613,11 +664,13 @@ namespace HordeAgent.Services
 				{
 					LeaseInfo.Lease.State = LeaseState.Cancelled;
 					LeaseInfo.Lease.Outcome = LeaseOutcome.Failed;
+					LeaseInfo.Lease.Result = null;
 				}
 				else
 				{
-					LeaseInfo.Lease.State = (Outcome == LeaseOutcome.Cancelled) ? LeaseState.Cancelled : LeaseState.Completed;
-					LeaseInfo.Lease.Outcome = Outcome;
+					LeaseInfo.Lease.State = (Result.Outcome == LeaseOutcome.Cancelled) ? LeaseState.Cancelled : LeaseState.Completed;
+					LeaseInfo.Lease.Outcome = Result.Outcome;
+					LeaseInfo.Lease.Result = ByteString.CopyFrom(Result.Output);
 				}
 				Logger.LogInformation("Transitioning lease {LeaseId} to {State}, outcome={Outcome}", LeaseInfo.Lease.Id, LeaseInfo.Lease.State, LeaseInfo.Lease.Outcome);
 			}
@@ -631,7 +684,7 @@ namespace HordeAgent.Services
 		/// <param name="AgentId">The current agent id</param>
 		/// <param name="LeaseInfo">Information about the lease</param>
 		/// <returns>Outcome from the lease</returns>
-		internal async Task<LeaseOutcome> HandleLeasePayloadAsync(IRpcConnection RpcConnection, string AgentId, LeaseInfo LeaseInfo)
+		internal async Task<LeaseResult> HandleLeasePayloadAsync(IRpcConnection RpcConnection, string AgentId, LeaseInfo LeaseInfo)
 		{
 			Any Payload = LeaseInfo.Lease.Payload;
 
@@ -639,7 +692,7 @@ namespace HordeAgent.Services
 			if (LeaseInfo.Lease.Payload.TryUnpack(out ActionTask))
 			{
 				Tracer.Instance.ActiveScope?.Span.SetTag("task", "Action");
-				Func<ILogger, Task<LeaseOutcome>> Handler = NewLogger => ActionAsync(RpcConnection, LeaseInfo.Lease.Id, ActionTask, NewLogger, LeaseInfo.CancellationTokenSource.Token);
+				Func<ILogger, Task<LeaseResult>> Handler = NewLogger => ActionAsync(RpcConnection, LeaseInfo.Lease.Id, ActionTask, NewLogger, LeaseInfo.CancellationTokenSource.Token);
 				return await HandleLeasePayloadWithLogAsync(RpcConnection, ActionTask.LogId, AgentId, null, null, Handler, LeaseInfo.CancellationTokenSource.Token);
 			}
 
@@ -647,7 +700,7 @@ namespace HordeAgent.Services
 			if (LeaseInfo.Lease.Payload.TryUnpack(out ConformTask))
 			{
 				Tracer.Instance.ActiveScope.Span.SetTag("task", "Conform");
-				Func<ILogger, Task<LeaseOutcome>> Handler = NewLogger => ConformAsync(RpcConnection, AgentId, LeaseInfo.Lease.Id, ConformTask, NewLogger, LeaseInfo.CancellationTokenSource.Token);
+				Func<ILogger, Task<LeaseResult>> Handler = NewLogger => ConformAsync(RpcConnection, AgentId, LeaseInfo.Lease.Id, ConformTask, NewLogger, LeaseInfo.CancellationTokenSource.Token);
 				return await HandleLeasePayloadWithLogAsync(RpcConnection, ConformTask.LogId, AgentId, null, null, Handler, LeaseInfo.CancellationTokenSource.Token);
 			}
 
@@ -655,7 +708,7 @@ namespace HordeAgent.Services
 			if (LeaseInfo.Lease.Payload.TryUnpack(out JobTask))
 			{
 				Tracer.Instance.ActiveScope.Span.SetTag("task", "Job");
-				Func<ILogger, Task<LeaseOutcome>> Handler = NewLogger => ExecuteJobAsync(RpcConnection, AgentId, LeaseInfo.Lease.Id, JobTask, NewLogger, LeaseInfo.CancellationTokenSource.Token);
+				Func<ILogger, Task<LeaseResult>> Handler = NewLogger => ExecuteJobAsync(RpcConnection, AgentId, LeaseInfo.Lease.Id, JobTask, NewLogger, LeaseInfo.CancellationTokenSource.Token);
 				return await HandleLeasePayloadWithLogAsync(RpcConnection, JobTask.LogId, AgentId, JobTask.JobId, JobTask.BatchId, Handler, LeaseInfo.CancellationTokenSource.Token);
 			}
 
@@ -663,7 +716,7 @@ namespace HordeAgent.Services
 			if (LeaseInfo.Lease.Payload.TryUnpack(out UpgradeTask))
 			{
 				Tracer.Instance.ActiveScope.Span.SetTag("task", "Upgrade");
-				Func<ILogger, Task<LeaseOutcome>> Handler = NewLogger => UpgradeAsync(RpcConnection, AgentId, UpgradeTask, NewLogger, LeaseInfo.CancellationTokenSource.Token);
+				Func<ILogger, Task<LeaseResult>> Handler = NewLogger => UpgradeAsync(RpcConnection, AgentId, UpgradeTask, NewLogger, LeaseInfo.CancellationTokenSource.Token);
 				return await HandleLeasePayloadWithLogAsync(RpcConnection, UpgradeTask.LogId, AgentId, null, null, Handler, LeaseInfo.CancellationTokenSource.Token);
 			}
 
@@ -671,7 +724,7 @@ namespace HordeAgent.Services
 			if (LeaseInfo.Lease.Payload.TryUnpack(out ShutdownTask))
 			{
 				Tracer.Instance.ActiveScope.Span.SetTag("task", "Shutdown");
-				Func<ILogger, Task<LeaseOutcome>> Handler = NewLogger => ShutdownAsync(RpcConnection, AgentId, ShutdownTask, NewLogger, LeaseInfo.CancellationTokenSource.Token);
+				Func<ILogger, Task<LeaseResult>> Handler = NewLogger => ShutdownAsync(RpcConnection, AgentId, ShutdownTask, NewLogger, LeaseInfo.CancellationTokenSource.Token);
 				return await HandleLeasePayloadWithLogAsync(RpcConnection, ShutdownTask.LogId, AgentId, null, null, Handler, LeaseInfo.CancellationTokenSource.Token);
 			}
 
@@ -679,12 +732,12 @@ namespace HordeAgent.Services
 			if (LeaseInfo.Lease.Payload.TryUnpack(out RestartTask))
 			{
 				Tracer.Instance.ActiveScope.Span.SetTag("task", "Restart");
-				Func<ILogger, Task<LeaseOutcome>> Handler = NewLogger => RestartAsync(RpcConnection, AgentId, RestartTask, NewLogger, LeaseInfo.CancellationTokenSource.Token);
+				Func<ILogger, Task<LeaseResult>> Handler = NewLogger => RestartAsync(RpcConnection, AgentId, RestartTask, NewLogger, LeaseInfo.CancellationTokenSource.Token);
 				return await HandleLeasePayloadWithLogAsync(RpcConnection, RestartTask.LogId, AgentId, null, null, Handler, LeaseInfo.CancellationTokenSource.Token);
 			}
 
 			Logger.LogError("Invalid lease payload type ({PayloadType})", Payload.TypeUrl);
-			return LeaseOutcome.Failed;
+			return LeaseResult.Failed;
 		}
 
 		/// <summary>
@@ -730,27 +783,27 @@ namespace HordeAgent.Services
 		/// <param name="ExecuteLease">Action to perform to execute the lease</param>
 		/// <param name="CancellationToken">The cancellation token for this lease</param>
 		/// <returns>Outcome from the lease</returns>
-		async Task<LeaseOutcome> HandleLeasePayloadWithLogAsync(IRpcConnection RpcConnection, string LogId, string AgentId, string? JobId, string? JobBatchId, Func<ILogger, Task<LeaseOutcome>> ExecuteLease, CancellationToken CancellationToken)
+		async Task<LeaseResult> HandleLeasePayloadWithLogAsync(IRpcConnection RpcConnection, string LogId, string AgentId, string? JobId, string? JobBatchId, Func<ILogger, Task<LeaseResult>> ExecuteLease, CancellationToken CancellationToken)
 		{
 			await using (JsonRpcLogger? ConformLogger = String.IsNullOrEmpty(LogId) ? null : new JsonRpcLogger(RpcConnection, LogId, JobId, JobBatchId, null, null, Logger))
 			{
 				ILogger LeaseLogger = (ConformLogger == null) ? (ILogger)Logger : new ForwardingLogger(Logger, new DefaultLoggerIndentHandler(ConformLogger));
 				try
 				{
-					LeaseOutcome Outcome = await ExecuteLease(LeaseLogger);
-					return Outcome;
+					LeaseResult Result = await ExecuteLease(LeaseLogger);
+					return Result;
 				}
 				catch (Exception Ex)
 				{
 					if (CancellationToken.IsCancellationRequested && IsCancellationException(Ex))
 					{
 						LeaseLogger.LogInformation("Lease was cancelled");
-						return LeaseOutcome.Cancelled;
+						return LeaseResult.Cancelled;
 					}
 					else
 					{
 						LeaseLogger.LogError(Ex, "Caught unhandled exception while attempting to execute lease:\n{Exception}", Ex);
-						return LeaseOutcome.Failed;
+						return LeaseResult.Failed;
 					}
 				}
 			}
@@ -765,7 +818,7 @@ namespace HordeAgent.Services
 		/// <param name="Logger">Logger for the task</param>
 		/// <param name="CancellationToken">Token used to cancel the operation</param>
 		/// <returns></returns>
-		internal async Task<LeaseOutcome> ActionAsync(IRpcConnection RpcConnection, string LeaseId, ActionTask ActionTask, ILogger Logger, CancellationToken CancellationToken)
+		internal async Task<LeaseResult> ActionAsync(IRpcConnection RpcConnection, string LeaseId, ActionTask ActionTask, ILogger Logger, CancellationToken CancellationToken)
 		{
 			DateTimeOffset ActionTaskStartTime = DateTimeOffset.UtcNow;
 			using (IRpcClientRef Client = await RpcConnection.GetClientRef(new RpcContext(), CancellationToken))
@@ -827,7 +880,7 @@ namespace HordeAgent.Services
 					}
 				}
 
-				return LeaseOutcome.Success;
+				return LeaseResult.Success;
 			}
 		}
 
@@ -841,7 +894,7 @@ namespace HordeAgent.Services
 		/// <param name="ConformLogger">Logger for the task</param>
 		/// <param name="CancellationToken">Token used to cancel the operation</param>
 		/// <returns>Async task</returns>
-		async Task<LeaseOutcome> ConformAsync(IRpcConnection RpcConnection, string AgentId, string LeaseId, ConformTask ConformTask, ILogger ConformLogger, CancellationToken CancellationToken)
+		async Task<LeaseResult> ConformAsync(IRpcConnection RpcConnection, string AgentId, string LeaseId, ConformTask ConformTask, ILogger ConformLogger, CancellationToken CancellationToken)
 		{
 			ConformLogger.LogInformation("Conforming, lease {LeaseId}", LeaseId);
 			TerminateProcesses(ConformLogger, CancellationToken);
@@ -875,7 +928,7 @@ namespace HordeAgent.Services
 				PendingWorkspaces = Response.PendingWorkspaces;
 			}
 
-			return LeaseOutcome.Success;
+			return LeaseResult.Success;
 		}
 
 		/// <summary>
@@ -888,7 +941,7 @@ namespace HordeAgent.Services
 		/// <param name="Logger">The logger to use for this lease</param>
 		/// <param name="CancellationToken">Token used to cancel the operation</param>
 		/// <returns>Async task</returns>
-		internal async Task<LeaseOutcome> ExecuteJobAsync(IRpcConnection RpcClient, string AgentId, string LeaseId, ExecuteJobTask ExecuteTask, ILogger Logger, CancellationToken CancellationToken)
+		internal async Task<LeaseResult> ExecuteJobAsync(IRpcConnection RpcClient, string AgentId, string LeaseId, ExecuteJobTask ExecuteTask, ILogger Logger, CancellationToken CancellationToken)
 		{
 			Logger.LogInformation("Executing job \"{JobName}\", jobId {JobId}, batchId {BatchId}, leaseId {LeaseId}", ExecuteTask.JobName, ExecuteTask.JobId, ExecuteTask.BatchId, LeaseId);
 			Tracer.Instance.ActiveScope?.Span.SetTag("jobId", ExecuteTask.JobId.ToString());
@@ -919,7 +972,7 @@ namespace HordeAgent.Services
 			await RpcClient.InvokeAsync(x => x.FinishBatchAsync(new FinishBatchRequest(ExecuteTask.JobId, ExecuteTask.BatchId, LeaseId), null, null, CancellationToken), new RpcContext(), CancellationToken);
 			Logger.LogInformation("Done.");
 
-			return LeaseOutcome.Success;
+			return LeaseResult.Success;
 		}
 
 		/// <summary>
@@ -1147,7 +1200,7 @@ namespace HordeAgent.Services
 		/// <param name="Logger">Logging device</param>
 		/// <param name="CancellationToken">Token used to cancel the operation</param>
 		/// <returns>Outcome of this operation</returns>
-		async Task<LeaseOutcome> UpgradeAsync(IRpcConnection RpcClient, string AgentId, UpgradeTask UpgradeTask, ILogger Logger, CancellationToken CancellationToken)
+		async Task<LeaseResult> UpgradeAsync(IRpcConnection RpcClient, string AgentId, UpgradeTask UpgradeTask, ILogger Logger, CancellationToken CancellationToken)
 		{
 			string RequiredVersion = UpgradeTask.SoftwareId;
 
@@ -1202,7 +1255,7 @@ namespace HordeAgent.Services
 					if (!FileReference.Exists(NewAssemblyFileName))
 					{
 						Logger.LogError("Unable to find {AgentExe} in extracted archive", NewAssemblyFileName);
-						return LeaseOutcome.Failed;
+						return LeaseResult.Failed;
 					}
 
 					Process.StartInfo.FileName = "dotnet";
@@ -1242,7 +1295,7 @@ namespace HordeAgent.Services
 				}
 			}
 
-			return LeaseOutcome.Success;
+			return LeaseResult.Success;
 		}
 
 		/// <summary>
@@ -1285,12 +1338,12 @@ namespace HordeAgent.Services
 		/// <param name="Logger">Logging device</param>
 		/// <param name="CancellationToken">Token used to cancel the operation</param>
 		/// <returns>Outcome of this operation</returns>
-		Task<LeaseOutcome> RestartAsync(IRpcConnection RpcClient, string AgentId, RestartTask RestartTask, ILogger Logger, CancellationToken CancellationToken)
+		Task<LeaseResult> RestartAsync(IRpcConnection RpcClient, string AgentId, RestartTask RestartTask, ILogger Logger, CancellationToken CancellationToken)
 		{
 			Logger.LogInformation("Setting restart flag");
 			RequestShutdown = true;
 			RestartAfterShutdown = true;
-			return Task.FromResult(LeaseOutcome.Success);
+			return Task.FromResult(LeaseResult.Success);
 		}
 
 		/// <summary>
@@ -1302,12 +1355,12 @@ namespace HordeAgent.Services
 		/// <param name="Logger">Logging device</param>
 		/// <param name="CancellationToken">Token used to cancel the operation</param>
 		/// <returns>Outcome of this operation</returns>
-		Task<LeaseOutcome> ShutdownAsync(IRpcConnection RpcClient, string AgentId, ShutdownTask ShutdownTask, ILogger Logger, CancellationToken CancellationToken)
+		Task<LeaseResult> ShutdownAsync(IRpcConnection RpcClient, string AgentId, ShutdownTask ShutdownTask, ILogger Logger, CancellationToken CancellationToken)
 		{
 			Logger.LogInformation("Setting shutdown flag");
 			RequestShutdown = true;
 			RestartAfterShutdown = false;
-			return Task.FromResult(LeaseOutcome.Success);
+			return Task.FromResult(LeaseResult.Success);
 		}
 
 		/// <summary>
diff --git a/Engine/Source/Programs/Horde/HordeServer/Collections/ILeaseCollection.cs b/Engine/Source/Programs/Horde/HordeServer/Collections/ILeaseCollection.cs
--- a/Engine/Source/Programs/Horde/HordeServer/Collections/ILeaseCollection.cs
+++ b/Engine/Source/Programs/Horde/HordeServer/Collections/ILeaseCollection.cs
@@ -77,7 +77,8 @@ namespace HordeServer.Collections
 		/// <param name="LeaseId">The lease to update</param>
 		/// <param name="FinishTime">Time at which the lease finished</param>
 		/// <param name="Outcome">Outcome of the lease</param>
+		/// <param name="Output">Output data from the task</param>
 		/// <returns>True if the lease was updated, false otherwise</returns>
-		Task<bool> TrySetOutcomeAsync(ObjectId LeaseId, DateTime FinishTime, LeaseOutcome Outcome);
+		Task<bool> TrySetOutcomeAsync(ObjectId LeaseId, DateTime FinishTime, LeaseOutcome Outcome, byte[]? Output);
 	}
 }
diff --git a/Engine/Source/Programs/Horde/HordeServer/Collections/Impl/LeaseCollection.cs b/Engine/Source/Programs/Horde/HordeServer/Collections/Impl/LeaseCollection.cs
--- a/Engine/Source/Programs/Horde/HordeServer/Collections/Impl/LeaseCollection.cs
+++ b/Engine/Source/Programs/Horde/HordeServer/Collections/Impl/LeaseCollection.cs
@@ -39,7 +39,11 @@ namespace HordeServer.Collections.Impl
 			public byte[] Payload { get; set; }
 			public LeaseOutcome Outcome { get; set; }
 
+			[BsonIgnoreIfNull]
+			public byte[]? Output { get; set; }
+
 			ReadOnlyMemory<byte> ILease.Payload => Payload;
+			ReadOnlyMemory<byte> ILease.Output => Output;
 
 			[BsonConstructor]
 			private LeaseDocument()
@@ -139,7 +143,7 @@ namespace HordeServer.Collections.Impl
 		}
 
 		/// <inheritdoc/>
-		public async Task<bool> TrySetOutcomeAsync(ObjectId LeaseId, DateTime FinishTime, LeaseOutcome Outcome)
+		public async Task<bool> TrySetOutcomeAsync(ObjectId LeaseId, DateTime FinishTime, LeaseOutcome Outcome, byte[]? Output)
 		{
 			FilterDefinitionBuilder<LeaseDocument> FilterBuilder = Builders<LeaseDocument>.Filter;
 			FilterDefinition<LeaseDocument> Filter = FilterBuilder.Eq(x => x.Id, LeaseId) & FilterBuilder.Eq(x => x.FinishTime, null);
@@ -147,6 +151,11 @@ namespace HordeServer.Collections.Impl
 			UpdateDefinitionBuilder<LeaseDocument> UpdateBuilder = Builders<LeaseDocument>.Update;
 			UpdateDefinition<LeaseDocument> Update = UpdateBuilder.Set(x => x.FinishTime, FinishTime).Set(x => x.Outcome, Outcome);
 
+			if(Output != null && Output.Length > 0)
+			{
+				Update = Update.Set(x => x.Output, Output);
+			}
+
 			UpdateResult Result = await Leases.UpdateOneAsync(Filter, Update);
 			return Result.ModifiedCount > 0;
 		}
diff --git a/Engine/Source/Programs/Horde/HordeServer/Services/AgentService.cs b/Engine/Source/Programs/Horde/HordeServer/Services/AgentService.cs
--- a/Engine/Source/Programs/Horde/HordeServer/Services/AgentService.cs
+++ b/Engine/Source/Programs/Horde/HordeServer/Services/AgentService.cs
@@ -287,7 +287,7 @@ namespace HordeServer.Services
 					foreach (AgentLease Lease in Agent.Leases)
 					{
 						Logger.LogDebug("Removing outstanding lease {LeaseId} during session create...", Lease.Id);
-						await RemoveLeaseAsync(Agent, Lease, UtcNow, LeaseOutcome.Failed);
+						await RemoveLeaseAsync(Agent, Lease, UtcNow, LeaseOutcome.Failed, null);
 					}
 
 					// Create a new session document
@@ -557,7 +557,7 @@ namespace HordeServer.Services
 							HordeCommon.Rpc.Messages.Lease? NewLease;
 							if (!LeaseIdToNewState.TryGetValue(Lease.Id, out NewLease) || NewLease.State == LeaseState.Cancelled || NewLease.State == LeaseState.Completed)
 							{
-								await RemoveLeaseAsync(Agent, Lease, UtcNow, LeaseOutcome.Cancelled);
+								await RemoveLeaseAsync(Agent, Lease, UtcNow, LeaseOutcome.Cancelled, null);
 								Leases.RemoveAt(Idx--);
 								bUpdateLeases = true;
 							}
@@ -569,7 +569,7 @@ namespace HordeServer.Services
 							{
 								if (NewLease.State == LeaseState.Cancelled || NewLease.State == LeaseState.Completed)
 								{
-									await RemoveLeaseAsync(Agent, Lease, UtcNow, NewLease.Outcome);
+									await RemoveLeaseAsync(Agent, Lease, UtcNow, NewLease.Outcome, NewLease.Result.ToByteArray());
 									Leases.RemoveAt(Idx--);
 								}
 								else if (NewLease.State == LeaseState.Active && Lease.State == LeaseState.Pending)
@@ -765,7 +765,7 @@ namespace HordeServer.Services
 				foreach(AgentLease Lease in Leases)
 				{
 					Logger.LogDebug("Removing lease {LeaseId} during session terminate...", Lease.Id);
-					await RemoveLeaseAsync(Agent, Lease, FinishTime, LeaseOutcome.Failed);
+					await RemoveLeaseAsync(Agent, Lease, FinishTime, LeaseOutcome.Failed, null);
 				}
 
 				// Update the session document
@@ -826,8 +826,9 @@ namespace HordeServer.Services
 		/// <param name="Lease">The lease to cancel</param>
 		/// <param name="UtcNow">The current time</param>
 		/// <param name="Outcome">Final status of the lease</param>
+		/// <param name="Output">Output from executing the task</param>
 		/// <returns>Async task</returns>
-		private async Task RemoveLeaseAsync(IAgent Agent, AgentLease Lease, DateTime UtcNow, LeaseOutcome Outcome)
+		private async Task RemoveLeaseAsync(IAgent Agent, AgentLease Lease, DateTime UtcNow, LeaseOutcome Outcome, byte[]? Output)
 		{
 			// Make sure the lease is terminated correctly
 			if (Lease.Payload == null)
@@ -843,7 +844,7 @@ namespace HordeServer.Services
 				{
 					if (Any.Is(TaskSource.Descriptor))
 					{
-						await TaskSource.OnLeaseCompletedAsync(Agent, Lease.Id, Any, Outcome);
+						await TaskSource.OnLeaseCompletedAsync(Agent, Lease.Id, Any, Outcome, Output);
 						break;
 					}
 				}
@@ -861,7 +862,7 @@ namespace HordeServer.Services
 			}
 
 			// Update the lease
-			await Leases.TrySetOutcomeAsync(Lease.Id, FinishTime, Outcome);
+			await Leases.TrySetOutcomeAsync(Lease.Id, FinishTime, Outcome, Output);
 		}
 
 		/// <summary>
diff --git a/Engine/Source/Programs/Horde/HordeServer/Services/ConsistencyService.cs b/Engine/Source/Programs/Horde/HordeServer/Services/ConsistencyService.cs
--- a/Engine/Source/Programs/Horde/HordeServer/Services/ConsistencyService.cs
+++ b/Engine/Source/Programs/Horde/HordeServer/Services/ConsistencyService.cs
@@ -59,7 +59,7 @@ namespace HordeServer.Services
 					ISession? Session = await SessionCollection.GetAsync(Lease.SessionId);
 					DateTime FinishTime = Session?.FinishTime ?? DateTime.UtcNow;
 					Logger.LogWarning("Setting finish time for lease {LeaseId} to {FinishTime}", Lease.Id, FinishTime);
-					await LeaseCollection.TrySetOutcomeAsync(Lease.Id, FinishTime, LeaseOutcome.Cancelled);
+					await LeaseCollection.TrySetOutcomeAsync(Lease.Id, FinishTime, LeaseOutcome.Cancelled, null);
 				}
 			}
 
