Missing files.

diff --git a/Engine/Source/Programs/Horde/HordeServer/Collections/Impl/LeaseCollection.cs b/Engine/Source/Programs/Horde/HordeServer/Collections/Impl/LeaseCollection.cs
--- a/Engine/Source/Programs/Horde/HordeServer/Collections/Impl/LeaseCollection.cs
+++ b/Engine/Source/Programs/Horde/HordeServer/Collections/Impl/LeaseCollection.cs
@@ -43,7 +43,7 @@ namespace HordeServer.Collections.Impl
 			public byte[]? Output { get; set; }
 
 			ReadOnlyMemory<byte> ILease.Payload => Payload;
-			ReadOnlyMemory<byte> ILease.Output => Output;
+			ReadOnlyMemory<byte> ILease.Output => Output ?? Array.Empty<byte>();
 
 			[BsonConstructor]
 			private LeaseDocument()
diff --git a/Engine/Source/Programs/Horde/HordeServer/Models/Lease.cs b/Engine/Source/Programs/Horde/HordeServer/Models/Lease.cs
--- a/Engine/Source/Programs/Horde/HordeServer/Models/Lease.cs
+++ b/Engine/Source/Programs/Horde/HordeServer/Models/Lease.cs
@@ -77,6 +77,11 @@ namespace HordeServer.Models
 		/// Outcome of the lease
 		/// </summary>
 		public LeaseOutcome Outcome { get; }
+
+		/// <summary>
+		/// Output from executing the lease
+		/// </summary>
+		public ReadOnlyMemory<byte> Output { get; }
 	}
 
 	/// <summary>
