Horde: Add a codepath for reporting lease output back to task sources.

diff --git a/Engine/Source/Programs/Horde/HordeServer/Services/AgentService.cs b/Engine/Source/Programs/Horde/HordeServer/Services/AgentService.cs
--- a/Engine/Source/Programs/Horde/HordeServer/Services/AgentService.cs
+++ b/Engine/Source/Programs/Horde/HordeServer/Services/AgentService.cs
@@ -569,7 +569,7 @@ namespace HordeServer.Services
 							{
 								if (NewLease.State == LeaseState.Cancelled || NewLease.State == LeaseState.Completed)
 								{
-									await RemoveLeaseAsync(Agent, Lease, UtcNow, NewLease.Outcome, NewLease.Result.ToByteArray());
+									await RemoveLeaseAsync(Agent, Lease, UtcNow, NewLease.Outcome, NewLease.Output.ToByteArray());
 									Leases.RemoveAt(Idx--);
 								}
 								else if (NewLease.State == LeaseState.Active && Lease.State == LeaseState.Pending)
diff --git a/Engine/Source/Programs/Horde/HordeServer/Tasks/ITaskSource.cs b/Engine/Source/Programs/Horde/HordeServer/Tasks/ITaskSource.cs
--- a/Engine/Source/Programs/Horde/HordeServer/Tasks/ITaskSource.cs
+++ b/Engine/Source/Programs/Horde/HordeServer/Tasks/ITaskSource.cs
@@ -110,8 +110,9 @@ namespace HordeServer.Tasks
 		/// <param name="LeaseId">The lease id</param>
 		/// <param name="Payload">The lease payload</param>
 		/// <param name="Outcome">Outcome of the lease</param>
+		/// <param name="Output">Output from the task</param>
 		/// <returns>Async task</returns>
-		Task OnLeaseCompletedAsync(IAgent Agent, ObjectId LeaseId, Any Payload, LeaseOutcome Outcome);
+		Task OnLeaseCompletedAsync(IAgent Agent, ObjectId LeaseId, Any Payload, LeaseOutcome Outcome, ReadOnlyMemory<byte> Output);
 	}
 
 	/// <summary>
@@ -127,13 +128,13 @@ namespace HordeServer.Tasks
 		public abstract Task<ITaskListener?> SubscribeAsync(IAgent Agent);
 
 		/// <inheritdoc/>
-		public Task OnLeaseCompletedAsync(IAgent Agent, ObjectId LeaseId, Any Task, LeaseOutcome Outcome)
+		public Task OnLeaseCompletedAsync(IAgent Agent, ObjectId LeaseId, Any Task, LeaseOutcome Outcome, ReadOnlyMemory<byte> Output)
 		{
-			return OnLeaseCompletedAsync(Agent, LeaseId, Task.Unpack<T>(), Outcome);
+			return OnLeaseCompletedAsync(Agent, LeaseId, Task.Unpack<T>(), Outcome, Output);
 		}
 
-		/// <inheritdoc cref="OnLeaseCompletedAsync(IAgent, ObjectId, Any, LeaseOutcome)"/>
-		public virtual Task OnLeaseCompletedAsync(IAgent Agent, ObjectId LeaseId, T TaskInput, LeaseOutcome Outcome)
+		/// <inheritdoc cref="OnLeaseCompletedAsync(IAgent, ObjectId, Any, LeaseOutcome, ReadOnlyMemory{byte})"/>
+		public virtual Task OnLeaseCompletedAsync(IAgent Agent, ObjectId LeaseId, T TaskInput, LeaseOutcome Outcome, ReadOnlyMemory<byte> Output)
 		{
 			return Task.CompletedTask;
 		}
diff --git a/Engine/Source/Programs/Horde/HordeServer/Tasks/Impl/ConformTaskSource.cs b/Engine/Source/Programs/Horde/HordeServer/Tasks/Impl/ConformTaskSource.cs
--- a/Engine/Source/Programs/Horde/HordeServer/Tasks/Impl/ConformTaskSource.cs
+++ b/Engine/Source/Programs/Horde/HordeServer/Tasks/Impl/ConformTaskSource.cs
@@ -211,7 +211,7 @@ namespace HordeServer.Tasks.Impl
 		}
 
 		/// <inheritdoc/>
-		public override Task OnLeaseCompletedAsync(IAgent Agent, ObjectId LeaseId, ConformTask Task, LeaseOutcome Outcome)
+		public override Task OnLeaseCompletedAsync(IAgent Agent, ObjectId LeaseId, ConformTask Task, LeaseOutcome Outcome, ReadOnlyMemory<byte> Output)
 		{
 			return ReleaseConformLeaseAsync(LeaseId);
 		}
diff --git a/Engine/Source/Programs/Horde/HordeServer/Tasks/Impl/JobTaskSource.cs b/Engine/Source/Programs/Horde/HordeServer/Tasks/Impl/JobTaskSource.cs
--- a/Engine/Source/Programs/Horde/HordeServer/Tasks/Impl/JobTaskSource.cs
+++ b/Engine/Source/Programs/Horde/HordeServer/Tasks/Impl/JobTaskSource.cs
@@ -1009,7 +1009,7 @@ namespace HordeServer.Tasks.Impl
 		}
 
 		/// <inheritdoc/>
-		public async Task OnLeaseCompletedAsync(IAgent Agent, ObjectId LeaseId, Any Any, LeaseOutcome Outcome)
+		public async Task OnLeaseCompletedAsync(IAgent Agent, ObjectId LeaseId, Any Any, LeaseOutcome Outcome, ReadOnlyMemory<byte> Output)
 		{
 			if (Outcome != LeaseOutcome.Success)
 			{
