Horde: Read the server config file from the application directory by default. Fixes basic "launch from unconfigured VS" exception.

diff --git a/Engine/Source/Programs/Horde/HordeServer/ServerSettings.cs b/Engine/Source/Programs/Horde/HordeServer/ServerSettings.cs
--- a/Engine/Source/Programs/Horde/HordeServer/ServerSettings.cs
+++ b/Engine/Source/Programs/Horde/HordeServer/ServerSettings.cs
@@ -409,7 +409,7 @@ namespace HordeServer
 		/// <summary>
 		/// Path to the root config file
 		/// </summary>
-		public string? ConfigPath { get; set; }
+		public string ConfigPath { get; set; } = "Defaults/globals.json";
 
 		/// <summary>
 		/// Settings for remote execution
diff --git a/Engine/Source/Programs/Horde/HordeServer/Services/ConfigService.cs b/Engine/Source/Programs/Horde/HordeServer/Services/ConfigService.cs
--- a/Engine/Source/Programs/Horde/HordeServer/Services/ConfigService.cs
+++ b/Engine/Source/Programs/Horde/HordeServer/Services/ConfigService.cs
@@ -422,7 +422,7 @@ namespace HordeServer.Services
 		{
 			if (Settings.ConfigPath != null)
 			{
-				Uri ConfigUri = CombinePaths(new Uri(FileReference.Combine(Program.DataDir, "_").FullName), Settings.ConfigPath);
+				Uri ConfigUri = CombinePaths(new Uri(FileReference.Combine(Program.AppDir, "_").FullName), Settings.ConfigPath);
 				await UpdateConfigAsync(ConfigUri);
 			}
 			return DateTime.UtcNow + TimeSpan.FromMinutes(1.0);
diff --git a/Engine/Source/Programs/Horde/HordeServer/appsettings.json b/Engine/Source/Programs/Horde/HordeServer/appsettings.json
--- a/Engine/Source/Programs/Horde/HordeServer/appsettings.json
+++ b/Engine/Source/Programs/Horde/HordeServer/appsettings.json
@@ -3,7 +3,6 @@
 	{
 		"DashboardUrl": "https://localhost:3000",
 		"DisableAuth": true,
-		"ConfigPath": "Defaults/globals.json",
 		"CorsEnabled": true,
 		"CorsOrigin": "https://localhost:5003;http://localhost:3000",
 		"LogFileProviderType": "Local"
