Horde: Add launchsettings.json to source control, to avoid launching IIS by default. We use always use Kestrel for Horde.

diff --git a/Engine/Source/Programs/Horde/HordeServer/Properties/launchSettings.json b/Engine/Source/Programs/Horde/HordeServer/Properties/launchSettings.json
new file mode 100644
--- /dev/null
+++ b/Engine/Source/Programs/Horde/HordeServer/Properties/launchSettings.json
@@ -0,0 +1,8 @@
+{
+  "profiles": {
+    "HordeServer": {
+      "commandName": "Project",
+      "commandLineArgs": ""
+    }
+  }
+}
\ No newline at end of file
