Horde: Use the OpenTracing API rather than hard-coding support for DataDog (though still initializes Datadog at startup).

diff --git a/Engine/Source/Programs/Horde/HordeServer/Collections/DatadogTracingCollection.cs b/Engine/Source/Programs/Horde/HordeServer/Collections/DatadogTracingCollection.cs
--- a/Engine/Source/Programs/Horde/HordeServer/Collections/DatadogTracingCollection.cs
+++ b/Engine/Source/Programs/Horde/HordeServer/Collections/DatadogTracingCollection.cs
@@ -7,6 +7,8 @@ using Datadog.Trace;
 using MongoDB.Bson;
 using MongoDB.Bson.Serialization;
 using MongoDB.Driver;
+using OpenTracing;
+using OpenTracing.Util;
 
 namespace HordeServer.Collections
 {
@@ -50,7 +52,7 @@ namespace HordeServer.Collections
 		public async Task<IAsyncCursor<TResult>> AggregateAsync<TResult>(PipelineDefinition<T, TResult> pipeline, AggregateOptions Options = null!,
 			CancellationToken  CancellationToken = new CancellationToken())
 		{
-			using Scope _ = Tracer.Instance.StartActive(Prefix + "AggregateAsync");
+			using IScope _ = GlobalTracer.Instance.BuildSpan(Prefix + "AggregateAsync").StartActive();
 			return await Collection.AggregateAsync(pipeline, Options,  CancellationToken);
 		}
 
@@ -58,7 +60,7 @@ namespace HordeServer.Collections
 		public async Task<IAsyncCursor<TResult>> AggregateAsync<TResult>(IClientSessionHandle Session, PipelineDefinition<T, TResult> pipeline, AggregateOptions Options = null!,
 			CancellationToken  CancellationToken = new CancellationToken())
 		{
-			using Scope _ = Tracer.Instance.StartActive(Prefix + "AggregateAsync");
+			using IScope _ = GlobalTracer.Instance.BuildSpan(Prefix + "AggregateAsync").StartActive();
 			return await Collection.AggregateAsync(Session, pipeline, Options, CancellationToken);
 		}
 
@@ -80,7 +82,7 @@ namespace HordeServer.Collections
 		public async Task<BulkWriteResult<T>> BulkWriteAsync(IEnumerable<WriteModel<T>> requests, BulkWriteOptions Options = null!,
 			CancellationToken  CancellationToken = new CancellationToken())
 		{
-			using Scope _ = Tracer.Instance.StartActive(Prefix + "BulkWriteAsync");
+			using IScope _ = GlobalTracer.Instance.BuildSpan(Prefix + "BulkWriteAsync").StartActive();
 			return await Collection.BulkWriteAsync(requests, Options,  CancellationToken);
 		}
 
@@ -88,7 +90,7 @@ namespace HordeServer.Collections
 		public async Task<BulkWriteResult<T>> BulkWriteAsync(IClientSessionHandle Session, IEnumerable<WriteModel<T>> requests, BulkWriteOptions Options = null!,
 			CancellationToken  CancellationToken = new CancellationToken())
 		{
-			using Scope _ = Tracer.Instance.StartActive(Prefix + "BulkWriteAsync");
+			using IScope _ = GlobalTracer.Instance.BuildSpan(Prefix + "BulkWriteAsync").StartActive();
 			return await Collection.BulkWriteAsync(Session, requests, Options,  CancellationToken);
 		}
 
@@ -110,7 +112,7 @@ namespace HordeServer.Collections
 		public async Task<long> CountAsync(FilterDefinition<T> filter, CountOptions Options = null!,
 			CancellationToken  CancellationToken = new CancellationToken())
 		{
-			using Scope _ = Tracer.Instance.StartActive(Prefix + "CountAsync");
+			using IScope _ = GlobalTracer.Instance.BuildSpan(Prefix + "CountAsync").StartActive();
 			return await Collection.CountAsync(filter, Options,  CancellationToken);
 		}
 
@@ -118,7 +120,7 @@ namespace HordeServer.Collections
 		public async Task<long> CountAsync(IClientSessionHandle Session, FilterDefinition<T> filter, CountOptions Options = null!,
 			CancellationToken  CancellationToken = new CancellationToken())
 		{
-			using Scope _ = Tracer.Instance.StartActive(Prefix + "CountAsync");
+			using IScope _ = GlobalTracer.Instance.BuildSpan(Prefix + "CountAsync").StartActive();
 			return await Collection.CountAsync(Session, filter, Options,  CancellationToken);
 		}
 
@@ -133,7 +135,7 @@ namespace HordeServer.Collections
 		public long CountDocuments(IClientSessionHandle Session, FilterDefinition<T> filter, CountOptions Options = null!,
 			CancellationToken  CancellationToken = new CancellationToken())
 		{
-			using Scope _ = Tracer.Instance.StartActive(Prefix + "CountDocumentsAsync");
+			using IScope _ = GlobalTracer.Instance.BuildSpan(Prefix + "CountDocumentsAsync").StartActive();
 			return Collection.CountDocuments(Session, filter, Options,  CancellationToken);
 		}
 
@@ -141,7 +143,7 @@ namespace HordeServer.Collections
 		public async Task<long> CountDocumentsAsync(FilterDefinition<T> filter, CountOptions Options = null!,
 			CancellationToken  CancellationToken = new CancellationToken())
 		{
-			using Scope _ = Tracer.Instance.StartActive(Prefix + "CountDocumentsAsync");
+			using IScope _ = GlobalTracer.Instance.BuildSpan(Prefix + "CountDocumentsAsync").StartActive();
 			return await Collection.CountDocumentsAsync(filter, Options,  CancellationToken);
 		}
 
@@ -149,7 +151,7 @@ namespace HordeServer.Collections
 		public async Task<long> CountDocumentsAsync(IClientSessionHandle Session, FilterDefinition<T> filter, CountOptions Options = null!,
 			CancellationToken  CancellationToken = new CancellationToken())
 		{
-			using Scope _ = Tracer.Instance.StartActive(Prefix + "CountDocumentsAsync");
+			using IScope _ = GlobalTracer.Instance.BuildSpan(Prefix + "CountDocumentsAsync").StartActive();
 			return await Collection.CountDocumentsAsync(Session, filter, Options,  CancellationToken);
 		}
 
@@ -176,7 +178,7 @@ namespace HordeServer.Collections
 		/// <inheritdoc />
 		public async Task<DeleteResult> DeleteManyAsync(FilterDefinition<T> filter, CancellationToken  CancellationToken = new CancellationToken())
 		{
-			using Scope _ = Tracer.Instance.StartActive(Prefix + "DeleteManyAsync");
+			using IScope _ = GlobalTracer.Instance.BuildSpan(Prefix + "DeleteManyAsync").StartActive();
 			return await Collection.DeleteManyAsync(filter,  CancellationToken);
 		}
 
@@ -184,7 +186,7 @@ namespace HordeServer.Collections
 		public async Task<DeleteResult> DeleteManyAsync(FilterDefinition<T> filter, DeleteOptions Options,
 			CancellationToken  CancellationToken = new CancellationToken())
 		{
-			using Scope _ = Tracer.Instance.StartActive(Prefix + "DeleteManyAsync");
+			using IScope _ = GlobalTracer.Instance.BuildSpan(Prefix + "DeleteManyAsync").StartActive();
 			return await Collection.DeleteManyAsync(filter, Options,  CancellationToken);
 		}
 
@@ -192,7 +194,7 @@ namespace HordeServer.Collections
 		public async Task<DeleteResult> DeleteManyAsync(IClientSessionHandle Session, FilterDefinition<T> filter, DeleteOptions Options = null!,
 			CancellationToken  CancellationToken = new CancellationToken())
 		{
-			using Scope _ = Tracer.Instance.StartActive(Prefix + "DeleteManyAsync");
+			using IScope _ = GlobalTracer.Instance.BuildSpan(Prefix + "DeleteManyAsync").StartActive();
 			return await Collection.DeleteManyAsync(Session, filter, Options,  CancellationToken);
 		}
 
@@ -219,7 +221,7 @@ namespace HordeServer.Collections
 		/// <inheritdoc />
 		public async Task<DeleteResult> DeleteOneAsync(FilterDefinition<T> filter, CancellationToken  CancellationToken = new CancellationToken())
 		{
-			using Scope _ = Tracer.Instance.StartActive(Prefix + "DeleteOneAsync");
+			using IScope _ = GlobalTracer.Instance.BuildSpan(Prefix + "DeleteOneAsync").StartActive();
 			return await Collection.DeleteOneAsync(filter,  CancellationToken);
 		}
 
@@ -227,7 +229,7 @@ namespace HordeServer.Collections
 		public async Task<DeleteResult> DeleteOneAsync(FilterDefinition<T> filter, DeleteOptions Options,
 			CancellationToken  CancellationToken = new CancellationToken())
 		{
-			using Scope _ = Tracer.Instance.StartActive(Prefix + "DeleteOneAsync");
+			using IScope _ = GlobalTracer.Instance.BuildSpan(Prefix + "DeleteOneAsync").StartActive();
 			return await Collection.DeleteOneAsync(filter, Options,  CancellationToken);
 		}
 
@@ -235,7 +237,7 @@ namespace HordeServer.Collections
 		public Task<DeleteResult> DeleteOneAsync(IClientSessionHandle Session, FilterDefinition<T> filter, DeleteOptions Options = null!,
 			CancellationToken  CancellationToken = new CancellationToken())
 		{
-			using Scope _ = Tracer.Instance.StartActive(Prefix + "DeleteOneAsync");
+			using IScope _ = GlobalTracer.Instance.BuildSpan(Prefix + "DeleteOneAsync").StartActive();
 			return Collection.DeleteOneAsync(Session, filter, Options,  CancellationToken);
 		}
 
@@ -257,7 +259,7 @@ namespace HordeServer.Collections
 		public async Task<IAsyncCursor<TField>> DistinctAsync<TField>(FieldDefinition<T, TField> field, FilterDefinition<T> filter, DistinctOptions Options = null!,
 			CancellationToken  CancellationToken = new CancellationToken())
 		{
-			using Scope _ = Tracer.Instance.StartActive(Prefix + "DistinctAsync");
+			using IScope _ = GlobalTracer.Instance.BuildSpan(Prefix + "DistinctAsync").StartActive();
 			return await Collection.DistinctAsync(field, filter, Options,  CancellationToken);
 		}
 
@@ -265,7 +267,7 @@ namespace HordeServer.Collections
 		public async Task<IAsyncCursor<TField>> DistinctAsync<TField>(IClientSessionHandle Session, FieldDefinition<T, TField> field, FilterDefinition<T> filter,
 			DistinctOptions Options = null!, CancellationToken  CancellationToken = new CancellationToken())
 		{
-			using Scope _ = Tracer.Instance.StartActive(Prefix + "DistinctAsync");
+			using IScope _ = GlobalTracer.Instance.BuildSpan(Prefix + "DistinctAsync").StartActive();
 			return await Collection.DistinctAsync(Session, field, filter, Options,  CancellationToken);
 		}
 
@@ -280,7 +282,7 @@ namespace HordeServer.Collections
 		public async Task<long> EstimatedDocumentCountAsync(EstimatedDocumentCountOptions Options = null!,
 			CancellationToken  CancellationToken = new CancellationToken())
 		{
-			using Scope _ = Tracer.Instance.StartActive(Prefix + "EstimatedDocumentCountAsync");
+			using IScope _ = GlobalTracer.Instance.BuildSpan(Prefix + "EstimatedDocumentCountAsync").StartActive();
 			return await Collection.EstimatedDocumentCountAsync(Options,  CancellationToken);
 		}
 
@@ -302,7 +304,7 @@ namespace HordeServer.Collections
 		public async Task<IAsyncCursor<TProjection>> FindAsync<TProjection>(FilterDefinition<T> filter, FindOptions<T, TProjection> Options = null!,
 			CancellationToken  CancellationToken = new CancellationToken())
 		{
-			using Scope _ = Tracer.Instance.StartActive(Prefix + "FindAsync");
+			using IScope _ = GlobalTracer.Instance.BuildSpan(Prefix + "FindAsync").StartActive();
 			return await Collection.FindAsync(filter, Options,  CancellationToken);
 		}
 
@@ -310,7 +312,7 @@ namespace HordeServer.Collections
 		public async Task<IAsyncCursor<TProjection>> FindAsync<TProjection>(IClientSessionHandle Session, FilterDefinition<T> filter, FindOptions<T, TProjection> Options = null!,
 			CancellationToken  CancellationToken = new CancellationToken())
 		{
-			using Scope _ = Tracer.Instance.StartActive(Prefix + "FindAsync");
+			using IScope _ = GlobalTracer.Instance.BuildSpan(Prefix + "FindAsync").StartActive();
 			return await Collection.FindAsync(Session, filter, Options,  CancellationToken);
 		}
 
@@ -332,7 +334,7 @@ namespace HordeServer.Collections
 		public async Task<TProjection> FindOneAndDeleteAsync<TProjection>(FilterDefinition<T> filter, FindOneAndDeleteOptions<T, TProjection> Options = null!,
 			CancellationToken  CancellationToken = new CancellationToken())
 		{
-			using Scope _ = Tracer.Instance.StartActive(Prefix + "FindOneAndDeleteAsync");
+			using IScope _ = GlobalTracer.Instance.BuildSpan(Prefix + "FindOneAndDeleteAsync").StartActive();
 			return await Collection.FindOneAndDeleteAsync(filter, Options,  CancellationToken);
 		}
 
@@ -340,7 +342,7 @@ namespace HordeServer.Collections
 		public async Task<TProjection> FindOneAndDeleteAsync<TProjection>(IClientSessionHandle Session, FilterDefinition<T> filter,
 			FindOneAndDeleteOptions<T, TProjection> Options = null!, CancellationToken  CancellationToken = new CancellationToken())
 		{
-			using Scope _ = Tracer.Instance.StartActive(Prefix + "FindOneAndDeleteAsync");
+			using IScope _ = GlobalTracer.Instance.BuildSpan(Prefix + "FindOneAndDeleteAsync").StartActive();
 			return await Collection.FindOneAndDeleteAsync(Session, filter, Options,  CancellationToken);
 		}
 
@@ -362,7 +364,7 @@ namespace HordeServer.Collections
 		public async Task<TProjection> FindOneAndReplaceAsync<TProjection>(FilterDefinition<T> filter, T replacement,
 			FindOneAndReplaceOptions<T, TProjection> Options = null!, CancellationToken  CancellationToken = new CancellationToken())
 		{
-			using Scope _ = Tracer.Instance.StartActive(Prefix + "FindOneAndReplaceAsync");
+			using IScope _ = GlobalTracer.Instance.BuildSpan(Prefix + "FindOneAndReplaceAsync").StartActive();
 			return await Collection.FindOneAndReplaceAsync(filter, replacement, Options,  CancellationToken);
 		}
 
@@ -370,7 +372,7 @@ namespace HordeServer.Collections
 		public async Task<TProjection> FindOneAndReplaceAsync<TProjection>(IClientSessionHandle Session, FilterDefinition<T> filter, T replacement,
 			FindOneAndReplaceOptions<T, TProjection> Options = null!, CancellationToken  CancellationToken = new CancellationToken())
 		{
-			using Scope _ = Tracer.Instance.StartActive(Prefix + "FindOneAndReplaceAsync");
+			using IScope _ = GlobalTracer.Instance.BuildSpan(Prefix + "FindOneAndReplaceAsync").StartActive();
 			return await Collection.FindOneAndReplaceAsync(Session, filter, replacement, Options,  CancellationToken);
 		}
 
@@ -393,7 +395,7 @@ namespace HordeServer.Collections
 		public async Task<TProjection> FindOneAndUpdateAsync<TProjection>(FilterDefinition<T> filter, UpdateDefinition<T> update,
 			FindOneAndUpdateOptions<T, TProjection> Options = null!, CancellationToken  CancellationToken = new CancellationToken())
 		{
-			using Scope _ = Tracer.Instance.StartActive(Prefix + "FindOneAndUpdateAsync");
+			using IScope _ = GlobalTracer.Instance.BuildSpan(Prefix + "FindOneAndUpdateAsync").StartActive();
 			return await Collection.FindOneAndUpdateAsync(filter, update, Options,  CancellationToken);
 		}
 
@@ -401,7 +403,7 @@ namespace HordeServer.Collections
 		public async Task<TProjection> FindOneAndUpdateAsync<TProjection>(IClientSessionHandle Session, FilterDefinition<T> filter, UpdateDefinition<T> update,
 			FindOneAndUpdateOptions<T, TProjection> Options = null!, CancellationToken  CancellationToken = new CancellationToken())
 		{
-			using Scope _ = Tracer.Instance.StartActive(Prefix + "FindOneAndUpdateAsync");
+			using IScope _ = GlobalTracer.Instance.BuildSpan(Prefix + "FindOneAndUpdateAsync").StartActive();
 			return await Collection.FindOneAndUpdateAsync(Session, filter, update, Options,  CancellationToken);
 		}
 
@@ -422,7 +424,7 @@ namespace HordeServer.Collections
 		/// <inheritdoc />
 		public async Task InsertOneAsync( T Document, CancellationToken CancellationToken)
 		{
-			using Scope _ = Tracer.Instance.StartActive(Prefix + "InsertOneAsync");
+			using IScope _ = GlobalTracer.Instance.BuildSpan(Prefix + "InsertOneAsync").StartActive();
 			await Collection.InsertOneAsync(Document, CancellationToken);
 		}
 
@@ -430,7 +432,7 @@ namespace HordeServer.Collections
 		public async Task InsertOneAsync( T Document, InsertOneOptions Options = null!,
 			CancellationToken  CancellationToken = new CancellationToken())
 		{
-			using Scope _ = Tracer.Instance.StartActive(Prefix + "InsertOneAsync");
+			using IScope _ = GlobalTracer.Instance.BuildSpan(Prefix + "InsertOneAsync").StartActive();
 			await Collection.InsertOneAsync(Document, Options,  CancellationToken);
 		}
 
@@ -438,7 +440,7 @@ namespace HordeServer.Collections
 		public async Task InsertOneAsync(IClientSessionHandle Session, T Document, InsertOneOptions Options = null!,
 			CancellationToken CancellationToken = new CancellationToken())
 		{
-			using Scope _ = Tracer.Instance.StartActive(Prefix + "InsertOneAsync");
+			using IScope _ = GlobalTracer.Instance.BuildSpan(Prefix + "InsertOneAsync").StartActive();
 			await Collection.InsertOneAsync(Session, Document, Options,  CancellationToken);
 		}
 
@@ -460,7 +462,7 @@ namespace HordeServer.Collections
 		public async Task InsertManyAsync(IEnumerable<T> Documents, InsertManyOptions Options = null!,
 			CancellationToken  CancellationToken = new CancellationToken())
 		{
-			using Scope _ = Tracer.Instance.StartActive(Prefix + "InsertManyAsync");
+			using IScope _ = GlobalTracer.Instance.BuildSpan(Prefix + "InsertManyAsync").StartActive();
 			await Collection.InsertManyAsync(Documents, Options,  CancellationToken);
 		}
 
@@ -468,7 +470,7 @@ namespace HordeServer.Collections
 		public async Task InsertManyAsync(IClientSessionHandle Session, IEnumerable<T> documents, InsertManyOptions Options = null!,
 			CancellationToken  CancellationToken = new CancellationToken())
 		{
-			using Scope _ = Tracer.Instance.StartActive(Prefix + "InsertManyAsync");
+			using IScope _ = GlobalTracer.Instance.BuildSpan(Prefix + "InsertManyAsync").StartActive();
 			await Collection.InsertManyAsync(Session, documents, Options,  CancellationToken);
 		}
 
@@ -490,7 +492,7 @@ namespace HordeServer.Collections
 		public async Task<IAsyncCursor<TResult>> MapReduceAsync<TResult>(BsonJavaScript map, BsonJavaScript reduce, MapReduceOptions<T, TResult> Options = null!,
 			CancellationToken  CancellationToken = new CancellationToken())
 		{
-			using Scope _ = Tracer.Instance.StartActive(Prefix + "MapReduceAsync");
+			using IScope _ = GlobalTracer.Instance.BuildSpan(Prefix + "MapReduceAsync").StartActive();
 			return await Collection.MapReduceAsync(map, reduce, Options,  CancellationToken);
 		}
 
@@ -498,7 +500,7 @@ namespace HordeServer.Collections
 		public async Task<IAsyncCursor<TResult>> MapReduceAsync<TResult>(IClientSessionHandle Session, BsonJavaScript map, BsonJavaScript reduce,
 			MapReduceOptions<T, TResult> Options = null!, CancellationToken  CancellationToken = new CancellationToken())
 		{
-			using Scope _ = Tracer.Instance.StartActive(Prefix + "MapReduceAsync");
+			using IScope _ = GlobalTracer.Instance.BuildSpan(Prefix + "MapReduceAsync").StartActive();
 			return await Collection.MapReduceAsync(Session, map, reduce, Options,  CancellationToken);
 		}
 
@@ -540,7 +542,7 @@ namespace HordeServer.Collections
 		public async Task<ReplaceOneResult> ReplaceOneAsync(FilterDefinition<T> filter, T replacement, ReplaceOptions Options = null!,
 			CancellationToken  CancellationToken = new CancellationToken())
 		{
-			using Scope _ = Tracer.Instance.StartActive(Prefix + "ReplaceOneAsync");
+			using IScope _ = GlobalTracer.Instance.BuildSpan(Prefix + "ReplaceOneAsync").StartActive();
 			return await Collection.ReplaceOneAsync(filter, replacement, Options,  CancellationToken);
 		}
 
@@ -548,7 +550,7 @@ namespace HordeServer.Collections
 		public async Task<ReplaceOneResult> ReplaceOneAsync(FilterDefinition<T> filter, T replacement, UpdateOptions Options,
 			CancellationToken  CancellationToken = new CancellationToken())
 		{
-			using Scope _ = Tracer.Instance.StartActive(Prefix + "ReplaceOneAsync");
+			using IScope _ = GlobalTracer.Instance.BuildSpan(Prefix + "ReplaceOneAsync").StartActive();
 			return await Collection.ReplaceOneAsync(filter, replacement, Options,  CancellationToken);
 		}
 
@@ -556,7 +558,7 @@ namespace HordeServer.Collections
 		public async Task<ReplaceOneResult> ReplaceOneAsync(IClientSessionHandle Session, FilterDefinition<T> filter, T replacement,
 			ReplaceOptions Options = null!, CancellationToken  CancellationToken = new CancellationToken())
 		{
-			using Scope _ = Tracer.Instance.StartActive(Prefix + "ReplaceOneAsync");
+			using IScope _ = GlobalTracer.Instance.BuildSpan(Prefix + "ReplaceOneAsync").StartActive();
 			return await Collection.ReplaceOneAsync(Session, filter, replacement, Options,  CancellationToken);
 		}
 
@@ -564,7 +566,7 @@ namespace HordeServer.Collections
 		public async Task<ReplaceOneResult> ReplaceOneAsync(IClientSessionHandle Session, FilterDefinition<T> filter, T replacement, UpdateOptions Options,
 			CancellationToken  CancellationToken = new CancellationToken())
 		{
-			using Scope _ = Tracer.Instance.StartActive(Prefix + "ReplaceOneAsync");
+			using IScope _ = GlobalTracer.Instance.BuildSpan(Prefix + "ReplaceOneAsync").StartActive();
 			return await Collection.ReplaceOneAsync(Session, filter, replacement, Options,  CancellationToken);
 		}
 
@@ -586,7 +588,7 @@ namespace HordeServer.Collections
 		public async Task<UpdateResult> UpdateManyAsync(FilterDefinition<T> filter, UpdateDefinition<T> update, UpdateOptions Options = null!,
 			CancellationToken  CancellationToken = new CancellationToken())
 		{
-			using Scope _ = Tracer.Instance.StartActive(Prefix + "UpdateManyAsync");
+			using IScope _ = GlobalTracer.Instance.BuildSpan(Prefix + "UpdateManyAsync").StartActive();
 			return await Collection.UpdateManyAsync(filter, update, Options,  CancellationToken);
 		}
 
@@ -594,7 +596,7 @@ namespace HordeServer.Collections
 		public async Task<UpdateResult> UpdateManyAsync(IClientSessionHandle Session, FilterDefinition<T> filter, UpdateDefinition<T> update,
 			UpdateOptions Options = null!, CancellationToken  CancellationToken = new CancellationToken())
 		{
-			using Scope _ = Tracer.Instance.StartActive(Prefix + "UpdateManyAsync");
+			using IScope _ = GlobalTracer.Instance.BuildSpan(Prefix + "UpdateManyAsync").StartActive();
 			return await Collection.UpdateManyAsync(Session, filter, update, Options,  CancellationToken);
 		}
 
@@ -616,7 +618,7 @@ namespace HordeServer.Collections
 		public async Task<UpdateResult> UpdateOneAsync(FilterDefinition<T> filter, UpdateDefinition<T> update, UpdateOptions Options = null!,
 			CancellationToken  CancellationToken = new CancellationToken())
 		{
-			using Scope _ = Tracer.Instance.StartActive(Prefix + "UpdateOneAsync");
+			using IScope _ = GlobalTracer.Instance.BuildSpan(Prefix + "UpdateOneAsync").StartActive();
 			return await Collection.UpdateOneAsync(filter, update, Options,  CancellationToken);
 		}
 
@@ -624,7 +626,7 @@ namespace HordeServer.Collections
 		public async Task<UpdateResult> UpdateOneAsync(IClientSessionHandle Session, FilterDefinition<T> filter, UpdateDefinition<T> update,
 			UpdateOptions Options = null!, CancellationToken  CancellationToken = new CancellationToken())
 		{
-			using Scope _ = Tracer.Instance.StartActive(Prefix + "UpdateOneAsync");
+			using IScope _ = GlobalTracer.Instance.BuildSpan(Prefix + "UpdateOneAsync").StartActive();
 			return await Collection.UpdateOneAsync(Session, filter, update, Options,  CancellationToken);
 		}
 
@@ -646,7 +648,7 @@ namespace HordeServer.Collections
 		public async Task<IChangeStreamCursor<TResult>> WatchAsync<TResult>(PipelineDefinition<ChangeStreamDocument<T>, TResult> pipeline, ChangeStreamOptions Options = null!,
 			CancellationToken  CancellationToken = new CancellationToken())
 		{
-			using Scope _ = Tracer.Instance.StartActive(Prefix + "WatchAsync");
+			using IScope _ = GlobalTracer.Instance.BuildSpan(Prefix + "WatchAsync").StartActive();
 			return await Collection.WatchAsync(pipeline, Options,  CancellationToken);
 		}
 
@@ -654,7 +656,7 @@ namespace HordeServer.Collections
 		public async Task<IChangeStreamCursor<TResult>> WatchAsync<TResult>(IClientSessionHandle Session, PipelineDefinition<ChangeStreamDocument<T>, TResult> pipeline, ChangeStreamOptions Options = null!,
 			CancellationToken  CancellationToken = new CancellationToken())
 		{
-			using Scope _ = Tracer.Instance.StartActive(Prefix + "WatchAsync");
+			using IScope _ = GlobalTracer.Instance.BuildSpan(Prefix + "WatchAsync").StartActive();
 			return await Collection.WatchAsync(Session, pipeline, Options,  CancellationToken);
 		}
 
diff --git a/Engine/Source/Programs/Horde/HordeServer/Collections/Impl/JobCollection.cs b/Engine/Source/Programs/Horde/HordeServer/Collections/Impl/JobCollection.cs
--- a/Engine/Source/Programs/Horde/HordeServer/Collections/Impl/JobCollection.cs
+++ b/Engine/Source/Programs/Horde/HordeServer/Collections/Impl/JobCollection.cs
@@ -21,6 +21,8 @@ using System.Threading.Tasks;
 using StreamId = HordeServer.Utilities.StringId<HordeServer.Models.IStream>;
 using TemplateRefId = HordeServer.Utilities.StringId<HordeServer.Models.TemplateRef>;
 using PoolId = HordeServer.Utilities.StringId<HordeServer.Models.IPool>;
+using OpenTracing.Util;
+using OpenTracing;
 
 namespace HordeServer.Collections.Impl
 {
@@ -467,7 +469,7 @@ namespace HordeServer.Collections.Impl
 			}
 
 			List<JobDocument> Results;
-			using (Scope Scope = Tracer.Instance.StartActive("Jobs.Find"))
+			using (IScope Scope = GlobalTracer.Instance.BuildSpan("Jobs.Find").StartActive())
 			{
 				IFindFluent<JobDocument, JobDocument> Query = Jobs.Find<JobDocument>(Filter).SortByDescending(x => x.CreateTimeUtc!);
 				if (Index != null)
diff --git a/Engine/Source/Programs/Horde/HordeServer/Controllers/JobsController.cs b/Engine/Source/Programs/Horde/HordeServer/Controllers/JobsController.cs
--- a/Engine/Source/Programs/Horde/HordeServer/Controllers/JobsController.cs
+++ b/Engine/Source/Programs/Horde/HordeServer/Controllers/JobsController.cs
@@ -24,6 +24,8 @@ using StreamId = HordeServer.Utilities.StringId<HordeServer.Models.IStream>;
 using TemplateRefId = HordeServer.Utilities.StringId<HordeServer.Models.TemplateRef>;
 using Microsoft.Extensions.Logging;
 using HordeServer.Notifications;
+using OpenTracing.Util;
+using OpenTracing;
 
 namespace HordeServer.Controllers
 {
@@ -575,7 +577,7 @@ namespace HordeServer.Controllers
 			}
 
 			List<IJob> Jobs;
-			using (Scope _ = Tracer.Instance.StartActive("FindJobs"))
+			using (IScope _ = GlobalTracer.Instance.BuildSpan("FindJobs").StartActive())
 			{
 				Jobs = await JobService.FindJobsAsync(JobIdValues, StreamIdValue, Name, TemplateRefIds, MinChange,
 					MaxChange, PreflightChange, PreflightStartedByUserIdValue, StartedByUserIdValue, MinCreateTime?.UtcDateTime, MaxCreateTime?.UtcDateTime, Target, State, Outcome,
@@ -587,11 +589,11 @@ namespace HordeServer.Controllers
 			List<object> Responses = new List<object>();
 			foreach (IJob Job in Jobs)
 			{
-				using Scope JobScope = Tracer.Instance.StartActive("JobIteration");
+				using IScope JobScope = GlobalTracer.Instance.BuildSpan("JobIteration").StartActive();
 				JobScope.Span.SetTag("jobId", Job.Id.ToString());
 				
 				bool ViewJobAuthorized;
-				using (Scope _ = Tracer.Instance.StartActive("AuthorizeViewJob"))
+				using (IScope _ = GlobalTracer.Instance.BuildSpan("AuthorizeViewJob").StartActive())
 				{
 					ViewJobAuthorized = await JobService.AuthorizeAsync(Job, AclAction.ViewJob, User, PermissionsCache);
 				}
@@ -599,18 +601,18 @@ namespace HordeServer.Controllers
 				if (ViewJobAuthorized)
 				{
 					IGraph Graph;
-					using (Scope _ = Tracer.Instance.StartActive("GetGraph"))
+					using (IScope _ = GlobalTracer.Instance.BuildSpan("GetGraph").StartActive())
 					{
 						Graph = await JobService.GetGraphAsync(Job);
 					}
 
 					bool bIncludeAcl;
-					using (Scope _ = Tracer.Instance.StartActive("AuthorizeViewPermissions"))
+					using (IScope _ = GlobalTracer.Instance.BuildSpan("AuthorizeViewPermissions").StartActive())
 					{
 						bIncludeAcl = await JobService.AuthorizeAsync(Job, AclAction.ViewPermissions, User, PermissionsCache);
 					}
 
-					using (Scope _ = Tracer.Instance.StartActive("CreateResponse"))
+					using (IScope _ = GlobalTracer.Instance.BuildSpan("CreateResponse").StartActive())
 					{
 						Responses.Add(Job.ToResponse(Graph, bIncludeAcl, Filter));
 					}
diff --git a/Engine/Source/Programs/Horde/HordeServer/HordeServer.csproj b/Engine/Source/Programs/Horde/HordeServer/HordeServer.csproj
--- a/Engine/Source/Programs/Horde/HordeServer/HordeServer.csproj
+++ b/Engine/Source/Programs/Horde/HordeServer/HordeServer.csproj
@@ -15,7 +15,7 @@
     <PackageReference Include="AWSSDK.Extensions.NETCore.Setup" Version="3.3.101" />
     <PackageReference Include="AWSSDK.S3" Version="3.5.6.5" />
     <PackageReference Include="AWSSDK.SecurityToken" Version="3.5.1.29" />
-    <PackageReference Include="Datadog.Trace" Version="1.18.0" />
+    <PackageReference Include="Datadog.Trace.OpenTracing" Version="1.18.0" />
     <PackageReference Include="DogStatsD-CSharp-Client" Version="6.0.0" />
     <PackageReference Include="Google.Protobuf" Version="3.12.3" />
     <PackageReference Include="Grpc.AspNetCore" Version="2.29.0" />
@@ -35,9 +35,11 @@
     </PackageReference>
     <PackageReference Include="Microsoft.Extensions.Caching.Memory" Version="3.1.6" />
     <PackageReference Include="MongoDB.Driver" Version="2.10.4" />
+    <PackageReference Include="OpenTracing" Version="0.12.1" />
     <PackageReference Include="Polly" Version="7.2.1" />
     <PackageReference Include="Polly.Extensions.Http" Version="3.0.0" />
     <PackageReference Include="Serilog.AspNetCore" Version="3.2.0" />
+    <PackageReference Include="Serilog.Enrichers.OpenTracing" Version="0.0.2" />
     <PackageReference Include="Serilog.Extensions.Hosting" Version="3.1.0" />
     <PackageReference Include="Serilog.Sinks.Console" Version="3.1.1" />
     <PackageReference Include="Serilog.Sinks.File" Version="4.1.0" />
diff --git a/Engine/Source/Programs/Horde/HordeServer/Logs/Builder/Impl/RedisLogBuilder.cs b/Engine/Source/Programs/Horde/HordeServer/Logs/Builder/Impl/RedisLogBuilder.cs
--- a/Engine/Source/Programs/Horde/HordeServer/Logs/Builder/Impl/RedisLogBuilder.cs
+++ b/Engine/Source/Programs/Horde/HordeServer/Logs/Builder/Impl/RedisLogBuilder.cs
@@ -14,6 +14,8 @@ using System.Linq;
 using System.Text.RegularExpressions;
 using System.Threading.Tasks;
 using HordeServer.Utilities;
+using OpenTracing;
+using OpenTracing.Util;
 
 namespace HordeServer.Logs.Builder
 {
@@ -84,7 +86,7 @@ namespace HordeServer.Logs.Builder
 
 			if(ChunkOffset == WriteOffset)
 			{
-				using Scope Scope = Tracer.Instance.StartActive("Redis.CreateChunk");
+				using IScope Scope = GlobalTracer.Instance.BuildSpan("Redis.CreateChunk").StartActive();
 				Scope.Span.SetTag("LogId", LogId.ToString());
 				Scope.Span.SetTag("Offset", ChunkOffset.ToString(CultureInfo.InvariantCulture));
 				Scope.Span.SetTag("WriteOffset", WriteOffset.ToString(CultureInfo.InvariantCulture));
@@ -104,7 +106,7 @@ namespace HordeServer.Logs.Builder
 			}
 
 			{
-				using Scope Scope = Tracer.Instance.StartActive("Redis.AppendChunk");
+				using IScope Scope = GlobalTracer.Instance.BuildSpan("Redis.AppendChunk").StartActive();
 				Scope.Span.SetTag("LogId", LogId.ToString());
 				Scope.Span.SetTag("Offset", ChunkOffset.ToString(CultureInfo.InvariantCulture));
 				Scope.Span.SetTag("WriteOffset", WriteOffset.ToString(CultureInfo.InvariantCulture));
@@ -134,7 +136,7 @@ namespace HordeServer.Logs.Builder
 			ChunkKeys Keys = new ChunkKeys(LogId, Offset);
 			for (; ; )
 			{
-				using Scope Scope = Tracer.Instance.StartActive("Redis.CompleteSubChunk");
+				using IScope Scope = GlobalTracer.Instance.BuildSpan("Redis.CompleteSubChunk").StartActive();
 				Scope.Span.SetTag("LogId", LogId.ToString());
 				Scope.Span.SetTag("Offset", Offset.ToString(CultureInfo.InvariantCulture));
 
@@ -174,7 +176,7 @@ namespace HordeServer.Logs.Builder
 		public async Task CompleteChunkAsync(ObjectId LogId, long Offset)
 		{
 			IDatabase RedisDb = RedisConnectionPool.GetDatabase();
-			using Scope Scope = Tracer.Instance.StartActive("Redis.CompleteChunk");
+			using IScope Scope = GlobalTracer.Instance.BuildSpan("Redis.CompleteChunk").StartActive();
 			Scope.Span.SetTag("LogId", LogId.ToString());
 			Scope.Span.SetTag("Offset", Offset.ToString(CultureInfo.InvariantCulture));
 
@@ -196,7 +198,7 @@ namespace HordeServer.Logs.Builder
 		public async Task RemoveChunkAsync(ObjectId LogId, long Offset)
 		{
 			IDatabase RedisDb = RedisConnectionPool.GetDatabase();
-			using Scope Scope = Tracer.Instance.StartActive("Redis.RemoveChunk");
+			using IScope Scope = GlobalTracer.Instance.BuildSpan("Redis.RemoveChunk").StartActive();
 			Scope.Span.SetTag("LogId", LogId.ToString());
 			Scope.Span.SetTag("Offset", Offset.ToString(CultureInfo.InvariantCulture));
 
@@ -221,7 +223,7 @@ namespace HordeServer.Logs.Builder
 			ChunkKeys Keys = new ChunkKeys(LogId, Offset);
 			for (; ; )
 			{
-				using Scope Scope = Tracer.Instance.StartActive("Redis.GetChunk");
+				using IScope Scope = GlobalTracer.Instance.BuildSpan("Redis.GetChunk").StartActive();
 				Scope.Span.SetTag("LogId", LogId.ToString());
 				Scope.Span.SetTag("Offset", Offset.ToString(CultureInfo.InvariantCulture));
 
diff --git a/Engine/Source/Programs/Horde/HordeServer/Logs/LogIndexData.cs b/Engine/Source/Programs/Horde/HordeServer/Logs/LogIndexData.cs
--- a/Engine/Source/Programs/Horde/HordeServer/Logs/LogIndexData.cs
+++ b/Engine/Source/Programs/Horde/HordeServer/Logs/LogIndexData.cs
@@ -3,6 +3,8 @@
 using Datadog.Trace;
 using EpicGames.Core;
 using HordeServer.Utilities;
+using OpenTracing;
+using OpenTracing.Util;
 using System;
 using System.Collections.Generic;
 using System.IO;
@@ -220,7 +222,7 @@ namespace HordeServer.Logs
 		/// <returns>Index data</returns>
 		public static LogIndexData Merge(IEnumerable<LogIndexData> Indexes)
 		{
-			using Scope Scope = Tracer.Instance.StartActive("LogIndexData.Merge");
+			using IScope Scope = GlobalTracer.Instance.BuildSpan("LogIndexData.Merge").StartActive();
 
 			// Create the combined block list
 			LogIndexBlock[] NewBlocks = Indexes.SelectMany(x => x.Blocks).ToArray();
diff --git a/Engine/Source/Programs/Horde/HordeServer/Program.cs b/Engine/Source/Programs/Horde/HordeServer/Program.cs
--- a/Engine/Source/Programs/Horde/HordeServer/Program.cs
+++ b/Engine/Source/Programs/Horde/HordeServer/Program.cs
@@ -22,8 +22,10 @@ using Microsoft.Extensions.Logging;
 using MongoDB.Bson;
 using MongoDB.Bson.Serialization;
 using MongoDB.Driver;
+using OpenTracing.Util;
 using Serilog;
 using Serilog.Configuration;
+using Serilog.Enrichers.OpenTracing;
 using Serilog.Events;
 using Serilog.Filters;
 using Serilog.Formatting.Json;
@@ -146,9 +148,12 @@ namespace HordeServer
 				LogDir = DirectoryReference.Combine(DataDir);
 			}
 
+			GlobalTracer.Register(Datadog.Trace.OpenTracing.OpenTracingTracerFactory.CreateTracer());
+
 			Serilog.Log.Logger = new LoggerConfiguration()
 				.WithHordeConfig(HordeSettings)
 				.Enrich.FromLogContext()
+				.Enrich.WithOpenTracingContext()
 				.WriteTo.Console(HordeSettings)
 				.WriteTo.File(Path.Combine(LogDir.FullName, "Log.txt"), outputTemplate: "[{Timestamp:HH:mm:ss} {Level:w3}] {Indent}{Message:l}{NewLine}{Exception} [{SourceContext}]", rollingInterval: RollingInterval.Day, rollOnFileSizeLimit: true, fileSizeLimitBytes: 20 * 1024 * 1024, retainedFileCountLimit: 10)
 				.WriteTo.File(new JsonFormatter(renderMessage: true), Path.Combine(LogDir.FullName, "Log.json"), rollingInterval: RollingInterval.Day, rollOnFileSizeLimit: true, fileSizeLimitBytes: 20 * 1024 * 1024, retainedFileCountLimit: 10)
diff --git a/Engine/Source/Programs/Horde/HordeServer/Services/DatabaseService.cs b/Engine/Source/Programs/Horde/HordeServer/Services/DatabaseService.cs
--- a/Engine/Source/Programs/Horde/HordeServer/Services/DatabaseService.cs
+++ b/Engine/Source/Programs/Horde/HordeServer/Services/DatabaseService.cs
@@ -36,6 +36,8 @@ using HordeServer.Collections;
 using System.Diagnostics;
 using System.Text.RegularExpressions;
 using System.Net.NetworkInformation;
+using OpenTracing.Util;
+using OpenTracing;
 
 namespace HordeServer.Services
 {
@@ -120,7 +122,7 @@ namespace HordeServer.Services
 		/// </summary>
 		class MongoDatadogTracer : IEventSubscriber
 		{
-			readonly ConcurrentDictionary<int, Scope> Scopes = new ConcurrentDictionary<int, Scope>();
+			readonly ConcurrentDictionary<int, IScope> Scopes = new ConcurrentDictionary<int, IScope>();
 			private readonly ReflectionEventSubscriber _subscriber;
 
 			public MongoDatadogTracer()
@@ -135,7 +137,7 @@ namespace HordeServer.Services
 			
 			public void Handle(CommandStartedEvent e)
 			{
-				Scope Scope = Tracer.Instance.StartActive("mongodb.command");
+				IScope Scope = GlobalTracer.Instance.BuildSpan("mongodb.command").StartActive();
 				string? DbName = null;
 				string? CollectionName = null;
 				string? Query = null;
@@ -159,7 +161,7 @@ namespace HordeServer.Services
 						break;
 				}
 
-				Scope.Span.Type = "db";
+				Scope.Span.SetTag("Type", "db");
 				Scope.Span.SetTag("Command", e.CommandName);
 				Scope.Span.SetTag("DbName", DbName);
 				Scope.Span.SetTag("Collection", CollectionName);
@@ -169,7 +171,7 @@ namespace HordeServer.Services
 
 			public void Handle(CommandSucceededEvent e)
 			{
-				if (Scopes.TryGetValue(e.RequestId, out Scope? Scope))
+				if (Scopes.TryGetValue(e.RequestId, out IScope? Scope))
 				{
 					Scope.Dispose();
 					Scopes.Remove(e.RequestId, out _);
@@ -178,9 +180,9 @@ namespace HordeServer.Services
 			
 			public void Handle(CommandFailedEvent e)
 			{
-				if (Scopes.TryGetValue(e.RequestId, out Scope? Scope))
+				if (Scopes.TryGetValue(e.RequestId, out IScope? Scope))
 				{
-					Scope.Span.Error = true;
+					Scope.Span.SetTag("Error", true);
 					Scope.Dispose();
 					Scopes.Remove(e.RequestId, out _);
 				}
diff --git a/Engine/Source/Programs/Horde/HordeServer/Services/Impl/IssueService.cs b/Engine/Source/Programs/Horde/HordeServer/Services/Impl/IssueService.cs
--- a/Engine/Source/Programs/Horde/HordeServer/Services/Impl/IssueService.cs
+++ b/Engine/Source/Programs/Horde/HordeServer/Services/Impl/IssueService.cs
@@ -31,6 +31,8 @@ using StreamId = HordeServer.Utilities.StringId<HordeServer.Models.IStream>;
 using TemplateRefId = HordeServer.Utilities.StringId<HordeServer.Models.TemplateRef>;
 using HordeServer.IssueHandlers.Impl;
 using HordeServer.Notifications;
+using OpenTracing;
+using OpenTracing.Util;
 
 namespace HordeServer.Services.Impl
 {
@@ -535,7 +537,7 @@ namespace HordeServer.Services.Impl
 		[SuppressMessage("Design", "CA1031:Do not catch general exception types", Justification = "<Pending>")]
 		public async Task UpdateCompleteStep(IJob Job, IGraph Graph, SubResourceId BatchId, SubResourceId StepId)
 		{
-			using Scope Scope = Tracer.Instance.StartActive("UpdateCompleteStep");
+			using IScope Scope = GlobalTracer.Instance.BuildSpan("UpdateCompleteStep").StartActive();
 			Scope.Span.SetTag("JobId", Job.Id.ToString());
 			Scope.Span.SetTag("BatchId", BatchId.ToString());
 			Scope.Span.SetTag("StepId", StepId.ToString());
diff --git a/Engine/Source/Programs/Horde/HordeServer/Services/JobService.cs b/Engine/Source/Programs/Horde/HordeServer/Services/JobService.cs
--- a/Engine/Source/Programs/Horde/HordeServer/Services/JobService.cs
+++ b/Engine/Source/Programs/Horde/HordeServer/Services/JobService.cs
@@ -28,6 +28,8 @@ using StreamId = HordeServer.Utilities.StringId<HordeServer.Models.IStream>;
 using TemplateRefId = HordeServer.Utilities.StringId<HordeServer.Models.TemplateRef>;
 using System.Globalization;
 using HordeServer.Tasks.Impl;
+using OpenTracing.Util;
+using OpenTracing;
 
 namespace HordeServer.Services
 {
@@ -866,7 +868,7 @@ namespace HordeServer.Services
 			// Update the step
 			if (await Jobs.TryUpdateStepAsync(Job, Graph, BatchId, StepId, NewState, NewOutcome, NewAbortRequested, NewAbortByUser, NewLogId, NewTriggerId, NewRetryByUser, NewPriority, NewReports, NewProperties))
 			{
-				using Scope DdScope = Tracer.Instance.StartActive("TryUpdateStepAsync");
+				using IScope DdScope = GlobalTracer.Instance.BuildSpan("TryUpdateStepAsync").StartActive();
 				DdScope.Span.SetTag("JobId", Job.Id.ToString());
 				DdScope.Span.SetTag("BatchId", BatchId.ToString());
 				DdScope.Span.SetTag("StepId", StepId.ToString());
@@ -878,7 +880,7 @@ namespace HordeServer.Services
 					// Send any updates for modified badges
 					if (OldLabelStates != null)
 					{
-						using Scope _ = Tracer.Instance.StartActive("Send badge updates");
+						using IScope _ = GlobalTracer.Instance.BuildSpan("Send badge updates").StartActive();
 						IReadOnlyList<(LabelState, LabelOutcome)> NewLabelStates = Job.GetLabelStates(Graph);
 						OnLabelUpdate?.Invoke(Job, OldLabelStates, NewLabelStates);
 						await JobTaskSource.UpdateUgsBadges(Job, Graph, OldLabelStates, NewLabelStates);
@@ -908,14 +910,14 @@ namespace HordeServer.Services
 					// Notify subscribers
 					if (NewState == JobStepState.Completed)
 					{
-						using Scope _ = Tracer.Instance.StartActive("Notify subscribers");
+						using IScope _ = GlobalTracer.Instance.BuildSpan("Notify subscribers").StartActive();
 						OnJobStepComplete?.Invoke(Job, Graph, BatchId, StepId);
 					}
 
 					// Create any downstream jobs
 					if (NewState == JobStepState.Completed && NewOutcome != JobStepOutcome.Failure)
 					{
-						using Scope _ = Tracer.Instance.StartActive("Create downstream jobs");
+						using IScope _ = GlobalTracer.Instance.BuildSpan("Create downstream jobs").StartActive();
 						for (int Idx = 0; Idx < Job.ChainedJobs.Count; Idx++)
 						{
 							IChainedJob JobTrigger = Job.ChainedJobs[Idx];
@@ -933,7 +935,7 @@ namespace HordeServer.Services
 					// Update the jobstep ref if it completed
 					if (NewState == JobStepState.Running || NewState == JobStepState.Completed || NewState == JobStepState.Aborted)
 					{
-						using Scope _ = Tracer.Instance.StartActive("Update job step ref");
+						using IScope _ = GlobalTracer.Instance.BuildSpan("Update job step ref").StartActive();
 						if (Job.TryGetBatch(BatchId, out IJobStepBatch? Batch) && Batch.TryGetStep(StepId, out IJobStep? Step) && Step.StartTimeUtc != null)
 						{
 							await JobStepRefs.UpdateAsync(Job, Batch, Step, Graph);
@@ -945,7 +947,7 @@ namespace HordeServer.Services
 					{
 						if (IssueService != null)
 						{
-							using Scope _ = Tracer.Instance.StartActive("Update issues (V2)");
+							using IScope _ = GlobalTracer.Instance.BuildSpan("Update issues (V2)").StartActive();
 							try
 							{
 								await IssueService.UpdateCompleteStep(Job, Graph, BatchId, StepId);
@@ -958,7 +960,7 @@ namespace HordeServer.Services
 					}
 				}
 
-				using (Scope DispatchScope = Tracer.Instance.StartActive("Update queued jobs"))
+				using (IScope DispatchScope = GlobalTracer.Instance.BuildSpan("Update queued jobs").StartActive())
 				{
 					// Notify the dispatch service that the job has changed
 					JobTaskSource.UpdateQueuedJob(Job, Graph);
diff --git a/Engine/Source/Programs/Horde/HordeServer/Services/LogFileService.cs b/Engine/Source/Programs/Horde/HordeServer/Services/LogFileService.cs
--- a/Engine/Source/Programs/Horde/HordeServer/Services/LogFileService.cs
+++ b/Engine/Source/Programs/Horde/HordeServer/Services/LogFileService.cs
@@ -31,6 +31,8 @@ using Datadog.Trace;
 
 using ILogger = Microsoft.Extensions.Logging.ILogger;
 using Stream = System.IO.Stream;
+using OpenTracing.Util;
+using OpenTracing;
 
 [assembly: InternalsVisibleTo("HordeServerTests")]
 
@@ -575,7 +577,7 @@ namespace HordeServer.Services
 		/// <inheritdoc/>
 		public async Task<ILogFile?> WriteLogDataAsync(ILogFile LogFile, long Offset, int LineIndex, ReadOnlyMemory<byte> Data, bool Flush, int MaxChunkLength, int MaxSubChunkLineCount)
 		{
-			using Scope Scope = Tracer.Instance.StartActive("WriteLogDataAsync");
+			using IScope Scope = GlobalTracer.Instance.BuildSpan("WriteLogDataAsync").StartActive();
 			Scope.Span.SetTag("LogId", LogFile.Id.ToString());
 			Scope.Span.SetTag("Offset", Offset.ToString(CultureInfo.InvariantCulture));
 			Scope.Span.SetTag("Length", Data.Length.ToString(CultureInfo.InvariantCulture));
@@ -878,7 +880,7 @@ namespace HordeServer.Services
 		/// <inheritdoc/>
 		public async Task<ILogEventData> GetEventDataAsync(ILogFile LogFile, int LineIndex, int LineCount)
 		{
-			using Scope Scope = Tracer.Instance.StartActive("GetEventDataAsync");
+			using IScope Scope = GlobalTracer.Instance.BuildSpan("GetEventDataAsync").StartActive();
 			Scope.Span.SetTag("LogId", LogFile.Id.ToString());
 			Scope.Span.SetTag("LineIndex", LineIndex.ToString(CultureInfo.InvariantCulture));
 			Scope.Span.SetTag("LineCount", LineCount.ToString(CultureInfo.InvariantCulture));
@@ -1246,7 +1248,7 @@ namespace HordeServer.Services
 			}
 
 			// Save stats for the index creation
-			using Scope Scope = Tracer.Instance.StartActive("CreateIndexAsync");
+			using IScope Scope = GlobalTracer.Instance.BuildSpan("CreateIndexAsync").StartActive();
 			Scope.Span.SetTag("LogId", LogFile.Id.ToString());
 			Scope.Span.SetTag("Length", (LastChunk.Offset + LastChunk.Length).ToString(CultureInfo.InvariantCulture));
 
@@ -1523,7 +1525,7 @@ namespace HordeServer.Services
 		{
 			Stopwatch Timer = Stopwatch.StartNew();
 
-			using Scope Scope = Tracer.Instance.StartActive("SearchLogDataAsync");
+			using IScope Scope = GlobalTracer.Instance.BuildSpan("SearchLogDataAsync").StartActive();
 			Scope.Span.SetTag("LogId", LogFile.Id.ToString());
 			Scope.Span.SetTag("Text", Text);
 			Scope.Span.SetTag("Count", Count.ToString(CultureInfo.InvariantCulture));
@@ -1553,7 +1555,7 @@ namespace HordeServer.Services
 				LogIndexData? IndexData = await ReadIndexAsync(LogFile, LogFile.IndexLength.Value);
 				if(IndexData != null && FirstLine < IndexData.LineCount)
 				{
-					using Scope IndexScope = Tracer.Instance.StartActive("Indexed");
+					using IScope IndexScope = GlobalTracer.Instance.BuildSpan("Indexed").StartActive();
 					IndexScope.Span.SetTag("LineCount", IndexData.LineCount.ToString(CultureInfo.InvariantCulture));
 
 					foreach(int LineIndex in IndexData.Search(FirstLine, SearchText, SearchStats))
diff --git a/Engine/Source/Programs/Horde/HordeServer/Services/ScheduleService.cs b/Engine/Source/Programs/Horde/HordeServer/Services/ScheduleService.cs
--- a/Engine/Source/Programs/Horde/HordeServer/Services/ScheduleService.cs
+++ b/Engine/Source/Programs/Horde/HordeServer/Services/ScheduleService.cs
@@ -15,6 +15,8 @@ using Microsoft.Extensions.Logging;
 using Microsoft.Extensions.Options;
 using MongoDB.Bson;
 using MongoDB.Driver;
+using OpenTracing;
+using OpenTracing.Util;
 using System;
 using System.Collections.Generic;
 using System.Diagnostics;
@@ -126,7 +128,7 @@ namespace HordeServer.Services
 			if (!DowntimeService.IsDowntimeActive)
 			{
 				Logger.LogInformation("Updating scheduled triggers...");
-				using (Scope Scope = Tracer.Instance.StartActive("ScheduleService Tick"))
+				using (IScope Scope = GlobalTracer.Instance.BuildSpan("ScheduleService Tick").StartActive())
 				{
 					NextTickTime = await UpdateSchedulesAsync();
 				}
diff --git a/Engine/Source/Programs/Horde/HordeServer/Utilities/CompressionExtensions.cs b/Engine/Source/Programs/Horde/HordeServer/Utilities/CompressionExtensions.cs
--- a/Engine/Source/Programs/Horde/HordeServer/Utilities/CompressionExtensions.cs
+++ b/Engine/Source/Programs/Horde/HordeServer/Utilities/CompressionExtensions.cs
@@ -3,6 +3,8 @@
 using Datadog.Trace;
 using EpicGames.Core;
 using ICSharpCode.SharpZipLib.BZip2;
+using OpenTracing;
+using OpenTracing.Util;
 using System;
 using System.Buffers.Binary;
 using System.Collections.Generic;
@@ -24,7 +26,7 @@ namespace HordeServer.Utilities
 		/// <returns>The compressed data</returns>
 		public static byte[] CompressBzip2(this ReadOnlyMemory<byte> Memory)
 		{
-			using Scope Scope = Tracer.Instance.StartActive("CompressBzip2");
+			using IScope Scope = GlobalTracer.Instance.BuildSpan("CompressBzip2").StartActive();
 			Scope.Span.SetTag("DecompressedSize", Memory.Length.ToString(CultureInfo.InvariantCulture));
 
 			byte[] CompressedData;
@@ -54,7 +56,7 @@ namespace HordeServer.Utilities
 		{
 			int DecompressedSize = BinaryPrimitives.ReadInt32LittleEndian(Memory.Span);
 
-			using Scope Scope = Tracer.Instance.StartActive("DecompressBzip2");
+			using IScope Scope = GlobalTracer.Instance.BuildSpan("DecompressBzip2").StartActive();
 			Scope.Span.SetTag("CompressedSize", Memory.Length.ToString(CultureInfo.InvariantCulture));
 			Scope.Span.SetTag("DecompressedSize", DecompressedSize.ToString(CultureInfo.InvariantCulture));
 
diff --git a/Engine/Source/Programs/Horde/HordeServer/Utilities/TickedBackgroundService.cs b/Engine/Source/Programs/Horde/HordeServer/Utilities/TickedBackgroundService.cs
--- a/Engine/Source/Programs/Horde/HordeServer/Utilities/TickedBackgroundService.cs
+++ b/Engine/Source/Programs/Horde/HordeServer/Utilities/TickedBackgroundService.cs
@@ -4,6 +4,8 @@ using Datadog.Trace;
 using EpicGames.Core;
 using Microsoft.Extensions.Hosting;
 using Microsoft.Extensions.Logging;
+using OpenTracing;
+using OpenTracing.Util;
 using System;
 using System.Collections.Generic;
 using System.Diagnostics;
@@ -73,7 +75,7 @@ namespace HordeServer.Utilities
 
 					// Run the tick method
 					Timer.Restart();
-					using (Scope Scope = Tracer.Instance.StartActive($"{GetType().Name}.{nameof(TickAsync)}"))
+					using (IScope Scope = GlobalTracer.Instance.BuildSpan($"{GetType().Name}.{nameof(TickAsync)}").StartActive())
 					{
 						try
 						{
