Horde: Fix tests.

diff --git a/Engine/Source/Programs/Horde/HordeAgent/Services/WorkerService.cs b/Engine/Source/Programs/Horde/HordeAgent/Services/WorkerService.cs
--- a/Engine/Source/Programs/Horde/HordeAgent/Services/WorkerService.cs
+++ b/Engine/Source/Programs/Horde/HordeAgent/Services/WorkerService.cs
@@ -664,13 +664,13 @@ namespace HordeAgent.Services
 				{
 					LeaseInfo.Lease.State = LeaseState.Cancelled;
 					LeaseInfo.Lease.Outcome = LeaseOutcome.Failed;
-					LeaseInfo.Lease.Result = null;
+					LeaseInfo.Lease.Output = null;
 				}
 				else
 				{
 					LeaseInfo.Lease.State = (Result.Outcome == LeaseOutcome.Cancelled) ? LeaseState.Cancelled : LeaseState.Completed;
 					LeaseInfo.Lease.Outcome = Result.Outcome;
-					LeaseInfo.Lease.Result = ByteString.CopyFrom(Result.Output);
+					LeaseInfo.Lease.Output = ByteString.CopyFrom(Result.Output);
 				}
 				Logger.LogInformation("Transitioning lease {LeaseId} to {State}, outcome={Outcome}", LeaseInfo.Lease.Id, LeaseInfo.Lease.State, LeaseInfo.Lease.Outcome);
 			}
diff --git a/Engine/Source/Programs/Horde/HordeAgentTests/TaskTest.cs b/Engine/Source/Programs/Horde/HordeAgentTests/TaskTest.cs
--- a/Engine/Source/Programs/Horde/HordeAgentTests/TaskTest.cs
+++ b/Engine/Source/Programs/Horde/HordeAgentTests/TaskTest.cs
@@ -1,4 +1,4 @@
-﻿// Copyright Epic Games, Inc. All Rights Reserved.
+// Copyright Epic Games, Inc. All Rights Reserved.
 
 using System;
 using System.IO;
@@ -90,7 +90,7 @@ namespace HordeAgentTests
 			
 			ActionTask ActionTask = CreateActionTask(InstanceName, Cas);
 			Assert.IsNull(ActionRpc.ActionResultRequest);
-			LeaseOutcome Outcome = await Ws.HandleLeasePayloadAsync(RpcConnection, "my-agent-id", CreateLeaseInfo(ActionTask));
+			LeaseOutcome Outcome = (await Ws.HandleLeasePayloadAsync(RpcConnection, "my-agent-id", CreateLeaseInfo(ActionTask))).Outcome;
 			Assert.AreEqual(LeaseOutcome.Success, Outcome);
 			Assert.IsNotNull(ActionRpc.ActionResultRequest);
 			Assert.IsNull(ActionRpc.ActionResultRequest!.Error);
@@ -103,7 +103,7 @@ namespace HordeAgentTests
 			if (!RuntimeInformation.IsOSPlatform(OSPlatform.Windows)) return;
 			
 			ActionTask ActionTask = CreateActionTask(InstanceName, Cas);
-			LeaseOutcome Outcome = await Ws.HandleLeasePayloadAsync(RpcConnection, "my-agent-id", CreateLeaseInfo(ActionTask));
+			LeaseOutcome Outcome = (await Ws.HandleLeasePayloadAsync(RpcConnection, "my-agent-id", CreateLeaseInfo(ActionTask))).Outcome;
 			Assert.AreEqual(LeaseOutcome.Failed, Outcome);
 			Assert.IsNotNull(ActionRpc.ActionResultRequest);
 			Assert.IsNull(ActionRpc.ActionResultRequest!.Error);
diff --git a/Engine/Source/Programs/Horde/HordeAgentTests/WorkerServiceTest.cs b/Engine/Source/Programs/Horde/HordeAgentTests/WorkerServiceTest.cs
--- a/Engine/Source/Programs/Horde/HordeAgentTests/WorkerServiceTest.cs
+++ b/Engine/Source/Programs/Horde/HordeAgentTests/WorkerServiceTest.cs
@@ -120,8 +120,8 @@ namespace HordeAgentTests
 
 			WorkerService Ws = GetWorkerService((a, b, c) => Executor);
 			Ws.StepAbortPollInterval = TimeSpan.FromMilliseconds(1);
-			LeaseOutcome Outcome = await Ws.ExecuteJobAsync(RpcConnection, "agentId1", "leaseId1", ExecuteJobTask,
-				WorkerLogger, Token);
+			LeaseOutcome Outcome = (await Ws.ExecuteJobAsync(RpcConnection, "agentId1", "leaseId1", ExecuteJobTask,
+				WorkerLogger, Token)).Outcome;
 
 			Assert.AreEqual(LeaseOutcome.Success, Outcome);
 			Assert.AreEqual(4, Client.UpdateStepRequests.Count);
