Horde: Move type aliases into namespace scope for simplicity. Also prevents weirdness with quick refactoring screwing up the order when adding new using declarations.

diff --git a/Engine/Source/Programs/Horde/HordeServer/Api/Agents.cs b/Engine/Source/Programs/Horde/HordeServer/Api/Agents.cs
--- a/Engine/Source/Programs/Horde/HordeServer/Api/Agents.cs
+++ b/Engine/Source/Programs/Horde/HordeServer/Api/Agents.cs
@@ -10,10 +10,10 @@ using System.ComponentModel.DataAnnotations;
 using System.Linq;
 using System.Threading.Tasks;
 
-using PoolId = HordeServer.Utilities.StringId<HordeServer.Models.IPool>;
-
 namespace HordeServer.Api
 {
+	using PoolId = StringId<IPool>;
+
 	/// <summary>
 	/// Parameters to register a new agent
 	/// </summary>
diff --git a/Engine/Source/Programs/Horde/HordeServer/Api/Events.cs b/Engine/Source/Programs/Horde/HordeServer/Api/Events.cs
--- a/Engine/Source/Programs/Horde/HordeServer/Api/Events.cs
+++ b/Engine/Source/Programs/Horde/HordeServer/Api/Events.cs
@@ -2,6 +2,8 @@
 
 using EpicGames.Core;
 using HordeCommon;
+using HordeServer.Models;
+using HordeServer.Utilities;
 using Microsoft.AspNetCore.Mvc;
 using System;
 using System.Collections.Generic;
@@ -9,11 +11,11 @@ using System.ComponentModel.DataAnnotations;
 using System.Linq;
 using System.Threading.Tasks;
 
-using StreamId = HordeServer.Utilities.StringId<HordeServer.Models.IStream>;
-using TemplateRefId = HordeServer.Utilities.StringId<HordeServer.Models.TemplateRef>;
-
 namespace HordeServer.Api
 {
+	using StreamId = StringId<IStream>;
+	using TemplateRefId = StringId<TemplateRef>;
+
 	/// <summary>
 	/// Base class for event records
 	/// </summary>
diff --git a/Engine/Source/Programs/Horde/HordeServer/Api/Issues.cs b/Engine/Source/Programs/Horde/HordeServer/Api/Issues.cs
--- a/Engine/Source/Programs/Horde/HordeServer/Api/Issues.cs
+++ b/Engine/Source/Programs/Horde/HordeServer/Api/Issues.cs
@@ -11,11 +11,11 @@ using System.Linq;
 using System.Text.RegularExpressions;
 using System.Threading.Tasks;
 
-using StreamId = HordeServer.Utilities.StringId<HordeServer.Models.IStream>;
-using TemplateRefId = HordeServer.Utilities.StringId<HordeServer.Models.TemplateRef>;
-
 namespace HordeServer.Api
 {
+	using StreamId = StringId<IStream>;
+	using TemplateRefId = StringId<TemplateRef>;
+
 	/// <summary>
 	/// Identifies a particular changelist and job
 	/// </summary>
diff --git a/Engine/Source/Programs/Horde/HordeServer/Api/Schedules.cs b/Engine/Source/Programs/Horde/HordeServer/Api/Schedules.cs
--- a/Engine/Source/Programs/Horde/HordeServer/Api/Schedules.cs
+++ b/Engine/Source/Programs/Horde/HordeServer/Api/Schedules.cs
@@ -1,6 +1,7 @@
 // Copyright Epic Games, Inc. All Rights Reserved.
 
 using HordeServer.Models;
+using HordeServer.Utilities;
 using System;
 using System.Collections.Generic;
 using System.ComponentModel.DataAnnotations;
@@ -11,10 +12,10 @@ using System.Text.Json;
 using System.Text.Json.Serialization;
 using System.Threading.Tasks;
 
-using TemplateRefId = HordeServer.Utilities.StringId<HordeServer.Models.TemplateRef>;
-
 namespace HordeServer.Api
 {
+	using TemplateRefId = StringId<TemplateRef>;
+
 	/// <summary>
 	/// Parameters to create a new schedule
 	/// </summary>
diff --git a/Engine/Source/Programs/Horde/HordeServer/Api/Streams.cs b/Engine/Source/Programs/Horde/HordeServer/Api/Streams.cs
--- a/Engine/Source/Programs/Horde/HordeServer/Api/Streams.cs
+++ b/Engine/Source/Programs/Horde/HordeServer/Api/Streams.cs
@@ -9,11 +9,10 @@ using System.ComponentModel.DataAnnotations;
 using System.Linq;
 using System.Threading.Tasks;
 
-using ProjectId = HordeServer.Utilities.StringId<HordeServer.Models.IProject>;
-using TemplateRefId = HordeServer.Utilities.StringId<HordeServer.Models.TemplateRef>;
-
 namespace HordeServer.Api
 {
+	using TemplateRefId = StringId<TemplateRef>;
+
 	/// <summary>
 	/// Query selecting the base changelist to use
 	/// </summary>
diff --git a/Engine/Source/Programs/Horde/HordeServer/Collections/DatadogTracingCollection.cs b/Engine/Source/Programs/Horde/HordeServer/Collections/DatadogTracingCollection.cs
--- a/Engine/Source/Programs/Horde/HordeServer/Collections/DatadogTracingCollection.cs
+++ b/Engine/Source/Programs/Horde/HordeServer/Collections/DatadogTracingCollection.cs
@@ -3,7 +3,6 @@
 using System.Collections.Generic;
 using System.Threading;
 using System.Threading.Tasks;
-using Datadog.Trace;
 using MongoDB.Bson;
 using MongoDB.Bson.Serialization;
 using MongoDB.Driver;
diff --git a/Engine/Source/Programs/Horde/HordeServer/Collections/IAgentCollection.cs b/Engine/Source/Programs/Horde/HordeServer/Collections/IAgentCollection.cs
--- a/Engine/Source/Programs/Horde/HordeServer/Collections/IAgentCollection.cs
+++ b/Engine/Source/Programs/Horde/HordeServer/Collections/IAgentCollection.cs
@@ -16,11 +16,11 @@ using System.Linq;
 using System.Linq.Expressions;
 using System.Threading.Tasks;
 
-using PoolId = HordeServer.Utilities.StringId<HordeServer.Models.IPool>;
-using AgentSoftwareChannelName = HordeServer.Utilities.StringId<HordeServer.Services.AgentSoftwareChannels>;
-
 namespace HordeServer.Collections
 {
+	using PoolId = StringId<IPool>;
+	using AgentSoftwareChannelName = StringId<AgentSoftwareChannels>;
+
 	/// <summary>
 	/// Interface for a collection of agent documents
 	/// </summary>
diff --git a/Engine/Source/Programs/Horde/HordeServer/Collections/IDeviceCollection.cs b/Engine/Source/Programs/Horde/HordeServer/Collections/IDeviceCollection.cs
--- a/Engine/Source/Programs/Horde/HordeServer/Collections/IDeviceCollection.cs
+++ b/Engine/Source/Programs/Horde/HordeServer/Collections/IDeviceCollection.cs
@@ -1,18 +1,17 @@
 // Copyright Epic Games, Inc. All Rights Reserved.
 
-using System.Threading.Tasks;
-
 using HordeServer.Models;
+using HordeServer.Utilities;
 using MongoDB.Bson;
 using System.Collections.Generic;
-
-using DeviceId = HordeServer.Utilities.StringId<HordeServer.Models.IDevice>;
-using DevicePlatformId = HordeServer.Utilities.StringId<HordeServer.Models.IDevicePlatform>;
-using DevicePoolId = HordeServer.Utilities.StringId<HordeServer.Models.IDevicePool>;
-using System;
+using System.Threading.Tasks;
 
 namespace HordeServer.Collections
 {
+	using DeviceId = StringId<IDevice>;
+	using DevicePlatformId = StringId<IDevicePlatform>;
+	using DevicePoolId = StringId<IDevicePool>;
+
 	/// <summary>
 	/// Device reservation request data
 	/// </summary>
diff --git a/Engine/Source/Programs/Horde/HordeServer/Collections/IIssueCollection.cs b/Engine/Source/Programs/Horde/HordeServer/Collections/IIssueCollection.cs
--- a/Engine/Source/Programs/Horde/HordeServer/Collections/IIssueCollection.cs
+++ b/Engine/Source/Programs/Horde/HordeServer/Collections/IIssueCollection.cs
@@ -20,13 +20,12 @@ using System.Linq.Expressions;
 using System.Text.RegularExpressions;
 using System.Threading.Tasks;
 using System.Transactions;
-using Datadog.Trace;
-
-using StreamId = HordeServer.Utilities.StringId<HordeServer.Models.IStream>;
-using TemplateRefId = HordeServer.Utilities.StringId<HordeServer.Models.TemplateRef>;
 
 namespace HordeServer.Collections
 {
+	using StreamId = StringId<IStream>;
+	using TemplateRefId = StringId<TemplateRef>;
+
 	/// <summary>
 	/// Fingerprint for an issue
 	/// </summary>
diff --git a/Engine/Source/Programs/Horde/HordeServer/Collections/IJobCollection.cs b/Engine/Source/Programs/Horde/HordeServer/Collections/IJobCollection.cs
--- a/Engine/Source/Programs/Horde/HordeServer/Collections/IJobCollection.cs
+++ b/Engine/Source/Programs/Horde/HordeServer/Collections/IJobCollection.cs
@@ -1,6 +1,5 @@
 // Copyright Epic Games, Inc. All Rights Reserved.
 
-using Datadog.Trace;
 using EpicGames.Core;
 using HordeServer.Api;
 using HordeCommon;
@@ -18,12 +17,12 @@ using System.IO;
 using System.Linq;
 using System.Threading.Tasks;
 
-using StreamId = HordeServer.Utilities.StringId<HordeServer.Models.IStream>;
-using TemplateRefId = HordeServer.Utilities.StringId<HordeServer.Models.TemplateRef>;
-using PoolId = HordeServer.Utilities.StringId<HordeServer.Models.IPool>;
-
 namespace HordeServer.Collections
 {
+	using PoolId = StringId<IPool>;
+	using StreamId = StringId<IStream>;
+	using TemplateRefId = StringId<TemplateRef>;
+
 	/// <summary>
 	/// Interface for a collection of job documents
 	/// </summary>
diff --git a/Engine/Source/Programs/Horde/HordeServer/Collections/IJobStepRefCollection.cs b/Engine/Source/Programs/Horde/HordeServer/Collections/IJobStepRefCollection.cs
--- a/Engine/Source/Programs/Horde/HordeServer/Collections/IJobStepRefCollection.cs
+++ b/Engine/Source/Programs/Horde/HordeServer/Collections/IJobStepRefCollection.cs
@@ -10,14 +10,14 @@ using System;
 using System.Collections.Generic;
 using System.Linq;
 using System.Threading.Tasks;
-
-using StreamId = HordeServer.Utilities.StringId<HordeServer.Models.IStream>;
-using PoolId = HordeServer.Utilities.StringId<HordeServer.Models.IPool>;
-using TemplateRefId = HordeServer.Utilities.StringId<HordeServer.Models.TemplateRef>;
 using HordeServer.Utilities;
 
 namespace HordeServer.Collections
 {
+	using PoolId = StringId<IPool>;
+	using StreamId = StringId<IStream>;
+	using TemplateRefId = StringId<TemplateRef>;
+
 	/// <summary>
 	/// Interface for a collection of JobStepRef documents
 	/// </summary>
diff --git a/Engine/Source/Programs/Horde/HordeServer/Collections/ILeaseCollection.cs b/Engine/Source/Programs/Horde/HordeServer/Collections/ILeaseCollection.cs
--- a/Engine/Source/Programs/Horde/HordeServer/Collections/ILeaseCollection.cs
+++ b/Engine/Source/Programs/Horde/HordeServer/Collections/ILeaseCollection.cs
@@ -12,11 +12,11 @@ using System.Collections.Generic;
 using System.Linq;
 using System.Threading.Tasks;
 
-using PoolId = HordeServer.Utilities.StringId<HordeServer.Models.IPool>;
-using StreamId = HordeServer.Utilities.StringId<HordeServer.Models.IStream>;
-
 namespace HordeServer.Collections
 {
+	using PoolId = StringId<IPool>;
+	using StreamId = StringId<IStream>;
+
 	/// <summary>
 	/// Interface for a collection of lease documents
 	/// </summary>
diff --git a/Engine/Source/Programs/Horde/HordeServer/Collections/IPoolCollection.cs b/Engine/Source/Programs/Horde/HordeServer/Collections/IPoolCollection.cs
--- a/Engine/Source/Programs/Horde/HordeServer/Collections/IPoolCollection.cs
+++ b/Engine/Source/Programs/Horde/HordeServer/Collections/IPoolCollection.cs
@@ -13,10 +13,10 @@ using System.Linq;
 using System.Linq.Expressions;
 using System.Threading.Tasks;
 
-using PoolId = HordeServer.Utilities.StringId<HordeServer.Models.IPool>;
-
 namespace HordeServer.Collections
 {
+	using PoolId = StringId<IPool>;
+
 	/// <summary>
 	/// Collection of pool documents
 	/// </summary>
diff --git a/Engine/Source/Programs/Horde/HordeServer/Collections/IProjectCollection.cs b/Engine/Source/Programs/Horde/HordeServer/Collections/IProjectCollection.cs
--- a/Engine/Source/Programs/Horde/HordeServer/Collections/IProjectCollection.cs
+++ b/Engine/Source/Programs/Horde/HordeServer/Collections/IProjectCollection.cs
@@ -15,10 +15,10 @@ using System.Linq;
 using System.Security.Claims;
 using System.Threading.Tasks;
 
-using ProjectId = HordeServer.Utilities.StringId<HordeServer.Models.IProject>;
-
 namespace HordeServer.Collections
 {
+	using ProjectId = StringId<IProject>;
+
 	/// <summary>
 	/// Interface for a collection of project documents
 	/// </summary>
diff --git a/Engine/Source/Programs/Horde/HordeServer/Collections/IStreamCollection.cs b/Engine/Source/Programs/Horde/HordeServer/Collections/IStreamCollection.cs
--- a/Engine/Source/Programs/Horde/HordeServer/Collections/IStreamCollection.cs
+++ b/Engine/Source/Programs/Horde/HordeServer/Collections/IStreamCollection.cs
@@ -13,12 +13,12 @@ using System.Diagnostics.CodeAnalysis;
 using System.Linq;
 using System.Threading.Tasks;
 
-using ProjectId = HordeServer.Utilities.StringId<HordeServer.Models.IProject>;
-using StreamId = HordeServer.Utilities.StringId<HordeServer.Models.IStream>;
-using TemplateRefId = HordeServer.Utilities.StringId<HordeServer.Models.TemplateRef>;
-
 namespace HordeServer.Collections
 {
+	using ProjectId = StringId<IProject>;
+	using StreamId = StringId<IStream>;
+	using TemplateRefId = StringId<TemplateRef>;
+
 	/// <summary>
 	/// Collection of stream documents
 	/// </summary>
diff --git a/Engine/Source/Programs/Horde/HordeServer/Collections/ISubscriptionCollection.cs b/Engine/Source/Programs/Horde/HordeServer/Collections/ISubscriptionCollection.cs
--- a/Engine/Source/Programs/Horde/HordeServer/Collections/ISubscriptionCollection.cs
+++ b/Engine/Source/Programs/Horde/HordeServer/Collections/ISubscriptionCollection.cs
@@ -17,11 +17,11 @@ using System.Reflection.Metadata.Ecma335;
 using System.Text;
 using System.Threading.Tasks;
 
-using StreamId = HordeServer.Utilities.StringId<HordeServer.Models.IStream>;
-using TemplateRefId = HordeServer.Utilities.StringId<HordeServer.Models.TemplateRef>;
-
 namespace HordeServer.Collections
 {
+	using StreamId = StringId<IStream>;
+	using TemplateRefId = StringId<TemplateRef>;
+
 	/// <summary>
 	/// Information for creating a new subscription
 	/// </summary>
diff --git a/Engine/Source/Programs/Horde/HordeServer/Collections/ITelemetryCollection.cs b/Engine/Source/Programs/Horde/HordeServer/Collections/ITelemetryCollection.cs
--- a/Engine/Source/Programs/Horde/HordeServer/Collections/ITelemetryCollection.cs
+++ b/Engine/Source/Programs/Horde/HordeServer/Collections/ITelemetryCollection.cs
@@ -7,11 +7,11 @@ using System.Collections.Generic;
 using System.Linq;
 using System.Threading.Tasks;
 
-using PoolId = HordeServer.Utilities.StringId<HordeServer.Models.IPool>;
-using StreamId = HordeServer.Utilities.StringId<HordeServer.Models.IStream>;
-
 namespace HordeServer.Collections
 {
+	using PoolId = StringId<IPool>;
+	using StreamId = StringId<IStream>;
+
 	/// <summary>
 	/// Collection of utilization collection
 	/// </summary>
diff --git a/Engine/Source/Programs/Horde/HordeServer/Collections/ITestDataCollection.cs b/Engine/Source/Programs/Horde/HordeServer/Collections/ITestDataCollection.cs
--- a/Engine/Source/Programs/Horde/HordeServer/Collections/ITestDataCollection.cs
+++ b/Engine/Source/Programs/Horde/HordeServer/Collections/ITestDataCollection.cs
@@ -2,6 +2,7 @@
 
 using HordeServer.Models;
 using HordeServer.Services;
+using HordeServer.Utilities;
 using MongoDB.Bson;
 using MongoDB.Bson.Serialization.Attributes;
 using MongoDB.Driver;
@@ -11,12 +12,12 @@ using System.Diagnostics.CodeAnalysis;
 using System.Linq;
 using System.Threading.Tasks;
 
-using ProjectId = HordeServer.Utilities.StringId<HordeServer.Models.IProject>;
-using StreamId = HordeServer.Utilities.StringId<HordeServer.Models.IStream>;
-using TemplateRefId = HordeServer.Utilities.StringId<HordeServer.Models.TemplateRef>;
-
 namespace HordeServer.Collections
 {
+	using ProjectId = StringId<IProject>;
+	using StreamId = StringId<IStream>;
+	using TemplateRefId = StringId<TemplateRef>;
+
 	/// <summary>
 	/// Collection of test data documents
 	/// </summary>
diff --git a/Engine/Source/Programs/Horde/HordeServer/Collections/Impl/AgentCollection.cs b/Engine/Source/Programs/Horde/HordeServer/Collections/Impl/AgentCollection.cs
--- a/Engine/Source/Programs/Horde/HordeServer/Collections/Impl/AgentCollection.cs
+++ b/Engine/Source/Programs/Horde/HordeServer/Collections/Impl/AgentCollection.cs
@@ -1,10 +1,9 @@
 // Copyright Epic Games, Inc. All Rights Reserved.
 
 using Google.Protobuf.WellKnownTypes;
-using HordeServer.Api;
 using HordeCommon;
-using HordeServer.Models;
 using HordeCommon.Rpc.Tasks;
+using HordeServer.Models;
 using HordeServer.Services;
 using HordeServer.Utilities;
 using MongoDB.Bson;
@@ -16,11 +15,11 @@ using System.Linq;
 using System.Linq.Expressions;
 using System.Threading.Tasks;
 
-using PoolId = HordeServer.Utilities.StringId<HordeServer.Models.IPool>;
-using AgentSoftwareChannelName = HordeServer.Utilities.StringId<HordeServer.Services.AgentSoftwareChannels>;
-
 namespace HordeServer.Collections.Impl
 {
+	using AgentSoftwareChannelName = StringId<AgentSoftwareChannels>;
+	using PoolId = StringId<IPool>;
+
 	/// <summary>
 	/// Collection of agent documents
 	/// </summary>
diff --git a/Engine/Source/Programs/Horde/HordeServer/Collections/Impl/AgentSoftwareCollection.cs b/Engine/Source/Programs/Horde/HordeServer/Collections/Impl/AgentSoftwareCollection.cs
--- a/Engine/Source/Programs/Horde/HordeServer/Collections/Impl/AgentSoftwareCollection.cs
+++ b/Engine/Source/Programs/Horde/HordeServer/Collections/Impl/AgentSoftwareCollection.cs
@@ -17,10 +17,10 @@ using System.Linq;
 using System.Text.Json;
 using System.Threading.Tasks;
 
-using AgentSoftwareVersion = HordeServer.Utilities.StringId<HordeServer.Collections.IAgentSoftwareCollection>;
-
 namespace HordeServer.Collections.Impl
 {
+	using AgentSoftwareVersion = StringId<IAgentSoftwareCollection>;
+
 	/// <summary>
 	/// Collection of agent software documents
 	/// </summary>
diff --git a/Engine/Source/Programs/Horde/HordeServer/Collections/Impl/DeviceCollection.cs b/Engine/Source/Programs/Horde/HordeServer/Collections/Impl/DeviceCollection.cs
--- a/Engine/Source/Programs/Horde/HordeServer/Collections/Impl/DeviceCollection.cs
+++ b/Engine/Source/Programs/Horde/HordeServer/Collections/Impl/DeviceCollection.cs
@@ -13,12 +13,11 @@ using System.Linq;
 using MongoDB.Bson;
 using System.Collections.Generic;
 
-using DeviceId = HordeServer.Utilities.StringId<HordeServer.Models.IDevice>;
-using DevicePlatformId = HordeServer.Utilities.StringId<HordeServer.Models.IDevicePlatform>;
-using DevicePoolId = HordeServer.Utilities.StringId<HordeServer.Models.IDevicePool>;
-
 namespace HordeServer.Collections.Impl
 {
+	using DeviceId = StringId<IDevice>;
+	using DevicePlatformId = StringId<IDevicePlatform>;
+	using DevicePoolId = StringId<IDevicePool>;
 
 	/// <summary>
 	/// Collection of device documents
diff --git a/Engine/Source/Programs/Horde/HordeServer/Collections/Impl/IssueCollection.cs b/Engine/Source/Programs/Horde/HordeServer/Collections/Impl/IssueCollection.cs
--- a/Engine/Source/Programs/Horde/HordeServer/Collections/Impl/IssueCollection.cs
+++ b/Engine/Source/Programs/Horde/HordeServer/Collections/Impl/IssueCollection.cs
@@ -18,11 +18,11 @@ using System.Linq.Expressions;
 using System.Security.Claims;
 using System.Threading.Tasks;
 
-using StreamId = HordeServer.Utilities.StringId<HordeServer.Models.IStream>;
-using TemplateRefId = HordeServer.Utilities.StringId<HordeServer.Models.TemplateRef>;
-
 namespace HordeServer.Collections.Impl
 {
+	using StreamId = StringId<IStream>;
+	using TemplateRefId = StringId<TemplateRef>;
+
 	class IssueCollection : IIssueCollection
 	{
 		class UpdateTransaction
diff --git a/Engine/Source/Programs/Horde/HordeServer/Collections/Impl/JobCollection.cs b/Engine/Source/Programs/Horde/HordeServer/Collections/Impl/JobCollection.cs
--- a/Engine/Source/Programs/Horde/HordeServer/Collections/Impl/JobCollection.cs
+++ b/Engine/Source/Programs/Horde/HordeServer/Collections/Impl/JobCollection.cs
@@ -1,6 +1,5 @@
 // Copyright Epic Games, Inc. All Rights Reserved.
 
-using Datadog.Trace;
 using EpicGames.Core;
 using HordeServer.Api;
 using HordeCommon;
@@ -17,15 +16,15 @@ using System.Diagnostics.CodeAnalysis;
 using System.IO;
 using System.Linq;
 using System.Threading.Tasks;
-
-using StreamId = HordeServer.Utilities.StringId<HordeServer.Models.IStream>;
-using TemplateRefId = HordeServer.Utilities.StringId<HordeServer.Models.TemplateRef>;
-using PoolId = HordeServer.Utilities.StringId<HordeServer.Models.IPool>;
 using OpenTracing.Util;
 using OpenTracing;
 
 namespace HordeServer.Collections.Impl
 {
+	using StreamId = StringId<IStream>;
+	using TemplateRefId = StringId<TemplateRef>;
+	using PoolId = StringId<IPool>;
+
 	/// <summary>
 	/// Wrapper around the jobs collection in a mongo DB
 	/// </summary>
diff --git a/Engine/Source/Programs/Horde/HordeServer/Collections/Impl/JobStepRefCollection.cs b/Engine/Source/Programs/Horde/HordeServer/Collections/Impl/JobStepRefCollection.cs
--- a/Engine/Source/Programs/Horde/HordeServer/Collections/Impl/JobStepRefCollection.cs
+++ b/Engine/Source/Programs/Horde/HordeServer/Collections/Impl/JobStepRefCollection.cs
@@ -10,14 +10,14 @@ using System;
 using System.Collections.Generic;
 using System.Linq;
 using System.Threading.Tasks;
-
-using StreamId = HordeServer.Utilities.StringId<HordeServer.Models.IStream>;
-using TemplateRefId = HordeServer.Utilities.StringId<HordeServer.Models.TemplateRef>;
-using PoolId = HordeServer.Utilities.StringId<HordeServer.Models.IPool>;
 using HordeServer.Utilities;
 
 namespace HordeServer.Collections.Impl
 {
+	using StreamId = StringId<IStream>;
+	using TemplateRefId = StringId<TemplateRef>;
+	using PoolId = StringId<IPool>;
+
 	/// <summary>
 	/// Collection of JobStepRef documents
 	/// </summary>
diff --git a/Engine/Source/Programs/Horde/HordeServer/Collections/Impl/LeaseCollection.cs b/Engine/Source/Programs/Horde/HordeServer/Collections/Impl/LeaseCollection.cs
--- a/Engine/Source/Programs/Horde/HordeServer/Collections/Impl/LeaseCollection.cs
+++ b/Engine/Source/Programs/Horde/HordeServer/Collections/Impl/LeaseCollection.cs
@@ -12,11 +12,11 @@ using System.Collections.Generic;
 using System.Linq;
 using System.Threading.Tasks;
 
-using PoolId = HordeServer.Utilities.StringId<HordeServer.Models.IPool>;
-using StreamId = HordeServer.Utilities.StringId<HordeServer.Models.IStream>;
-
 namespace HordeServer.Collections.Impl
 {
+	using PoolId = StringId<IPool>;
+	using StreamId = StringId<IStream>;
+
 	/// <summary>
 	/// Collection of lease documents
 	/// </summary>
diff --git a/Engine/Source/Programs/Horde/HordeServer/Collections/Impl/PoolCollection.cs b/Engine/Source/Programs/Horde/HordeServer/Collections/Impl/PoolCollection.cs
--- a/Engine/Source/Programs/Horde/HordeServer/Collections/Impl/PoolCollection.cs
+++ b/Engine/Source/Programs/Horde/HordeServer/Collections/Impl/PoolCollection.cs
@@ -13,10 +13,10 @@ using System.Linq;
 using System.Linq.Expressions;
 using System.Threading.Tasks;
 
-using PoolId = HordeServer.Utilities.StringId<HordeServer.Models.IPool>;
-
 namespace HordeServer.Collections.Impl
 {
+	using PoolId = StringId<IPool>;
+
 	/// <summary>
 	/// Collection of pool documents
 	/// </summary>
diff --git a/Engine/Source/Programs/Horde/HordeServer/Collections/Impl/ProjectCollection.cs b/Engine/Source/Programs/Horde/HordeServer/Collections/Impl/ProjectCollection.cs
--- a/Engine/Source/Programs/Horde/HordeServer/Collections/Impl/ProjectCollection.cs
+++ b/Engine/Source/Programs/Horde/HordeServer/Collections/Impl/ProjectCollection.cs
@@ -15,10 +15,10 @@ using System.Linq;
 using System.Security.Claims;
 using System.Threading.Tasks;
 
-using ProjectId = HordeServer.Utilities.StringId<HordeServer.Models.IProject>;
-
 namespace HordeServer.Collections.Impl
 {
+	using ProjectId = StringId<IProject>;
+
 	/// <summary>
 	/// Wraps functionality for manipulating projects
 	/// </summary>
diff --git a/Engine/Source/Programs/Horde/HordeServer/Collections/Impl/StreamCollection.cs b/Engine/Source/Programs/Horde/HordeServer/Collections/Impl/StreamCollection.cs
--- a/Engine/Source/Programs/Horde/HordeServer/Collections/Impl/StreamCollection.cs
+++ b/Engine/Source/Programs/Horde/HordeServer/Collections/Impl/StreamCollection.cs
@@ -14,12 +14,12 @@ using System.IO;
 using System.Linq;
 using System.Threading.Tasks;
 
-using ProjectId = HordeServer.Utilities.StringId<HordeServer.Models.IProject>;
-using StreamId = HordeServer.Utilities.StringId<HordeServer.Models.IStream>;
-using TemplateRefId = HordeServer.Utilities.StringId<HordeServer.Models.TemplateRef>;
-
 namespace HordeServer.Collections.Impl
 {
+	using ProjectId = StringId<IProject>;
+	using StreamId = StringId<IStream>;
+	using TemplateRefId = StringId<TemplateRef>;
+
 	/// <summary>
 	/// Collection of stream documents
 	/// </summary>
diff --git a/Engine/Source/Programs/Horde/HordeServer/Collections/Impl/SubscriptionCollection.cs b/Engine/Source/Programs/Horde/HordeServer/Collections/Impl/SubscriptionCollection.cs
--- a/Engine/Source/Programs/Horde/HordeServer/Collections/Impl/SubscriptionCollection.cs
+++ b/Engine/Source/Programs/Horde/HordeServer/Collections/Impl/SubscriptionCollection.cs
@@ -17,11 +17,11 @@ using System.Reflection.Metadata.Ecma335;
 using System.Text;
 using System.Threading.Tasks;
 
-using StreamId = HordeServer.Utilities.StringId<HordeServer.Models.IStream>;
-using TemplateRefId = HordeServer.Utilities.StringId<HordeServer.Models.TemplateRef>;
-
 namespace HordeServer.Collections.Impl
 {
+	using StreamId = StringId<IStream>;
+	using TemplateRefId = StringId<TemplateRef>;
+
 	/// <summary>
 	/// Collection of subscription documents
 	/// </summary>
diff --git a/Engine/Source/Programs/Horde/HordeServer/Collections/Impl/TelemetryCollection.cs b/Engine/Source/Programs/Horde/HordeServer/Collections/Impl/TelemetryCollection.cs
--- a/Engine/Source/Programs/Horde/HordeServer/Collections/Impl/TelemetryCollection.cs
+++ b/Engine/Source/Programs/Horde/HordeServer/Collections/Impl/TelemetryCollection.cs
@@ -14,11 +14,11 @@ using System.Linq;
 using System.Threading;
 using System.Threading.Tasks;
 
-using PoolId = HordeServer.Utilities.StringId<HordeServer.Models.IPool>;
-using StreamId = HordeServer.Utilities.StringId<HordeServer.Models.IStream>;
-
 namespace HordeServer.Collections.Impl
 {
+	using PoolId = StringId<IPool>;
+	using StreamId = StringId<IStream>;
+
 	/// <summary>
 	/// Collection of utilization data
 	/// </summary>
diff --git a/Engine/Source/Programs/Horde/HordeServer/Collections/Impl/TestDataCollection.cs b/Engine/Source/Programs/Horde/HordeServer/Collections/Impl/TestDataCollection.cs
--- a/Engine/Source/Programs/Horde/HordeServer/Collections/Impl/TestDataCollection.cs
+++ b/Engine/Source/Programs/Horde/HordeServer/Collections/Impl/TestDataCollection.cs
@@ -2,6 +2,7 @@
 
 using HordeServer.Models;
 using HordeServer.Services;
+using HordeServer.Utilities;
 using MongoDB.Bson;
 using MongoDB.Bson.Serialization.Attributes;
 using MongoDB.Driver;
@@ -11,12 +12,12 @@ using System.Diagnostics.CodeAnalysis;
 using System.Linq;
 using System.Threading.Tasks;
 
-using ProjectId = HordeServer.Utilities.StringId<HordeServer.Models.IProject>;
-using StreamId = HordeServer.Utilities.StringId<HordeServer.Models.IStream>;
-using TemplateRefId = HordeServer.Utilities.StringId<HordeServer.Models.TemplateRef>;
-
 namespace HordeServer.Collections.Impl
 {
+	using ProjectId = StringId<IProject>;
+	using StreamId = StringId<IStream>;
+	using TemplateRefId = StringId<TemplateRef>;
+
 	/// <summary>
 	/// Collection of test data documents
 	/// </summary>
diff --git a/Engine/Source/Programs/Horde/HordeServer/Controllers/AgentSoftwareController.cs b/Engine/Source/Programs/Horde/HordeServer/Controllers/AgentSoftwareController.cs
--- a/Engine/Source/Programs/Horde/HordeServer/Controllers/AgentSoftwareController.cs
+++ b/Engine/Source/Programs/Horde/HordeServer/Controllers/AgentSoftwareController.cs
@@ -15,10 +15,10 @@ using HordeServer.Models;
 using MongoDB.Bson;
 using HordeServer.Utilities;
 
-using AgentSoftwareChannelName = HordeServer.Utilities.StringId<HordeServer.Services.AgentSoftwareChannels>;
-
 namespace HordeServer.Controllers
 {
+	using AgentSoftwareChannelName = StringId<AgentSoftwareChannels>;
+
 	/// <summary>
 	/// Controller for the /api/v1/software endpoint
 	/// </summary>
diff --git a/Engine/Source/Programs/Horde/HordeServer/Controllers/AgentsController.cs b/Engine/Source/Programs/Horde/HordeServer/Controllers/AgentsController.cs
--- a/Engine/Source/Programs/Horde/HordeServer/Controllers/AgentsController.cs
+++ b/Engine/Source/Programs/Horde/HordeServer/Controllers/AgentsController.cs
@@ -21,11 +21,11 @@ using System.Security.Claims;
 using System.Text.Json;
 using System.Threading.Tasks;
 
-using PoolId = HordeServer.Utilities.StringId<HordeServer.Models.IPool>;
-using AgentSoftwareChannelName = HordeServer.Utilities.StringId<HordeServer.Services.AgentSoftwareChannels>;
-
 namespace HordeServer.Controllers
 {
+	using PoolId = StringId<IPool>;
+	using AgentSoftwareChannelName = StringId<AgentSoftwareChannels>;
+
 	/// <summary>
 	/// Controller for the /api/v1/agents endpoint
 	/// </summary>
diff --git a/Engine/Source/Programs/Horde/HordeServer/Controllers/DebugController.cs b/Engine/Source/Programs/Horde/HordeServer/Controllers/DebugController.cs
--- a/Engine/Source/Programs/Horde/HordeServer/Controllers/DebugController.cs
+++ b/Engine/Source/Programs/Horde/HordeServer/Controllers/DebugController.cs
@@ -25,12 +25,6 @@ using System.Text;
 using System.Threading.Tasks;
 using System.Web;
 
-using PoolId = HordeServer.Utilities.StringId<HordeServer.Models.IPool>;
-using ProjectId = HordeServer.Utilities.StringId<HordeServer.Models.IProject>;
-using TemplateRefId = HordeServer.Utilities.StringId<HordeServer.Models.TemplateRef>;
-using AgentSoftwareVersion = HordeServer.Utilities.StringId<HordeServer.Collections.IAgentSoftwareCollection>;
-
-using IStream = HordeServer.Models.IStream;
 using MongoDB.Driver;
 using System.Text.Json;
 using System.Threading;
@@ -40,11 +34,18 @@ using HordeServer.Storage;
 using Microsoft.Extensions.Hosting;
 using Microsoft.Extensions.Options;
 using StatsdClient;
-using P4Debugging = Perforce.P4.P4Debugging;
+
 using Amazon.S3;
 
 namespace HordeServer.Controllers
 {
+	using PoolId = StringId<IPool>;
+	using ProjectId = StringId<IProject>;
+	using TemplateRefId = StringId<TemplateRef>;
+	using AgentSoftwareVersion = StringId<IAgentSoftwareCollection>;
+	using IStream = HordeServer.Models.IStream;
+	using P4Debugging = Perforce.P4.P4Debugging;
+
 	/// <summary>
 	/// Public endpoints for the debug controller
 	/// </summary>
diff --git a/Engine/Source/Programs/Horde/HordeServer/Controllers/DevicesController.cs b/Engine/Source/Programs/Horde/HordeServer/Controllers/DevicesController.cs
--- a/Engine/Source/Programs/Horde/HordeServer/Controllers/DevicesController.cs
+++ b/Engine/Source/Programs/Horde/HordeServer/Controllers/DevicesController.cs
@@ -1,29 +1,27 @@
 // Copyright Epic Games, Inc. All Rights Reserved.
 
-using HordeServer.Services;
-using System.Threading.Tasks;
-
 using HordeServer.Api;
+using HordeServer.Collections;
 using HordeServer.Models;
-using MongoDB.Driver;
+using HordeServer.Services;
+using HordeServer.Utilities;
+using Microsoft.AspNetCore.Authorization;
 using Microsoft.AspNetCore.Mvc;
-using HordeServer.Collections;
-using System;
-using System.Linq;
-
 using Microsoft.Extensions.Logging;
 using MongoDB.Bson;
+using MongoDB.Driver;
+using System;
 using System.Collections.Generic;
-
-using DeviceId = HordeServer.Utilities.StringId<HordeServer.Models.IDevice>;
-using DevicePlatformId = HordeServer.Utilities.StringId<HordeServer.Models.IDevicePlatform>;
-using DevicePoolId = HordeServer.Utilities.StringId<HordeServer.Models.IDevicePool>;
-using Microsoft.AspNetCore.Authorization;
+using System.Linq;
 using System.Text.Json;
-
+using System.Threading.Tasks;
 
 namespace HordeServer.Controllers
 {
+	using DeviceId = StringId<IDevice>;
+	using DevicePlatformId = StringId<IDevicePlatform>;
+	using DevicePoolId = StringId<IDevicePool>;
+
 	/// <summary>
 	/// Controller for device service
 	/// </summary>
diff --git a/Engine/Source/Programs/Horde/HordeServer/Controllers/IssuesController.cs b/Engine/Source/Programs/Horde/HordeServer/Controllers/IssuesController.cs
--- a/Engine/Source/Programs/Horde/HordeServer/Controllers/IssuesController.cs
+++ b/Engine/Source/Programs/Horde/HordeServer/Controllers/IssuesController.cs
@@ -19,10 +19,10 @@ using System.Linq;
 using System.Text;
 using System.Threading.Tasks;
 
-using StreamId = HordeServer.Utilities.StringId<HordeServer.Models.IStream>;
-
 namespace HordeServer.Controllers
 {
+	using StreamId = StringId<IStream>;
+
 	/// <summary>
 	/// Controller for the /api/v1/issues endpoint
 	/// </summary>
diff --git a/Engine/Source/Programs/Horde/HordeServer/Controllers/JobsController.cs b/Engine/Source/Programs/Horde/HordeServer/Controllers/JobsController.cs
--- a/Engine/Source/Programs/Horde/HordeServer/Controllers/JobsController.cs
+++ b/Engine/Source/Programs/Horde/HordeServer/Controllers/JobsController.cs
@@ -1,6 +1,5 @@
 // Copyright Epic Games, Inc. All Rights Reserved.
 
-using Datadog.Trace;
 using EpicGames.Core;
 using HordeServer.Api;
 using HordeServer.Collections;
@@ -20,8 +19,6 @@ using System.Security.Claims;
 using System.Text.RegularExpressions;
 using System.Threading.Tasks;
 
-using StreamId = HordeServer.Utilities.StringId<HordeServer.Models.IStream>;
-using TemplateRefId = HordeServer.Utilities.StringId<HordeServer.Models.TemplateRef>;
 using Microsoft.Extensions.Logging;
 using HordeServer.Notifications;
 using OpenTracing.Util;
@@ -29,6 +26,9 @@ using OpenTracing;
 
 namespace HordeServer.Controllers
 {
+	using StreamId = StringId<IStream>;
+	using TemplateRefId = StringId<TemplateRef>;
+
 	/// <summary>
 	/// Controller for the /api/v1/jobs endpoing
 	/// </summary>
diff --git a/Engine/Source/Programs/Horde/HordeServer/Controllers/LeasesController.cs b/Engine/Source/Programs/Horde/HordeServer/Controllers/LeasesController.cs
--- a/Engine/Source/Programs/Horde/HordeServer/Controllers/LeasesController.cs
+++ b/Engine/Source/Programs/Horde/HordeServer/Controllers/LeasesController.cs
@@ -22,11 +22,11 @@ using System.Text.Json;
 using System.Threading.Tasks;
 using Microsoft.Extensions.Logging;
 
-using PoolId = HordeServer.Utilities.StringId<HordeServer.Models.IPool>;
-using AgentSoftwareChannelName = HordeServer.Utilities.StringId<HordeServer.Services.AgentSoftwareChannels>;
-
 namespace HordeServer.Controllers
 {
+	using PoolId = StringId<IPool>;
+	using AgentSoftwareChannelName = StringId<AgentSoftwareChannels>;
+
 	/// <summary>
 	/// Controller for the /api/v1/leases endpoint
 	/// </summary>
diff --git a/Engine/Source/Programs/Horde/HordeServer/Controllers/PoolsController.cs b/Engine/Source/Programs/Horde/HordeServer/Controllers/PoolsController.cs
--- a/Engine/Source/Programs/Horde/HordeServer/Controllers/PoolsController.cs
+++ b/Engine/Source/Programs/Horde/HordeServer/Controllers/PoolsController.cs
@@ -13,10 +13,10 @@ using System.ComponentModel.DataAnnotations;
 using System.Linq;
 using System.Threading.Tasks;
 
-using PoolId = HordeServer.Utilities.StringId<HordeServer.Models.IPool>;
-
 namespace HordeServer.Controllers
 {
+	using PoolId = StringId<IPool>;
+
 	/// <summary>
 	/// Controller for the /api/v1/pools endpoint
 	/// </summary>
diff --git a/Engine/Source/Programs/Horde/HordeServer/Controllers/ProjectsController.cs b/Engine/Source/Programs/Horde/HordeServer/Controllers/ProjectsController.cs
--- a/Engine/Source/Programs/Horde/HordeServer/Controllers/ProjectsController.cs
+++ b/Engine/Source/Programs/Horde/HordeServer/Controllers/ProjectsController.cs
@@ -13,10 +13,10 @@ using System.ComponentModel.DataAnnotations;
 using System.Linq;
 using System.Threading.Tasks;
 
-using ProjectId = HordeServer.Utilities.StringId<HordeServer.Models.IProject>;
-
 namespace HordeServer.Controllers
 {
+	using ProjectId = StringId<IProject>;
+
 	/// <summary>
 	/// Controller for the /api/v1/projects endpoint
 	/// </summary>
diff --git a/Engine/Source/Programs/Horde/HordeServer/Controllers/StreamsController.cs b/Engine/Source/Programs/Horde/HordeServer/Controllers/StreamsController.cs
--- a/Engine/Source/Programs/Horde/HordeServer/Controllers/StreamsController.cs
+++ b/Engine/Source/Programs/Horde/HordeServer/Controllers/StreamsController.cs
@@ -1,29 +1,25 @@
 // Copyright Epic Games, Inc. All Rights Reserved.
 
-using EpicGames.Core;
 using HordeServer.Api;
+using HordeServer.Collections;
 using HordeServer.Models;
 using HordeServer.Services;
 using HordeServer.Utilities;
 using Microsoft.AspNetCore.Authorization;
 using Microsoft.AspNetCore.Mvc;
-using MongoDB.Bson;
 using System;
 using System.Collections.Generic;
-using System.ComponentModel.DataAnnotations;
 using System.Diagnostics.CodeAnalysis;
+using System.IO;
 using System.Linq;
 using System.Threading.Tasks;
 
-using ProjectId = HordeServer.Utilities.StringId<HordeServer.Models.IProject>;
-using StreamId = HordeServer.Utilities.StringId<HordeServer.Models.IStream>;
-using TemplateRefId = HordeServer.Utilities.StringId<HordeServer.Models.TemplateRef>;
-using HordeServer.Collections;
-using System.Security.Claims;
-using System.IO;
-
 namespace HordeServer.Controllers
 {
+	using ProjectId = StringId<IProject>;
+	using StreamId = StringId<IStream>;
+	using TemplateRefId = StringId<TemplateRef>;
+
 	/// <summary>
 	/// Controller for the /api/v1/streams endpoint
 	/// </summary>
diff --git a/Engine/Source/Programs/Horde/HordeServer/Controllers/TestDataController.cs b/Engine/Source/Programs/Horde/HordeServer/Controllers/TestDataController.cs
--- a/Engine/Source/Programs/Horde/HordeServer/Controllers/TestDataController.cs
+++ b/Engine/Source/Programs/Horde/HordeServer/Controllers/TestDataController.cs
@@ -18,10 +18,10 @@ using System.IO;
 using System.Linq;
 using System.Threading.Tasks;
 
-using StreamId = HordeServer.Utilities.StringId<HordeServer.Models.IStream>;
-
 namespace HordeServer.Controllers
 {
+	using StreamId = StringId<IStream>;
+
 	/// <summary>
 	/// Controller for the /api/v1/testdata endpoint
 	/// </summary>
diff --git a/Engine/Source/Programs/Horde/HordeServer/Controllers/UgsController.cs b/Engine/Source/Programs/Horde/HordeServer/Controllers/UgsController.cs
--- a/Engine/Source/Programs/Horde/HordeServer/Controllers/UgsController.cs
+++ b/Engine/Source/Programs/Horde/HordeServer/Controllers/UgsController.cs
@@ -19,12 +19,12 @@ using System.ComponentModel.DataAnnotations;
 using System.IO;
 using System.Linq;
 using System.Threading.Tasks;
-
-using StreamId = HordeServer.Utilities.StringId<HordeServer.Models.IStream>;
 using Microsoft.Extensions.Logging;
 
 namespace HordeServer.Controllers
 {
+	using StreamId = StringId<IStream>;
+
 	/// <summary>
 	/// Controller for the /api/v1/issues endpoint
 	/// </summary>
diff --git a/Engine/Source/Programs/Horde/HordeServer/Logs/Builder/Impl/RedisLogBuilder.cs b/Engine/Source/Programs/Horde/HordeServer/Logs/Builder/Impl/RedisLogBuilder.cs
--- a/Engine/Source/Programs/Horde/HordeServer/Logs/Builder/Impl/RedisLogBuilder.cs
+++ b/Engine/Source/Programs/Horde/HordeServer/Logs/Builder/Impl/RedisLogBuilder.cs
@@ -1,6 +1,5 @@
 // Copyright Epic Games, Inc. All Rights Reserved.
 
-using Datadog.Trace;
 using EpicGames.Core;
 using HordeServer.Api;
 using Microsoft.Extensions.Logging;
diff --git a/Engine/Source/Programs/Horde/HordeServer/Logs/LogIndexData.cs b/Engine/Source/Programs/Horde/HordeServer/Logs/LogIndexData.cs
--- a/Engine/Source/Programs/Horde/HordeServer/Logs/LogIndexData.cs
+++ b/Engine/Source/Programs/Horde/HordeServer/Logs/LogIndexData.cs
@@ -1,6 +1,5 @@
 // Copyright Epic Games, Inc. All Rights Reserved.
 
-using Datadog.Trace;
 using EpicGames.Core;
 using HordeServer.Utilities;
 using OpenTracing;
diff --git a/Engine/Source/Programs/Horde/HordeServer/Models/Agent.cs b/Engine/Source/Programs/Horde/HordeServer/Models/Agent.cs
--- a/Engine/Source/Programs/Horde/HordeServer/Models/Agent.cs
+++ b/Engine/Source/Programs/Horde/HordeServer/Models/Agent.cs
@@ -16,15 +16,17 @@ using System.IO;
 using System.Linq;
 using System.Security.Cryptography;
 using Serilog;
-using PoolId = HordeServer.Utilities.StringId<HordeServer.Models.IPool>;
-using StreamId = HordeServer.Utilities.StringId<HordeServer.Models.IStream>;
-using AgentSoftwareVersion = HordeServer.Utilities.StringId<HordeServer.Collections.IAgentSoftwareCollection>;
-using AgentSoftwareChannelName = HordeServer.Utilities.StringId<HordeServer.Services.AgentSoftwareChannels>;
 using System.Diagnostics;
 using System.Threading.Tasks;
+using HordeServer.Collections;
 
 namespace HordeServer.Models
 {
+	using PoolId = StringId<IPool>;
+	using StreamId = StringId<IStream>;
+	using AgentSoftwareVersion = StringId<IAgentSoftwareCollection>;
+	using AgentSoftwareChannelName = StringId<AgentSoftwareChannels>;
+
 	/// <summary>
 	/// Information about a workspace synced to an agent
 	/// </summary>
diff --git a/Engine/Source/Programs/Horde/HordeServer/Models/AgentSoftware.cs b/Engine/Source/Programs/Horde/HordeServer/Models/AgentSoftware.cs
--- a/Engine/Source/Programs/Horde/HordeServer/Models/AgentSoftware.cs
+++ b/Engine/Source/Programs/Horde/HordeServer/Models/AgentSoftware.cs
@@ -2,6 +2,7 @@
 
 using HordeServer.Api;
 using HordeServer.Collections;
+using HordeServer.Services;
 using HordeServer.Utilities;
 using MongoDB.Bson;
 using MongoDB.Bson.Serialization.Attributes;
@@ -13,11 +14,11 @@ using System.Linq;
 using System.Text.Json;
 using System.Threading.Tasks;
 
-using AgentSoftwareVersion = HordeServer.Utilities.StringId<HordeServer.Collections.IAgentSoftwareCollection>;
-using AgentSoftwareChannelName = HordeServer.Utilities.StringId<HordeServer.Services.AgentSoftwareChannels>;
-
 namespace HordeServer.Models
 {
+	using AgentSoftwareVersion = StringId<IAgentSoftwareCollection>;
+	using AgentSoftwareChannelName = StringId<AgentSoftwareChannels>;
+
 	/// <summary>
 	/// A software channel
 	/// </summary>
diff --git a/Engine/Source/Programs/Horde/HordeServer/Models/Change.cs b/Engine/Source/Programs/Horde/HordeServer/Models/Change.cs
--- a/Engine/Source/Programs/Horde/HordeServer/Models/Change.cs
+++ b/Engine/Source/Programs/Horde/HordeServer/Models/Change.cs
@@ -13,8 +13,6 @@ using System.Linq;
 using System.Text.RegularExpressions;
 using System.Threading.Tasks;
 
-using StreamId = HordeServer.Utilities.StringId<HordeServer.Models.IStream>;
-
 namespace HordeServer.Models
 {
 	/// <summary>
diff --git a/Engine/Source/Programs/Horde/HordeServer/Models/Config.cs b/Engine/Source/Programs/Horde/HordeServer/Models/Config.cs
--- a/Engine/Source/Programs/Horde/HordeServer/Models/Config.cs
+++ b/Engine/Source/Programs/Horde/HordeServer/Models/Config.cs
@@ -9,11 +9,11 @@ using System.ComponentModel.DataAnnotations;
 using System.Linq;
 using System.Threading.Tasks;
 
-using ProjectId = HordeServer.Utilities.StringId<HordeServer.Models.IProject>;
-using StreamId = HordeServer.Utilities.StringId<HordeServer.Models.IStream>;
-
 namespace HordeServer.Models
 {
+	using ProjectId = StringId<IProject>;
+	using StreamId = StringId<IStream>;
+
 	/// <summary>
 	/// Global configuration
 	/// </summary>
diff --git a/Engine/Source/Programs/Horde/HordeServer/Models/Device.cs b/Engine/Source/Programs/Horde/HordeServer/Models/Device.cs
--- a/Engine/Source/Programs/Horde/HordeServer/Models/Device.cs
+++ b/Engine/Source/Programs/Horde/HordeServer/Models/Device.cs
@@ -5,14 +5,14 @@ using MongoDB.Bson;
 using MongoDB.Bson.Serialization.Attributes;
 
 using System.Collections.Generic;
-
-using DeviceId = HordeServer.Utilities.StringId<HordeServer.Models.IDevice>;
-using DevicePlatformId = HordeServer.Utilities.StringId<HordeServer.Models.IDevicePlatform>;
-using DevicePoolId = HordeServer.Utilities.StringId<HordeServer.Models.IDevicePool>;
+using HordeServer.Utilities;
 
 namespace HordeServer.Models
 {
-	
+	using DeviceId = StringId<IDevice>;
+	using DevicePlatformId = StringId<IDevicePlatform>;
+	using DevicePoolId = StringId<IDevicePool>;
+
 	/// <summary>
 	/// A reservation containing one or more devices
 	/// </summary>
diff --git a/Engine/Source/Programs/Horde/HordeServer/Models/Event.cs b/Engine/Source/Programs/Horde/HordeServer/Models/Event.cs
--- a/Engine/Source/Programs/Horde/HordeServer/Models/Event.cs
+++ b/Engine/Source/Programs/Horde/HordeServer/Models/Event.cs
@@ -1,17 +1,18 @@
 // Copyright Epic Games, Inc. All Rights Reserved.
 
 using HordeServer.Api;
+using HordeServer.Utilities;
 using System;
 using System.Collections.Generic;
 using System.Diagnostics.CodeAnalysis;
 using System.Linq;
 using System.Threading.Tasks;
 
-using StreamId = HordeServer.Utilities.StringId<HordeServer.Models.IStream>;
-using TemplateRefId = HordeServer.Utilities.StringId<HordeServer.Models.TemplateRef>;
-
 namespace HordeServer.Models
 {
+	using StreamId = StringId<IStream>;
+	using TemplateRefId = StringId<TemplateRef>;
+
 	/// <summary>
 	/// Base interface for events that users can subscribe for notifications to
 	/// </summary>
diff --git a/Engine/Source/Programs/Horde/HordeServer/Models/IIssue.cs b/Engine/Source/Programs/Horde/HordeServer/Models/IIssue.cs
--- a/Engine/Source/Programs/Horde/HordeServer/Models/IIssue.cs
+++ b/Engine/Source/Programs/Horde/HordeServer/Models/IIssue.cs
@@ -15,11 +15,11 @@ using System.Diagnostics.CodeAnalysis;
 using System.Linq;
 using System.Runtime.CompilerServices;
 
-using StreamId = HordeServer.Utilities.StringId<HordeServer.Models.IStream>;
-using TemplateRefId = HordeServer.Utilities.StringId<HordeServer.Models.TemplateRef>;
-
 namespace HordeServer.Models
 {
+	using StreamId = StringId<IStream>;
+	using TemplateRefId = StringId<TemplateRef>;
+
 	/// <summary>
 	/// The severity of an issue
 	/// </summary>
diff --git a/Engine/Source/Programs/Horde/HordeServer/Models/IIssueSpan.cs b/Engine/Source/Programs/Horde/HordeServer/Models/IIssueSpan.cs
--- a/Engine/Source/Programs/Horde/HordeServer/Models/IIssueSpan.cs
+++ b/Engine/Source/Programs/Horde/HordeServer/Models/IIssueSpan.cs
@@ -8,11 +8,11 @@ using System.Collections.Generic;
 using System.Linq;
 using System.Threading.Tasks;
 
-using StreamId = HordeServer.Utilities.StringId<HordeServer.Models.IStream>;
-using TemplateRefId = HordeServer.Utilities.StringId<HordeServer.Models.TemplateRef>;
-
 namespace HordeServer.Models
 {
+	using StreamId = StringId<IStream>;
+	using TemplateRefId = StringId<TemplateRef>;
+
 	/// <summary>
 	/// Trace of a set of errors in a single step across multiple jobs
 	/// </summary>
diff --git a/Engine/Source/Programs/Horde/HordeServer/Models/Job.cs b/Engine/Source/Programs/Horde/HordeServer/Models/Job.cs
--- a/Engine/Source/Programs/Horde/HordeServer/Models/Job.cs
+++ b/Engine/Source/Programs/Horde/HordeServer/Models/Job.cs
@@ -17,13 +17,13 @@ using System.Text;
 using System.Text.Json.Serialization;
 using System.Threading.Tasks;
 
-using ReportPlacement = HordeCommon.Rpc.ReportPlacement;
-using StreamId = HordeServer.Utilities.StringId<HordeServer.Models.IStream>;
-using TemplateRefId = HordeServer.Utilities.StringId<HordeServer.Models.TemplateRef>;
-using PoolId = HordeServer.Utilities.StringId<HordeServer.Models.IPool>;
-
 namespace HordeServer.Models
 {
+	using ReportPlacement = HordeCommon.Rpc.ReportPlacement;
+	using StreamId = StringId<HordeServer.Models.IStream>;
+	using TemplateRefId = StringId<HordeServer.Models.TemplateRef>;
+	using PoolId = StringId<HordeServer.Models.IPool>;
+
 	/// <summary>
 	/// Report for a job or jobstep
 	/// </summary>
diff --git a/Engine/Source/Programs/Horde/HordeServer/Models/JobStepRef.cs b/Engine/Source/Programs/Horde/HordeServer/Models/JobStepRef.cs
--- a/Engine/Source/Programs/Horde/HordeServer/Models/JobStepRef.cs
+++ b/Engine/Source/Programs/Horde/HordeServer/Models/JobStepRef.cs
@@ -10,12 +10,12 @@ using System.Collections.Generic;
 using System.Linq;
 using System.Threading.Tasks;
 
-using StreamId = HordeServer.Utilities.StringId<HordeServer.Models.IStream>;
-using TemplateRefId = HordeServer.Utilities.StringId<HordeServer.Models.TemplateRef>;
-using PoolId = HordeServer.Utilities.StringId<HordeServer.Models.IPool>;
-
 namespace HordeServer.Models
 {
+	using StreamId = StringId<IStream>;
+	using TemplateRefId = StringId<TemplateRef>;
+	using PoolId = StringId<IPool>;
+
 	/// <summary>
 	/// Unique id struct for JobStepRef objects. Includes a job id, batch id, and step id to uniquely identify the step.
 	/// </summary>
diff --git a/Engine/Source/Programs/Horde/HordeServer/Models/Lease.cs b/Engine/Source/Programs/Horde/HordeServer/Models/Lease.cs
--- a/Engine/Source/Programs/Horde/HordeServer/Models/Lease.cs
+++ b/Engine/Source/Programs/Horde/HordeServer/Models/Lease.cs
@@ -10,14 +10,14 @@ using System;
 using System.Collections.Generic;
 using System.Linq;
 using System.Threading.Tasks;
-
-using PoolId = HordeServer.Utilities.StringId<HordeServer.Models.IPool>;
-using StreamId = HordeServer.Utilities.StringId<HordeServer.Models.IStream>;
 using Google.Protobuf.WellKnownTypes;
 using Google.Protobuf;
 
 namespace HordeServer.Models
 {
+	using PoolId = StringId<IPool>;
+	using StreamId = StringId<IStream>;
+
 	/// <summary>
 	/// Document describing a lease. This exists to permanently record a lease; the agent object tracks internal state of any active leases through AgentLease objects.
 	/// </summary>
diff --git a/Engine/Source/Programs/Horde/HordeServer/Models/Pool.cs b/Engine/Source/Programs/Horde/HordeServer/Models/Pool.cs
--- a/Engine/Source/Programs/Horde/HordeServer/Models/Pool.cs
+++ b/Engine/Source/Programs/Horde/HordeServer/Models/Pool.cs
@@ -10,10 +10,10 @@ using System.Linq;
 using System.Security.Permissions;
 using System.Threading.Tasks;
 
-using PoolId = HordeServer.Utilities.StringId<HordeServer.Models.IPool>;
-
 namespace HordeServer.Models
 {
+	using PoolId = StringId<IPool>;
+
 	/// <summary>
 	/// A pool of machines
 	/// </summary>
diff --git a/Engine/Source/Programs/Horde/HordeServer/Models/Project.cs b/Engine/Source/Programs/Horde/HordeServer/Models/Project.cs
--- a/Engine/Source/Programs/Horde/HordeServer/Models/Project.cs
+++ b/Engine/Source/Programs/Horde/HordeServer/Models/Project.cs
@@ -12,10 +12,10 @@ using System.Linq;
 using System.Text.RegularExpressions;
 using System.Threading.Tasks;
 
-using ProjectId = HordeServer.Utilities.StringId<HordeServer.Models.IProject>;
-
 namespace HordeServer.Models
 {
+	using ProjectId = StringId<IProject>;
+
 	/// <summary>
 	/// Specifies a category of streams to display on the dashboard
 	/// </summary>
diff --git a/Engine/Source/Programs/Horde/HordeServer/Models/Schedule.cs b/Engine/Source/Programs/Horde/HordeServer/Models/Schedule.cs
--- a/Engine/Source/Programs/Horde/HordeServer/Models/Schedule.cs
+++ b/Engine/Source/Programs/Horde/HordeServer/Models/Schedule.cs
@@ -10,10 +10,10 @@ using System.Collections.Generic;
 using System.Linq;
 using System.Threading.Tasks;
 
-using TemplateRefId = HordeServer.Utilities.StringId<HordeServer.Models.TemplateRef>;
-
 namespace HordeServer.Models
 {
+	using TemplateRefId = StringId<TemplateRef>;
+
 	/// <summary>
 	/// Specifies a pattern of times that this schedule should run. Each schedule may have multiple patterns.
 	/// </summary>
diff --git a/Engine/Source/Programs/Horde/HordeServer/Models/Stream.cs b/Engine/Source/Programs/Horde/HordeServer/Models/Stream.cs
--- a/Engine/Source/Programs/Horde/HordeServer/Models/Stream.cs
+++ b/Engine/Source/Programs/Horde/HordeServer/Models/Stream.cs
@@ -17,15 +17,15 @@ using System.Security.Cryptography;
 using System.Text;
 using System.Text.RegularExpressions;
 using System.Threading.Tasks;
-
-using PoolId = HordeServer.Utilities.StringId<HordeServer.Models.IPool>;
-using ProjectId = HordeServer.Utilities.StringId<HordeServer.Models.IProject>;
-using StreamId = HordeServer.Utilities.StringId<HordeServer.Models.IStream>;
-using TemplateRefId = HordeServer.Utilities.StringId<HordeServer.Models.TemplateRef>;
 using System.Runtime.CompilerServices;
 
 namespace HordeServer.Models
 {
+	using PoolId = StringId<IPool>;
+	using ProjectId = StringId<IProject>;
+	using StreamId = StringId<IStream>;
+	using TemplateRefId = StringId<TemplateRef>;
+
 	/// <summary>
 	/// Exception thrown when stream validation fails
 	/// </summary>
diff --git a/Engine/Source/Programs/Horde/HordeServer/Models/StreamTab.cs b/Engine/Source/Programs/Horde/HordeServer/Models/StreamTab.cs
--- a/Engine/Source/Programs/Horde/HordeServer/Models/StreamTab.cs
+++ b/Engine/Source/Programs/Horde/HordeServer/Models/StreamTab.cs
@@ -1,16 +1,17 @@
 // Copyright Epic Games, Inc. All Rights Reserved.
 
 using HordeServer.Api;
+using HordeServer.Utilities;
 using MongoDB.Bson.Serialization.Attributes;
 using System;
 using System.Collections.Generic;
 using System.Linq;
 using System.Threading.Tasks;
 
-using TemplateRefId = HordeServer.Utilities.StringId<HordeServer.Models.TemplateRef>;
-
 namespace HordeServer.Models
 {
+	using TemplateRefId = StringId<TemplateRef>;
+
 	/// <summary>
 	/// Information about a page to display in the dashboard for a stream
 	/// </summary>
diff --git a/Engine/Source/Programs/Horde/HordeServer/Models/Subscription.cs b/Engine/Source/Programs/Horde/HordeServer/Models/Subscription.cs
--- a/Engine/Source/Programs/Horde/HordeServer/Models/Subscription.cs
+++ b/Engine/Source/Programs/Horde/HordeServer/Models/Subscription.cs
@@ -8,9 +8,6 @@ using System.Collections.Generic;
 using System.Linq;
 using System.Threading.Tasks;
 
-using StreamId = HordeServer.Utilities.StringId<HordeServer.Models.IStream>;
-using TemplateRefId = HordeServer.Utilities.StringId<HordeServer.Models.TemplateRef>;
-
 namespace HordeServer.Models
 {
 	/// <summary>
diff --git a/Engine/Source/Programs/Horde/HordeServer/Models/Telemetry.cs b/Engine/Source/Programs/Horde/HordeServer/Models/Telemetry.cs
--- a/Engine/Source/Programs/Horde/HordeServer/Models/Telemetry.cs
+++ b/Engine/Source/Programs/Horde/HordeServer/Models/Telemetry.cs
@@ -6,11 +6,11 @@ using System.Collections.Generic;
 using System.Linq;
 using System.Threading.Tasks;
 
-using PoolId = HordeServer.Utilities.StringId<HordeServer.Models.IPool>;
-using StreamId = HordeServer.Utilities.StringId<HordeServer.Models.IStream>;
-
 namespace HordeServer.Models
 {
+	using PoolId = StringId<IPool>;
+	using StreamId = StringId<IStream>;
+
 	/// <summary>
 	/// Information about the utilization of a pool
 	/// </summary>
diff --git a/Engine/Source/Programs/Horde/HordeServer/Models/TestData.cs b/Engine/Source/Programs/Horde/HordeServer/Models/TestData.cs
--- a/Engine/Source/Programs/Horde/HordeServer/Models/TestData.cs
+++ b/Engine/Source/Programs/Horde/HordeServer/Models/TestData.cs
@@ -1,16 +1,17 @@
 // Copyright Epic Games, Inc. All Rights Reserved.
 
+using HordeServer.Utilities;
 using MongoDB.Bson;
 using System;
 using System.Collections.Generic;
 using System.Linq;
 using System.Threading.Tasks;
 
-using StreamId = HordeServer.Utilities.StringId<HordeServer.Models.IStream>;
-using TemplateRefId = HordeServer.Utilities.StringId<HordeServer.Models.TemplateRef>;
-
 namespace HordeServer.Models
 {
+	using StreamId = StringId<IStream>;
+	using TemplateRefId = StringId<TemplateRef>;
+
 	/// <summary>
 	/// Stores information about the results of a test
 	/// </summary>
diff --git a/Engine/Source/Programs/Horde/HordeServer/Notifications/Impl/SlackNotificationSink.cs b/Engine/Source/Programs/Horde/HordeServer/Notifications/Impl/SlackNotificationSink.cs
--- a/Engine/Source/Programs/Horde/HordeServer/Notifications/Impl/SlackNotificationSink.cs
+++ b/Engine/Source/Programs/Horde/HordeServer/Notifications/Impl/SlackNotificationSink.cs
@@ -34,8 +34,6 @@ using System.Text.RegularExpressions;
 using System.Threading;
 using System.Threading.Tasks;
 
-using StreamId = HordeServer.Utilities.StringId<HordeServer.Models.IStream>;
-
 namespace HordeServer.Notifications.Impl
 {
 	/// <summary>
diff --git a/Engine/Source/Programs/Horde/HordeServer/Services/AgentService.cs b/Engine/Source/Programs/Horde/HordeServer/Services/AgentService.cs
--- a/Engine/Source/Programs/Horde/HordeServer/Services/AgentService.cs
+++ b/Engine/Source/Programs/Horde/HordeServer/Services/AgentService.cs
@@ -1,6 +1,5 @@
 // Copyright Epic Games, Inc. All Rights Reserved.
 
-using Datadog.Trace;
 using EpicGames.Core;
 using Google.Protobuf;
 using Google.Protobuf.Reflection;
@@ -30,14 +29,15 @@ using System.Text;
 using System.Threading;
 using System.Threading.Tasks;
 using StatsdClient;
-using PoolId = HordeServer.Utilities.StringId<HordeServer.Models.IPool>;
-using AgentSoftwareVersion = HordeServer.Utilities.StringId<HordeServer.Collections.IAgentSoftwareCollection>;
-using AgentSoftwareChannelName = HordeServer.Utilities.StringId<HordeServer.Services.AgentSoftwareChannels>;
 using System.Collections.Concurrent;
 using Microsoft.Extensions.Caching.Memory;
 
 namespace HordeServer.Services
 {
+	using PoolId = StringId<IPool>;
+	using AgentSoftwareVersion = StringId<IAgentSoftwareCollection>;
+	using AgentSoftwareChannelName = StringId<AgentSoftwareChannels>;
+
 	/// <summary>
 	/// Wraps funtionality for manipulating agents
 	/// </summary>
diff --git a/Engine/Source/Programs/Horde/HordeServer/Services/AgentSoftwareService.cs b/Engine/Source/Programs/Horde/HordeServer/Services/AgentSoftwareService.cs
--- a/Engine/Source/Programs/Horde/HordeServer/Services/AgentSoftwareService.cs
+++ b/Engine/Source/Programs/Horde/HordeServer/Services/AgentSoftwareService.cs
@@ -14,8 +14,6 @@ using System.IO;
 using System.Linq;
 using System.Threading.Channels;
 using System.Threading.Tasks;
-
-using AgentSoftwareChannelName = HordeServer.Utilities.StringId<HordeServer.Services.AgentSoftwareChannels>;
 using EpicGames.Core;
 using Microsoft.Extensions.Logging;
 using System.Runtime.InteropServices;
@@ -24,6 +22,8 @@ using System.IO.Compression;
 
 namespace HordeServer.Services
 {
+	using AgentSoftwareChannelName = StringId<AgentSoftwareChannels>;
+
 	/// <summary>
 	/// Information about a channel
 	/// </summary>
diff --git a/Engine/Source/Programs/Horde/HordeServer/Services/AutoscaleService.cs b/Engine/Source/Programs/Horde/HordeServer/Services/AutoscaleService.cs
--- a/Engine/Source/Programs/Horde/HordeServer/Services/AutoscaleService.cs
+++ b/Engine/Source/Programs/Horde/HordeServer/Services/AutoscaleService.cs
@@ -22,12 +22,13 @@ using System.Threading.Tasks;
 using HordeCommon;
 using Microsoft.Extensions.Options;
 using StatsdClient;
-using PoolId = HordeServer.Utilities.StringId<HordeServer.Models.IPool>;
 
 [assembly: InternalsVisibleTo("HordeServerTests")]
 
 namespace HordeServer.Services
 {
+	using PoolId = StringId<IPool>;
+
 	/// <summary>
 	/// Service for managing the autoscaling of agent pools
 	/// </summary>
diff --git a/Engine/Source/Programs/Horde/HordeServer/Services/ConfigService.cs b/Engine/Source/Programs/Horde/HordeServer/Services/ConfigService.cs
--- a/Engine/Source/Programs/Horde/HordeServer/Services/ConfigService.cs
+++ b/Engine/Source/Programs/Horde/HordeServer/Services/ConfigService.cs
@@ -22,11 +22,11 @@ using System.Text.RegularExpressions;
 using System.Threading;
 using System.Threading.Tasks;
 
-using ProjectId = HordeServer.Utilities.StringId<HordeServer.Models.IProject>;
-using StreamId = HordeServer.Utilities.StringId<HordeServer.Models.IStream>;
-
 namespace HordeServer.Services
 {
+	using ProjectId = StringId<IProject>;
+	using StreamId = StringId<IStream>;
+
 	/// <summary>
 	/// Polls Perforce for stream config changes
 	/// </summary>
diff --git a/Engine/Source/Programs/Horde/HordeServer/Services/DatabaseService.cs b/Engine/Source/Programs/Horde/HordeServer/Services/DatabaseService.cs
--- a/Engine/Source/Programs/Horde/HordeServer/Services/DatabaseService.cs
+++ b/Engine/Source/Programs/Horde/HordeServer/Services/DatabaseService.cs
@@ -23,15 +23,9 @@ using System.Runtime.InteropServices;
 using System.Security.Cryptography.X509Certificates;
 using System.Text;
 using System.Threading.Tasks;
-using IStream = HordeServer.Models.IStream;
-
-using PoolId = HordeServer.Utilities.StringId<HordeServer.Models.IPool>;
-using ProjectId = HordeServer.Utilities.StringId<HordeServer.Models.IProject>;
-using TemplateRefId = HordeServer.Utilities.StringId<HordeServer.Models.TemplateRef>;
 using System.Net.Sockets;
 using MongoDB.Driver.Core.Events;
 using System.Threading;
-using Datadog.Trace;
 using HordeServer.Collections;
 using System.Diagnostics;
 using System.Text.RegularExpressions;
@@ -41,6 +35,11 @@ using OpenTracing;
 
 namespace HordeServer.Services
 {
+	using IStream = HordeServer.Models.IStream;
+	using PoolId = StringId<IPool>;
+	using ProjectId = StringId<IProject>;
+	using TemplateRefId = StringId<TemplateRef>;
+
 	/// <summary>
 	/// Singleton for accessing the database
 	/// </summary>
diff --git a/Engine/Source/Programs/Horde/HordeServer/Services/DeviceService.cs b/Engine/Source/Programs/Horde/HordeServer/Services/DeviceService.cs
--- a/Engine/Source/Programs/Horde/HordeServer/Services/DeviceService.cs
+++ b/Engine/Source/Programs/Horde/HordeServer/Services/DeviceService.cs
@@ -12,10 +12,6 @@ using Microsoft.Extensions.Logging;
 using HordeCommon;
 using MongoDB.Bson;
 using System.Collections.Generic;
-
-using DeviceId = HordeServer.Utilities.StringId<HordeServer.Models.IDevice>;
-using DevicePlatformId = HordeServer.Utilities.StringId<HordeServer.Models.IDevicePlatform>;
-using DevicePoolId = HordeServer.Utilities.StringId<HordeServer.Models.IDevicePool>;
 using System.Security.Claims;
 using System.Diagnostics.CodeAnalysis;
 using HordeServer.Notifications;
@@ -23,6 +19,10 @@ using System.Linq;
 
 namespace HordeServer.Services
 {
+	using DeviceId = StringId<IDevice>;
+	using DevicePlatformId = StringId<IDevicePlatform>;
+	using DevicePoolId = StringId<IDevicePool>;
+
 	/// <summary>
 	/// Device management service
 	/// </summary>
diff --git a/Engine/Source/Programs/Horde/HordeServer/Services/IIssueService.cs b/Engine/Source/Programs/Horde/HordeServer/Services/IIssueService.cs
--- a/Engine/Source/Programs/Horde/HordeServer/Services/IIssueService.cs
+++ b/Engine/Source/Programs/Horde/HordeServer/Services/IIssueService.cs
@@ -1,6 +1,7 @@
 // Copyright Epic Games, Inc. All Rights Reserved.
 
 using HordeServer.Models;
+using HordeServer.Utilities;
 using Microsoft.Extensions.Hosting;
 using MongoDB.Bson;
 using System;
@@ -8,10 +9,10 @@ using System.Collections.Generic;
 using System.Linq;
 using System.Threading.Tasks;
 
-using StreamId = HordeServer.Utilities.StringId<HordeServer.Models.IStream>;
-
 namespace HordeServer.Services
 {
+	using StreamId = StringId<IStream>;
+
 	/// <summary>
 	/// Service which manages build health issues
 	/// </summary>
diff --git a/Engine/Source/Programs/Horde/HordeServer/Services/Impl/IssueService.cs b/Engine/Source/Programs/Horde/HordeServer/Services/Impl/IssueService.cs
--- a/Engine/Source/Programs/Horde/HordeServer/Services/Impl/IssueService.cs
+++ b/Engine/Source/Programs/Horde/HordeServer/Services/Impl/IssueService.cs
@@ -25,10 +25,6 @@ using System.Text;
 using System.Text.RegularExpressions;
 using System.Threading;
 using System.Threading.Tasks;
-using Datadog.Trace;
-
-using StreamId = HordeServer.Utilities.StringId<HordeServer.Models.IStream>;
-using TemplateRefId = HordeServer.Utilities.StringId<HordeServer.Models.TemplateRef>;
 using HordeServer.IssueHandlers.Impl;
 using HordeServer.Notifications;
 using OpenTracing;
@@ -36,6 +32,9 @@ using OpenTracing.Util;
 
 namespace HordeServer.Services.Impl
 {
+	using StreamId = StringId<IStream>;
+	using TemplateRefId = StringId<TemplateRef>;
+
 	/// <summary>
 	/// Detailed issue information
 	/// </summary>
diff --git a/Engine/Source/Programs/Horde/HordeServer/Services/JobService.cs b/Engine/Source/Programs/Horde/HordeServer/Services/JobService.cs
--- a/Engine/Source/Programs/Horde/HordeServer/Services/JobService.cs
+++ b/Engine/Source/Programs/Horde/HordeServer/Services/JobService.cs
@@ -23,9 +23,6 @@ using System.Security.Claims;
 using System.Text.RegularExpressions;
 using System.Threading;
 using System.Threading.Tasks;
-using Datadog.Trace;
-using StreamId = HordeServer.Utilities.StringId<HordeServer.Models.IStream>;
-using TemplateRefId = HordeServer.Utilities.StringId<HordeServer.Models.TemplateRef>;
 using System.Globalization;
 using HordeServer.Tasks.Impl;
 using OpenTracing.Util;
@@ -33,6 +30,9 @@ using OpenTracing;
 
 namespace HordeServer.Services
 {
+	using StreamId = StringId<IStream>;
+	using TemplateRefId = StringId<TemplateRef>;
+
 	/// <summary>
 	/// Exception thrown when attempting to retry executing a node that does not allow retries
 	/// </summary>
diff --git a/Engine/Source/Programs/Horde/HordeServer/Services/LogFileService.cs b/Engine/Source/Programs/Horde/HordeServer/Services/LogFileService.cs
--- a/Engine/Source/Programs/Horde/HordeServer/Services/LogFileService.cs
+++ b/Engine/Source/Programs/Horde/HordeServer/Services/LogFileService.cs
@@ -27,7 +27,6 @@ using System.Text.Json;
 using System.Text;
 using System.Threading.Tasks;
 using System.Threading;
-using Datadog.Trace;
 
 using ILogger = Microsoft.Extensions.Logging.ILogger;
 using Stream = System.IO.Stream;
diff --git a/Engine/Source/Programs/Horde/HordeServer/Services/MetricService.cs b/Engine/Source/Programs/Horde/HordeServer/Services/MetricService.cs
--- a/Engine/Source/Programs/Horde/HordeServer/Services/MetricService.cs
+++ b/Engine/Source/Programs/Horde/HordeServer/Services/MetricService.cs
@@ -6,7 +6,6 @@ using System;
 using System.Threading;
 using System.Threading.Tasks;
 using StatsdClient;
-using PoolId = HordeServer.Utilities.StringId<HordeServer.Models.IPool>;
 
 namespace HordeServer.Services
 {
diff --git a/Engine/Source/Programs/Horde/HordeServer/Services/PoolService.cs b/Engine/Source/Programs/Horde/HordeServer/Services/PoolService.cs
--- a/Engine/Source/Programs/Horde/HordeServer/Services/PoolService.cs
+++ b/Engine/Source/Programs/Horde/HordeServer/Services/PoolService.cs
@@ -15,10 +15,10 @@ using System.Linq;
 using System.Threading;
 using System.Threading.Tasks;
 
-using PoolId = HordeServer.Utilities.StringId<HordeServer.Models.IPool>;
-
 namespace HordeServer.Services
 {
+	using PoolId = StringId<IPool>;
+
 	/// <summary>
 	/// Wraps functionality for manipulating pools
 	/// </summary>
diff --git a/Engine/Source/Programs/Horde/HordeServer/Services/PoolUpdateService.cs b/Engine/Source/Programs/Horde/HordeServer/Services/PoolUpdateService.cs
--- a/Engine/Source/Programs/Horde/HordeServer/Services/PoolUpdateService.cs
+++ b/Engine/Source/Programs/Horde/HordeServer/Services/PoolUpdateService.cs
@@ -10,10 +10,10 @@ using System.Linq;
 using System.Threading;
 using System.Threading.Tasks;
 
-using PoolId = HordeServer.Utilities.StringId<HordeServer.Models.IPool>;
-
 namespace HordeServer.Services
 {
+	using PoolId = StringId<IPool>;
+
 	/// <summary>
 	/// Periodically updates pool documents to contain the correct workspaces
 	/// </summary>
diff --git a/Engine/Source/Programs/Horde/HordeServer/Services/ProjectService.cs b/Engine/Source/Programs/Horde/HordeServer/Services/ProjectService.cs
--- a/Engine/Source/Programs/Horde/HordeServer/Services/ProjectService.cs
+++ b/Engine/Source/Programs/Horde/HordeServer/Services/ProjectService.cs
@@ -13,10 +13,10 @@ using System.Linq;
 using System.Security.Claims;
 using System.Threading.Tasks;
 
-using ProjectId = HordeServer.Utilities.StringId<HordeServer.Models.IProject>;
-
 namespace HordeServer.Services
 {
+	using ProjectId = StringId<IProject>;
+
 	/// <summary>
 	/// Cache of information about job ACLs
 	/// </summary>
diff --git a/Engine/Source/Programs/Horde/HordeServer/Services/RpcService.cs b/Engine/Source/Programs/Horde/HordeServer/Services/RpcService.cs
--- a/Engine/Source/Programs/Horde/HordeServer/Services/RpcService.cs
+++ b/Engine/Source/Programs/Horde/HordeServer/Services/RpcService.cs
@@ -21,16 +21,17 @@ using System.Runtime.CompilerServices;
 using System.Security.Claims;
 using System.Text;
 using System.Threading.Tasks;
-using AgentSoftwareChannelName = HordeServer.Utilities.StringId<HordeServer.Services.AgentSoftwareChannels>;
-using AgentSoftwareVersion = HordeServer.Utilities.StringId<HordeServer.Collections.IAgentSoftwareCollection>;
-using IStream = HordeServer.Models.IStream;
-using StreamId = HordeServer.Utilities.StringId<HordeServer.Models.IStream>;
 using HordeServer.Tasks.Impl;
 
 [assembly: InternalsVisibleTo("HordeAgentTests")]
 
 namespace HordeServer.Services
 {
+	using AgentSoftwareChannelName = StringId<AgentSoftwareChannels>;
+	using AgentSoftwareVersion = StringId<IAgentSoftwareCollection>;
+	using IStream = HordeServer.Models.IStream;
+	using StreamId = StringId<IStream>;
+
 	/// <summary>
 	/// Implements the Horde gRPC service for bots updating their status and dequeing work
 	/// </summary>
diff --git a/Engine/Source/Programs/Horde/HordeServer/Services/ScheduleService.cs b/Engine/Source/Programs/Horde/HordeServer/Services/ScheduleService.cs
--- a/Engine/Source/Programs/Horde/HordeServer/Services/ScheduleService.cs
+++ b/Engine/Source/Programs/Horde/HordeServer/Services/ScheduleService.cs
@@ -1,7 +1,6 @@
 // Copyright Epic Games, Inc. All Rights Reserved.
 
 using Amazon.S3.Model.Internal.MarshallTransformations;
-using Datadog.Trace;
 using EpicGames.Core;
 using HordeCommon;
 using HordeServer.Api;
@@ -25,11 +24,12 @@ using System.Security.Claims;
 using System.Threading;
 using System.Threading.Tasks;
 using TimeZoneConverter;
-using StreamId = HordeServer.Utilities.StringId<HordeServer.Models.IStream>;
-using TemplateRefId = HordeServer.Utilities.StringId<HordeServer.Models.TemplateRef>;
 
 namespace HordeServer.Services
 {
+	using StreamId = StringId<IStream>;
+	using TemplateRefId = StringId<TemplateRef>;
+
 	/// <summary>
 	/// Manipulates schedule instances
 	/// </summary>
diff --git a/Engine/Source/Programs/Horde/HordeServer/Services/SoftwareService.cs b/Engine/Source/Programs/Horde/HordeServer/Services/SoftwareService.cs
--- a/Engine/Source/Programs/Horde/HordeServer/Services/SoftwareService.cs
+++ b/Engine/Source/Programs/Horde/HordeServer/Services/SoftwareService.cs
@@ -18,10 +18,10 @@ using System.Text;
 using System.Text.Json;
 using System.Threading.Tasks;
 
-using SoftwareChannelId = HordeServer.Utilities.StringId<HordeServer.Services.SoftwareVersionsSingleton>;
-
 namespace HordeServer.Services
 {
+	using SoftwareChannelId = StringId<SoftwareVersionsSingleton>;
+
 	/// <summary>
 	/// Singleton document used to track
 	/// </summary>
diff --git a/Engine/Source/Programs/Horde/HordeServer/Services/StreamService.cs b/Engine/Source/Programs/Horde/HordeServer/Services/StreamService.cs
--- a/Engine/Source/Programs/Horde/HordeServer/Services/StreamService.cs
+++ b/Engine/Source/Programs/Horde/HordeServer/Services/StreamService.cs
@@ -14,15 +14,15 @@ using System.Security.Claims;
 using System.Text;
 using System.Threading.Tasks;
 using System.Linq.Expressions;
-
-using ProjectId = HordeServer.Utilities.StringId<HordeServer.Models.IProject>;
-using StreamId = HordeServer.Utilities.StringId<HordeServer.Models.IStream>;
-using TemplateRefId = HordeServer.Utilities.StringId<HordeServer.Models.TemplateRef>;
 using HordeServer.Collections;
 using Microsoft.Extensions.Caching.Memory;
 
 namespace HordeServer.Services
 {
+	using ProjectId = StringId<IProject>;
+	using StreamId = StringId<IStream>;
+	using TemplateRefId = StringId<TemplateRef>;
+
 	/// <summary>
 	/// Cache of information about stream ACLs
 	/// </summary>
diff --git a/Engine/Source/Programs/Horde/HordeServer/Services/TelemetryService.cs b/Engine/Source/Programs/Horde/HordeServer/Services/TelemetryService.cs
--- a/Engine/Source/Programs/Horde/HordeServer/Services/TelemetryService.cs
+++ b/Engine/Source/Programs/Horde/HordeServer/Services/TelemetryService.cs
@@ -13,10 +13,10 @@ using System.Linq;
 using System.Threading;
 using System.Threading.Tasks;
 
-using PoolId = HordeServer.Utilities.StringId<HordeServer.Models.IPool>;
-
 namespace HordeServer.Services
 {
+	using PoolId = StringId<IPool>;
+
 	/// <summary>
 	/// Service which updates telemetry periodically
 	/// </summary>
diff --git a/Engine/Source/Programs/Horde/HordeServer/Tasks/Impl/ActionTaskSource.cs b/Engine/Source/Programs/Horde/HordeServer/Tasks/Impl/ActionTaskSource.cs
--- a/Engine/Source/Programs/Horde/HordeServer/Tasks/Impl/ActionTaskSource.cs
+++ b/Engine/Source/Programs/Horde/HordeServer/Tasks/Impl/ActionTaskSource.cs
@@ -24,11 +24,12 @@ using System.Linq;
 using System.Threading;
 using System.Threading.Tasks;
 using Microsoft.Extensions.Options;
-
-using PoolId = HordeServer.Utilities.StringId<HordeServer.Models.IPool>;
+using HordeServer.Utilities;
 
 namespace HordeServer.Tasks.Impl
 {
+	using PoolId = StringId<IPool>;
+
 	/// <summary>
 	/// Store either a successful or failed result for an execution
 	/// </summary>
diff --git a/Engine/Source/Programs/Horde/HordeServer/Tasks/Impl/JobTaskSource.cs b/Engine/Source/Programs/Horde/HordeServer/Tasks/Impl/JobTaskSource.cs
--- a/Engine/Source/Programs/Horde/HordeServer/Tasks/Impl/JobTaskSource.cs
+++ b/Engine/Source/Programs/Horde/HordeServer/Tasks/Impl/JobTaskSource.cs
@@ -26,13 +26,13 @@ using System.Runtime.CompilerServices;
 using System.Text;
 using System.Threading;
 using System.Threading.Tasks;
-
-using PoolId = HordeServer.Utilities.StringId<HordeServer.Models.IPool>;
-using StreamId = HordeServer.Utilities.StringId<HordeServer.Models.IStream>;
 using HordeServer.Services;
 
 namespace HordeServer.Tasks.Impl
 {
+	using PoolId = StringId<IPool>;
+	using StreamId = StringId<IStream>;
+
 	/// <summary>
 	/// Background service to dispatch pending work to agents in priority order.
 	/// </summary>
diff --git a/Engine/Source/Programs/Horde/HordeServer/Tasks/Impl/RestartTaskSource.cs b/Engine/Source/Programs/Horde/HordeServer/Tasks/Impl/RestartTaskSource.cs
--- a/Engine/Source/Programs/Horde/HordeServer/Tasks/Impl/RestartTaskSource.cs
+++ b/Engine/Source/Programs/Horde/HordeServer/Tasks/Impl/RestartTaskSource.cs
@@ -14,8 +14,6 @@ using System.Collections.Generic;
 using System.Linq;
 using System.Threading.Tasks;
 
-using AgentSoftwareVersion = HordeServer.Utilities.StringId<HordeServer.Collections.IAgentSoftwareCollection>;
-
 namespace HordeServer.Tasks.Impl
 {
 	class RestartTaskSource : TaskSourceBase<RestartTask>
diff --git a/Engine/Source/Programs/Horde/HordeServer/Tasks/Impl/ShutdownTaskSource.cs b/Engine/Source/Programs/Horde/HordeServer/Tasks/Impl/ShutdownTaskSource.cs
--- a/Engine/Source/Programs/Horde/HordeServer/Tasks/Impl/ShutdownTaskSource.cs
+++ b/Engine/Source/Programs/Horde/HordeServer/Tasks/Impl/ShutdownTaskSource.cs
@@ -14,8 +14,6 @@ using System.Collections.Generic;
 using System.Linq;
 using System.Threading.Tasks;
 
-using AgentSoftwareVersion = HordeServer.Utilities.StringId<HordeServer.Collections.IAgentSoftwareCollection>;
-
 namespace HordeServer.Tasks.Impl
 {
 	class ShutdownTaskSource : TaskSourceBase<ShutdownTask>
diff --git a/Engine/Source/Programs/Horde/HordeServer/Tasks/Impl/UpgradeTaskSource.cs b/Engine/Source/Programs/Horde/HordeServer/Tasks/Impl/UpgradeTaskSource.cs
--- a/Engine/Source/Programs/Horde/HordeServer/Tasks/Impl/UpgradeTaskSource.cs
+++ b/Engine/Source/Programs/Horde/HordeServer/Tasks/Impl/UpgradeTaskSource.cs
@@ -8,16 +8,17 @@ using HordeCommon.Rpc.Tasks;
 using HordeServer.Api;
 using HordeServer.Models;
 using HordeServer.Services;
+using HordeServer.Utilities;
 using MongoDB.Bson;
 using System;
 using System.Collections.Generic;
 using System.Linq;
 using System.Threading.Tasks;
 
-using AgentSoftwareChannelName = HordeServer.Utilities.StringId<HordeServer.Services.AgentSoftwareChannels>;
-
 namespace HordeServer.Tasks.Impl
 {
+	using AgentSoftwareChannelName = StringId<AgentSoftwareChannels>;
+
 	class UpgradeTaskSource : TaskSourceBase<UpgradeTask>
 	{
 		AgentSoftwareService AgentSoftwareService;
diff --git a/Engine/Source/Programs/Horde/HordeServer/Utilities/CompressionExtensions.cs b/Engine/Source/Programs/Horde/HordeServer/Utilities/CompressionExtensions.cs
--- a/Engine/Source/Programs/Horde/HordeServer/Utilities/CompressionExtensions.cs
+++ b/Engine/Source/Programs/Horde/HordeServer/Utilities/CompressionExtensions.cs
@@ -1,6 +1,5 @@
 // Copyright Epic Games, Inc. All Rights Reserved.
 
-using Datadog.Trace;
 using EpicGames.Core;
 using ICSharpCode.SharpZipLib.BZip2;
 using OpenTracing;
diff --git a/Engine/Source/Programs/Horde/HordeServer/Utilities/TickedBackgroundService.cs b/Engine/Source/Programs/Horde/HordeServer/Utilities/TickedBackgroundService.cs
--- a/Engine/Source/Programs/Horde/HordeServer/Utilities/TickedBackgroundService.cs
+++ b/Engine/Source/Programs/Horde/HordeServer/Utilities/TickedBackgroundService.cs
@@ -1,6 +1,5 @@
 // Copyright Epic Games, Inc. All Rights Reserved.
 
-using Datadog.Trace;
 using EpicGames.Core;
 using Microsoft.Extensions.Hosting;
 using Microsoft.Extensions.Logging;
