Horde: Use the OpenTracing API in the Horde Agent rather than having a hard dependency on Datadog (though Datadog is still initialized by default at startup)

diff --git a/Engine/Source/Programs/Horde/HordeAgent/Commands/Service/RunCommand.cs b/Engine/Source/Programs/Horde/HordeAgent/Commands/Service/RunCommand.cs
--- a/Engine/Source/Programs/Horde/HordeAgent/Commands/Service/RunCommand.cs
+++ b/Engine/Source/Programs/Horde/HordeAgent/Commands/Service/RunCommand.cs
@@ -17,10 +17,12 @@ using System.Security.Cryptography.X509Certificates;
 using System.Text;
 using System.Threading.Tasks;
 using EpicGames.Core;
-using Datadog.Trace.Configuration;
-using Datadog.Trace;
 using System.IO;
 using System.Reflection;
+using OpenTracing;
+using OpenTracing.Util;
+using Datadog.Trace.Configuration;
+using Datadog.Trace;
 
 namespace HordeAgent.Modes.Service
 {
@@ -149,6 +151,9 @@ namespace HordeAgent.Modes.Service
 			Settings.LogsInjectionEnabled = true;
 
 			Tracer.Instance = new Tracer(Settings);
+
+			ITracer OpenTracer = Datadog.Trace.OpenTracing.OpenTracingTracerFactory.WrapTracer(Tracer.Instance);
+			GlobalTracer.Register(OpenTracer);
 		}
 	}
 }
diff --git a/Engine/Source/Programs/Horde/HordeAgent/Execution/BuildGraphExecutor.cs b/Engine/Source/Programs/Horde/HordeAgent/Execution/BuildGraphExecutor.cs
--- a/Engine/Source/Programs/Horde/HordeAgent/Execution/BuildGraphExecutor.cs
+++ b/Engine/Source/Programs/Horde/HordeAgent/Execution/BuildGraphExecutor.cs
@@ -1,6 +1,5 @@
 // Copyright Epic Games, Inc. All Rights Reserved.
 
-using Datadog.Trace;
 using EpicGames.Core;
 using Grpc.Core;
 using HordeAgent.Execution.Interfaces;
@@ -9,6 +8,8 @@ using HordeAgent.Utility;
 using HordeCommon;
 using HordeCommon.Rpc;
 using Microsoft.Extensions.Logging;
+using OpenTracing;
+using OpenTracing.Util;
 using System;
 using System.Collections.Generic;
 using System.Diagnostics;
@@ -600,7 +601,7 @@ namespace HordeAgent.Execution
 		protected async Task<int> ExecuteAutomationToolAsync(BeginStepResponse Step, DirectoryReference WorkspaceDir, string Arguments, IReadOnlyDictionary<string, string> EnvVars, IReadOnlyDictionary<string, string> Credentials, ILogger Logger, CancellationToken CancellationToken)
 		{
 			int Result;
-			using (Scope Scope = Tracer.Instance.StartActive("BuildGraph"))
+			using (IScope Scope = GlobalTracer.Instance.BuildSpan("BuildGraph").StartActive())
 			{
 				if (!CompileAutomationTool)
 				{
@@ -832,7 +833,7 @@ namespace HordeAgent.Execution
 							}
 							await ArtifactUploader.UploadAsync(RpcConnection, JobId, BatchId, Step.StepId, "Trace.json", TraceFile, Logger, CancellationToken.None);
 
-							CreateTracingData(RootSpan);
+							CreateTracingData(GlobalTracer.Instance.ActiveSpan, RootSpan);
 						}
 					}
 					
@@ -941,27 +942,32 @@ namespace HordeAgent.Execution
 			await RpcConnection.InvokeAsync(x => x.CreateReportAsync(Request), new RpcContext(), CancellationToken.None);
 		}
 
-		private void CreateTracingData(TraceSpan Span)
+		private ISpan CreateTracingData(ISpan Parent, TraceSpan Span)
 		{
-			using (Scope Scope = Tracer.Instance.StartActive(Span.Name, serviceName: Span.Service, startTime: new DateTime(Span.Start, DateTimeKind.Utc)))
+			ISpan NewSpan = GlobalTracer.Instance.BuildSpan(Span.Name)
+				.AsChildOf(Parent)
+				.WithServiceName(Span.Service)
+				.WithResourceName(Span.Resource)
+				.WithStartTimestamp(new DateTime(Span.Start, DateTimeKind.Utc))
+				.Start();
+
+			if (Span.Properties != null)
 			{
-				Scope.Span.ResourceName = Span.Resource;
-				if (Span.Properties != null)
+				foreach (KeyValuePair<string, string> Pair in Span.Properties)
 				{
-					foreach (KeyValuePair<string, string> Pair in Span.Properties)
-					{
-						Scope.Span.SetTag(Pair.Key, Pair.Value);
-					}
+					NewSpan.SetTag(Pair.Key, Pair.Value);
 				}
-				if (Span.Children != null)
+			}
+			if (Span.Children != null)
+			{
+				foreach (TraceSpan Child in Span.Children)
 				{
-					foreach (TraceSpan Child in Span.Children)
-					{
-						CreateTracingData(Child);
-					}
+					CreateTracingData(NewSpan, Child);
 				}
-				Scope.Span.Finish(new DateTime(Span.Finish, DateTimeKind.Utc));
 			}
+
+			NewSpan.Finish(new DateTime(Span.Finish, DateTimeKind.Utc));
+			return NewSpan;
 		}
 
 		protected async Task UploadTestDataAsync(string JobStepId, IEnumerable<KeyValuePair<string, object>> TestData)
diff --git a/Engine/Source/Programs/Horde/HordeAgent/Execution/PerforceExecutor.cs b/Engine/Source/Programs/Horde/HordeAgent/Execution/PerforceExecutor.cs
--- a/Engine/Source/Programs/Horde/HordeAgent/Execution/PerforceExecutor.cs
+++ b/Engine/Source/Programs/Horde/HordeAgent/Execution/PerforceExecutor.cs
@@ -27,7 +27,8 @@ using System.Threading.Tasks;
 using EpicGames.Core;
 using EpicGames.Perforce;
 using Google.Protobuf.WellKnownTypes;
-using Datadog.Trace;
+using OpenTracing;
+using OpenTracing.Util;
 
 namespace HordeAgent.Execution
 {
@@ -60,10 +61,8 @@ namespace HordeAgent.Execution
 			// Setup and sync the autosdk workspace
 			if (AutoSdkWorkspaceInfo != null)
 			{
-				using (Scope Scope = Tracer.Instance.StartActive("Workspace"))
+				using (IScope Scope = GlobalTracer.Instance.BuildSpan("Workspace").WithResourceName("AutoSDK").StartActive())
 				{
-					Scope.Span.ResourceName = "AutoSDK";
-
 					AutoSdkWorkspace = await Utility.WorkspaceInfo.SetupWorkspaceAsync(AutoSdkWorkspaceInfo, RootDir, Logger, CancellationToken);
 
 					DirectoryReference LegacyDir = DirectoryReference.Combine(AutoSdkWorkspace.MetadataDir, "HostWin64");
@@ -87,10 +86,8 @@ namespace HordeAgent.Execution
 				}
 			}
 
-			using (Scope Scope = Tracer.Instance.StartActive("Workspace"))
+			using (IScope Scope = GlobalTracer.Instance.BuildSpan("Workspace").WithResourceName(WorkspaceInfo.Identifier).StartActive())
 			{
-				Scope.Span.ResourceName = WorkspaceInfo.Identifier;
-
 				// Sync the regular workspace
 				Workspace = await Utility.WorkspaceInfo.SetupWorkspaceAsync(WorkspaceInfo, RootDir, Logger, CancellationToken);
 
diff --git a/Engine/Source/Programs/Horde/HordeAgent/HordeAgent.csproj b/Engine/Source/Programs/Horde/HordeAgent/HordeAgent.csproj
--- a/Engine/Source/Programs/Horde/HordeAgent/HordeAgent.csproj
+++ b/Engine/Source/Programs/Horde/HordeAgent/HordeAgent.csproj
@@ -22,7 +22,7 @@
   <ItemGroup>
     <PackageReference Include="AWSSDK.Core" Version="3.5.1.49" />
     <PackageReference Include="AWSSDK.EC2" Version="3.5.1.2" />
-    <PackageReference Include="Datadog.Trace" Version="1.18.0" />
+    <PackageReference Include="Datadog.Trace.OpenTracing" Version="1.18.0" />
     <PackageReference Include="Google.Protobuf" Version="3.12.3" />
     <PackageReference Include="Grpc.Net.Client" Version="2.29.0" />
     <PackageReference Include="Grpc.Tools" Version="2.30.0">
@@ -36,7 +36,9 @@
     <PackageReference Include="Microsoft.Extensions.Http.Polly" Version="3.1.5" />
     <PackageReference Include="Microsoft.Extensions.Logging.Abstractions" Version="5.0.0" />
     <PackageReference Include="Microsoft.Extensions.Options.DataAnnotations" Version="3.1.5" />
+    <PackageReference Include="OpenTracing" Version="0.12.1" />
     <PackageReference Include="Serilog" Version="2.9.0" />
+    <PackageReference Include="Serilog.Enrichers.OpenTracing" Version="0.0.2" />
     <PackageReference Include="Serilog.Extensions.Hosting" Version="3.1.0" />
     <PackageReference Include="Serilog.Extensions.Logging" Version="3.0.1" />
     <PackageReference Include="Serilog.Sinks.Console" Version="3.1.1" />
diff --git a/Engine/Source/Programs/Horde/HordeAgent/Logging.cs b/Engine/Source/Programs/Horde/HordeAgent/Logging.cs
--- a/Engine/Source/Programs/Horde/HordeAgent/Logging.cs
+++ b/Engine/Source/Programs/Horde/HordeAgent/Logging.cs
@@ -11,6 +11,7 @@ using System;
 using System.IO;
 using System.Runtime.InteropServices;
 using EpicGames.Core;
+using Serilog.Enrichers.OpenTracing;
 
 namespace HordeAgent
 {
@@ -82,6 +83,7 @@ namespace HordeAgent
 				.MinimumLevel.Override("Serilog.AspNetCore.RequestLoggingMiddleware", LogEventLevel.Warning)
 				.MinimumLevel.ControlledBy(LogLevelSwitch)
 				.Enrich.FromLogContext()
+				.Enrich.WithOpenTracingContext()
 				.Enrich.With<VersionLogEnricher>()
 				.WriteTo.Console(outputTemplate: "[{Timestamp:HH:mm:ss} {Level:w3}] {Indent}{Message:l}{NewLine}{Exception}", theme: Theme)
 				.WriteTo.File(FileReference.Combine(Program.DataDir, "Log-.txt").FullName, fileSizeLimitBytes: 50 * 1024 * 1024, rollingInterval: RollingInterval.Day, rollOnFileSizeLimit: true, retainedFileCountLimit: 10)
diff --git a/Engine/Source/Programs/Horde/HordeAgent/Services/WorkerService.cs b/Engine/Source/Programs/Horde/HordeAgent/Services/WorkerService.cs
--- a/Engine/Source/Programs/Horde/HordeAgent/Services/WorkerService.cs
+++ b/Engine/Source/Programs/Horde/HordeAgent/Services/WorkerService.cs
@@ -42,17 +42,17 @@ using EpicGames.Core;
 using System.ComponentModel;
 using System.Diagnostics.CodeAnalysis;
 using System.Runtime.CompilerServices;
-using Datadog.Trace;
 using Serilog.Context;
 using Polly;
 using Microsoft.Extensions.Http;
 using Polly.Extensions.Http;
 using Amazon.Util;
 using Amazon.EC2;
-using Scope = Datadog.Trace.Scope;
 using Amazon.EC2.Model;
 using Status = Grpc.Core.Status;
 using Google.Protobuf;
+using OpenTracing.Util;
+using OpenTracing;
 
 [assembly: InternalsVisibleTo("HordeAgentTests")]
 
@@ -637,12 +637,11 @@ namespace HordeAgent.Services
 		/// <returns>Async task</returns>
 		async Task HandleLeaseAsync(IRpcConnection RpcConnection, string AgentId, LeaseInfo LeaseInfo)
 		{
-			using Scope Scope = Tracer.Instance.StartActive("HandleLease");
-			Scope.Span.ResourceName = LeaseInfo.Lease.Id;
+			using IScope Scope = GlobalTracer.Instance.BuildSpan("HandleLease").WithResourceName(LeaseInfo.Lease.Id).StartActive();
 			Scope.Span.SetTag("LeaseId", LeaseInfo.Lease.Id);
 			Scope.Span.SetTag("AgentId", AgentId);
-			using IDisposable TraceProperty = LogContext.PushProperty("dd.trace_id", CorrelationIdentifier.TraceId.ToString());
-			using IDisposable SpanProperty = LogContext.PushProperty("dd.span_id", CorrelationIdentifier.SpanId.ToString());
+//			using IDisposable TraceProperty = LogContext.PushProperty("dd.trace_id", CorrelationIdentifier.TraceId.ToString());
+//			using IDisposable SpanProperty = LogContext.PushProperty("dd.span_id", CorrelationIdentifier.SpanId.ToString());
 
 			Logger.LogInformation("Handling lease {LeaseId}", LeaseInfo.Lease.Id);
 
@@ -691,7 +690,7 @@ namespace HordeAgent.Services
 			ActionTask ActionTask;
 			if (LeaseInfo.Lease.Payload.TryUnpack(out ActionTask))
 			{
-				Tracer.Instance.ActiveScope?.Span.SetTag("task", "Action");
+				GlobalTracer.Instance.ActiveSpan?.SetTag("task", "Action");
 				Func<ILogger, Task<LeaseResult>> Handler = NewLogger => ActionAsync(RpcConnection, LeaseInfo.Lease.Id, ActionTask, NewLogger, LeaseInfo.CancellationTokenSource.Token);
 				return await HandleLeasePayloadWithLogAsync(RpcConnection, ActionTask.LogId, AgentId, null, null, Handler, LeaseInfo.CancellationTokenSource.Token);
 			}
@@ -699,7 +698,7 @@ namespace HordeAgent.Services
 			ConformTask ConformTask;
 			if (LeaseInfo.Lease.Payload.TryUnpack(out ConformTask))
 			{
-				Tracer.Instance.ActiveScope.Span.SetTag("task", "Conform");
+				GlobalTracer.Instance.ActiveSpan?.SetTag("task", "Conform");
 				Func<ILogger, Task<LeaseResult>> Handler = NewLogger => ConformAsync(RpcConnection, AgentId, LeaseInfo.Lease.Id, ConformTask, NewLogger, LeaseInfo.CancellationTokenSource.Token);
 				return await HandleLeasePayloadWithLogAsync(RpcConnection, ConformTask.LogId, AgentId, null, null, Handler, LeaseInfo.CancellationTokenSource.Token);
 			}
@@ -707,7 +706,7 @@ namespace HordeAgent.Services
 			ExecuteJobTask JobTask;
 			if (LeaseInfo.Lease.Payload.TryUnpack(out JobTask))
 			{
-				Tracer.Instance.ActiveScope.Span.SetTag("task", "Job");
+				GlobalTracer.Instance.ActiveSpan?.SetTag("task", "Job");
 				Func<ILogger, Task<LeaseResult>> Handler = NewLogger => ExecuteJobAsync(RpcConnection, AgentId, LeaseInfo.Lease.Id, JobTask, NewLogger, LeaseInfo.CancellationTokenSource.Token);
 				return await HandleLeasePayloadWithLogAsync(RpcConnection, JobTask.LogId, AgentId, JobTask.JobId, JobTask.BatchId, Handler, LeaseInfo.CancellationTokenSource.Token);
 			}
@@ -715,7 +714,7 @@ namespace HordeAgent.Services
 			UpgradeTask UpgradeTask;
 			if (LeaseInfo.Lease.Payload.TryUnpack(out UpgradeTask))
 			{
-				Tracer.Instance.ActiveScope.Span.SetTag("task", "Upgrade");
+				GlobalTracer.Instance.ActiveSpan?.SetTag("task", "Upgrade");
 				Func<ILogger, Task<LeaseResult>> Handler = NewLogger => UpgradeAsync(RpcConnection, AgentId, UpgradeTask, NewLogger, LeaseInfo.CancellationTokenSource.Token);
 				return await HandleLeasePayloadWithLogAsync(RpcConnection, UpgradeTask.LogId, AgentId, null, null, Handler, LeaseInfo.CancellationTokenSource.Token);
 			}
@@ -723,7 +722,7 @@ namespace HordeAgent.Services
 			ShutdownTask ShutdownTask;
 			if (LeaseInfo.Lease.Payload.TryUnpack(out ShutdownTask))
 			{
-				Tracer.Instance.ActiveScope.Span.SetTag("task", "Shutdown");
+				GlobalTracer.Instance.ActiveSpan?.SetTag("task", "Shutdown");
 				Func<ILogger, Task<LeaseResult>> Handler = NewLogger => ShutdownAsync(RpcConnection, AgentId, ShutdownTask, NewLogger, LeaseInfo.CancellationTokenSource.Token);
 				return await HandleLeasePayloadWithLogAsync(RpcConnection, ShutdownTask.LogId, AgentId, null, null, Handler, LeaseInfo.CancellationTokenSource.Token);
 			}
@@ -731,7 +730,7 @@ namespace HordeAgent.Services
 			RestartTask RestartTask;
 			if (LeaseInfo.Lease.Payload.TryUnpack(out RestartTask))
 			{
-				Tracer.Instance.ActiveScope.Span.SetTag("task", "Restart");
+				GlobalTracer.Instance.ActiveSpan?.SetTag("task", "Restart");
 				Func<ILogger, Task<LeaseResult>> Handler = NewLogger => RestartAsync(RpcConnection, AgentId, RestartTask, NewLogger, LeaseInfo.CancellationTokenSource.Token);
 				return await HandleLeasePayloadWithLogAsync(RpcConnection, RestartTask.LogId, AgentId, null, null, Handler, LeaseInfo.CancellationTokenSource.Token);
 			}
@@ -944,9 +943,9 @@ namespace HordeAgent.Services
 		internal async Task<LeaseResult> ExecuteJobAsync(IRpcConnection RpcClient, string AgentId, string LeaseId, ExecuteJobTask ExecuteTask, ILogger Logger, CancellationToken CancellationToken)
 		{
 			Logger.LogInformation("Executing job \"{JobName}\", jobId {JobId}, batchId {BatchId}, leaseId {LeaseId}", ExecuteTask.JobName, ExecuteTask.JobId, ExecuteTask.BatchId, LeaseId);
-			Tracer.Instance.ActiveScope?.Span.SetTag("jobId", ExecuteTask.JobId.ToString());
-			Tracer.Instance.ActiveScope?.Span.SetTag("jobName", ExecuteTask.JobName.ToString());
-			Tracer.Instance.ActiveScope?.Span.SetTag("batchId", ExecuteTask.BatchId.ToString());	
+			GlobalTracer.Instance.ActiveSpan?.SetTag("jobId", ExecuteTask.JobId.ToString());
+			GlobalTracer.Instance.ActiveSpan?.SetTag("jobName", ExecuteTask.JobName.ToString());
+			GlobalTracer.Instance.ActiveSpan?.SetTag("batchId", ExecuteTask.BatchId.ToString());	
 
 			// Start executing the current batch
 			BeginBatchResponse Batch = await RpcClient.InvokeAsync(x => x.BeginBatchAsync(new BeginBatchRequest(ExecuteTask.JobId, ExecuteTask.BatchId, LeaseId), null, null, CancellationToken), new RpcContext(), CancellationToken);
@@ -1001,7 +1000,7 @@ namespace HordeAgent.Services
 			BatchLogger.LogInformation("Initializing...");
 			using (BatchLogger.BeginIndentScope("  "))
 			{
-				using Scope Scope = Tracer.Instance.StartActive("Initialize");
+				using IScope Scope = GlobalTracer.Instance.BuildSpan("Initialize").StartActive();
 				await Executor.InitializeAsync(BatchLogger, CancellationToken);
 			}
 
@@ -1031,12 +1030,11 @@ namespace HordeAgent.Services
 				BatchLogger.LogInformation("Starting job {JobId}, batch {BatchId}, step {StepId}", ExecuteTask.JobId, ExecuteTask.BatchId, Step.StepId);
 
 				// Create a trace span
-				using Scope Scope = Tracer.Instance.StartActive("Execute");
-				Scope.Span.ResourceName = Step.Name;
+				using IScope Scope = GlobalTracer.Instance.BuildSpan("Execute").WithResourceName(Step.Name).StartActive();
 				Scope.Span.SetTag("stepId", Step.StepId);
 				Scope.Span.SetTag("logId", Step.LogId);
-				using IDisposable TraceProperty = LogContext.PushProperty("dd.trace_id", CorrelationIdentifier.TraceId.ToString());
-				using IDisposable SpanProperty = LogContext.PushProperty("dd.span_id", CorrelationIdentifier.SpanId.ToString());
+//				using IDisposable TraceProperty = LogContext.PushProperty("dd.trace_id", CorrelationIdentifier.TraceId.ToString());
+//				using IDisposable SpanProperty = LogContext.PushProperty("dd.span_id", CorrelationIdentifier.SpanId.ToString());
 
 				// Update the context to include information about this step
 				JobStepOutcome StepOutcome;
@@ -1096,7 +1094,7 @@ namespace HordeAgent.Services
 			BatchLogger.LogInformation("Finalizing...");
 			using (BatchLogger.BeginIndentScope("  "))
 			{
-				using Scope Scope = Tracer.Instance.StartActive("Finalize");
+				using IScope Scope = GlobalTracer.Instance.BuildSpan("Finalize").StartActive();
 				await Executor.FinalizeAsync(BatchLogger, CancellationToken);
 			}
 		}
diff --git a/Engine/Source/Programs/Horde/HordeAgent/Utility/RpcConnection.cs b/Engine/Source/Programs/Horde/HordeAgent/Utility/RpcConnection.cs
--- a/Engine/Source/Programs/Horde/HordeAgent/Utility/RpcConnection.cs
+++ b/Engine/Source/Programs/Horde/HordeAgent/Utility/RpcConnection.cs
@@ -1,12 +1,13 @@
 // Copyright Epic Games, Inc. All Rights Reserved.
 
-using Datadog.Trace;
 using EpicGames.Core;
 using Grpc.Core;
 using Grpc.Net.Client;
 using HordeCommon.Rpc;
 using Microsoft.Extensions.Hosting;
 using Microsoft.Extensions.Logging;
+using OpenTracing;
+using OpenTracing.Util;
 using System;
 using System.Collections.Concurrent;
 using System.Collections.Generic;
@@ -640,7 +641,7 @@ namespace HordeAgent.Utility
 		/// <returns>Async task</returns>
 		async Task HandleConnectionAsync(int ConnectionId, TaskCompletionSource<bool> ReconnectTaskSource)
 		{
-			using Scope Scope = Tracer.Instance.StartActive("RpcConnection {ConnectionId}");
+			using IScope Scope = GlobalTracer.Instance.BuildSpan("RpcConnection {ConnectionId}").StartActive();
 
 			Logger.LogInformation("Connection {ConnectionId}: Creating connection", ConnectionId);
 			try
