Missing file

diff --git a/Engine/Source/Programs/Horde/HordeAgent/Tracing.cs b/Engine/Source/Programs/Horde/HordeAgent/Tracing.cs
new file mode 100644
--- /dev/null
+++ b/Engine/Source/Programs/Horde/HordeAgent/Tracing.cs
@@ -0,0 +1,28 @@
+using OpenTracing;
+using System;
+using System.Collections.Generic;
+using System.Text;
+
+namespace HordeAgent
+{
+	static class Tracing
+	{
+		public static ISpanBuilder WithResourceName(this ISpanBuilder Builder, string? ResourceName)
+		{
+			if (ResourceName != null)
+			{
+				Builder = Builder.WithTag(Datadog.Trace.OpenTracing.DatadogTags.ResourceName, ResourceName);
+			}
+			return Builder;
+		}
+
+		public static ISpanBuilder WithServiceName(this ISpanBuilder Builder, string? ServiceName)
+		{
+			if(ServiceName != null)
+			{
+				Builder = Builder.WithTag(Datadog.Trace.OpenTracing.DatadogTags.ServiceName, ServiceName);
+			}
+			return Builder;
+		}
+}
+}
