Missing copyright notice.

diff --git a/Engine/Source/Programs/Horde/HordeAgent/Tracing.cs b/Engine/Source/Programs/Horde/HordeAgent/Tracing.cs
--- a/Engine/Source/Programs/Horde/HordeAgent/Tracing.cs
+++ b/Engine/Source/Programs/Horde/HordeAgent/Tracing.cs
@@ -1,3 +1,5 @@
+// Copyright Epic Games, Inc. All Rights Reserved.
+
 using OpenTracing;
 using System;
 using System.Collections.Generic;
