Horde: Add standard datadog properties into log events.

diff --git a/Engine/Source/Programs/Horde/HordeAgent/Commands/Service/RunCommand.cs b/Engine/Source/Programs/Horde/HordeAgent/Commands/Service/RunCommand.cs
--- a/Engine/Source/Programs/Horde/HordeAgent/Commands/Service/RunCommand.cs
+++ b/Engine/Source/Programs/Horde/HordeAgent/Commands/Service/RunCommand.cs
@@ -104,10 +104,11 @@ namespace HordeAgent.Modes.Service
 					AgentSettings Settings = new AgentSettings();
 					ConfigSection.Bind(Settings);
 
-					ServerProfile ServerProfile = Settings.GetCurrentServerProfile();
-					ConfigureTracing(ServerProfile.Environment, Program.Version);
+					ITracer OpenTracer = Datadog.Trace.OpenTracing.OpenTracingTracerFactory.CreateTracer();
+					GlobalTracer.Register(OpenTracer);
 
-					Logging.Version = Program.Version;
+					ServerProfile ServerProfile = Settings.GetCurrentServerProfile();
+					Logging.SetEnv(ServerProfile.Environment);
 
 					Services.AddHttpClient(Program.HordeServerClientName, Config =>
 					{
@@ -141,19 +142,5 @@ namespace HordeAgent.Modes.Service
 
 			return 0;
 		}
-
-		static void ConfigureTracing(string Environment, string Version)
-		{
-			TracerSettings Settings = TracerSettings.FromDefaultSources();
-			Settings.Environment = Environment;
-			Settings.ServiceName = "hordeagent";
-			Settings.ServiceVersion = Version;
-			Settings.LogsInjectionEnabled = true;
-
-			Tracer.Instance = new Tracer(Settings);
-
-			ITracer OpenTracer = Datadog.Trace.OpenTracing.OpenTracingTracerFactory.WrapTracer(Tracer.Instance);
-			GlobalTracer.Register(OpenTracer);
-		}
 	}
 }
diff --git a/Engine/Source/Programs/Horde/HordeAgent/Logging.cs b/Engine/Source/Programs/Horde/HordeAgent/Logging.cs
--- a/Engine/Source/Programs/Horde/HordeAgent/Logging.cs
+++ b/Engine/Source/Programs/Horde/HordeAgent/Logging.cs
@@ -12,11 +12,23 @@ using System.IO;
 using System.Runtime.InteropServices;
 using EpicGames.Core;
 using Serilog.Enrichers.OpenTracing;
+using System.Collections.Generic;
+using System.Threading;
 
 namespace HordeAgent
 {
 	static class Logging
 	{
+		static object LockObject = new object();
+		static string Env = "default";
+		static int PropsVersion = 1;
+
+		public static void SetEnv(string NewEnv)
+		{
+			Env = NewEnv;
+			Interlocked.Increment(ref PropsVersion);
+		}
+
 		static Lazy<Serilog.ILogger> Logger = new Lazy<Serilog.ILogger>(CreateSerilogLogger, true);
 		
 		/// <summary>
@@ -26,11 +38,34 @@ namespace HordeAgent
 		
 		public static LoggingLevelSwitch LogLevelSwitch =  new LoggingLevelSwitch();
 
-		private class VersionLogEnricher : ILogEventEnricher
+		private class DatadogLogEnricher : ILogEventEnricher
 		{
+			object LockObject = new object();
+			int CachedPropsVersion = -1;
+			List<LogEventProperty> Properties = new List<LogEventProperty>();
+
 			public void Enrich(LogEvent LogEvent, ILogEventPropertyFactory PropertyFactory)
 			{
-				LogEvent.AddOrUpdateProperty(PropertyFactory.CreateProperty("Version", Version));
+				if (CachedPropsVersion != PropsVersion)
+				{
+					lock (LockObject)
+					{
+						if (CachedPropsVersion != PropsVersion)
+						{
+							List<LogEventProperty> NewProperties = new List<LogEventProperty>();
+							NewProperties.Add(PropertyFactory.CreateProperty("dd.env", Env));
+							NewProperties.Add(PropertyFactory.CreateProperty("dd.service", "hordeagent"));
+							NewProperties.Add(PropertyFactory.CreateProperty("dd.version", Version));
+							Properties = NewProperties;
+						}
+					}
+				}
+
+				List<LogEventProperty> CurrentProperties = Properties;
+				foreach (LogEventProperty CurrentProperty in CurrentProperties)
+				{
+					LogEvent.AddOrUpdateProperty(CurrentProperty);
+				}
 			}
 		}
 		
@@ -84,7 +119,7 @@ namespace HordeAgent
 				.MinimumLevel.ControlledBy(LogLevelSwitch)
 				.Enrich.FromLogContext()
 				.Enrich.WithOpenTracingContext()
-				.Enrich.With<VersionLogEnricher>()
+				.Enrich.With<DatadogLogEnricher>()
 				.WriteTo.Console(outputTemplate: "[{Timestamp:HH:mm:ss} {Level:w3}] {Indent}{Message:l}{NewLine}{Exception}", theme: Theme)
 				.WriteTo.File(FileReference.Combine(Program.DataDir, "Log-.txt").FullName, fileSizeLimitBytes: 50 * 1024 * 1024, rollingInterval: RollingInterval.Day, rollOnFileSizeLimit: true, retainedFileCountLimit: 10)
 				.WriteTo.File(new JsonFormatter(renderMessage: true), FileReference.Combine(Program.DataDir, "Log-.json").FullName, fileSizeLimitBytes: 50 * 1024 * 1024, rollingInterval: RollingInterval.Day, rollOnFileSizeLimit: true, retainedFileCountLimit: 10)
