Horde: Revert changes to attempt to fix missing logs.

diff --git a/Engine/Source/Programs/Horde/HordeAgent/Commands/Service/RunCommand.cs b/Engine/Source/Programs/Horde/HordeAgent/Commands/Service/RunCommand.cs
--- a/Engine/Source/Programs/Horde/HordeAgent/Commands/Service/RunCommand.cs
+++ b/Engine/Source/Programs/Horde/HordeAgent/Commands/Service/RunCommand.cs
@@ -104,10 +104,9 @@ namespace HordeAgent.Modes.Service
 					AgentSettings Settings = new AgentSettings();
 					ConfigSection.Bind(Settings);
 
-					ITracer OpenTracer = Datadog.Trace.OpenTracing.OpenTracingTracerFactory.CreateTracer();
-					GlobalTracer.Register(OpenTracer);
-
 					ServerProfile ServerProfile = Settings.GetCurrentServerProfile();
+					ConfigureTracing(ServerProfile.Environment, Program.Version);
+
 					Logging.SetEnv(ServerProfile.Environment);
 
 					Services.AddHttpClient(Program.HordeServerClientName, Config =>
@@ -142,5 +141,19 @@ namespace HordeAgent.Modes.Service
 
 			return 0;
 		}
+
+		static void ConfigureTracing(string Environment, string Version)
+		{
+			TracerSettings Settings = TracerSettings.FromDefaultSources();
+			Settings.Environment = Environment;
+			Settings.ServiceName = "hordeagent";
+			Settings.ServiceVersion = Version;
+			Settings.LogsInjectionEnabled = true;
+
+			Tracer.Instance = new Tracer(Settings);
+
+			ITracer OpenTracer = Datadog.Trace.OpenTracing.OpenTracingTracerFactory.WrapTracer(Tracer.Instance);
+			GlobalTracer.Register(OpenTracer);
+		}
 	}
 }
