Horde: Fix missing version field on log messages.

diff --git a/Engine/Source/Programs/Horde/HordeAgent/Logging.cs b/Engine/Source/Programs/Horde/HordeAgent/Logging.cs
--- a/Engine/Source/Programs/Horde/HordeAgent/Logging.cs
+++ b/Engine/Source/Programs/Horde/HordeAgent/Logging.cs
@@ -19,53 +19,24 @@ namespace HordeAgent
 {
 	static class Logging
 	{
-		static object LockObject = new object();
 		static string Env = "default";
-		static int PropsVersion = 1;
 
 		public static void SetEnv(string NewEnv)
 		{
 			Env = NewEnv;
-			Interlocked.Increment(ref PropsVersion);
 		}
 
 		static Lazy<Serilog.ILogger> Logger = new Lazy<Serilog.ILogger>(CreateSerilogLogger, true);
 		
-		/// <summary>
-		/// Current version of the HordeAgent
-		/// </summary>
-		public static string Version = "";
-		
 		public static LoggingLevelSwitch LogLevelSwitch =  new LoggingLevelSwitch();
 
 		private class DatadogLogEnricher : ILogEventEnricher
 		{
-			object LockObject = new object();
-			int CachedPropsVersion = -1;
-			List<LogEventProperty> Properties = new List<LogEventProperty>();
-
 			public void Enrich(LogEvent LogEvent, ILogEventPropertyFactory PropertyFactory)
 			{
-				if (CachedPropsVersion != PropsVersion)
-				{
-					lock (LockObject)
-					{
-						if (CachedPropsVersion != PropsVersion)
-						{
-							List<LogEventProperty> NewProperties = new List<LogEventProperty>();
-							NewProperties.Add(PropertyFactory.CreateProperty("dd.env", Env));
-							NewProperties.Add(PropertyFactory.CreateProperty("dd.service", "hordeagent"));
-							NewProperties.Add(PropertyFactory.CreateProperty("dd.version", Version));
-							Properties = NewProperties;
-						}
-					}
-				}
-
-				List<LogEventProperty> CurrentProperties = Properties;
-				foreach (LogEventProperty CurrentProperty in CurrentProperties)
-				{
-					LogEvent.AddOrUpdateProperty(CurrentProperty);
-				}
+				LogEvent.AddOrUpdateProperty(PropertyFactory.CreateProperty("dd.env", Env));
+				LogEvent.AddOrUpdateProperty(PropertyFactory.CreateProperty("dd.service", "hordeagent"));
+				LogEvent.AddOrUpdateProperty(PropertyFactory.CreateProperty("dd.version", Program.Version));
 			}
 		}
 		
