Horde: Add Datadog dd.trace_id and dd.span_id to log events.

diff --git a/Engine/Source/Programs/Horde/HordeAgent/HordeAgent.csproj b/Engine/Source/Programs/Horde/HordeAgent/HordeAgent.csproj
--- a/Engine/Source/Programs/Horde/HordeAgent/HordeAgent.csproj
+++ b/Engine/Source/Programs/Horde/HordeAgent/HordeAgent.csproj
@@ -38,7 +38,6 @@
     <PackageReference Include="Microsoft.Extensions.Options.DataAnnotations" Version="3.1.5" />
     <PackageReference Include="OpenTracing" Version="0.12.1" />
     <PackageReference Include="Serilog" Version="2.9.0" />
-    <PackageReference Include="Serilog.Enrichers.OpenTracing" Version="0.0.2" />
     <PackageReference Include="Serilog.Extensions.Hosting" Version="3.1.0" />
     <PackageReference Include="Serilog.Extensions.Logging" Version="3.0.1" />
     <PackageReference Include="Serilog.Sinks.Console" Version="3.1.1" />
diff --git a/Engine/Source/Programs/Horde/HordeAgent/Logging.cs b/Engine/Source/Programs/Horde/HordeAgent/Logging.cs
--- a/Engine/Source/Programs/Horde/HordeAgent/Logging.cs
+++ b/Engine/Source/Programs/Horde/HordeAgent/Logging.cs
@@ -11,9 +11,10 @@ using System;
 using System.IO;
 using System.Runtime.InteropServices;
 using EpicGames.Core;
-using Serilog.Enrichers.OpenTracing;
 using System.Collections.Generic;
 using System.Threading;
+using OpenTracing.Util;
+using OpenTracing;
 
 namespace HordeAgent
 {
@@ -37,6 +38,13 @@ namespace HordeAgent
 				LogEvent.AddOrUpdateProperty(PropertyFactory.CreateProperty("dd.env", Env));
 				LogEvent.AddOrUpdateProperty(PropertyFactory.CreateProperty("dd.service", "hordeagent"));
 				LogEvent.AddOrUpdateProperty(PropertyFactory.CreateProperty("dd.version", Program.Version));
+
+				ISpan? Span = GlobalTracer.Instance?.ActiveSpan;
+				if (Span != null)
+				{
+					LogEvent.AddPropertyIfAbsent(PropertyFactory.CreateProperty("dd.trace_id", Span.Context.TraceId));
+					LogEvent.AddPropertyIfAbsent(PropertyFactory.CreateProperty("dd.span_id", Span.Context.SpanId));
+				}
 			}
 		}
 		
@@ -89,7 +97,6 @@ namespace HordeAgent
 				.MinimumLevel.Override("Serilog.AspNetCore.RequestLoggingMiddleware", LogEventLevel.Warning)
 				.MinimumLevel.ControlledBy(LogLevelSwitch)
 				.Enrich.FromLogContext()
-				.Enrich.WithOpenTracingContext()
 				.Enrich.With<DatadogLogEnricher>()
 				.WriteTo.Console(outputTemplate: "[{Timestamp:HH:mm:ss} {Level:w3}] {Indent}{Message:l}{NewLine}{Exception}", theme: Theme)
 				.WriteTo.File(FileReference.Combine(Program.DataDir, "Log-.txt").FullName, fileSizeLimitBytes: 50 * 1024 * 1024, rollingInterval: RollingInterval.Day, rollOnFileSizeLimit: true, retainedFileCountLimit: 10)
