Horde: Use a custom enricher to set trace and span id in log events.

diff --git a/Engine/Source/Programs/Horde/HordeServer/Program.cs b/Engine/Source/Programs/Horde/HordeServer/Program.cs
--- a/Engine/Source/Programs/Horde/HordeServer/Program.cs
+++ b/Engine/Source/Programs/Horde/HordeServer/Program.cs
@@ -22,9 +22,11 @@ using Microsoft.Extensions.Logging;
 using MongoDB.Bson;
 using MongoDB.Bson.Serialization;
 using MongoDB.Driver;
+using OpenTracing;
 using OpenTracing.Util;
 using Serilog;
 using Serilog.Configuration;
+using Serilog.Core;
 using Serilog.Enrichers.OpenTracing;
 using Serilog.Events;
 using Serilog.Filters;
@@ -102,6 +104,19 @@ namespace HordeServer
 		}
 	}
 
+	class DatadogLogEnricher : ILogEventEnricher
+	{
+		public void Enrich(LogEvent LogEvent, ILogEventPropertyFactory PropertyFactory)
+		{
+			ISpan? Span = GlobalTracer.Instance?.ActiveSpan;
+			if (Span != null)
+			{
+				LogEvent.AddPropertyIfAbsent(PropertyFactory.CreateProperty("dd.trace_id", Span.Context.TraceId));
+				LogEvent.AddPropertyIfAbsent(PropertyFactory.CreateProperty("dd.span_id", Span.Context.SpanId));
+			}
+		}
+	}
+
 	class Program
 	{
 		public static DirectoryReference AppDir { get; } = GetAppDir();
@@ -153,7 +168,7 @@ namespace HordeServer
 			Serilog.Log.Logger = new LoggerConfiguration()
 				.WithHordeConfig(HordeSettings)
 				.Enrich.FromLogContext()
-				.Enrich.WithOpenTracingContext()
+				.Enrich.With<DatadogLogEnricher>()
 				.WriteTo.Console(HordeSettings)
 				.WriteTo.File(Path.Combine(LogDir.FullName, "Log.txt"), outputTemplate: "[{Timestamp:HH:mm:ss} {Level:w3}] {Indent}{Message:l}{NewLine}{Exception} [{SourceContext}]", rollingInterval: RollingInterval.Day, rollOnFileSizeLimit: true, fileSizeLimitBytes: 20 * 1024 * 1024, retainedFileCountLimit: 10)
 				.WriteTo.File(new JsonFormatter(renderMessage: true), Path.Combine(LogDir.FullName, "Log.json"), rollingInterval: RollingInterval.Day, rollOnFileSizeLimit: true, fileSizeLimitBytes: 20 * 1024 * 1024, retainedFileCountLimit: 10)
