Horde: Remove legacy code for upgrading pool collections.

diff --git a/Engine/Source/Programs/Horde/HordeServer/Collections/Impl/PoolCollection.cs b/Engine/Source/Programs/Horde/HordeServer/Collections/Impl/PoolCollection.cs
--- a/Engine/Source/Programs/Horde/HordeServer/Collections/Impl/PoolCollection.cs
+++ b/Engine/Source/Programs/Horde/HordeServer/Collections/Impl/PoolCollection.cs
@@ -79,19 +79,6 @@ namespace HordeServer.Collections.Impl
 		public PoolCollection(DatabaseService DatabaseService)
 		{
 			Pools = DatabaseService.GetCollection<PoolDocument>("Pools");
-
-			if (!DatabaseService.ReadOnlyMode)
-			{
-				// Drop any legacy indexes
-				List<BsonDocument> IndexList = Pools.Indexes.List().ToList();
-				if (IndexList.Any(x => { BsonValue Name; return x.TryGetValue("name", out Name) && Name.AsString == "Deleted_1_Name_1"; }))
-				{
-					Pools.Indexes.DropOne("Deleted_1_Name_1");
-				}
-				
-				// Remove any pools with an objectid for unique id
-				Pools.DeleteMany("{ \"_id\": { \"$type\": 7 } }");
-			}
 		}
 
 		/// <inheritdoc/>
