Horde: Fix configuration from local files on paths containing spaces.

diff --git a/Engine/Source/Programs/Horde/HordeServer/Services/ConfigService.cs b/Engine/Source/Programs/Horde/HordeServer/Services/ConfigService.cs
--- a/Engine/Source/Programs/Horde/HordeServer/Services/ConfigService.cs
+++ b/Engine/Source/Programs/Horde/HordeServer/Services/ConfigService.cs
@@ -381,7 +381,7 @@ namespace HordeServer.Services
 			switch(ConfigPath.Scheme)
 			{
 				case FileScheme:
-					return File.ReadAllBytesAsync(ConfigPath.AbsolutePath);
+					return File.ReadAllBytesAsync(ConfigPath.LocalPath);
 				case PerforceScheme:
 					return PerforceService.PrintAsync(ConfigPath.Host, ConfigPath.AbsolutePath);
 				default:
