Horde: Fix tracing across gRPC calls.

diff --git a/Engine/Source/Programs/Horde/HordeAgent/HordeAgent.csproj b/Engine/Source/Programs/Horde/HordeAgent/HordeAgent.csproj
--- a/Engine/Source/Programs/Horde/HordeAgent/HordeAgent.csproj
+++ b/Engine/Source/Programs/Horde/HordeAgent/HordeAgent.csproj
@@ -37,6 +37,7 @@
     <PackageReference Include="Microsoft.Extensions.Logging.Abstractions" Version="5.0.0" />
     <PackageReference Include="Microsoft.Extensions.Options.DataAnnotations" Version="3.1.5" />
     <PackageReference Include="OpenTracing" Version="0.12.1" />
+    <PackageReference Include="OpenTracing.Contrib.Grpc" Version="0.4.0" />
     <PackageReference Include="Serilog" Version="2.9.0" />
     <PackageReference Include="Serilog.Extensions.Hosting" Version="3.1.0" />
     <PackageReference Include="Serilog.Extensions.Logging" Version="3.0.1" />
diff --git a/Engine/Source/Programs/Horde/HordeAgent/Utility/RpcConnection.cs b/Engine/Source/Programs/Horde/HordeAgent/Utility/RpcConnection.cs
--- a/Engine/Source/Programs/Horde/HordeAgent/Utility/RpcConnection.cs
+++ b/Engine/Source/Programs/Horde/HordeAgent/Utility/RpcConnection.cs
@@ -2,11 +2,13 @@
 
 using EpicGames.Core;
 using Grpc.Core;
+using Grpc.Core.Interceptors;
 using Grpc.Net.Client;
 using HordeCommon.Rpc;
 using Microsoft.Extensions.Hosting;
 using Microsoft.Extensions.Logging;
 using OpenTracing;
+using OpenTracing.Contrib.Grpc.Interceptors;
 using OpenTracing.Util;
 using System;
 using System.Collections.Concurrent;
@@ -670,7 +672,7 @@ namespace HordeAgent.Utility
 		{
 			using (GrpcChannel Channel = CreateGrpcChannel())
 			{
-				HordeRpc.HordeRpcClient Client = new HordeRpc.HordeRpcClient(Channel);
+				HordeRpc.HordeRpcClient Client = new HordeRpc.HordeRpcClient(Channel.Intercept(new ClientTracingInterceptor(GlobalTracer.Instance)));
 
 				RpcSubConnection? SubConnection = null;
 				try
diff --git a/Engine/Source/Programs/Horde/HordeServer/HordeServer.csproj b/Engine/Source/Programs/Horde/HordeServer/HordeServer.csproj
--- a/Engine/Source/Programs/Horde/HordeServer/HordeServer.csproj
+++ b/Engine/Source/Programs/Horde/HordeServer/HordeServer.csproj
@@ -36,6 +36,7 @@
     <PackageReference Include="Microsoft.Extensions.Caching.Memory" Version="3.1.6" />
     <PackageReference Include="MongoDB.Driver" Version="2.10.4" />
     <PackageReference Include="OpenTracing" Version="0.12.1" />
+    <PackageReference Include="OpenTracing.Contrib.Grpc" Version="0.4.0" />
     <PackageReference Include="Polly" Version="7.2.1" />
     <PackageReference Include="Polly.Extensions.Http" Version="3.0.0" />
     <PackageReference Include="Serilog.AspNetCore" Version="3.2.0" />
diff --git a/Engine/Source/Programs/Horde/HordeServer/Startup.cs b/Engine/Source/Programs/Horde/HordeServer/Startup.cs
--- a/Engine/Source/Programs/Horde/HordeServer/Startup.cs
+++ b/Engine/Source/Programs/Horde/HordeServer/Startup.cs
@@ -77,6 +77,8 @@ using HordeServer.Storage.Backends;
 using HordeServer.Storage.Services;
 using HordeServer.Commits.Impl;
 using HordeServer.Commits;
+using OpenTracing.Contrib.Grpc.Interceptors;
+using OpenTracing.Util;
 
 namespace HordeServer
 {
@@ -293,6 +295,7 @@ namespace HordeServer
 				Options.MaxReceiveMessageSize = 20 * 1024 * 1024; // 20 MB (packaged builds of Horde agent can be large) 
 				Options.Interceptors.Add(typeof(LifetimeGrpcInterceptor));
 				Options.Interceptors.Add(typeof(GrpcExceptionInterceptor));
+				Options.Interceptors.Add<ServerTracingInterceptor>(GlobalTracer.Instance);
 			});
 			Services.AddGrpcReflection();
 
