Horde: Add support for max exec timeout on REAPI executions

Will prevent hung jobs idling forever. Tests added for cancelled and timeouts executions

diff --git a/Engine/Source/Programs/Horde/HordeAgent/Execution/ActionExecutor.cs b/Engine/Source/Programs/Horde/HordeAgent/Execution/ActionExecutor.cs
--- a/Engine/Source/Programs/Horde/HordeAgent/Execution/ActionExecutor.cs
+++ b/Engine/Source/Programs/Horde/HordeAgent/Execution/ActionExecutor.cs
@@ -85,6 +85,67 @@ namespace HordeAgent
 			this.Logger = Logger;
 		}
 
+		/// <summary>
+		/// Read stdout/stderr and wait until process completes
+		/// </summary>
+		/// <param name="Process">The process reading from</param>
+		/// <param name="Timeout">Max execution timeout for process</param>
+		/// <param name="CancelToken">Cancellation token</param>
+		/// <returns>stdout and stderr data</returns>
+		/// <exception cref="RpcException">Raised for either timeout or cancel</exception>
+		private async Task<(byte[] StdOutData, byte[] StdOutErr)> ReadProcessStreams(ManagedProcess Process, TimeSpan Timeout, CancellationToken CancelToken)
+		{
+			using MemoryStream StdOutStream = new MemoryStream();
+			using MemoryStream StdErrStream = new MemoryStream();
+
+			void CheckIfCancelled()
+			{
+				if (CancelToken.IsCancellationRequested)
+				{
+					throw new RpcException(new Grpc.Core.Status(StatusCode.Cancelled, "Action cancelled"));
+				}
+			}
+			
+			try
+			{
+				CheckIfCancelled();
+
+				// Read stdout/stderr without cancellation token.
+				Task StdOutReadTask = Process.StdOut.CopyToAsync(StdOutStream);
+				Task StdErrReadTask = Process.StdErr.CopyToAsync(StdErrStream);
+				
+				// Instead, create a separate task that will wait for either timeout or cancellation
+				// as cancel interruptions are not reliable with Stream.CopyToAsync()
+				Task TimeoutTask = Task.Delay(Timeout, CancelToken);
+				
+				for (;;)
+				{
+					Task t = await Task.WhenAny(StdOutReadTask, StdErrReadTask, TimeoutTask);
+					bool BothStreamsRead = StdOutReadTask.IsCompleted && StdErrReadTask.IsCompleted;
+					
+					if (BothStreamsRead) { break; }
+					else if (t == StdOutReadTask) {} // Ignore, wait for stderr to complete
+					else if (t == StdErrReadTask) {} // Ignore, wait for stdout to complete
+					else if (t == TimeoutTask)
+					{
+						throw new RpcException(new Grpc.Core.Status(StatusCode.DeadlineExceeded, $"Action timed out after {Timeout.TotalMilliseconds} ms"));
+					}
+					else
+					{
+						throw new Exception("Unhandled task condition waiting for process streams");
+					}
+				}
+			}
+			catch (TaskCanceledException)
+			{
+				// Raising an exception will force the Process variable to be disposed and terminated.
+				CheckIfCancelled();
+				throw;
+			}
+
+			return (StdOutStream.ToArray(), StdErrStream.ToArray());
+		}
+		
 		/// <summary>
 		/// Execute an action
 		/// </summary>
@@ -143,22 +204,13 @@ namespace HordeAgent
 				}
 
 				Logger.LogInformation("Executing {FileName} with arguments {Arguments}", FileName, Arguments);
+
+				TimeSpan Timeout = Action.Timeout == null ? TimeSpan.FromMinutes(5) : Action.Timeout.ToTimeSpan();
 				DateTimeOffset ExecutionStartTime = DateTimeOffset.UtcNow;
 				using (ManagedProcess Process = new ManagedProcess(ProcessGroup, FileName, Arguments, WorkingDirectory, NewEnvironment, ProcessPriorityClass.Normal, ManagedProcessFlags.None))
 				{
-					byte[] StdOutData;
-					byte[] StdErrData;
-					using (MemoryStream StdOutStream = new MemoryStream())
-					using (MemoryStream StdErrStream = new MemoryStream())
-					{
-						await Task.WhenAll(Process.StdOut.CopyToAsync(StdOutStream, CancellationToken), Process.StdErr.CopyToAsync(StdErrStream, CancellationToken));
-						if (CancellationToken.IsCancellationRequested)
-						{
-							throw new RpcException(new Grpc.Core.Status(StatusCode.Cancelled, "Action Cancelled"));
-						}
-						StdOutData = StdOutStream.ToArray();
-						StdErrData = StdErrStream.ToArray();
-					}
+					(byte[] StdOutData, byte[] StdErrData) = await ReadProcessStreams(Process, Timeout, CancellationToken);
+					
 					DateTimeOffset ExecutionCompletedTime = DateTimeOffset.UtcNow;
 						
 
diff --git a/Engine/Source/Programs/Horde/HordeAgentTests/ActionExecutorTest.cs b/Engine/Source/Programs/Horde/HordeAgentTests/ActionExecutorTest.cs
--- a/Engine/Source/Programs/Horde/HordeAgentTests/ActionExecutorTest.cs
+++ b/Engine/Source/Programs/Horde/HordeAgentTests/ActionExecutorTest.cs
@@ -113,6 +113,36 @@ namespace HordeAgentTests
 				new DirectoryReference(TempDir), DateTimeOffset.UtcNow, CancellationToken.None));
 		}
 		
+		
+		[TestMethod]
+		public async Task ExecuteActionWithinTimeout()
+		{
+			ActionExecutor Executor = new ActionExecutor(AgentName, InstanceName, CasClient, null!, ActionRpcClient, Logger);
+			ActionTask ActionTask1 = TaskTest.CreateActionTask(InstanceName, CasClient, CommandSleep: TimeSpan.FromMilliseconds(200), Timeout: TimeSpan.FromMilliseconds(1500));
+			await Executor.ExecuteActionAsync("myLeaseId1", ActionTask1, new DirectoryReference(TempDir),
+				DateTimeOffset.UtcNow, CancellationToken.None);
+		}
+		
+		[TestMethod]
+		public async Task ExecuteActionCancel()
+		{
+			ActionExecutor Executor = new ActionExecutor(AgentName, InstanceName, CasClient, null!, ActionRpcClient, Logger);
+			ActionTask ActionTask = TaskTest.CreateActionTask(InstanceName, CasClient, CommandSleep: TimeSpan.FromMilliseconds(15000));
+			CancellationTokenSource Cts = new CancellationTokenSource();
+			Cts.CancelAfter(200);
+			RpcException Ex = await Assert.ThrowsExceptionAsync<RpcException>(() => Executor.ExecuteActionAsync("myLeaseId1", ActionTask, new DirectoryReference(TempDir), DateTimeOffset.UtcNow, Cts.Token));
+			Assert.AreEqual(StatusCode.Cancelled, Ex.Status.StatusCode);
+		}
+		
+		[TestMethod]
+		public async Task ExecuteActionExceedingTimeout()
+		{
+			ActionExecutor Executor = new ActionExecutor(AgentName, InstanceName, CasClient, null!, ActionRpcClient, Logger);
+			ActionTask ActionTask = TaskTest.CreateActionTask(InstanceName, CasClient, CommandSleep: TimeSpan.FromMilliseconds(15000), Timeout: TimeSpan.FromMilliseconds(200));
+			RpcException Ex = await Assert.ThrowsExceptionAsync<RpcException>(() => Executor.ExecuteActionAsync("myLeaseId1", ActionTask, new DirectoryReference(TempDir), DateTimeOffset.UtcNow, CancellationToken.None));
+			Assert.AreEqual(StatusCode.DeadlineExceeded, Ex.Status.StatusCode);
+		}
+		
 		private static ILogger CreateLogger()
 		{
 			LoggerFactory LoggerFactory = new LoggerFactory();
diff --git a/Engine/Source/Programs/Horde/HordeAgentTests/TaskTest.cs b/Engine/Source/Programs/Horde/HordeAgentTests/TaskTest.cs
--- a/Engine/Source/Programs/Horde/HordeAgentTests/TaskTest.cs
+++ b/Engine/Source/Programs/Horde/HordeAgentTests/TaskTest.cs
@@ -132,12 +132,14 @@ namespace HordeAgentTests
 		/// </summary>
 		/// <param name="InstanceName">Instance name</param>
 		/// <param name="Cas">The content-addressable store being used</param>
+		/// <param name="Timeout">Timeout before killing the command</param>
 		/// <param name="UploadBin">True to upload the remote exec test binary to CAS</param>
 		/// <param name="UploadDir">True to upload directory to CAS</param>
 		/// <param name="UploadCommand">True to upload the command to CAS</param>
 		/// <param name="UploadAction">True to upload action to CAS</param>
 		/// <returns>An ActionTask populated in CAS</returns>
-		public static ActionTask CreateActionTask(string InstanceName, FakeCasClient Cas, bool UploadBin = true, bool UploadDir = true, bool UploadCommand = true, bool UploadAction = true)
+		public static ActionTask CreateActionTask(string InstanceName, FakeCasClient Cas, TimeSpan? CommandSleep = null, int CommandExitCode = 0, TimeSpan? Timeout = null,
+			bool UploadBin = true, bool UploadDir = true, bool UploadCommand = true, bool UploadAction = true)
 		{
 			string BinName = "remote-exec-test-bin.exe";
 			string AssemblyPath = Path.GetDirectoryName(System.Reflection.Assembly.GetExecutingAssembly().Location)!;
@@ -164,6 +166,16 @@ namespace HordeAgentTests
 
 			Build.Bazel.Remote.Execution.V2.Command Command = new Build.Bazel.Remote.Execution.V2.Command();
 			Command.Arguments.Add(ExeFile.Name);
+
+			int CommandSleepMs = 0;
+			if (CommandSleep != null)
+			{
+				CommandSleepMs = (int) CommandSleep.Value.TotalMilliseconds;
+			}
+			
+			Command.Arguments.Add(CommandSleepMs.ToString());
+			Command.Arguments.Add(CommandExitCode.ToString());
+			
 			byte[] CommandData = Command.ToByteArray();
 			if (UploadCommand)
 			{
@@ -174,6 +186,11 @@ namespace HordeAgentTests
 			Action.CommandDigest = GetDigest(CommandData);
 			Action.InputRootDigest = GetDigest(DirData);
 			Action.DoNotCache = true;
+
+			if (Timeout != null)
+			{
+				Action.Timeout = Duration.FromTimeSpan(Timeout.Value);	
+			}
 			
 			byte[] ActionData = Action.ToByteArray();
 			if (UploadAction)
