Horde: Fix bug with incorrect exception getting raised during cancelling of REAPI actions

diff --git a/Engine/Source/Programs/Horde/HordeAgent/Execution/ActionExecutor.cs b/Engine/Source/Programs/Horde/HordeAgent/Execution/ActionExecutor.cs
--- a/Engine/Source/Programs/Horde/HordeAgent/Execution/ActionExecutor.cs
+++ b/Engine/Source/Programs/Horde/HordeAgent/Execution/ActionExecutor.cs
@@ -128,6 +128,7 @@ namespace HordeAgent
 					else if (t == StdErrReadTask) {} // Ignore, wait for stdout to complete
 					else if (t == TimeoutTask)
 					{
+						CheckIfCancelled();
 						throw new RpcException(new Grpc.Core.Status(StatusCode.DeadlineExceeded, $"Action timed out after {Timeout.TotalMilliseconds} ms"));
 					}
 					else
diff --git a/Engine/Source/Programs/Horde/HordeAgentTests/ActionExecutorTest.cs b/Engine/Source/Programs/Horde/HordeAgentTests/ActionExecutorTest.cs
--- a/Engine/Source/Programs/Horde/HordeAgentTests/ActionExecutorTest.cs
+++ b/Engine/Source/Programs/Horde/HordeAgentTests/ActionExecutorTest.cs
@@ -129,7 +129,7 @@ namespace HordeAgentTests
 			ActionExecutor Executor = new ActionExecutor(AgentName, InstanceName, CasClient, null!, ActionRpcClient, Logger);
 			ActionTask ActionTask = TaskTest.CreateActionTask(InstanceName, CasClient, CommandSleep: TimeSpan.FromMilliseconds(15000));
 			CancellationTokenSource Cts = new CancellationTokenSource();
-			Cts.CancelAfter(200);
+			Cts.CancelAfter(2000);
 			RpcException Ex = await Assert.ThrowsExceptionAsync<RpcException>(() => Executor.ExecuteActionAsync("myLeaseId1", ActionTask, new DirectoryReference(TempDir), DateTimeOffset.UtcNow, Cts.Token));
 			Assert.AreEqual(StatusCode.Cancelled, Ex.Status.StatusCode);
 		}
