Horde: Increase max gRPC payload sizes to accommodate larger Horde Agent updates

diff --git a/Engine/Source/Programs/Horde/HordeAgent/Commands/Utilities/PublishCommand.cs b/Engine/Source/Programs/Horde/HordeAgent/Commands/Utilities/PublishCommand.cs
--- a/Engine/Source/Programs/Horde/HordeAgent/Commands/Utilities/PublishCommand.cs
+++ b/Engine/Source/Programs/Horde/HordeAgent/Commands/Utilities/PublishCommand.cs
@@ -105,7 +105,13 @@ namespace HordeAgent.Commands.Utilities
 			HttpClient.BaseAddress = new Uri(ServerProfile.Url);
 			HttpClient.DefaultRequestHeaders.Authorization = new AuthenticationHeaderValue("Bearer", Token);
 
-			using (GrpcChannel RpcChannel = GrpcChannel.ForAddress(ServerProfile.Url, new GrpcChannelOptions { HttpClient = HttpClient }))
+			GrpcChannelOptions ChannelOptions = new GrpcChannelOptions
+			{
+				HttpClient = HttpClient,
+				MaxReceiveMessageSize = 100 * 1024 * 1024, // 100 MB
+				MaxSendMessageSize = 100 * 1024 * 1024 // 100 MB
+			};
+			using (GrpcChannel RpcChannel = GrpcChannel.ForAddress(ServerProfile.Url, ChannelOptions))
 			{
 				HordeRpc.HordeRpcClient RpcClient = new HordeRpc.HordeRpcClient(RpcChannel);
 
diff --git a/Engine/Source/Programs/Horde/HordeServer/Startup.cs b/Engine/Source/Programs/Horde/HordeServer/Startup.cs
--- a/Engine/Source/Programs/Horde/HordeServer/Startup.cs
+++ b/Engine/Source/Programs/Horde/HordeServer/Startup.cs
@@ -292,7 +292,7 @@ namespace HordeServer
 			Services.AddGrpc(Options =>
 			{
 				Options.EnableDetailedErrors = true;
-				Options.MaxReceiveMessageSize = 20 * 1024 * 1024; // 20 MB (packaged builds of Horde agent can be large) 
+				Options.MaxReceiveMessageSize = 200 * 1024 * 1024; // 100 MB (packaged builds of Horde agent can be large) 
 				Options.Interceptors.Add(typeof(LifetimeGrpcInterceptor));
 				Options.Interceptors.Add(typeof(GrpcExceptionInterceptor));
 				Options.Interceptors.Add<ServerTracingInterceptor>(GlobalTracer.Instance);
