Horde: Support device renewel json response for streams
#rnx

diff --git a/Engine/Source/Programs/Horde/HordeServer/Controllers/DevicesController.cs b/Engine/Source/Programs/Horde/HordeServer/Controllers/DevicesController.cs
--- a/Engine/Source/Programs/Horde/HordeServer/Controllers/DevicesController.cs
+++ b/Engine/Source/Programs/Horde/HordeServer/Controllers/DevicesController.cs
@@ -820,7 +820,17 @@ namespace HordeServer.Controllers
 				return BadRequest();
 			}
 
-			return Ok();
+            List<IDevice> Devices = await DeviceService.GetDevicesAsync(Reservation.Devices);
+
+            GetLegacyReservationResponse Response = new GetLegacyReservationResponse();
+
+			Response.Guid = Reservation.LegacyGuid;
+			Response.DeviceNames = Reservation.Devices.Select(DeviceId => Devices.First(D => D.Id == DeviceId).Name).ToArray();
+			Response.HostName = Reservation.Hostname ?? "";
+			Response.StartDateTime = Reservation.CreateTimeUtc.ToString("O", System.Globalization.CultureInfo.InvariantCulture);
+			Response.Duration = "00::10:00"; // matches gauntlet duration
+
+			return new JsonResult(Response, new JsonSerializerOptions() { PropertyNamingPolicy = null });
 		}
 
 		/// <summary>
