Horde: Add a default Perforce cluster (perforce:1666) to allow Horde to run out of the box.

diff --git a/Engine/Source/Programs/Horde/HordeServer/Models/Globals.cs b/Engine/Source/Programs/Horde/HordeServer/Models/Globals.cs
--- a/Engine/Source/Programs/Horde/HordeServer/Models/Globals.cs
+++ b/Engine/Source/Programs/Horde/HordeServer/Models/Globals.cs
@@ -375,16 +375,32 @@ namespace HordeServer.Models
 		/// </summary>
 		/// <param name="Name">Name of the cluster</param>
 		/// <returns></returns>
-		public PerforceCluster? FindPerforceCluster(string? Name)
+		public PerforceCluster FindPerforceCluster(string? Name)
 		{
+			List<PerforceCluster> Clusters = PerforceClusters;
+			if (Clusters.Count == 0)
+			{
+				Clusters = DefaultClusters;
+			}
+
 			if (Name == null)
 			{
-				return PerforceClusters.FirstOrDefault();
+				return Clusters.FirstOrDefault();
 			}
 			else
 			{
-				return PerforceClusters.FirstOrDefault(x => String.Equals(x.Name, Name, StringComparison.OrdinalIgnoreCase));
+				return Clusters.FirstOrDefault(x => String.Equals(x.Name, Name, StringComparison.OrdinalIgnoreCase));
 			}
 		}
+
+		static List<PerforceCluster> DefaultClusters { get; } = GetDefaultClusters();
+
+		static List<PerforceCluster> GetDefaultClusters()
+		{
+			PerforceCluster Cluster = new PerforceCluster();
+			Cluster.Name = "Default";
+			Cluster.Servers.Add(new PerforceServer { ServerAndPort = "perforce:1666", Properties = new List<string> { "HordeServer=1" } });
+			return new List<PerforceCluster> { Cluster };
+		}
 	}
 }
