Horde: Fix wrong response type on device controller
#rnx

diff --git a/Engine/Source/Programs/Horde/HordeServer/Controllers/DevicesController.cs b/Engine/Source/Programs/Horde/HordeServer/Controllers/DevicesController.cs
--- a/Engine/Source/Programs/Horde/HordeServer/Controllers/DevicesController.cs
+++ b/Engine/Source/Programs/Horde/HordeServer/Controllers/DevicesController.cs
@@ -802,7 +802,8 @@ namespace HordeServer.Controllers
 		/// </summary>
 		[HttpPut]
 		[Route("/api/v1/reservations/{ReservationGuid}")]
-		public async Task<ActionResult> UpdateReservationV1Async(string ReservationGuid /* [FromBody] string Duration */)
+		[ProducesResponseType(typeof(GetLegacyReservationResponse), 200)]
+		public async Task<ActionResult<GetLegacyReservationResponse>> UpdateReservationV1Async(string ReservationGuid /* [FromBody] string Duration */)
 		{
 			IDeviceReservation? Reservation = await DeviceService.TryGetReservationFromLegacyGuidAsync(ReservationGuid);
 
