Horde: Remove package pushing agent over size limit.

diff --git a/Engine/Source/Programs/Horde/HordeAgent/HordeAgent.csproj b/Engine/Source/Programs/Horde/HordeAgent/HordeAgent.csproj
--- a/Engine/Source/Programs/Horde/HordeAgent/HordeAgent.csproj
+++ b/Engine/Source/Programs/Horde/HordeAgent/HordeAgent.csproj
@@ -37,7 +37,6 @@
     <PackageReference Include="Microsoft.Extensions.Logging.Abstractions" Version="5.0.0" />
     <PackageReference Include="Microsoft.Extensions.Options.DataAnnotations" Version="3.1.5" />
     <PackageReference Include="OpenTracing" Version="0.12.1" />
-    <PackageReference Include="OpenTracing.Contrib.Grpc" Version="0.4.0" />
     <PackageReference Include="Serilog" Version="2.9.0" />
     <PackageReference Include="Serilog.Extensions.Hosting" Version="3.1.0" />
     <PackageReference Include="Serilog.Extensions.Logging" Version="3.0.1" />
diff --git a/Engine/Source/Programs/Horde/HordeAgent/Utility/RpcConnection.cs b/Engine/Source/Programs/Horde/HordeAgent/Utility/RpcConnection.cs
--- a/Engine/Source/Programs/Horde/HordeAgent/Utility/RpcConnection.cs
+++ b/Engine/Source/Programs/Horde/HordeAgent/Utility/RpcConnection.cs
@@ -8,7 +8,6 @@ using HordeCommon.Rpc;
 using Microsoft.Extensions.Hosting;
 using Microsoft.Extensions.Logging;
 using OpenTracing;
-using OpenTracing.Contrib.Grpc.Interceptors;
 using OpenTracing.Util;
 using System;
 using System.Collections.Concurrent;
@@ -672,7 +671,7 @@ namespace HordeAgent.Utility
 		{
 			using (GrpcChannel Channel = CreateGrpcChannel())
 			{
-				HordeRpc.HordeRpcClient Client = new HordeRpc.HordeRpcClient(Channel.Intercept(new ClientTracingInterceptor(GlobalTracer.Instance)));
+				HordeRpc.HordeRpcClient Client = new HordeRpc.HordeRpcClient(Channel);//.Intercept(new ClientTracingInterceptor(GlobalTracer.Instance)));
 
 				RpcSubConnection? SubConnection = null;
 				try
