Horde: Fix timespan string
#rnx

diff --git a/Engine/Source/Programs/Horde/HordeServer/Controllers/DevicesController.cs b/Engine/Source/Programs/Horde/HordeServer/Controllers/DevicesController.cs
--- a/Engine/Source/Programs/Horde/HordeServer/Controllers/DevicesController.cs
+++ b/Engine/Source/Programs/Horde/HordeServer/Controllers/DevicesController.cs
@@ -829,7 +829,7 @@ namespace HordeServer.Controllers
 			Response.DeviceNames = Reservation.Devices.Select(DeviceId => Devices.First(D => D.Id == DeviceId).Name).ToArray();
 			Response.HostName = Reservation.Hostname ?? "";
 			Response.StartDateTime = Reservation.CreateTimeUtc.ToString("O", System.Globalization.CultureInfo.InvariantCulture);
-			Response.Duration = "00::10:00"; // matches gauntlet duration
+			Response.Duration = "00:10:00"; // matches gauntlet duration
 
 			return new JsonResult(Response, new JsonSerializerOptions() { PropertyNamingPolicy = null });
 		}
