Horde: New API for remote execution.

Now allows for a stateless API and supports work-stealing between pods (via a shared task list stored in Redis). Clients now send requests for work to be executed with messages returned through a particular channel, which can be polled for outputs.

ExecuteV2 command in HordeAgent allows exercising the new code via the command line. Now mirrors the entire input directory by default, and only takes a JSON file to describe the command to perform.

diff --git a/Engine/Source/Programs/Horde/HordeAgent/Commands/Execution/ExecuteCommandV2.cs b/Engine/Source/Programs/Horde/HordeAgent/Commands/Execution/ExecuteCommandV2.cs
new file mode 100644
--- /dev/null
+++ b/Engine/Source/Programs/Horde/HordeAgent/Commands/Execution/ExecuteCommandV2.cs
@@ -0,0 +1,299 @@
+// Copyright Epic Games, Inc. All Rights Reserved.
+
+using Google.Protobuf;
+using Microsoft.Extensions.Logging;
+using System;
+using System.Collections.Concurrent;
+using System.Collections.Generic;
+using System.IO;
+using System.Linq;
+using System.Text;
+using System.Threading.Tasks;
+using EpicGames.Core;
+using System.Text.Json;
+using HordeAgent.Services;
+using Grpc.Net.Client;
+using System.Net.Http.Headers;
+using Grpc.Core;
+using Google.LongRunning;
+using Microsoft.Extensions.Hosting;
+using Microsoft.Extensions.DependencyInjection;
+using Microsoft.Extensions.Configuration;
+using Serilog.Events;
+using EpicGames.Serialization;
+using EpicGames.Horde.Compute;
+using EpicGames.Horde.Storage;
+
+namespace HordeAgent.Commands
+{
+	/// <summary>
+	/// Installs the agent as a service
+	/// </summary>
+	[Command("ExecuteV2", "Executes a command through the remote execution API")]
+	class ExecuteCommandV2 : Command
+	{
+		class JsonRequirements
+		{
+			public string? Condition { get; set; }
+			public Dictionary<string, int> Resources { get; set; } = new Dictionary<string, int>();
+			public bool Exclusive { get; set; }
+		}
+
+		class JsonComputeTask
+		{
+			public string Executable { get; set; } = String.Empty;
+			public List<string> Arguments { get; set; } = new List<string>();
+			public Dictionary<string, string> EnvVars { get; set; } = new Dictionary<string, string>();
+			public List<string> OutputPaths { get; set; } = new List<string>();
+			public string WorkingDirectory { get; set; } = String.Empty;
+			public JsonRequirements Requirements { get; set; } = new JsonRequirements();
+		}
+
+		/// <summary>
+		/// Input file describing the work to execute
+		/// </summary>
+		[CommandLine(Required = true)]
+		public FileReference Input = null!;
+
+		/// <summary>
+		/// The input directory. By default, the directory containing the input file will be used.
+		/// </summary>
+		[CommandLine]
+		public DirectoryReference? InputDir = null;
+
+		/// <summary>
+		/// Apply a random salt to the cached value
+		/// </summary>
+		[CommandLine("-Salt")]
+		public bool RandomSalt;
+
+		/// <summary>
+		/// Add a known salt value to the cache
+		/// </summary>
+		[CommandLine("-Salt=")]
+		public string? Salt = null;
+						
+		/// <summary>
+		/// Skip checking if a result is already available in the action cache
+		/// </summary>
+		[CommandLine("-SkipCacheLookup")]
+		public bool SkipCacheLookup = false;
+		
+		/// <summary>
+		/// Directory to download the output files to. If not set, no results will be downloaded.
+		/// </summary>
+		[CommandLine("-OutputDir")]
+		public string? OutputDir = null;
+	
+		/// <summary>
+		/// Log verbosity level (use normal Serilog levels such as debug, warning or info)
+		/// </summary>
+		[CommandLine("-LogLevel")]
+		public string LogLevelStr = "debug";
+
+		static IoHash CreateSandbox(DirectoryInfo BaseDirInfo, Dictionary<IoHash, byte[]> UploadList)
+		{
+			DirectoryTree Tree = new DirectoryTree();
+
+			foreach (DirectoryInfo SubDirInfo in BaseDirInfo.EnumerateDirectories())
+			{
+				IoHash SubDirHash = CreateSandbox(SubDirInfo, UploadList);
+				Tree.Directories.Add(new DirectoryNode(SubDirInfo.Name, SubDirHash));
+			}
+			Tree.Directories.SortBy(x => x.Name, Utf8StringComparer.Ordinal);
+
+			foreach (FileInfo FileInfo in BaseDirInfo.EnumerateFiles())
+			{
+				byte[] Data = File.ReadAllBytes(FileInfo.FullName);
+				IoHash Hash = IoHash.Compute(Data);
+				UploadList[Hash] = Data;
+				Tree.Files.Add(new FileNode(FileInfo.Name, Hash));
+			}
+			Tree.Files.SortBy(x => x.Name, Utf8StringComparer.Ordinal);
+
+			return AddCbObject(UploadList, Tree);
+		}
+
+		/// <inheritdoc/>
+		public override async Task<int> ExecuteAsync(ILogger Logger)
+		{
+			InputDir ??= Input.Directory;
+
+			if (Enum.TryParse(LogLevelStr, true, out LogEventLevel LogEventLevel))
+			{
+				Logging.LogLevelSwitch.MinimumLevel = LogEventLevel;
+			}
+			else
+			{
+				Console.WriteLine($"Unable to parse log level: {this.LogLevelStr}");
+				return 0;
+			}
+
+			IConfiguration Configuration = new ConfigurationBuilder()
+				.AddEnvironmentVariables()
+				.AddJsonFile("appsettings.json")
+				.Build();
+
+			IConfigurationSection ConfigSection = Configuration.GetSection(AgentSettings.SectionName);
+			IHostBuilder HostBuilder = Host.CreateDefaultBuilder()
+				.ConfigureLogging(Builder =>
+				{
+					Builder.SetMinimumLevel(LogLevel.Warning);
+				})
+				.ConfigureServices(Services =>
+				{
+					Services.AddOptions<AgentSettings>().Configure(Options => ConfigSection.Bind(Options)).ValidateDataAnnotations();
+					Services.AddLogging();
+					Services.AddSingleton<GrpcService>();
+				});
+
+			using (IHost Host = HostBuilder.Build())
+			{
+				GrpcService GrpcService = Host.Services.GetRequiredService<GrpcService>();
+				
+				byte[] Json = await FileReference.ReadAllBytesAsync(Input);
+				JsonComputeTask JsonComputeTask = JsonSerializer.Deserialize<JsonComputeTask>(Json, new JsonSerializerOptions { PropertyNameCaseInsensitive = true });
+
+				ByteString SaltBytes = ByteString.Empty;
+				if (Salt != null)
+				{
+					SaltBytes = ByteString.CopyFromUtf8(Salt);
+				}
+				else if(RandomSalt)
+				{
+					SaltBytes = ByteString.CopyFrom(Guid.NewGuid().ToByteArray());
+				}
+
+				if (SaltBytes.Length > 0)
+				{
+					Logger.LogInformation("Using salt: {SaltBytes}", StringUtils.FormatHexString(SaltBytes.ToByteArray()));
+				}
+
+				Dictionary<IoHash, byte[]> Blobs = new Dictionary<IoHash, byte[]>();
+				IoHash SandboxHash = CreateSandbox(InputDir.ToDirectoryInfo(), Blobs);
+
+				ComputeTask Task = new ComputeTask(JsonComputeTask.Executable, JsonComputeTask.Arguments.ConvertAll<Utf8String>(x => x), JsonComputeTask.WorkingDirectory, SandboxHash);
+				Task.EnvVars = JsonComputeTask.EnvVars.ToDictionary(x => (Utf8String)x.Key, x => (Utf8String)x.Value);
+				Task.OutputPaths.AddRange(JsonComputeTask.OutputPaths.Select(x => (Utf8String)x));
+				IoHash TaskHash = AddCbObject(Blobs, Task);
+
+				JsonRequirements JsonRequirements = JsonComputeTask.Requirements;
+				Requirements Requirements = new Requirements(JsonRequirements.Condition ?? String.Empty);
+				Requirements.Resources = JsonRequirements.Resources.ToDictionary(x => (Utf8String)x.Key, x => x.Value);
+				Requirements.Exclusive = JsonRequirements.Exclusive;
+				IoHash RequirementsHash = AddCbObject(Blobs, Requirements);
+
+				await ExecuteAction(Logger, GrpcService, TaskHash, RequirementsHash, Blobs);	
+			}
+			return 0;
+		}
+
+		static IoHash AddCbObject<T>(Dictionary<IoHash, byte[]> HashToData, T Source)
+		{
+			CbObject Obj = CbSerializer.Serialize<T>(Source);
+			IoHash Hash = Obj.GetHash();
+			HashToData[Hash] = Obj.GetView().ToArray();
+			return Hash;
+		}
+
+		private async Task ExecuteAction(ILogger Logger, GrpcService GrpcService, IoHash TaskHash, IoHash RequirementsHash, Dictionary<IoHash, byte[]> UploadList)
+		{
+			using GrpcChannel Channel = GrpcService.CreateGrpcChannel();
+
+			const string NamespaceId = "default";
+
+			using GrpcChannel CasChannel = Channel;
+
+			BlobStore.BlobStoreClient BlobStore = new BlobStore.BlobStoreClient(CasChannel);
+			foreach (KeyValuePair<IoHash, byte[]> Pair in UploadList)
+			{
+				AddBlobsRequest Request = new AddBlobsRequest();
+				Request.NamespaceId = NamespaceId;
+				Request.Blobs.Add(new AddBlobRequest(Pair.Key, Pair.Value));
+				await BlobStore.AddAsync(Request);
+			}
+
+			// Execute the action
+			ComputeRpc.ComputeRpcClient ComputeClient = new ComputeRpc.ComputeRpcClient(Channel);
+
+			AddTasksRpcRequest AddTasksRequest = new AddTasksRpcRequest(Guid.NewGuid().ToString(), NamespaceId, RequirementsHash, new List<IoHash> { TaskHash }, SkipCacheLookup);
+			await ComputeClient.AddTasksAsync(AddTasksRequest);
+
+			using (AsyncDuplexStreamingCall<GetTaskUpdatesRpcRequest, GetTaskUpdatesRpcResponse> Call = ComputeClient.GetTaskUpdates())
+			{
+				await Call.RequestStream.WriteAsync(new GetTaskUpdatesRpcRequest(AddTasksRequest.ChannelId));
+
+				TaskCompletionSource<bool> CompleteNowTask = new TaskCompletionSource<bool>();
+				Task CompleteTask = Task.Run(async () => { await CompleteNowTask.Task; await Call.RequestStream.CompleteAsync(); } );
+
+				while (await Call.ResponseStream.MoveNext())
+				{
+					GetTaskUpdatesRpcResponse Response = Call.ResponseStream.Current;
+					Logger.LogInformation("{OperationName}: Execution state: {State}", (IoHash)Response.TaskHash, Response.State.ToString());
+					if (!String.IsNullOrEmpty(Response.AgentId) || !String.IsNullOrEmpty(Response.LeaseId))
+					{
+						Logger.LogInformation("{OperationName}: Running on agent {AgentId} under lease {LeaseId}", (IoHash)Response.TaskHash, Response.AgentId, Response.LeaseId);
+					}
+					if (Response.ResultHash != null)
+					{
+						await HandleCompleteTask(BlobStore, NamespaceId, Response.ResultHash, Logger);
+					}
+					if (Response.State == ComputeTaskState.Complete)
+					{
+						CompleteNowTask.TrySetResult(true);
+					}
+				}
+
+				await CompleteTask;
+			}
+		}
+
+		async Task HandleCompleteTask(BlobStore.BlobStoreClient Client, string NamespaceId, IoHash Hash, ILogger Logger)
+		{
+			ComputeTaskResult Result = await Client.GetObjectAsync<ComputeTaskResult>(NamespaceId, Hash);
+			Logger.LogInformation("exit: {ExitCode}", Result.ExitCode);
+
+			await LogTaskOutputAsync(Client, "stdout", NamespaceId, Result.StdOutHash, Logger);
+			await LogTaskOutputAsync(Client, "stderr", NamespaceId, Result.StdErrHash, Logger);
+
+			if (Result.OutputHash != null && OutputDir != null)
+			{
+				await WriteOutputAsync(Client, NamespaceId, Result.OutputHash.Value, new DirectoryReference(OutputDir));
+			}
+		}
+
+		async Task LogTaskOutputAsync(BlobStore.BlobStoreClient Client, string Channel, string NamespaceId, IoHash? LogHash, ILogger Logger)
+		{
+			if (LogHash != null)
+			{
+				byte[] StdOutData = await Client.GetBlobAsync(NamespaceId, LogHash.Value);
+				if (StdOutData.Length > 0)
+				{
+					foreach (string Line in Encoding.UTF8.GetString(StdOutData).Split('\n'))
+					{
+						Logger.LogDebug("{Channel}: {Line}", Channel, Line);
+					}
+				}
+			}
+		}
+
+		async Task WriteOutputAsync(BlobStore.BlobStoreClient Client, string NamespaceId, IoHash TreeHash, DirectoryReference OutputDir)
+		{
+			DirectoryTree Tree = await Client.GetObjectAsync<DirectoryTree>(NamespaceId, TreeHash);
+
+			List<Task> Tasks = new List<Task>();
+			foreach (FileNode File in Tree.Files)
+			{
+				FileReference OutputFile = FileReference.Combine(OutputDir, File.Name.ToString());
+				Tasks.Add(Client.GetFileAsync(NamespaceId, File.Hash, OutputFile));
+			}
+			foreach (DirectoryNode Directory in Tree.Directories)
+			{
+				DirectoryReference NextOutputDir = DirectoryReference.Combine(OutputDir, Directory.Name.ToString());
+				Tasks.Add(WriteOutputAsync(Client, NamespaceId, Directory.Hash, NextOutputDir));
+			}
+
+			await Task.WhenAll(Tasks);
+		}
+	}
+}
diff --git a/Engine/Source/Programs/Horde/HordeAgent/Execution/ComputeTaskExecutor.cs b/Engine/Source/Programs/Horde/HordeAgent/Execution/ComputeTaskExecutor.cs
new file mode 100644
--- /dev/null
+++ b/Engine/Source/Programs/Horde/HordeAgent/Execution/ComputeTaskExecutor.cs
@@ -0,0 +1,334 @@
+// Copyright Epic Games, Inc. All Rights Reserved.
+
+using EpicGames.Core;
+using EpicGames.Horde;
+using EpicGames.Horde.Storage;
+using Google.LongRunning;
+using Google.Protobuf;
+using Grpc.Core;
+using Microsoft.Extensions.Options;
+using System;
+using System.Collections.Generic;
+using System.IO;
+using System.Linq;
+using System.Threading;
+using System.Threading.Tasks;
+using System.Diagnostics;
+using System.Text;
+using Google.Protobuf.WellKnownTypes;
+using Microsoft.Extensions.Hosting;
+using System.Collections.Concurrent;
+using Grpc.Net.Client;
+using HordeAgent.Utility;
+using HordeCommon.Rpc.Tasks;
+using Microsoft.Extensions.Logging;
+using HordeCommon.Rpc;
+using EpicGames.Horde.Common.RemoteExecution;
+using EpicGames.Serialization;
+using EpicGames.Horde.Compute;
+
+namespace HordeAgent
+{
+	/// <summary>
+	/// Executes remote actions in a sandbox
+	/// </summary>
+	class ComputeTaskExecutor
+	{
+		class OutputTree
+		{
+			public Dictionary<string, Task<IoHash>> Files = new Dictionary<string, Task<IoHash>>();
+			public Dictionary<string, OutputTree> SubDirs = new Dictionary<string, OutputTree>();
+		}
+
+		/// <summary>
+		/// The blob store
+		/// </summary>
+		BlobStore.BlobStoreClient BlobStore;
+
+		/// <summary>
+		/// Logger for internal executor output
+		/// </summary>
+		ILogger Logger;
+
+		/// <summary>
+		/// Constructor
+		/// </summary>
+		/// <param name="GrpcChannel"></param>
+		/// <param name="Logger"></param>
+		public ComputeTaskExecutor(GrpcChannel GrpcChannel, ILogger Logger)
+		{
+			this.BlobStore = new BlobStore.BlobStoreClient(GrpcChannel);
+			this.Logger = Logger;
+		}
+
+		/// <summary>
+		/// Read stdout/stderr and wait until process completes
+		/// </summary>
+		/// <param name="Process">The process reading from</param>
+		/// <param name="Timeout">Max execution timeout for process</param>
+		/// <param name="CancelToken">Cancellation token</param>
+		/// <returns>stdout and stderr data</returns>
+		/// <exception cref="RpcException">Raised for either timeout or cancel</exception>
+		private async Task<(byte[] StdOutData, byte[] StdOutErr)> ReadProcessStreams(ManagedProcess Process, TimeSpan Timeout, CancellationToken CancelToken)
+		{
+			// Read stdout/stderr without cancellation token.
+			using MemoryStream StdOutStream = new MemoryStream();
+			Task StdOutReadTask = Process.StdOut.CopyToAsync(StdOutStream);
+
+			using MemoryStream StdErrStream = new MemoryStream();
+			Task StdErrReadTask = Process.StdErr.CopyToAsync(StdErrStream);
+
+			Task OutputTask = Task.WhenAll(StdOutReadTask, StdErrReadTask);
+
+			// Instead, create a separate task that will wait for either timeout or cancellation
+			// as cancel interruptions are not reliable with Stream.CopyToAsync()
+			Task TimeoutTask = Task.Delay(Timeout, CancelToken);
+
+			Task WaitTask = await Task.WhenAny(OutputTask, TimeoutTask);
+			if (WaitTask == TimeoutTask)
+			{
+				throw new RpcException(new Grpc.Core.Status(StatusCode.DeadlineExceeded, $"Action timed out after {Timeout.TotalMilliseconds} ms"));
+			}
+			
+			return (StdOutStream.ToArray(), StdErrStream.ToArray());
+		}
+
+		/// <summary>
+		/// Execute an action
+		/// </summary>
+		/// <param name="LeaseId">The lease id</param>
+		/// <param name="ComputeTaskMessage">Task to execute</param>
+		/// <param name="SandboxDir">Directory to use as a sandbox for execution</param>
+		/// <param name="CancellationToken">Token used to cancel the operation</param>
+		/// <returns>The action result</returns>
+		public async Task<ComputeTaskResultMessage> ExecuteAsync(string LeaseId, ComputeTaskMessage ComputeTaskMessage, DirectoryReference SandboxDir, CancellationToken CancellationToken)
+		{
+			ComputeTask? Task = await BlobStore.GetObjectAsync<ComputeTask>(ComputeTaskMessage.NamespaceId, ComputeTaskMessage.TaskHash);
+			if (Task == null)
+			{
+				throw new RpcException(new Grpc.Core.Status(StatusCode.NotFound, "Unable to find action message"));
+			}
+			Logger.LogInformation("Executing task {Hash} for lease ID {LeaseId}", ComputeTaskMessage.TaskHash, LeaseId);
+
+			DirectoryReference.CreateDirectory(SandboxDir);
+			FileUtils.ForceDeleteDirectoryContents(SandboxDir);
+
+			DateTimeOffset InputFetchStart = DateTimeOffset.UtcNow;
+			DirectoryTree InputDirectory = await BlobStore.GetObjectAsync<DirectoryTree>(ComputeTaskMessage.NamespaceId, Task.SandboxHash);
+			await SetupSandboxAsync(ComputeTaskMessage.NamespaceId, InputDirectory, SandboxDir);
+			DateTimeOffset InputFetchCompleted = DateTimeOffset.UtcNow;
+
+			ComputeTaskResultMessage ResultMessage;
+			using (ManagedProcessGroup ProcessGroup = new ManagedProcessGroup())
+			{
+				string FileName = FileReference.Combine(SandboxDir, Task.WorkingDirectory.ToString(), Task.Executable.ToString()).FullName;
+				string WorkingDirectory = DirectoryReference.Combine(SandboxDir, Task.WorkingDirectory.ToString()).FullName;
+
+				Dictionary<string, string> NewEnvironment = new Dictionary<string, string>();
+				foreach (System.Collections.DictionaryEntry? Entry in Environment.GetEnvironmentVariables())
+				{
+					NewEnvironment[Entry!.Value.Key.ToString()!] = Entry!.Value.Value!.ToString()!;
+				}
+				foreach(KeyValuePair<Utf8String, Utf8String> Pair in Task.EnvVars)
+				{
+					NewEnvironment[Pair.Key.ToString()] = Pair.Value.ToString();
+				}
+
+				string Arguments = CommandLineArguments.Join(Task.Arguments.Select(x => x.ToString()));
+				Logger.LogInformation("Executing {FileName} with arguments {Arguments}", FileName, Arguments);
+
+				TimeSpan Timeout = TimeSpan.FromMinutes(5.0);// Action.Timeout == null ? TimeSpan.FromMinutes(5) : Action.Timeout.ToTimeSpan();
+				DateTimeOffset ExecutionStartTime = DateTimeOffset.UtcNow;
+				using (ManagedProcess Process = new ManagedProcess(ProcessGroup, FileName, Arguments, WorkingDirectory, NewEnvironment, ProcessPriorityClass.Normal, ManagedProcessFlags.None))
+				{
+					(byte[] StdOutData, byte[] StdErrData) = await ReadProcessStreams(Process, Timeout, CancellationToken);
+
+					foreach (string Line in Encoding.UTF8.GetString(StdOutData).Split('\n'))
+					{
+						Logger.LogInformation("stdout: {Line}", Line);
+					}
+					foreach (string Line in Encoding.UTF8.GetString(StdErrData).Split('\n'))
+					{
+						Logger.LogInformation("stderr: {Line}", Line);
+					}
+
+					Logger.LogInformation("exit: {ExitCode}", Process.ExitCode);
+
+					ComputeTaskResult Result = new ComputeTaskResult(Process.ExitCode);
+					Result.StdOutHash = await BlobStore.PutBlobAsync(ComputeTaskMessage.NamespaceId, StdOutData);
+					Result.StdErrHash = await BlobStore.PutBlobAsync(ComputeTaskMessage.NamespaceId, StdErrData);
+
+					FileReference[] OutputFiles = ResolveOutputPaths(SandboxDir, Task.OutputPaths.Select(x => x.ToString())).OrderBy(x => x.FullName, StringComparer.Ordinal).ToArray();
+					if (OutputFiles.Length > 0)
+					{
+						foreach (FileReference OutputFile in OutputFiles)
+						{
+							Logger.LogInformation("output: {File}", OutputFile.MakeRelativeTo(SandboxDir));
+						}
+						Result.OutputHash = await PutOutput(ComputeTaskMessage.NamespaceId, SandboxDir, OutputFiles);
+					}
+
+					ResultMessage = new ComputeTaskResultMessage();
+					ResultMessage.OutputHash = await BlobStore.PutObjectAsync(ComputeTaskMessage.NamespaceId, Result);
+
+
+					/*
+										Result.ExecutionMetadata.OutputUploadStartTimestamp = Timestamp.FromDateTimeOffset(DateTimeOffset.UtcNow);
+										foreach (var (FileRef, OutputPath) in ResolveOutputPaths(SandboxDir, Command.OutputPaths))
+										{
+											byte[] Bytes = await FileReference.ReadAllBytesAsync(FileRef);
+											Digest Digest = await Storage.PutBulkDataAsync(InstanceName, Bytes);
+
+											OutputFile OutputFileInfo = new OutputFile();
+											OutputFileInfo.Path = OutputPath;
+											OutputFileInfo.Digest = Digest;
+											Result.OutputFiles.Add(OutputFileInfo);
+
+											Logger.LogInformation("Uploaded {File} (digest: {Digest}, size: {Size})", FileRef, Digest.Hash, Digest.SizeBytes);
+										}
+										Result.ExecutionMetadata.OutputUploadCompletedTimestamp = Timestamp.FromDateTimeOffset(DateTimeOffset.UtcNow);
+										Result.ExecutionMetadata.WorkerCompletedTimestamp = Timestamp.FromDateTimeOffset(DateTimeOffset.UtcNow);
+
+										PostActionResultRequest PostResultRequest = new PostActionResultRequest();
+										PostResultRequest.LeaseId = LeaseId;
+										PostResultRequest.ActionDigest = ActionTask.Digest;
+										PostResultRequest.Result = Result;
+										await ActionRpc.PostActionResultAsync(PostResultRequest);
+
+										await CacheActionResultAsync(Action, ActionTask, Result);
+					*/
+					// Command.OutputPaths
+					// Command.OutputDirectories
+				}
+			}
+			return ResultMessage;
+		}
+
+		/// <summary>
+		/// Downloads files to the sandbox
+		/// </summary>
+		/// <param name="NamespaceId">Namespace for fetching data</param>
+		/// <param name="InputDirectory">The directory spec</param>
+		/// <param name="OutputDir">Output directory on disk</param>
+		/// <returns>Async task</returns>
+		async Task SetupSandboxAsync(string NamespaceId, DirectoryTree InputDirectory, DirectoryReference OutputDir)
+		{
+			DirectoryReference.CreateDirectory(OutputDir);
+
+			async Task DownloadFile(FileNode FileNode)
+			{
+				FileReference File = FileReference.Combine(OutputDir, FileNode.Name.ToString());
+				Logger.LogInformation("Downloading {File} (digest: {Digest})", File, FileNode.Hash);
+				byte[] Data = await BlobStore.GetBlobAsync(NamespaceId, FileNode.Hash);
+				Logger.LogInformation("Writing {File} (digest: {Digest})", File, FileNode.Hash);
+				await FileReference.WriteAllBytesAsync(File, Data);
+			}
+
+			async Task DownloadDir(DirectoryNode DirectoryNode)
+			{
+				DirectoryTree InputSubDirectory = await BlobStore.GetObjectAsync<DirectoryTree>(NamespaceId, DirectoryNode.Hash);
+				DirectoryReference OutputSubDirectory = DirectoryReference.Combine(OutputDir, DirectoryNode.Name.ToString());
+				await SetupSandboxAsync(NamespaceId, InputSubDirectory, OutputSubDirectory);
+			}
+
+			List<Task> Tasks = new List<Task>();
+			Tasks.AddRange(InputDirectory.Files.Select(x => Task.Run(() => DownloadFile(x))));
+			Tasks.AddRange(InputDirectory.Directories.Select(x => Task.Run(() => DownloadDir(x))));
+			await Task.WhenAll(Tasks);
+		}
+
+		Task<IoHash> PutOutput(string NamespaceId, DirectoryReference BaseDir, IEnumerable<FileReference> Files)
+		{
+			List<FileReference> SortedFiles = Files.OrderBy(x => x.FullName, StringComparer.Ordinal).ToList();
+			return PutDirectoryTree(NamespaceId, BaseDir.FullName.Length, SortedFiles, 0, SortedFiles.Count);
+		}
+
+		async Task<IoHash> PutDirectoryTree(string NamespaceId, int BaseDirLen, List<FileReference> SortedFiles, int MinIdx, int MaxIdx)
+		{
+			List<Task<FileNode>> Files = new List<Task<FileNode>>();
+			List<Task<DirectoryNode>> Trees = new List<Task<DirectoryNode>>();
+
+			while (MinIdx < MaxIdx)
+			{
+				FileReference File = SortedFiles[MinIdx];
+
+				int NextMinIdx = MinIdx + 1;
+
+				int NextDirLen = File.FullName.IndexOf(Path.DirectorySeparatorChar, BaseDirLen + 1);
+				if (NextDirLen == -1)
+				{
+					string Name = File.FullName.Substring(BaseDirLen + 1);
+					Files.Add(CreateFileNode(NamespaceId, Name, File));
+				}
+				else
+				{
+					string Name = File.FullName.Substring(BaseDirLen + 1, NextDirLen - (BaseDirLen + 1));
+					while (NextMinIdx < MaxIdx)
+					{
+						string NextFile = SortedFiles[NextMinIdx].FullName;
+						if (NextFile.Length < NextDirLen || String.Compare(Name, 0, NextFile, BaseDirLen, Name.Length, StringComparison.Ordinal) == 0)
+						{
+							break;
+						}
+						NextMinIdx++;
+					}
+					Trees.Add(CreateDirectoryNode(NamespaceId, Name, NextDirLen, SortedFiles, MinIdx, NextMinIdx));
+				}
+
+				MinIdx = NextMinIdx;
+			}
+
+			DirectoryTree Tree = new DirectoryTree();
+			Tree.Files.AddRange(await Task.WhenAll(Files));
+			Tree.Directories.AddRange(await Task.WhenAll(Trees));
+			return await BlobStore.PutObjectAsync(NamespaceId, Tree);
+		}
+
+		async Task<DirectoryNode> CreateDirectoryNode(string NamespaceId, string Name, int BaseDirLen, List<FileReference> SortedFiles, int MinIdx, int MaxIdx)
+		{
+			IoHash Hash = await PutDirectoryTree(NamespaceId, BaseDirLen, SortedFiles, MinIdx, MaxIdx);
+			return new DirectoryNode(Name, Hash);
+		}
+
+		async Task<FileNode> CreateFileNode(string NamespaceId, string Name, FileReference File)
+		{
+			IoHash Hash = await BlobStore.PutBlobAsync(NamespaceId, await FileReference.ReadAllBytesAsync(File));
+			return new FileNode(Name, Hash);
+		}
+
+		/// <summary>
+		/// Resolves a list of output paths into file references
+		///
+		/// The REAPI spec allows directories to be specified as an output path which require all sub dirs and files
+		/// to be resolved.
+		/// </summary>
+		/// <param name="SandboxDir">Base directory where execution is taking place</param>
+		/// <param name="OutputPaths">List of output paths relative to SandboxDir</param>
+		/// <returns>List of resolved paths (incl expanded dirs)</returns>
+		internal static HashSet<FileReference> ResolveOutputPaths(DirectoryReference SandboxDir, IEnumerable<string> OutputPaths)
+		{
+			HashSet<FileReference> Files = new HashSet<FileReference>();
+			foreach (string OutputPath in OutputPaths)
+			{
+				DirectoryReference DirRef = DirectoryReference.Combine(SandboxDir, OutputPath);
+				if (DirectoryReference.Exists(DirRef))
+				{
+					IEnumerable<FileReference> ListedFiles = DirectoryReference.EnumerateFiles(DirRef, "*", SearchOption.AllDirectories);
+					foreach (FileReference ListedFileRef in ListedFiles)
+					{
+						Files.Add(ListedFileRef);
+					}
+				}
+				else
+				{
+					FileReference FileRef = FileReference.Combine(SandboxDir, OutputPath);
+					if (FileReference.Exists(FileRef))
+					{
+						Files.Add(FileRef);
+					}
+				}
+			}
+			return Files;
+		}
+	}
+}
diff --git a/Engine/Source/Programs/Horde/HordeAgent/HordeAgent.csproj b/Engine/Source/Programs/Horde/HordeAgent/HordeAgent.csproj
--- a/Engine/Source/Programs/Horde/HordeAgent/HordeAgent.csproj
+++ b/Engine/Source/Programs/Horde/HordeAgent/HordeAgent.csproj
@@ -82,7 +82,9 @@
 
   <ItemGroup>
     <ProjectReference Include="..\..\Shared\EpicGames.Core\EpicGames.Core.csproj" />
+    <ProjectReference Include="..\..\Shared\EpicGames.Horde.Client\EpicGames.Horde.Client.csproj" />
     <ProjectReference Include="..\..\Shared\EpicGames.Horde.Common\EpicGames.Horde.Common.csproj" />
+    <ProjectReference Include="..\..\Shared\EpicGames.Horde\EpicGames.Horde.csproj" />
     <ProjectReference Include="..\..\Shared\EpicGames.Perforce.Managed\EpicGames.Perforce.Managed.csproj" />
     <ProjectReference Include="..\..\Shared\EpicGames.Perforce\EpicGames.Perforce.csproj" />
     <ProjectReference Include="..\HordeCommon\HordeCommon.csproj" />
diff --git a/Engine/Source/Programs/Horde/HordeAgent/Services/WorkerService.cs b/Engine/Source/Programs/Horde/HordeAgent/Services/WorkerService.cs
--- a/Engine/Source/Programs/Horde/HordeAgent/Services/WorkerService.cs
+++ b/Engine/Source/Programs/Horde/HordeAgent/Services/WorkerService.cs
@@ -663,13 +663,13 @@ namespace HordeAgent.Services
 				{
 					LeaseInfo.Lease.State = LeaseState.Cancelled;
 					LeaseInfo.Lease.Outcome = LeaseOutcome.Failed;
-					LeaseInfo.Lease.Output = null;
+					LeaseInfo.Lease.Output = ByteString.Empty;
 				}
 				else
 				{
 					LeaseInfo.Lease.State = (Result.Outcome == LeaseOutcome.Cancelled) ? LeaseState.Cancelled : LeaseState.Completed;
 					LeaseInfo.Lease.Outcome = Result.Outcome;
-					LeaseInfo.Lease.Output = ByteString.CopyFrom(Result.Output);
+					LeaseInfo.Lease.Output = (Result.Output != null)? ByteString.CopyFrom(Result.Output) : ByteString.Empty;
 				}
 				Logger.LogInformation("Transitioning lease {LeaseId} to {State}, outcome={Outcome}", LeaseInfo.Lease.Id, LeaseInfo.Lease.State, LeaseInfo.Lease.Outcome);
 			}
@@ -695,6 +695,13 @@ namespace HordeAgent.Services
 				return await HandleLeasePayloadWithLogAsync(RpcConnection, ActionTask.LogId, AgentId, null, null, Handler, LeaseInfo.CancellationTokenSource.Token);
 			}
 
+			ComputeTaskMessage ComputeTask;
+			if (LeaseInfo.Lease.Payload.TryUnpack(out ComputeTask))
+			{
+				GlobalTracer.Instance.ActiveSpan?.SetTag("task", "Compute");
+				return await ComputeAsync(RpcConnection, LeaseInfo.Lease.Id, ComputeTask, Logger, LeaseInfo.CancellationTokenSource.Token);
+			}
+
 			ConformTask ConformTask;
 			if (LeaseInfo.Lease.Payload.TryUnpack(out ConformTask))
 			{
@@ -883,6 +890,53 @@ namespace HordeAgent.Services
 			}
 		}
 
+		/// <summary>
+		/// Execute a remote execution aciton
+		/// </summary>
+		/// <param name="RpcConnection">RPC client for communicating with the server</param>
+		/// <param name="LeaseId">The lease id</param>
+		/// <param name="ComputeTask">The action task parameters</param>
+		/// <param name="Logger">Logger for the task</param>
+		/// <param name="CancellationToken">Token used to cancel the operation</param>
+		/// <returns></returns>
+		internal async Task<LeaseResult> ComputeAsync(IRpcConnection RpcConnection, string LeaseId, ComputeTaskMessage ComputeTask, ILogger Logger, CancellationToken CancellationToken)
+		{
+			DateTimeOffset ActionTaskStartTime = DateTimeOffset.UtcNow;
+			using (IRpcClientRef Client = await RpcConnection.GetClientRef(new RpcContext(), CancellationToken))
+			{
+				DirectoryReference LeaseDir = DirectoryReference.Combine(WorkingDir, "Compute", LeaseId);
+				DirectoryReference.CreateDirectory(LeaseDir);
+
+				ComputeTaskExecutor Executor = new ComputeTaskExecutor(Client.Channel, Logger);
+				try
+				{
+					ComputeTaskResultMessage Result = await Executor.ExecuteAsync(LeaseId, ComputeTask, LeaseDir, CancellationToken);
+					return new LeaseResult(Result.ToByteArray());
+				}
+				catch (Exception re)
+				{
+/*					Logger.LogError(re, "Unhandled exception during remote execution");
+					PostActionResultRequest PostResultRequest = new PostActionResultRequest();
+					PostResultRequest.LeaseId = LeaseId;
+					PostResultRequest.ActionDigest = ComputeTask.Digest;
+					PostResultRequest.Error = new RpcException(new Status(StatusCode.Unknown, "Unhandled exception during remote execution: " + re.Message), re.Message).Status;
+					await Executor.ActionRpc.PostActionResultAsync(PostResultRequest);
+*/
+					throw re;
+				}
+				finally
+				{
+					try
+					{
+						DirectoryReference.Delete(LeaseDir, true);
+					}
+					catch
+					{
+					}
+				}
+			}
+		}
+
 		/// <summary>
 		/// Conform a machine
 		/// </summary>
diff --git a/Engine/Source/Programs/Horde/HordeCommon/HordeCommon.csproj b/Engine/Source/Programs/Horde/HordeCommon/HordeCommon.csproj
--- a/Engine/Source/Programs/Horde/HordeCommon/HordeCommon.csproj
+++ b/Engine/Source/Programs/Horde/HordeCommon/HordeCommon.csproj
@@ -15,6 +15,10 @@
     <PackageReference Include="Microsoft.Extensions.Logging" Version="5.0.0" />
   </ItemGroup>
 
+  <ItemGroup>
+    <ProjectReference Include="..\..\Shared\EpicGames.Horde\EpicGames.Horde.csproj" />
+  </ItemGroup>
+
   <ItemGroup>
     <Protobuf Include="Protos\horde\horde_rpc_messages.proto" ProtoRoot="Protos" />
     <Protobuf Include="Protos\horde\common\agent_status.proto" ProtoRoot="Protos" />
@@ -29,6 +33,7 @@
     <Protobuf Include="Protos\horde\messages\agent_requirements.proto" ProtoRoot="Protos" />
     <Protobuf Include="Protos\horde\messages\agent_workspace.proto" ProtoRoot="Protos" />
     <Protobuf Include="Protos\horde\messages\lease.proto" ProtoRoot="Protos" />
+    <Protobuf Include="Protos\horde\tasks\compute_task.proto" ProtoRoot="Protos;..\..\Shared\EpicGames.Horde\Protos" />
     <Protobuf Include="Protos\horde\tasks\conform_task.proto" ProtoRoot="Protos" />
     <Protobuf Include="Protos\horde\tasks\job_task.proto" ProtoRoot="Protos" />
     <Protobuf Include="Protos\horde\tasks\restart_task.proto" ProtoRoot="Protos" />
diff --git a/Engine/Source/Programs/Horde/HordeServer/Compute/IComputeService.cs b/Engine/Source/Programs/Horde/HordeServer/Compute/IComputeService.cs
new file mode 100644
--- /dev/null
+++ b/Engine/Source/Programs/Horde/HordeServer/Compute/IComputeService.cs
@@ -0,0 +1,106 @@
+using EpicGames.Core;
+using EpicGames.Horde.Compute;
+using HordeServer.Models;
+using HordeServer.Storage;
+using HordeServer.Tasks;
+using HordeServer.Utilities;
+using Microsoft.Extensions.Hosting;
+using MongoDB.Bson;
+using System;
+using System.Collections.Generic;
+using System.Linq;
+using System.Threading;
+using System.Threading.Tasks;
+
+namespace HordeServer.Compute
+{
+	using ChannelId = StringId<IComputeChannel>;
+	using NamespaceId = StringId<INamespace>;
+
+	/// <summary>
+	/// Unused, except for ChannelId alias.
+	/// </summary>
+	interface IComputeChannel
+	{
+	}
+
+	/// <summary>
+	/// Status of a compute task
+	/// </summary>
+	interface IComputeTaskStatus
+	{
+		/// <summary>
+		/// The input hash
+		/// </summary>
+		public IoHash TaskHash { get; }
+
+		/// <summary>
+		/// Timestamp for the event
+		/// </summary>
+		public DateTime Time { get; }
+
+		/// <summary>
+		/// State of the task
+		/// </summary>
+		public ComputeTaskState State { get; }
+
+		/// <summary>
+		/// When transitioning to the executing state, includes the agent name
+		/// </summary>
+		public AgentId? AgentId { get; }
+
+		/// <summary>
+		/// The lease id
+		/// </summary>
+		public ObjectId? LeaseId { get; }
+
+		/// <summary>
+		/// Hash of the result, if complete
+		/// </summary>
+		public IoHash? ResultHash { get; }
+	}
+
+	/// <summary>
+	/// Interface for the compute service
+	/// </summary>
+	interface IComputeService : INewTaskSource
+	{
+		/// <summary>
+		/// Post tasks to be executed to a channel
+		/// </summary>
+		/// <param name="NamespaceId">Namespace of referenced blobs</param>
+		/// <param name="RequirementsHash">Hash of the requirements document for execution</param>
+		/// <param name="TaskHashes">List of task hashes</param>
+		/// <param name="ChannelId">Unique identifier of the client</param>
+		/// <returns>Async task</returns>
+		Task AddTasksAsync(NamespaceId NamespaceId, IoHash RequirementsHash, List<IoHash> TaskHashes, ChannelId ChannelId);
+
+		/// <summary>
+		/// Dequeue completed items from a queue
+		/// </summary>
+		/// <param name="ChannelId">Queue to remove items from</param>
+		/// <param name="CancellationToken">Cancellation token to stop waiting for items</param>
+		/// <returns>List of status updates</returns>
+		Task<List<IComputeTaskStatus>> GetTaskUpdatesAsync(ChannelId ChannelId, CancellationToken CancellationToken);
+	}
+
+	/// <summary>
+	/// Extension methods for <see cref="IComputeService"/>
+	/// </summary>
+	static class ComputeServiceExtensions
+	{
+		/// <summary>
+		/// 
+		/// </summary>
+		/// <param name="ComputeService"></param>
+		/// <param name="NamespaceId"></param>
+		/// <param name="RequirementsHash">Hash of the requirements document for execution</param>
+		/// <param name="TaskHash">The task hash</param>
+		/// <param name="ChannelId"></param>
+		/// <returns></returns>
+		public static Task AddTaskAsync(this IComputeService ComputeService, NamespaceId NamespaceId, IoHash RequirementsHash, IoHash TaskHash, ChannelId ChannelId)
+		{
+			return ComputeService.AddTasksAsync(NamespaceId, RequirementsHash, new List<IoHash> { TaskHash }, ChannelId);
+		}
+	}
+}
diff --git a/Engine/Source/Programs/Horde/HordeServer/Compute/Impl/ComputeService.cs b/Engine/Source/Programs/Horde/HordeServer/Compute/Impl/ComputeService.cs
new file mode 100644
--- /dev/null
+++ b/Engine/Source/Programs/Horde/HordeServer/Compute/Impl/ComputeService.cs
@@ -0,0 +1,281 @@
+// Copyright Epic Games, Inc. All Rights Reserved.
+
+using EpicGames.Core;
+using Google.LongRunning;
+using Google.Protobuf;
+using Google.Protobuf.Reflection;
+using Google.Protobuf.WellKnownTypes;
+using Grpc.Core;
+using HordeCommon;
+using HordeCommon.Rpc.Tasks;
+using HordeServer.Api;
+using HordeServer.Models;
+using HordeServer.Rpc;
+using HordeServer.Services;
+using Microsoft.Extensions.Logging;
+using MongoDB.Bson;
+using System;
+using System.Collections.Concurrent;
+using System.Collections.Generic;
+using System.Diagnostics;
+using System.Diagnostics.CodeAnalysis;
+using System.Linq;
+using System.Threading;
+using System.Threading.Tasks;
+using Microsoft.Extensions.Options;
+using HordeServer.Tasks;
+using HordeServer.Utilities;
+using HordeServer.Storage;
+using StackExchange.Redis;
+using MongoDB.Bson.Serialization.Attributes;
+using MongoDB.Bson.Serialization;
+using EpicGames.Serialization;
+using Microsoft.Extensions.Caching.Memory;
+using EpicGames.Horde.Compute;
+using Microsoft.Extensions.Hosting;
+using System.Threading.Channels;
+using EpicGames.Redis;
+
+namespace HordeServer.Compute.Impl
+{
+	using ChannelId = StringId<IComputeChannel>;
+	using NamespaceId = StringId<INamespace>;
+
+	/// <summary>
+	/// Information about a particular task
+	/// </summary>
+	[RedisConverter(typeof(RedisCbConverter<>))]
+	class ComputeTaskInfo
+	{
+		[CbField("h")]
+		public IoHash TaskHash { get; set; }
+
+		[CbField("c")]
+		public ChannelId ChannelId { get; set; }
+
+		private ComputeTaskInfo()
+		{
+		}
+
+		public ComputeTaskInfo(IoHash TaskHash, ChannelId ChannelId)
+		{
+			this.TaskHash = TaskHash;
+			this.ChannelId = ChannelId;
+		}
+	}
+
+	/// <summary>
+	/// Current status of a task
+	/// </summary>
+	[RedisConverter(typeof(RedisCbConverter<>))]
+	class ComputeTaskStatus : IComputeTaskStatus
+	{
+		/// <inheritdoc/>
+		[CbField("h")]
+		public IoHash TaskHash { get; set; }
+
+		/// <inheritdoc/>
+		[CbField("t")]
+		public DateTime Time { get; set; }
+
+		/// <inheritdoc/>
+		[CbField("s")]
+		public ComputeTaskState State { get; set; }
+
+		/// <inheritdoc cref="IComputeTaskStatus.AgentId"/>
+		[CbField("a")]
+		public Utf8String AgentId { get; set; }
+
+		/// <inheritdoc cref="IComputeTaskStatus.LeaseId"/>
+		[CbField("l")]
+		byte[]? LeaseIdBytes { get; set; }
+
+		/// <inheritdoc/>
+		[CbField("r")]
+		public IoHash? ResultHash { get; set; }
+
+		/// <inheritdoc/>
+		AgentId? IComputeTaskStatus.AgentId => AgentId.IsEmpty ? (AgentId?)null : new AgentId(AgentId.ToString());
+
+		/// <inheritdoc/>
+		public ObjectId? LeaseId
+		{
+			get => new ObjectId(LeaseIdBytes);
+			set => LeaseIdBytes = value?.ToByteArray();
+		}
+
+		private ComputeTaskStatus()
+		{
+		}
+
+		public ComputeTaskStatus(IoHash TaskHash, ComputeTaskState State)
+		{
+			this.TaskHash = TaskHash;
+			this.Time = DateTime.UtcNow;
+			this.State = State;
+		}
+	}
+
+	/// <summary>
+	/// Dispatches remote actions. Does not implement any cross-pod communication to satisfy leases; only agents connected to this server instance will be stored.
+	/// </summary>
+	class ComputeService : IComputeService, IDisposable
+	{
+		public MessageDescriptor Descriptor => ComputeTaskMessage.Descriptor;
+
+		public NamespaceId NamespaceId { get; } = new NamespaceId("default");
+
+		ITaskScheduler<ComputeTaskInfo> TaskScheduler;
+		RedisMessageQueue<ComputeTaskStatus> MessageQueue;
+		ILogger Logger;
+
+		/// <summary>
+		/// Constructor
+		/// </summary>
+		/// <param name="Redis">Redis instance</param>
+		/// <param name="ObjectCollection"></param>
+		/// <param name="Logger">The logger instance</param>
+		public ComputeService(IDatabase Redis, IObjectCollection ObjectCollection, ILogger<ComputeService> Logger)
+		{
+			this.TaskScheduler = new RedisTaskScheduler<ComputeTaskInfo>(Redis, ObjectCollection, NamespaceId, "compute/tasks/");
+			this.MessageQueue = new RedisMessageQueue<ComputeTaskStatus>(Redis, "compute/messages/");
+			this.Logger = Logger;
+		}
+
+		/// <inheritdoc/>
+		public void Dispose()
+		{
+			MessageQueue.Dispose();
+		}
+
+		/// <inheritdoc/>
+		public async Task AddTasksAsync(NamespaceId NamespaceId, IoHash RequirementsHash, List<IoHash> TaskHashes, ChannelId ChannelId)
+		{
+			List<Task> Tasks = new List<Task>();
+			foreach (IoHash TaskHash in TaskHashes)
+			{
+				ComputeTaskInfo TaskInfo = new ComputeTaskInfo(TaskHash, ChannelId);
+				Tasks.Add(TaskScheduler.EnqueueAsync(TaskInfo, RequirementsHash));
+			}
+			await Task.WhenAll(Tasks);
+		}
+
+		/// <inheritdoc/>
+		public async Task<AgentLease?> TryAssignLeaseAsync(IAgent Agent, CancellationToken CancellationToken)
+		{
+			Console.WriteLine("TRY ASSIGN LEASE");
+			TaskSchedulerEntry<ComputeTaskInfo>? Entry = await TaskScheduler.DequeueAsync(Agent, CancellationToken);
+			if(Entry != null)
+			{
+				Console.WriteLine("GOT IT");
+				ComputeTaskMessage ComputeTask = new ComputeTaskMessage();
+				ComputeTask.ChannelId = Entry.Item.ChannelId.ToString();
+				ComputeTask.NamespaceId = NamespaceId.ToString();
+				ComputeTask.RequirementsHash = Entry.RequirementsHash;
+				ComputeTask.TaskHash = Entry.Item.TaskHash;
+
+				string LeaseName = $"Remote action ({Entry.Item.TaskHash})";
+				byte[] Payload = Any.Pack(ComputeTask).ToByteArray();
+				return new AgentLease(ObjectId.GenerateNewId(), LeaseName, null, null, null, LeaseState.Pending, Payload, new AgentRequirements(), null);
+			}
+			return null;
+		}
+
+		/// <inheritdoc/>
+		public Task CancelLeaseAsync(IAgent Agent, ObjectId LeaseId, Any Payload)
+		{
+			ComputeTaskMessage Message = Payload.Unpack<ComputeTaskMessage>();
+			ComputeTaskInfo TaskInfo = new ComputeTaskInfo(Message.TaskHash, new ChannelId(Message.ChannelId));
+			return TaskScheduler.EnqueueAsync(TaskInfo, Message.RequirementsHash);
+		}
+
+		/// <inheritdoc/>
+		public async Task<List<IComputeTaskStatus>> GetTaskUpdatesAsync(ChannelId ChannelId, CancellationToken CancellationToken)
+		{
+			List<ComputeTaskStatus> Messages = await MessageQueue.ReadAsync(ChannelId.ToString(), CancellationToken);
+			return Messages.ConvertAll<IComputeTaskStatus>(x => x);
+		}
+
+		public Task OnLeaseStartedAsync(IAgent Agent, ObjectId LeaseId, Any Payload)
+		{
+			ComputeTaskMessage ComputeTask = Payload.Unpack<ComputeTaskMessage>();
+
+			ComputeTaskStatus Status = new ComputeTaskStatus(ComputeTask.TaskHash, ComputeTaskState.Executing);
+			return MessageQueue.PostAsync(ComputeTask.ChannelId, Status);
+		}
+
+		public Task OnLeaseFinishedAsync(IAgent Agent, ObjectId LeaseId, Any Payload, LeaseOutcome Outcome, ReadOnlyMemory<byte> Output)
+		{
+			ComputeTaskMessage ComputeTask = Payload.Unpack<ComputeTaskMessage>();
+
+			Logger.LogInformation("Task complete: {Hash}", (IoHash)ComputeTask.TaskHash);
+			ComputeTaskResultMessage Result = ComputeTaskResultMessage.Parser.ParseFrom(Output.ToArray());
+
+			ComputeTaskStatus Status = new ComputeTaskStatus(ComputeTask.TaskHash, ComputeTaskState.Complete);
+			Status.LeaseId = LeaseId;
+			if (Result.OutputHash != null)
+			{
+				Status.ResultHash = Result.OutputHash;
+			}
+
+			return MessageQueue.PostAsync(ComputeTask.ChannelId, Status);
+		}
+	}
+
+	/// <summary>
+	/// Implementation of the gRPC compute service interface
+	/// </summary>
+	class ComputeRpcServer : ComputeRpc.ComputeRpcBase
+	{
+		IComputeService ComputeService;
+
+		public ComputeRpcServer(IComputeService ComputeService)
+		{
+			this.ComputeService = ComputeService;
+		}
+
+		public override async Task<AddTasksRpcResponse> AddTasks(AddTasksRpcRequest RpcRequest, ServerCallContext Context)
+		{
+			NamespaceId NamespaceId = new NamespaceId(RpcRequest.NamespaceId);
+			ChannelId ChannelId = new ChannelId(RpcRequest.ChannelId);
+			await ComputeService.AddTasksAsync(NamespaceId, RpcRequest.RequirementsHash, RpcRequest.TaskHashes.Select(x => (IoHash)x).ToList(), ChannelId);
+			return new AddTasksRpcResponse();
+		}
+
+		public override async Task GetTaskUpdates(IAsyncStreamReader<GetTaskUpdatesRpcRequest> RequestStream, IServerStreamWriter<GetTaskUpdatesRpcResponse> ResponseStream, ServerCallContext Context)
+		{
+			Task<bool> MoveNextTask = RequestStream.MoveNext();
+			while(await MoveNextTask)
+			{
+				GetTaskUpdatesRpcRequest Request = RequestStream.Current;
+
+				ChannelId ChannelId = new ChannelId(Request.ChannelId);
+				using (CancellationTokenSource CancellationSource = new CancellationTokenSource())
+				{
+					MoveNextTask = MoveNextAndCancel(RequestStream, CancellationSource);
+					while (!CancellationSource.IsCancellationRequested)
+					{
+						List<IComputeTaskStatus> Updates = await ComputeService.GetTaskUpdatesAsync(ChannelId, CancellationSource.Token);
+						foreach (IComputeTaskStatus Update in Updates)
+						{
+							GetTaskUpdatesRpcResponse Response = new GetTaskUpdatesRpcResponse();
+							Response.TaskHash = Update.TaskHash;
+							Response.Time = Timestamp.FromDateTime(Update.Time);
+							Response.State = Update.State;
+							Response.ResultHash = Update.ResultHash;
+							await ResponseStream.WriteAsync(Response);
+						}
+					}
+					await MoveNextTask;
+				}
+			}
+		}
+
+		static async Task<bool> MoveNextAndCancel(IAsyncStreamReader<GetTaskUpdatesRpcRequest> RequestStream, CancellationTokenSource CancellationSource)
+		{
+			bool Result = await RequestStream.MoveNext();
+			CancellationSource.Cancel();
+			return Result;
+		}
+	}
+}
diff --git a/Engine/Source/Programs/Horde/HordeServer/Compute/Impl/RedisMessageQueue.cs b/Engine/Source/Programs/Horde/HordeServer/Compute/Impl/RedisMessageQueue.cs
new file mode 100644
--- /dev/null
+++ b/Engine/Source/Programs/Horde/HordeServer/Compute/Impl/RedisMessageQueue.cs
@@ -0,0 +1,169 @@
+// Copyright Epic Games, Inc. All Rights Reserved.
+
+using EpicGames.Redis;
+using HordeServer.Compute.Impl;
+using Microsoft.Extensions.Logging;
+using StackExchange.Redis;
+using System;
+using System.Collections.Concurrent;
+using System.Collections.Generic;
+using System.Linq;
+using System.Threading;
+using System.Threading.Tasks;
+
+namespace HordeServer.Compute.Impl
+{
+	/// <summary>
+	/// Interface for a distributed message queue
+	/// </summary>
+	/// <typeparam name="T">The message type</typeparam>
+	interface IMessageQueue<T> where T : class
+	{
+		/// <summary>
+		/// Adds messages to the queue
+		/// </summary>
+		/// <param name="ChannelId">The channel to post to</param>
+		/// <param name="Message">Message to post</param>
+		Task PostAsync(string ChannelId, T Message);
+
+		/// <summary>
+		/// Waits for a message to be available on the given channel
+		/// </summary>
+		/// <param name="ChannelId">The channel to read from</param>
+		/// <param name="CancellationToken">May be signalled to stop the wait, without throwing an exception</param>
+		/// <returns>True if a message was available, false otherwise</returns>
+		Task<List<T>> ReadAsync(string ChannelId, CancellationToken CancellationToken);
+	}
+
+	/// <summary>
+	/// Implementation of a message queue using Redis
+	/// </summary>
+	/// <typeparam name="T">The message type. Must be serailizable by RedisSerializer.</typeparam>
+	class RedisMessageQueue<T> : IMessageQueue<T>, IDisposable where T : class
+	{
+		IDatabase Redis;
+		RedisKey KeyPrefix;
+		RedisChannel UpdateChannel;
+		Dictionary<string, TaskCompletionSource<bool>> ChannelWakeEvents = new Dictionary<string, TaskCompletionSource<bool>>();
+
+		// Time after which entries should be removed
+		TimeSpan ExpireTime { get; set; } = TimeSpan.FromSeconds(30);
+
+		/// <summary>
+		/// Constructor
+		/// </summary>
+		/// <param name="Redis">The Redis database instance</param>
+		/// <param name="KeyPrefix">Prefix for keys to use in the database</param>
+		public RedisMessageQueue(IDatabase Redis, RedisKey KeyPrefix)
+		{
+			this.Redis = Redis;
+			this.KeyPrefix = KeyPrefix;
+			this.UpdateChannel = KeyPrefix.Append("updates").ToString();
+
+			Redis.Multiplexer.GetSubscriber().Subscribe(UpdateChannel, OnChannelUpdate);
+		}
+
+		/// <inheritdoc/>
+		public void Dispose()
+		{
+			Redis.Multiplexer.GetSubscriber().Unsubscribe(UpdateChannel, OnChannelUpdate);
+		}
+
+		/// <summary>
+		/// Callback for a message being posted to a channel
+		/// </summary>
+		/// <param name="Channel"></param>
+		/// <param name="Value"></param>
+		void OnChannelUpdate(RedisChannel Channel, RedisValue Value)
+		{
+			TaskCompletionSource<bool>? CompletionSource;
+			lock (ChannelWakeEvents)
+			{
+				ChannelWakeEvents.TryGetValue(Value.ToString(), out CompletionSource);
+			}
+			if (CompletionSource != null)
+			{
+				CompletionSource.TrySetResult(true);
+			}
+		}
+
+		/// <summary>
+		/// Gets the set of messages for a channel
+		/// </summary>
+		/// <param name="ChannelId"></param>
+		/// <returns></returns>
+		RedisList<T> GetChannel(string ChannelId)
+		{
+			return new RedisList<T>(KeyPrefix.Append(ChannelId));
+		}
+
+		/// <inheritdoc/>
+		public async Task PostAsync(string ChannelId, T Message)
+		{
+			RedisList<T> Channel = GetChannel(ChannelId);
+
+			long Length = await Redis.ListRightPushAsync(Channel, Message);
+			if (Length == 1)
+			{
+				await Redis.PublishAsync(UpdateChannel, ChannelId, CommandFlags.FireAndForget);
+			}
+			await Redis.KeyExpireAsync(Channel.Key, ExpireTime, CommandFlags.FireAndForget);
+		}
+
+		async Task<bool> ReadMessagesAsync(RedisList<T> List, List<T> Messages)
+		{
+			T? Message = await Redis.ListLeftPopAsync(List);
+			while(Message != null)
+			{
+				Messages.Add(Message);
+				Message = await Redis.ListLeftPopAsync(List);
+			}
+			return Messages.Count > 0;
+		}
+
+		/// <inheritdoc/>
+		public async Task<List<T>> ReadAsync(string ChannelId, CancellationToken CancellationToken)
+		{
+			List<T> Messages = new List<T>();
+
+			RedisList<T> Channel = GetChannel(ChannelId);
+			while (!CancellationToken.IsCancellationRequested && !await ReadMessagesAsync(Channel, Messages))
+			{
+				// Register for notifications on this channel
+				TaskCompletionSource<bool>? CompletionSource;
+				lock (ChannelWakeEvents)
+				{
+					if (!ChannelWakeEvents.TryGetValue(ChannelId, out CompletionSource))
+					{
+						CompletionSource = new TaskCompletionSource<bool>();
+						ChannelWakeEvents.Add(ChannelId, CompletionSource);
+					}
+				}
+
+				try
+				{
+					// Read the current queue state again, in case it was modified since we registered.
+					if(await ReadMessagesAsync(Channel, Messages))
+					{
+						break;
+					}
+
+					// Wait for messages to be available
+					using (CancellationToken.Register(() => CompletionSource.TrySetResult(false)))
+					{
+						await CompletionSource.Task;
+					}
+				}
+				finally
+				{
+					lock (ChannelWakeEvents)
+					{
+						ChannelWakeEvents.Remove(ChannelId);
+					}
+				}
+			}
+
+			return Messages;
+		}
+	}
+}
diff --git a/Engine/Source/Programs/Horde/HordeServer/Compute/Impl/RedisTaskScheduler.cs b/Engine/Source/Programs/Horde/HordeServer/Compute/Impl/RedisTaskScheduler.cs
new file mode 100644
--- /dev/null
+++ b/Engine/Source/Programs/Horde/HordeServer/Compute/Impl/RedisTaskScheduler.cs
@@ -0,0 +1,661 @@
+// Copyright Epic Games, Inc. All Rights Reserved.
+
+using EpicGames.Core;
+using EpicGames.Horde.Compute;
+using EpicGames.Redis;
+using EpicGames.Serialization;
+using HordeServer.Models;
+using HordeServer.Services;
+using HordeServer.Storage;
+using HordeServer.Utilities;
+using Microsoft.Extensions.Caching.Memory;
+using Microsoft.Extensions.Logging;
+using StackExchange.Redis;
+using System;
+using System.Collections.Concurrent;
+using System.Collections.Generic;
+using System.Diagnostics;
+using System.Linq;
+using System.Threading;
+using System.Threading.Channels;
+using System.Threading.Tasks;
+
+namespace HordeServer.Compute.Impl
+{
+	using NamespaceId = StringId<INamespace>;
+
+	/// <summary>
+	/// Describes an entry for a compute queue.
+	/// </summary>
+	/// <typeparam name="T">The type describing the work to be performed. Must be serializable to a RedisValue via RedisSerializer.</typeparam>
+	class TaskSchedulerEntry<T>
+	{
+		/// <summary>
+		/// The task description, encoded as a RedisValue
+		/// </summary>
+		public T Item { get; }
+
+		/// <summary>
+		/// Hash of the <see cref="Requirements"/> object in CAS for the agent to execute the task
+		/// </summary>
+		public IoHash RequirementsHash { get; }
+
+		/// <summary>
+		/// Constructor
+		/// </summary>
+		/// <param name="Item">The task item</param>
+		/// <param name="RequirementsHash">Requirements for executing the task</param>
+		public TaskSchedulerEntry(T Item, IoHash RequirementsHash)
+		{
+			this.Item = Item;
+			this.RequirementsHash = RequirementsHash;
+		}
+	}
+
+	/// <summary>
+	/// Interface for a queue of compute operations, each of which may have different requirements for the executing machines.
+	/// </summary>
+	interface ITaskScheduler<T>
+	{
+		/// <summary>
+		/// Adds a task to the queue
+		/// </summary>
+		/// <param name="Item">The item to add</param>
+		/// <param name="RequirementsHash">Hash of a <see cref="Requirements"/> object stored in the CAS which describes the agent to execute the task</param>
+		Task EnqueueAsync(T Item, IoHash RequirementsHash);
+
+		/// <summary>
+		/// Inserts a previously dequeued task back at the front of the queue
+		/// </summary>
+		/// <param name="Item">The item to add</param>
+		/// <param name="RequirementsHash">Hash of a <see cref="Requirements"/> object stored in the CAS which describes the agent to execute the task</param>
+		Task RequeueAsync(T Item, IoHash RequirementsHash);
+
+		/// <summary>
+		/// Dequeues a task that the given agent can execute
+		/// </summary>
+		/// <param name="Agent">The agent to execute the task</param>
+		/// <param name="Token">Cancellation token for the operation. Will return a null entry rather than throwing an exception.</param>
+		/// <returns>Information about the task to be executed</returns>
+		Task<TaskSchedulerEntry<T>?> DequeueAsync(IAgent Agent, CancellationToken Token = default);
+	}
+
+	/// <summary>
+	/// Implementation of <see cref="ITaskScheduler{T}"/> using Redis for storage
+	/// </summary>
+	/// <typeparam name="T">The task definition type</typeparam>
+	class RedisTaskScheduler<T> : ITaskScheduler<T>, IDisposable where T : class
+	{
+		class Listener
+		{
+			public IAgent Agent { get; }
+			public TaskCompletionSource<TaskSchedulerEntry<T>?> CompletionSource { get; }
+
+			public Listener(IAgent Agent)
+			{
+				this.Agent = Agent;
+				this.CompletionSource = new TaskCompletionSource<TaskSchedulerEntry<T>?>();
+			}
+		}
+
+		IDatabase Redis;
+		IObjectCollection ObjectCollection;
+		RedisKey BaseKey;
+		NamespaceId NamespaceId;
+		RedisSet<IoHash> QueueIndex;
+		RedisHash<IoHash, DateTime> ActiveQueues; // Queues which are actively being dequeued from
+		ReadOnlyHashSet<IoHash> LocalActiveQueues = new HashSet<IoHash>();
+		Stopwatch ResetActiveQueuesTimer = Stopwatch.StartNew();
+
+		List<Listener> Listeners = new List<Listener>();
+		RedisChannel<IoHash> NewQueueChannel;
+		Task QueueUpdateTask;
+		CancellationTokenSource CancellationSource = new CancellationTokenSource();
+		IMemoryCache RequirementsCache;
+
+		/// <summary>
+		/// Constructor
+		/// </summary>
+		/// <param name="Redis">The Redis database instance</param>
+		/// <param name="ObjectCollection">Interface to the CAS object store</param>
+		/// <param name="BaseKey">Base key for all keys used by this scheduler</param>
+		/// <param name="NamespaceId">The namespace identifier</param>
+		public RedisTaskScheduler(IDatabase Redis, IObjectCollection ObjectCollection, NamespaceId NamespaceId, RedisKey BaseKey)
+		{
+			this.Redis = Redis;
+			this.ObjectCollection = ObjectCollection;
+			this.NamespaceId = NamespaceId;
+			this.BaseKey = BaseKey;
+			this.QueueIndex = new RedisSet<IoHash>(BaseKey.Append("index"));
+			this.ActiveQueues = new RedisHash<IoHash, DateTime>(BaseKey.Append("active"));
+			this.NewQueueChannel = new RedisChannel<IoHash>(BaseKey.Append("new_queues").ToString());
+			this.RequirementsCache = new MemoryCache(new MemoryCacheOptions());
+
+			QueueUpdateTask = Task.Run(() => UpdateQueuesAsync(CancellationSource.Token));
+		}
+
+		/// <inheritdoc/>
+		public async ValueTask DisposeAsync()
+		{
+			if (!QueueUpdateTask.IsCompleted)
+			{
+				CancellationSource.Cancel();
+				await QueueUpdateTask;
+			}
+			RequirementsCache.Dispose();
+			CancellationSource.Dispose();
+		}
+
+		/// <inheritdoc/>
+		public void Dispose()
+		{
+			DisposeAsync().AsTask().Wait();
+		}
+
+		/// <summary>
+		/// Gets the key for a list of tasks for a particular queue
+		/// </summary>
+		/// <param name="RequirementsHash"></param>
+		/// <returns></returns>
+		RedisKey GetQueueKey(IoHash RequirementsHash)
+		{
+			return BaseKey.Append(RedisSerializer.Serialize(RequirementsHash).AsKey());
+		}
+
+		/// <summary>
+		/// Gets the key for a list of tasks for a particular queue
+		/// </summary>
+		/// <param name="RequirementsHash"></param>
+		/// <returns></returns>
+		RedisList<T> GetQueue(IoHash RequirementsHash)
+		{
+			return new RedisList<T>(GetQueueKey(RequirementsHash));
+		}
+
+		/// <summary>
+		/// Adds a queue to the index, atomically checking that the queue exists
+		/// </summary>
+		/// <param name="RequirementsHash">Hash of the queue requirements objectadd</param>
+		async ValueTask AddQueueToIndexAsync(IoHash RequirementsHash)
+		{
+			ITransaction Transaction = Redis.CreateTransaction();
+			Transaction.AddCondition(Condition.KeyExists(GetQueueKey(RequirementsHash)));
+			_ = Transaction.SetAddAsync(QueueIndex, RequirementsHash);
+			await Transaction.ExecuteAsync(CommandFlags.FireAndForget);
+		}
+
+		/// <summary>
+		/// Removes a queue from the index, atomically checking that it is empty
+		/// </summary>
+		/// <param name="RequirementsHash">Hash of the queue requirements objectadd</param>
+		async ValueTask RemoveQueueFromIndexAsync(IoHash RequirementsHash)
+		{
+			ITransaction Transaction = Redis.CreateTransaction();
+			Transaction.AddCondition(Condition.KeyNotExists(GetQueueKey(RequirementsHash)));
+			_ = Transaction.SetRemoveAsync(QueueIndex, RequirementsHash);
+			await Transaction.ExecuteAsync(CommandFlags.FireAndForget);
+		}
+
+		/// <summary>
+		/// Adds an item to a particular queue, creating and adding that queue to the index if necessary
+		/// </summary>
+		/// <param name="Item">The item to be added</param>
+		/// <param name="RequirementsHash">Requirements for the task</param>
+		public async Task EnqueueAsync(T Item, IoHash RequirementsHash)
+		{
+			RedisList<T> List = GetQueue(RequirementsHash);
+			if (await Redis.ListRightPushAsync(List, Item) == 1)
+			{
+				await AddQueueToIndexAsync(RequirementsHash);
+			}
+		}
+
+		/// <summary>
+		/// Adds an item to the front of a particular queue, creating and adding that queue to the index if necessary
+		/// </summary>
+		/// <param name="Item">The item to be added</param>
+		/// <param name="RequirementsHash">Requirements for the task</param>
+		public async Task RequeueAsync(T Item, IoHash RequirementsHash)
+		{
+			RedisList<T> Queue = GetQueue(RequirementsHash);
+			if (await Redis.ListLeftPushAsync(Queue, Item) == 1)
+			{
+				await AddQueueToIndexAsync(RequirementsHash);
+			}
+		}
+
+		/// <summary>
+		/// Dequeues an item for execution by the given agent
+		/// </summary>
+		/// <param name="Agent">The agent to dequeue an item for</param>
+		/// <param name="Token">Cancellation token for waiting for an item</param>
+		/// <returns>The dequeued item, or null if no item is available</returns>
+		public async Task<TaskSchedulerEntry<T>?> DequeueAsync(IAgent Agent, CancellationToken Token = default)
+		{
+			// Compare against all the list of cached queues to see if we can dequeue something from any of them
+			Listener? Listener = null;
+			try
+			{
+				IoHash[] Queues = await Redis.SetMembersAsync(QueueIndex);
+				while (!Token.IsCancellationRequested)
+				{
+					// Try to dequeue an item from the list
+					TaskSchedulerEntry<T>? Entry = await TryAssignToLocalAgentAsync(Queues, Agent);
+					if (Entry != null)
+					{
+						return Entry;
+					}
+
+					// Create and register a listener for this waiter 
+					if (Listener == null)
+					{
+						Listener = new Listener(Agent);
+						lock (Listeners)
+						{
+							Listeners.Add(Listener);
+						}
+					}
+					else
+					{
+						using (IDisposable Registration = Token.Register(() => Listener.CompletionSource.TrySetResult(null)))
+						{
+							return await Listener.CompletionSource.Task;
+						}
+					}
+				}
+			}
+			finally
+			{
+				if (Listener != null)
+				{
+					lock (Listeners)
+					{
+						Listeners.Remove(Listener);
+					}
+				}
+			}
+			return null;
+		}
+
+		async Task<TaskSchedulerEntry<T>?> TryAssignToLocalAgentAsync(IoHash[] Queues, IAgent Agent)
+		{
+			foreach (IoHash RequirementsHash in Queues)
+			{
+				Requirements? Requirements = await GetRequirementsAsync(RequirementsHash);
+				if (Requirements != null && CheckRequirements(Agent, Requirements))
+				{
+					T? Item = await DequeueAsync(RequirementsHash);
+					if (Item != null)
+					{
+						return new TaskSchedulerEntry<T>(Item, RequirementsHash);
+					}
+				}
+			}
+			return null;
+		}
+
+		/// <summary>
+		/// Dequeues the item at the front of a queue
+		/// </summary>
+		/// <param name="RequirementsHash">Key of the queue to remove from</param>
+		/// <returns>The dequeued item, or null if the queue is empty</returns>
+		public async Task<T?> DequeueAsync(IoHash RequirementsHash)
+		{
+			await AddActiveQueue(RequirementsHash);
+
+			T? Item = await Redis.ListLeftPopAsync(GetQueue(RequirementsHash));
+			if (Item == null)
+			{
+				await RemoveQueueFromIndexAsync(RequirementsHash);
+			}
+			return Item;
+		}
+
+		/// <summary>
+		/// Marks a queue key as being actively monitored, preventing it being returned by <see cref="GetInactiveQueuesAsync"/>
+		/// </summary>
+		/// <param name="Key">The queue key</param>
+		/// <returns></returns>
+		async ValueTask AddActiveQueue(IoHash Key)
+		{
+			// Periodically clear out the set of active keys
+			TimeSpan ResetTime = TimeSpan.FromSeconds(10.0);
+			if (ResetActiveQueuesTimer.Elapsed > ResetTime)
+			{
+				lock (ResetActiveQueuesTimer)
+				{
+					if (ResetActiveQueuesTimer.Elapsed > ResetTime)
+					{
+						LocalActiveQueues = new HashSet<IoHash>();
+						ResetActiveQueuesTimer.Reset();
+					}
+				}
+			}
+
+			// Check if the set of active keys already contains the key we're adding. In order to optimize the 
+			// common case under heavy load where the key is in the set, updating it creates a full copy of it. Any
+			// readers can thus access it without the need for any locking.
+			ReadOnlyHashSet<IoHash> LocalActiveKeysCopy = LocalActiveQueues;
+			if (!LocalActiveKeysCopy.Contains(Key))
+			{
+				for (; ; )
+				{
+					HashSet<IoHash> NewLocalActiveKeys = new HashSet<IoHash>(LocalActiveKeysCopy);
+					if (!NewLocalActiveKeys.Add(Key))
+					{
+						break;
+					}
+					if (Interlocked.CompareExchange(ref LocalActiveQueues, NewLocalActiveKeys, LocalActiveKeysCopy) == LocalActiveKeysCopy)
+					{
+						await Redis.HashSetAsync(ActiveQueues, Key, DateTime.UtcNow);
+						break;
+					}
+				}
+			}
+		}
+
+		/// <summary>
+		/// Find any inactive keys
+		/// </summary>
+		/// <returns></returns>
+		public async Task<List<IoHash>> GetInactiveQueuesAsync()
+		{
+			HashSet<IoHash> Keys = new HashSet<IoHash>(await Redis.SetMembersAsync(QueueIndex));
+			HashSet<IoHash> InvalidKeys = new HashSet<IoHash>();
+
+			DateTime MinTime = DateTime.UtcNow - TimeSpan.FromMinutes(1.0);
+
+			HashEntry<IoHash, DateTime>[] Entries = await Redis.HashGetAllAsync(ActiveQueues);
+			foreach (HashEntry<IoHash, DateTime> Entry in Entries)
+			{
+				if (Entry.Value < MinTime)
+				{
+					InvalidKeys.Add(Entry.Name);
+				}
+				else
+				{
+					Keys.Remove(Entry.Name);
+				}
+			}
+
+			if (InvalidKeys.Count > 0)
+			{
+				await Redis.HashDeleteAsync(ActiveQueues, InvalidKeys.ToArray());
+			}
+
+			return Keys.ToList();
+		}
+
+		async Task UpdateQueuesAsync(CancellationToken CancellationToken)
+		{
+			Channel<IoHash> NewQueues = Channel.CreateUnbounded<IoHash>();
+
+			ISubscriber Subscriber = Redis.Multiplexer.GetSubscriber();
+			await using var _ = await Subscriber.SubscribeAsync(NewQueueChannel, (_, v) => NewQueues.Writer.TryWrite(v));
+
+			while (await NewQueues.Reader.WaitToReadAsync(CancellationToken))
+			{
+				HashSet<IoHash> NewQueueKeys = new HashSet<IoHash>();
+				while (NewQueues.Reader.TryRead(out IoHash Key))
+				{
+					NewQueueKeys.Add(Key);
+				}
+				foreach (IoHash NewQueueKey in NewQueueKeys)
+				{
+					await TryDispatchToNewQueueAsync(NewQueueKey);
+				}
+			}
+		}
+
+		async Task<bool> TryDispatchToNewQueueAsync(IoHash RequirementsHash)
+		{
+			RedisList<T> Queue = GetQueue(RequirementsHash);
+
+			// Get the requirements for this queue
+			Requirements? Requirements = await GetRequirementsAsync(RequirementsHash);
+			if (Requirements == null)
+			{
+				throw new Exception($"Unable to find requirements with hash {RequirementsHash}");
+			}
+
+			// Find a local listener that can execute the work
+			TaskSchedulerEntry<T>? Entry = null;
+			try
+			{
+				for (; ; )
+				{
+					// Look for a listener that can execute the task
+					Listener? Listener;
+					lock (Listeners)
+					{
+						Listener = Listeners.FirstOrDefault(x => !x.CompletionSource.Task.IsCompleted && CheckRequirements(x.Agent, Requirements));
+					}
+					if (Listener == null)
+					{
+						return false;
+					}
+
+					// Pop an entry from the queue
+					if (Entry == null)
+					{
+						T? TaskData = await DequeueAsync(RequirementsHash);
+						if (TaskData == null)
+						{
+							return false;
+						}
+						Entry = new TaskSchedulerEntry<T>(TaskData, RequirementsHash);
+					}
+
+					// Assign it to the listener
+					if (Listener.CompletionSource.TrySetResult(Entry))
+					{
+						Entry = null;
+					}
+				}
+			}
+			finally
+			{
+				if (Entry != null)
+				{
+					await RequeueAsync(Entry.Item, Entry.RequirementsHash);
+				}
+			}
+		}
+
+		static bool CheckRequirements(IAgent Agent, Requirements Requirements)
+		{
+			_ = Agent;
+			_ = Requirements;
+			// TODO: check requirements
+			// TODO: cache requirements
+			return true;
+		}
+
+		/// <summary>
+		/// Gets the requirements object from the CAS
+		/// </summary>
+		/// <param name="RequirementsHash"></param>
+		/// <returns></returns>
+		async Task<Requirements?> GetRequirementsAsync(IoHash RequirementsHash)
+		{
+			Requirements? Requirements;
+			if (!RequirementsCache.TryGetValue(RequirementsHash, out Requirements))
+			{
+				Requirements = await ObjectCollection.GetAsync<Requirements>(NamespaceId, RequirementsHash);
+				using (ICacheEntry Entry = RequirementsCache.CreateEntry(RequirementsHash))
+				{
+					Entry.SetSlidingExpiration(TimeSpan.FromMinutes(10.0));
+					Entry.SetValue(Requirements);
+				}
+			}
+			return Requirements;
+		}
+
+#if false
+		class QueueStats
+		{
+			QueueKey Key;
+			double Score;
+			int Length;
+		}
+
+		interface IQueueIndex
+		{
+			public IReadOnlyList<QueueKey> Keys { get; }
+
+			public async Task AddAsync(QueueKey Key);
+
+			public async Task RemoveAsync(QueueKey Key);
+		}
+#if false
+
+
+		DatabaseService DatabaseService;
+		IObjectCollection ObjectCollection;
+		IDatabase Redis;
+		RedisKey BaseKey;
+		RedisSortedSet<QueueKey> Index;
+		RedisChannel<QueueKey> NewQueueChannel;
+		List<Listener> Listeners = new List<Listener>();
+		IMemoryCache RequirementsCache = new MemoryCache(new MemoryCacheOptions());
+		List<QueueKey> CachedQueues = new List<QueueKey>();
+		Task? UpdateQueuesTask;
+		BackgroundTick? BackgroundTick;
+		ILogger Logger;
+
+		public RedisTaskScheduler(DatabaseService DatabaseService, IObjectCollection ObjectCollection, IDatabase Redis, RedisKey BaseKey, ILogger<RedisTaskScheduler<T>> Logger)
+		{
+			this.DatabaseService = DatabaseService;
+			this.ObjectCollection = ObjectCollection;
+			this.Redis = Redis;
+			this.BaseKey = BaseKey;
+			this.Index = new RedisSortedSet<QueueKey>(BaseKey.Append("index"));
+			this.NewQueueChannel = new RedisChannel<QueueKey>(BaseKey.Append("new_queues").ToString());
+			this.Logger = Logger;
+
+			UpdateQueuesTask = Task.Run(() => CheckNewQueues());
+		}
+
+		public async Task StopAsync()
+		{
+			if (UpdateQueuesTask != null)
+			{
+				NewQueues.Writer.TryComplete();
+				await UpdateQueuesTask;
+			}
+		}
+
+		public void Dispose()
+		{
+			StopAsync().Wait();
+			BackgroundTick?.Dispose();
+			RequirementsCache.Dispose();
+		}
+
+		RedisList<T> GetQueue(QueueKey QueueKey)
+		{
+			return new RedisList<T>(QueueKey.GetKey(BaseKey));
+		}
+
+
+		/// <inheritdoc/>
+		public async Task EnqueueAsync(TaskSchedulerEntry<T> Entry, bool AtFront = false)
+		{
+			if (!await TryAssignToLocalListenerAsync(Entry))
+			{
+				await AddToSharedQueue(Entry, AtFront);
+			}
+		}
+
+		async Task<bool> TryAssignToLocalListenerAsync(TaskSchedulerEntry<T> Entry)
+		{
+			// Get the requirements for this entry
+			Requirements? Requirements = await GetRequirementsAsync(Entry.NamespaceId, Entry.RequirementsHash);
+			if (Requirements == null)
+			{
+				throw new Exception($"Unable to find requirements with hash {Entry.RequirementsHash}");
+			}
+
+			// Find a local listener that can execute the work
+			for (; ; )
+			{
+				Listener? Listener;
+				lock (Listeners)
+				{
+					Listener = Listeners.FirstOrDefault(x => !x.CompletionSource.Task.IsCompleted && CheckRequirements(x.Agent, Requirements));
+				}
+
+				if (Listener == null)
+				{
+					return false;
+				}
+				else if (Listener.CompletionSource.TrySetResult(Entry))
+				{
+					return true;
+				}
+			}
+		}
+
+		private async Task AddToSharedQueue(TaskSchedulerEntry<T> Entry, bool AtFront)
+		{
+			// TODO: check if there's already something listening
+
+			QueueKey QueueKey = new QueueKey(Entry);
+			RedisList<T> Queue = new RedisList<T>(QueueKey.GetKey(BaseKey));
+
+			// Insert the item in to the appropriate queue
+			long NewLength;
+			if (AtFront)
+			{
+				NewLength = await Redis.ListLeftPushAsync(Queue, Entry.TaskData);
+			}
+			else
+			{
+				NewLength = await Redis.ListRightPushAsync(Queue, Entry.TaskData);
+			}
+
+			// Check if this was a new queue; if so, we need to add a new entry to the index for it, and notify other pods
+			if (NewLength == 1)
+			{
+				await Redis.SortedSetAddAsync(Index, QueueKey, DateTime.UtcNow.Ticks);
+				await Redis.PublishAsync(NewQueueChannel, QueueKey);
+			}
+		}
+
+
+		static bool CheckRequirements(IAgent Agent, Requirements Requirements)
+		{
+			_ = Agent;
+			_ = Requirements;
+			// TODO: check requirements
+			// TODO: cache requirements
+			return true;
+		}
+
+		List<QueueKey> ModifyCachedQueues(Action<List<QueueKey>> Action)
+		{
+			for (; ; )
+			{
+				List<QueueKey> QueueKeys = CachedQueues;
+				if(TryModifyCachedQueueKeys(QueueKeys, Action))
+				{
+					return QueueKeys;
+				}
+			}
+		}
+
+		bool TryModifyCachedQueueKeys(List<QueueKey> InitialQueues, Action<List<QueueKey>> Action)
+		{
+			if (CachedQueues == InitialQueues)
+			{
+				List<QueueKey> NewQueues = new List<QueueKey>(InitialQueues);
+				Action(NewQueues);
+				return Interlocked.CompareExchange(ref CachedQueues, NewQueues, InitialQueues) == InitialQueues;
+			}
+			return false;
+		}
+
+#endif
+#endif
+	}
+}
diff --git a/Engine/Source/Programs/Horde/HordeServer/HordeServer.csproj b/Engine/Source/Programs/Horde/HordeServer/HordeServer.csproj
--- a/Engine/Source/Programs/Horde/HordeServer/HordeServer.csproj
+++ b/Engine/Source/Programs/Horde/HordeServer/HordeServer.csproj
@@ -47,6 +47,7 @@
     <PackageReference Include="SharpZipLib" Version="1.2.0" />
     <PackageReference Include="StackExchange.Redis" Version="2.2.4" />
     <PackageReference Include="Swashbuckle.AspNetCore" Version="5.5.0" />
+    <PackageReference Include="System.Linq.Async" Version="5.0.0" />
     <PackageReference Include="TimeZoneConverter" Version="3.2.0" />
   </ItemGroup>
 
@@ -63,6 +64,7 @@
   <ItemGroup>
     <Protobuf Include="..\..\Shared\EpicGames.Horde\Protos\horde\storage\ref_table.proto" ProtoRoot="..\..\Shared\EpicGames.Horde\Protos" GrpcServices="Server" />
     <Protobuf Include="..\..\Shared\EpicGames.Horde\Protos\horde\storage\blob_store.proto" ProtoRoot="..\..\Shared\EpicGames.Horde\Protos" GrpcServices="Server" />
+    <Protobuf Include="..\..\Shared\EpicGames.Horde\Protos\horde\compute\compute_rpc.proto" ProtoRoot="..\..\Shared\EpicGames.Horde\Protos" GrpcServices="Server" />
     <Protobuf Include="..\..\Shared\EpicGames.Horde.Common\Protos\google\bytestream\bytestream.proto" ProtoRoot="..\..\Shared\EpicGames.Horde.Common\Protos" GrpcServices="Server" />
     <Protobuf Include="..\..\Shared\EpicGames.Horde.Common\Protos\horde\action_rpc.proto" ProtoRoot="..\..\Shared\EpicGames.Horde.Common\Protos" GrpcServices="Server" />
     <Protobuf Include="..\..\Shared\EpicGames.Horde.Common\Protos\horde\tasks\action_task.proto" ProtoRoot="..\..\Shared\EpicGames.Horde.Common\Protos" GrpcServices="Server" />
@@ -77,6 +79,7 @@
     <ProjectReference Include="..\..\Shared\EpicGames.Horde\EpicGames.Horde.csproj" />
     <ProjectReference Include="..\..\Shared\EpicGames.IoHash\EpicGames.IoHash.csproj" />
     <ProjectReference Include="..\..\Shared\EpicGames.Perforce\EpicGames.Perforce.csproj" />
+    <ProjectReference Include="..\..\Shared\EpicGames.Redis\EpicGames.Redis.csproj" />
     <ProjectReference Include="..\..\Shared\EpicGames.Serialization\EpicGames.Serialization.csproj" />
     <ProjectReference Include="..\HordeCommon\HordeCommon.csproj" />
     <ProjectReference Include="..\..\..\ThirdParty\Perforce\P4Api.Net\p4api.net\p4api.net.csproj" />
diff --git a/Engine/Source/Programs/Horde/HordeServer/Startup.cs b/Engine/Source/Programs/Horde/HordeServer/Startup.cs
--- a/Engine/Source/Programs/Horde/HordeServer/Startup.cs
+++ b/Engine/Source/Programs/Horde/HordeServer/Startup.cs
@@ -79,6 +79,9 @@ using HordeServer.Commits.Impl;
 using HordeServer.Commits;
 using OpenTracing.Contrib.Grpc.Interceptors;
 using OpenTracing.Util;
+using EpicGames.Horde.Compute;
+using HordeServer.Compute.Impl;
+using HordeServer.Compute;
 
 namespace HordeServer
 {
@@ -533,6 +536,7 @@ namespace HordeServer
 			Services.AddSingleton<JobTaskSource>();
 			Services.AddSingleton<ActionTaskSource>();
 			Services.AddSingleton<ConformTaskSource>();
+			Services.AddSingleton<IComputeService, ComputeService>();
 
 			Services.AddSingleton<ITaskSource, UpgradeTaskSource>();
 			Services.AddSingleton<ITaskSource, ShutdownTaskSource>();
@@ -540,6 +544,7 @@ namespace HordeServer
 			Services.AddSingleton<ITaskSource, ConformTaskSource>(Provider => Provider.GetRequiredService<ConformTaskSource>());
 			Services.AddSingleton<ITaskSource, JobTaskSource>(Provider => Provider.GetRequiredService<JobTaskSource>());
 			Services.AddSingleton<ITaskSource, ActionTaskSource>(Provider => Provider.GetRequiredService<ActionTaskSource>());
+			Services.AddSingleton<ITaskSource>(Provider => new NewTaskSourceWrapper(Provider.GetRequiredService<IComputeService>()));
 
 			Services.AddHostedService(Provider => Provider.GetRequiredService<ConformTaskSource>());
 
@@ -772,6 +777,7 @@ namespace HordeServer
 				Endpoints.MapGrpcService<ContentStorageService>();
 				Endpoints.MapGrpcService<ExecutionService>();
 
+				Endpoints.MapGrpcService<ComputeRpcServer>();
 				Endpoints.MapGrpcService<BlobRpc>();
 				Endpoints.MapGrpcService<BlobStoreRpc>();
 				Endpoints.MapGrpcService<RefRpc>();
diff --git a/Engine/Source/Programs/Horde/HordeServer/Storage/Collections/BlobCollection.cs b/Engine/Source/Programs/Horde/HordeServer/Storage/Collections/BlobCollection.cs
--- a/Engine/Source/Programs/Horde/HordeServer/Storage/Collections/BlobCollection.cs
+++ b/Engine/Source/Programs/Horde/HordeServer/Storage/Collections/BlobCollection.cs
@@ -66,6 +66,11 @@ namespace HordeServer.Storage.Collections
 		/// <returns></returns>
 		static string GetPath(NamespaceId NsId, IoHash Hash)
 		{
+			if (NsId.IsEmpty)
+			{
+				throw new InvalidOperationException("Namespace id may not be empty");
+			}
+
 			string HashText = Hash.ToString();
 			return $"blobs/{NsId}/{HashText[0..2]}/{HashText[2..4]}/{HashText}.blob";
 		}
diff --git a/Engine/Source/Programs/Horde/HordeServer/Storage/IObjectCollection.cs b/Engine/Source/Programs/Horde/HordeServer/Storage/IObjectCollection.cs
--- a/Engine/Source/Programs/Horde/HordeServer/Storage/IObjectCollection.cs
+++ b/Engine/Source/Programs/Horde/HordeServer/Storage/IObjectCollection.cs
@@ -57,6 +57,40 @@ namespace HordeServer.Storage
 	/// </summary>
 	public static class ObjectCollectionExtensions
 	{
+		/// <summary>
+		/// Adds an object
+		/// </summary>
+		/// <typeparam name="T"></typeparam>
+		/// <param name="ObjectCollection"></param>
+		/// <param name="NsId"></param>
+		/// <param name="Value"></param>
+		/// <returns></returns>
+		public static async Task<IoHash> AddAsync<T>(this IObjectCollection ObjectCollection, NamespaceId NsId, T Value) where T : class
+		{
+			CbObject ObjectValue = CbSerializer.Serialize<T>(Value);
+			IoHash Hash = ObjectValue.GetHash();
+			await ObjectCollection.AddAsync(NsId, Hash, ObjectValue);
+			return Hash;
+		}
+
+		/// <summary>
+		/// Gets a typed object from the store
+		/// </summary>
+		/// <typeparam name="T"></typeparam>
+		/// <param name="ObjectCollection"></param>
+		/// <param name="NsId"></param>
+		/// <param name="Hash"></param>
+		/// <returns></returns>
+		public static async Task<T?> GetAsync<T>(this IObjectCollection ObjectCollection, NamespaceId NsId, IoHash Hash) where T : class
+		{
+			CbObject? Object = await ObjectCollection.GetAsync(NsId, Hash);
+			if (Object == null)
+			{
+				return null;
+			}
+			return CbSerializer.Deserialize<T>(Object.AsField());
+		}
+
 		/// <summary>
 		/// Test whether a single blob exists
 		/// </summary>
diff --git a/Engine/Source/Programs/Horde/HordeServer/Tasks/ITaskSource.cs b/Engine/Source/Programs/Horde/HordeServer/Tasks/ITaskSource.cs
--- a/Engine/Source/Programs/Horde/HordeServer/Tasks/ITaskSource.cs
+++ b/Engine/Source/Programs/Horde/HordeServer/Tasks/ITaskSource.cs
@@ -10,6 +10,7 @@ using MongoDB.Bson;
 using System;
 using System.Collections.Generic;
 using System.Linq;
+using System.Threading;
 using System.Threading.Tasks;
 
 namespace HordeServer.Tasks
@@ -139,4 +140,120 @@ namespace HordeServer.Tasks
 			return Task.CompletedTask;
 		}
 	}
+
+	/// <summary>
+	/// Handler for a certain lease type
+	/// </summary>
+	public interface INewTaskSource
+	{
+		/// <summary>
+		/// Descriptor for this lease type
+		/// </summary>
+		MessageDescriptor Descriptor { get; }
+
+		/// <summary>
+		/// Waits for a lease to be availabe
+		/// </summary>
+		/// <param name="Agent">The agent to assign a lease to</param>
+		/// <param name="CancellationToken">Cancellation token for the wait</param>
+		/// <returns>New lease object</returns>
+		Task<AgentLease?> TryAssignLeaseAsync(IAgent Agent, CancellationToken CancellationToken);
+
+		/// <summary>
+		/// Cancel a lease that was previously assigned to an agent, allowing it to be assigned out again
+		/// </summary>
+		/// <param name="Agent">The agent that was assigned the lease</param>
+		/// <param name="LeaseId">The lease id</param>
+		/// <param name="Payload">Payload for the lease</param>
+		/// <returns></returns>
+		Task CancelLeaseAsync(IAgent Agent, ObjectId LeaseId, Any Payload);
+
+		/// <summary>
+		/// Notification that a lease has been started
+		/// </summary>
+		/// <param name="Agent">The agent executing the lease</param>
+		/// <param name="LeaseId">The lease id</param>
+		/// <param name="Payload">Payload for the lease</param>
+		Task OnLeaseStartedAsync(IAgent Agent, ObjectId LeaseId, Any Payload);
+
+		/// <summary>
+		/// Notification that a task has completed
+		/// </summary>
+		/// <param name="Agent">The agent that was allocated to the lease</param>
+		/// <param name="LeaseId">The lease id</param>
+		/// <param name="Payload">The lease payload</param>
+		/// <param name="Outcome">Outcome of the lease</param>
+		/// <param name="Output">Output from the task</param>
+		Task OnLeaseFinishedAsync(IAgent Agent, ObjectId LeaseId, Any Payload, LeaseOutcome Outcome, ReadOnlyMemory<byte> Output);
+	}
+
+	class NewTaskSourceWrapper : ITaskSource
+	{
+		class TaskListener : ITaskListener
+		{
+			NewTaskSourceWrapper Outer;
+			IAgent Agent;
+			CancellationTokenSource CancellationTokenSource;
+
+			public Task<NewLeaseInfo?> LeaseTask { get; }
+
+			public bool Accepted { get; set; }
+
+			public TaskListener(NewTaskSourceWrapper Outer, IAgent Agent)
+			{
+				this.Outer = Outer;
+				this.Agent = Agent;
+				this.CancellationTokenSource = new CancellationTokenSource();
+				this.LeaseTask = GetLeaseAsync();
+			}
+
+			async Task<NewLeaseInfo?> GetLeaseAsync()
+			{
+				AgentLease? Lease = await Outer.TaskSource.TryAssignLeaseAsync(Agent, CancellationTokenSource.Token);
+				return (Lease == null)? null : new NewLeaseInfo(Lease);
+			}
+
+			public async ValueTask DisposeAsync()
+			{
+				CancellationTokenSource.Cancel();
+				NewLeaseInfo? NewLease = await LeaseTask;
+
+				if (NewLease != null)
+				{
+					if (Accepted)
+					{
+						Console.WriteLine("ACCEPTED");
+						await Outer.TaskSource.OnLeaseStartedAsync(Agent, NewLease.Lease.Id, Any.Parser.ParseFrom(NewLease.Lease.Payload));
+					}
+					else 
+					{
+						Console.WriteLine("CANCELLED");
+						await Outer.TaskSource.CancelLeaseAsync(Agent, NewLease.Lease.Id, Any.Parser.ParseFrom(NewLease.Lease.Payload));
+					}
+				}
+
+				CancellationTokenSource.Dispose();
+			}
+		}
+
+		INewTaskSource TaskSource;
+
+		public MessageDescriptor Descriptor => TaskSource.Descriptor;
+
+		public NewTaskSourceWrapper(INewTaskSource TaskSource)
+		{
+			this.TaskSource = TaskSource;
+		}
+
+		public Task<ITaskListener?> SubscribeAsync(IAgent Agent)
+		{
+			return Task.FromResult<ITaskListener?>(new TaskListener(this, Agent));
+		}
+
+		public Task OnLeaseCompletedAsync(IAgent Agent, ObjectId LeaseId, Any Payload, LeaseOutcome Outcome, ReadOnlyMemory<byte> Output)
+		{
+			return TaskSource.OnLeaseFinishedAsync(Agent, LeaseId, Payload, Outcome, Output);
+		}
+	}
 }
+
diff --git a/Engine/Source/Programs/Horde/HordeServer/Utilities/BackgroundTick.cs b/Engine/Source/Programs/Horde/HordeServer/Utilities/BackgroundTick.cs
new file mode 100644
--- /dev/null
+++ b/Engine/Source/Programs/Horde/HordeServer/Utilities/BackgroundTick.cs
@@ -0,0 +1,92 @@
+// Copyright Epic Games, Inc. All Rights Reserved.
+
+using Microsoft.Extensions.Logging;
+using System;
+using System.Collections.Generic;
+using System.Diagnostics;
+using System.Linq;
+using System.Threading;
+using System.Threading.Tasks;
+
+namespace HordeServer.Utilities
+{
+	/// <summary>
+	/// Calls a method on a specified interval on a background task
+	/// </summary>
+	public sealed class BackgroundTick : IAsyncDisposable, IDisposable
+	{
+		Func<CancellationToken, Task> TickFunc;
+		TimeSpan Interval;
+		Task BackgroundTask;
+		CancellationTokenSource CancellationTokenSource;
+		ILogger Logger;
+
+		/// <summary>
+		/// Constructor
+		/// </summary>
+		/// <param name="TickFunc">The method to call when the interval elapses</param>
+		/// <param name="Interval">Inveral between calls</param>
+		/// <param name="Logger">Logger for errors produced by the tick function</param>
+		public BackgroundTick(Func<CancellationToken, Task> TickFunc, TimeSpan Interval, ILogger Logger)
+		{
+			this.TickFunc = TickFunc;
+			this.Interval = Interval;
+			this.BackgroundTask = Task.Run(BackgroundTickAsync);
+			this.CancellationTokenSource = new CancellationTokenSource();
+			this.Logger = Logger;
+		}
+
+		/// <summary>
+		/// Signals the tick function to stop executing, and waits for it to complete
+		/// </summary>
+		public async Task StopAsync()
+		{
+			CancellationTokenSource.Cancel();
+			await BackgroundTask;
+		}
+
+		/// <summary>
+		/// Background task executor
+		/// </summary>
+		async Task BackgroundTickAsync()
+		{
+			CancellationToken CancellationToken = CancellationTokenSource.Token;
+			for (; ; )
+			{
+				try
+				{
+					Stopwatch Timer = Stopwatch.StartNew();
+					await TickFunc(CancellationToken);
+
+					TimeSpan Delay = Interval - Timer.Elapsed;
+					if (Delay > TimeSpan.Zero)
+					{
+						await Task.Delay(Delay, CancellationToken);
+					}
+				}
+				catch (OperationCanceledException) when (CancellationToken.IsCancellationRequested)
+				{
+					break;
+				}
+				catch (Exception Ex)
+				{
+					Logger.LogError(Ex, "Exception while executing background task");
+				}
+			}
+		}
+
+		/// <inheritdoc/>
+		public async ValueTask DisposeAsync()
+		{
+			await StopAsync();
+			CancellationTokenSource.Dispose();
+		}
+
+		/// <inheritdoc/>
+		public void Dispose()
+		{
+			StopAsync().Wait();
+			CancellationTokenSource.Dispose();
+		}
+	}
+}
diff --git a/Engine/Source/Programs/Horde/HordeServer/Utilities/RedisExtensions.cs b/Engine/Source/Programs/Horde/HordeServer/Utilities/RedisExtensions.cs
deleted file mode 100644
diff --git a/Engine/Source/Programs/Horde/HordeServer/Utilities/RedisMessageQueue.cs b/Engine/Source/Programs/Horde/HordeServer/Utilities/RedisMessageQueue.cs
deleted file mode 100644
diff --git a/Engine/Source/Programs/Horde/HordeServer/Utilities/RedisSerializer.cs b/Engine/Source/Programs/Horde/HordeServer/Utilities/RedisSerializer.cs
deleted file mode 100644
diff --git a/Engine/Source/Programs/Horde/HordeServer/Utilities/StringId.cs b/Engine/Source/Programs/Horde/HordeServer/Utilities/StringId.cs
--- a/Engine/Source/Programs/Horde/HordeServer/Utilities/StringId.cs
+++ b/Engine/Source/Programs/Horde/HordeServer/Utilities/StringId.cs
@@ -1,6 +1,8 @@
 // Copyright Epic Games, Inc. All Rights Reserved.
 
 using EpicGames.Core;
+using EpicGames.Serialization;
+using EpicGames.Serialization.Converters;
 using HordeServer.Models;
 using MongoDB.Bson.Serialization;
 using MongoDB.Bson.Serialization.Attributes;
@@ -23,6 +25,7 @@ namespace HordeServer.Utilities
 	/// </summary>
 	[JsonSchemaString]
 	[TypeConverter(typeof(StringIdTypeConverter))]
+	[CbConverter(typeof(CbStringIdConverter<>))]
 	public struct StringId<T> : IEquatable<StringId<T>>
 	{
 		/// <summary>
@@ -273,6 +276,24 @@ namespace HordeServer.Utilities
 		}
 	}
 
+	sealed class CbStringIdConverter<T> : CbConverter<StringId<T>>
+	{
+		public override StringId<T> Read(CbField Field)
+		{
+			return new StringId<T>(Field.AsString().ToString());
+		}
+
+		public override void Write(CbWriter Writer, StringId<T> Value)
+		{
+			Writer.WriteStringValue(Value.ToString());
+		}
+
+		public override void WriteNamed(CbWriter Writer, Utf8String Name, StringId<T> Value)
+		{
+			Writer.WriteString(Name, Value.ToString());
+		}
+	}
+
 	/// <summary>
 	/// Type converter from strings to PropertyFilter objects
 	/// </summary>
@@ -304,6 +325,10 @@ namespace HordeServer.Utilities
 		/// <inheritdoc/>
 		public override bool CanConvertTo(ITypeDescriptorContext Context, Type DestinationType)
 		{
+			if (DestinationType == typeof(string))
+			{
+				return true;
+			}
 			if (DestinationType.IsGenericType)
 			{
 				Type GenericTypeDefinition = DestinationType.GetGenericTypeDefinition();
@@ -322,7 +347,14 @@ namespace HordeServer.Utilities
 		/// <inheritdoc/>
 		public override object? ConvertTo(ITypeDescriptorContext Context, CultureInfo Culture, object Value, Type DestinationType)
 		{
-			return Activator.CreateInstance(DestinationType, Value);
+			if (DestinationType == typeof(string))
+			{
+				return Value.ToString();
+			}
+			else
+			{
+				return Activator.CreateInstance(DestinationType, Value);
+			}
 		}
 	}
 }
diff --git a/Engine/Source/Programs/Shared/EpicGames.Horde.Client/EpicGames.Horde.Client.csproj b/Engine/Source/Programs/Shared/EpicGames.Horde.Client/EpicGames.Horde.Client.csproj
--- a/Engine/Source/Programs/Shared/EpicGames.Horde.Client/EpicGames.Horde.Client.csproj
+++ b/Engine/Source/Programs/Shared/EpicGames.Horde.Client/EpicGames.Horde.Client.csproj
@@ -43,6 +43,9 @@
     <ProjectReference Include="..\EpicGames.Serialization\EpicGames.Serialization.csproj" />
   </ItemGroup>
   <ItemGroup>
+    <Content Include="..\EpicGames.Horde\Protos\horde\compute\compute_rpc.proto">
+      <Link>Protos\horde\compute\compute_rpc.proto</Link>
+    </Content>
     <Content Include="..\EpicGames.Horde\Protos\horde\storage\blob_store.proto">
       <Link>Protos\horde\storage\blob_store.proto</Link>
     </Content>
@@ -51,6 +54,7 @@
     </Content>
   </ItemGroup>
   <ItemGroup>
+    <Protobuf Include="..\EpicGames.Horde\Protos\horde\compute\compute_rpc.proto" ProtoRoot="..\EpicGames.Horde\Protos" GrpcServices="Client" />
     <Protobuf Include="..\EpicGames.Horde\Protos\horde\storage\blob_store.proto" ProtoRoot="..\EpicGames.Horde\Protos" GrpcServices="Client" />
     <Protobuf Include="..\EpicGames.Horde\Protos\horde\storage\ref_table.proto" ProtoRoot="..\EpicGames.Horde\Protos" GrpcServices="Client" />
   </ItemGroup>
diff --git a/Engine/Source/Programs/Shared/EpicGames.Horde.Client/Storage/BlobStoreExtensions.cs b/Engine/Source/Programs/Shared/EpicGames.Horde.Client/Storage/BlobStoreExtensions.cs
new file mode 100644
--- /dev/null
+++ b/Engine/Source/Programs/Shared/EpicGames.Horde.Client/Storage/BlobStoreExtensions.cs
@@ -0,0 +1,126 @@
+// Copyright Epic Games, Inc. All Rights Reserved.
+
+using EpicGames.Core;
+using EpicGames.Serialization;
+using Google.Protobuf;
+using Grpc.Core;
+using System;
+using System.Collections.Generic;
+using System.Text;
+using System.Threading.Tasks;
+
+namespace EpicGames.Horde.Storage
+{
+	/// <summary>
+	/// Extension methods for <see cref="BlobStore.BlobStoreClient"/>
+	/// </summary>
+	public static class BlobStoreExtensions
+	{
+		/// <summary>
+		/// Gets a blob of data from the store
+		/// </summary>
+		/// <param name="BlobStore">The blob store instance</param>
+		/// <param name="NamespaceId">Namespace for the data</param>
+		/// <param name="Hash">Hash of the blob</param>
+		/// <returns>Byte array of blob data</returns>
+		public static async Task<byte[]> GetBlobAsync(this BlobStore.BlobStoreClient BlobStore, string NamespaceId, IoHash Hash)
+		{
+			GetBlobsRequest Request = new GetBlobsRequest();
+			Request.NamespaceId = NamespaceId;
+			Request.Blobs.Add(Hash);
+
+			GetBlobsResponse Response = await BlobStore.GetAsync(Request);
+			foreach (GetBlobResponse BlobResponse in Response.Blobs)
+			{
+				if (BlobResponse.Exists)
+				{
+					return BlobResponse.Data.ToByteArray();
+				}
+			}
+
+			throw new RpcException(new Status(StatusCode.NotFound, $"Unable to fetch blob {Hash}"));
+		}
+
+		/// <summary>
+		/// Copies a blob of data from the store to a file on disk
+		/// </summary>
+		/// <param name="BlobStore">The blob store instance</param>
+		/// <param name="NamespaceId">Namespace for the data</param>
+		/// <param name="Hash">Hash of the blob</param>
+		/// <param name="File">File to write. Any parent directories will be created if they do not already exist.</param>
+		/// <returns>Byte array of blob data</returns>
+		public static async Task GetFileAsync(this BlobStore.BlobStoreClient BlobStore, string NamespaceId, IoHash Hash, FileReference File)
+		{
+			byte[] Data = await GetBlobAsync(BlobStore, NamespaceId, Hash);
+			DirectoryReference.CreateDirectory(File.Directory);
+			await FileReference.WriteAllBytesAsync(File, Data);
+		}
+
+		/// <summary>
+		/// Reads a compact binary object from the store
+		/// </summary>
+		/// <typeparam name="T">The type of object to return</typeparam>
+		/// <param name="BlobStore">The blob store instance</param>
+		/// <param name="NamespaceId">Namespace for the data</param>
+		/// <param name="Hash">Hash of the blob</param>
+		/// <returns>Byte array of blob data</returns>
+		public static async Task<T> GetObjectAsync<T>(this BlobStore.BlobStoreClient BlobStore, string NamespaceId, IoHash Hash) where T : class
+		{
+			byte[] Data = await GetBlobAsync(BlobStore, NamespaceId, Hash);
+			return CbSerializer.Deserialize<T>(new CbField(Data));
+		}
+
+		/// <summary>
+		/// Writes a blob of data to the store
+		/// </summary>
+		/// <typeparam name="T">The type of object to return</typeparam>
+		/// <param name="BlobStore">The blob store instance</param>
+		/// <param name="NamespaceId">Namespace for the data</param>
+		/// <param name="Data">The data to write</param>
+		/// <returns>Byte array of blob data</returns>
+		public static async Task<IoHash> PutBlobAsync(this BlobStore.BlobStoreClient BlobStore, string NamespaceId, ReadOnlyMemory<byte> Data)
+		{
+			IoHash Hash = IoHash.Compute(Data.Span);
+
+			AddBlobsRequest Request = new AddBlobsRequest();
+			Request.NamespaceId = NamespaceId;
+
+			AddBlobRequest BlobRequest = new AddBlobRequest();
+			BlobRequest.Hash = Hash;
+			BlobRequest.Data = ByteString.CopyFrom(Data.Span);
+
+			Request.Blobs.Add(BlobRequest);
+			await BlobStore.AddAsync(Request);
+
+			return Hash;
+		}
+
+		/// <summary>
+		/// Writes a file to the store
+		/// </summary>
+		/// <typeparam name="T">The type of object to return</typeparam>
+		/// <param name="BlobStore">The blob store instance</param>
+		/// <param name="NamespaceId">Namespace for the data</param>
+		/// <param name="File">Location of the file to write</param>
+		/// <returns>Byte array of blob data</returns>
+		public static async Task<IoHash> PutFileAsync(this BlobStore.BlobStoreClient BlobStore, string NamespaceId, FileReference File)
+		{
+			byte[] Data = await FileReference.ReadAllBytesAsync(File);
+			return await PutBlobAsync(BlobStore, NamespaceId, Data);
+		}
+
+		/// <summary>
+		/// Writes a compact binary object to the store
+		/// </summary>
+		/// <typeparam name="T">The type of object to return</typeparam>
+		/// <param name="BlobStore">The blob store instance</param>
+		/// <param name="NamespaceId">Namespace for the data</param>
+		/// <param name="Item">The item to write</param>
+		/// <returns>Byte array of blob data</returns>
+		public static Task<IoHash> PutObjectAsync<T>(this BlobStore.BlobStoreClient BlobStore, string NamespaceId, T Item)
+		{
+			ReadOnlyMemory<byte> Data = CbSerializer.Serialize(Item).GetView();
+			return PutBlobAsync(BlobStore, NamespaceId, Data);
+		}
+	}
+}
diff --git a/Engine/Source/Programs/Shared/EpicGames.Horde/Common/IoHashWrapper.cs b/Engine/Source/Programs/Shared/EpicGames.Horde/Common/IoHashWrapper.cs
--- a/Engine/Source/Programs/Shared/EpicGames.Horde/Common/IoHashWrapper.cs
+++ b/Engine/Source/Programs/Shared/EpicGames.Horde/Common/IoHashWrapper.cs
@@ -4,6 +4,7 @@ using EpicGames.Core;
 using Google.Protobuf;
 using System;
 using System.Collections.Generic;
+using System.Diagnostics.CodeAnalysis;
 using System.Text;
 using System.Threading.Tasks;
 
@@ -29,7 +30,7 @@ namespace EpicGames.Horde.Common
 		/// <param name="Hash"></param>
 		public static implicit operator IoHash(IoHashWrapper Hash)
 		{
-			return new IoHash(Hash.ToByteArray());
+			return new IoHash(Hash.Hash.ToByteArray());
 		}
 
 		/// <summary>
diff --git a/Engine/Source/Programs/Shared/EpicGames.Horde/Compute/ComputeTask.cs b/Engine/Source/Programs/Shared/EpicGames.Horde/Compute/ComputeTask.cs
--- a/Engine/Source/Programs/Shared/EpicGames.Horde/Compute/ComputeTask.cs
+++ b/Engine/Source/Programs/Shared/EpicGames.Horde/Compute/ComputeTask.cs
@@ -55,6 +55,13 @@ namespace EpicGames.Horde.Compute
 		[CbField("o")]
 		public List<Utf8String> OutputPaths { get; set; } = new List<Utf8String>();
 
+		/// <summary>
+		/// Default constructor for serialization
+		/// </summary>
+		private ComputeTask()
+		{
+		}
+
 		/// <summary>
 		/// Constructor
 		/// </summary>
diff --git a/Engine/Source/Programs/Shared/EpicGames.Horde/Compute/ComputeTaskResult.cs b/Engine/Source/Programs/Shared/EpicGames.Horde/Compute/ComputeTaskResult.cs
--- a/Engine/Source/Programs/Shared/EpicGames.Horde/Compute/ComputeTaskResult.cs
+++ b/Engine/Source/Programs/Shared/EpicGames.Horde/Compute/ComputeTaskResult.cs
@@ -37,6 +37,13 @@ namespace EpicGames.Horde.Compute
 		[CbField("o")]
 		public IoHash? OutputHash { get; set; }
 
+		/// <summary>
+		/// Constructor
+		/// </summary>
+		private ComputeTaskResult()
+		{
+		}
+
 		/// <summary>
 		/// Constructor
 		/// </summary>
diff --git a/Engine/Source/Programs/Shared/EpicGames.Horde/Compute/Types.cs b/Engine/Source/Programs/Shared/EpicGames.Horde/Compute/Types.cs
new file mode 100644
--- /dev/null
+++ b/Engine/Source/Programs/Shared/EpicGames.Horde/Compute/Types.cs
@@ -0,0 +1,29 @@
+using EpicGames.Core;
+using EpicGames.Horde.Common;
+using System;
+using System.Collections.Generic;
+using System.Linq;
+using System.Text;
+
+namespace EpicGames.Horde.Compute
+{
+	partial class AddTasksRpcRequest
+	{
+		public AddTasksRpcRequest(string ChannelId, string NamespaceId, IoHash RequirementsHash, List<IoHash> TaskHashes, bool DoNotCache)
+		{
+			this.ChannelId = ChannelId;
+			this.NamespaceId = NamespaceId;
+			this.RequirementsHash = RequirementsHash;
+			this.TaskHashes.Add(TaskHashes.Select(x => (IoHashWrapper)x));
+			this.DoNotCache = DoNotCache;
+		}
+	}
+
+	partial class GetTaskUpdatesRpcRequest
+	{
+		public GetTaskUpdatesRpcRequest(string ChannelId)
+		{
+			this.ChannelId = ChannelId;
+		}
+	}
+}
diff --git a/Engine/Source/Programs/Shared/EpicGames.Horde/EpicGames.Horde.csproj b/Engine/Source/Programs/Shared/EpicGames.Horde/EpicGames.Horde.csproj
--- a/Engine/Source/Programs/Shared/EpicGames.Horde/EpicGames.Horde.csproj
+++ b/Engine/Source/Programs/Shared/EpicGames.Horde/EpicGames.Horde.csproj
@@ -44,6 +44,7 @@
   <ItemGroup>
     <Protobuf Include="Protos\horde\common\cbobject_wrapper.proto" ProtoRoot="Protos" />
     <Protobuf Include="Protos\horde\common\iohash_wrapper.proto" ProtoRoot="Protos" />
+    <Protobuf Include="Protos\horde\compute\compute_rpc_messages.proto" ProtoRoot="Protos" />
     <Protobuf Include="Protos\horde\storage\blob_store_messages.proto" ProtoRoot="Protos" />
     <Protobuf Include="Protos\horde\storage\ref_table_messages.proto" ProtoRoot="Protos" />
   </ItemGroup>
