Horde: Fix missing copyright
#rnx

diff --git a/Engine/Source/Programs/Horde/HordeServer/Compute/IComputeService.cs b/Engine/Source/Programs/Horde/HordeServer/Compute/IComputeService.cs
--- a/Engine/Source/Programs/Horde/HordeServer/Compute/IComputeService.cs
+++ b/Engine/Source/Programs/Horde/HordeServer/Compute/IComputeService.cs
@@ -1,3 +1,5 @@
+// Copyright Epic Games, Inc. All Rights Reserved.
+
 using EpicGames.Core;
 using EpicGames.Horde.Compute;
 using HordeServer.Models;
diff --git a/Engine/Source/Programs/Shared/EpicGames.Horde/Compute/Types.cs b/Engine/Source/Programs/Shared/EpicGames.Horde/Compute/Types.cs
--- a/Engine/Source/Programs/Shared/EpicGames.Horde/Compute/Types.cs
+++ b/Engine/Source/Programs/Shared/EpicGames.Horde/Compute/Types.cs
@@ -1,3 +1,5 @@
+// Copyright Epic Games, Inc. All Rights Reserved.
+
 using EpicGames.Core;
 using EpicGames.Horde.Common;
 using System;
