Horde: Better logging for agent connection lost events in REAPI

diff --git a/Engine/Source/Programs/Horde/HordeServer/Tasks/Impl/ActionTaskSource.cs b/Engine/Source/Programs/Horde/HordeServer/Tasks/Impl/ActionTaskSource.cs
--- a/Engine/Source/Programs/Horde/HordeServer/Tasks/Impl/ActionTaskSource.cs
+++ b/Engine/Source/Programs/Horde/HordeServer/Tasks/Impl/ActionTaskSource.cs
@@ -491,8 +491,8 @@ namespace HordeServer.Tasks.Impl
 
 				void OnConnectionLost()
 				{
-					string Message = $"Connection lost to agent. LeaseId={Lease.Id} OperationId={Operation!.Id}";
-					Logger.LogInformation(Message);
+					string Message = $"Connection lost to agent. LeaseId={Lease.Id} OperationId={Operation!.Id}.";
+					Logger.LogInformation("Connection lost to agent. LeaseId={LeaseId} OperationId={OperationId} Stacktrade={Stacktrace}.", Lease.Id, Operation!.Id, Environment.StackTrace);
 					Operation.TrySetResult(new ActionExecuteResult(new Google.Rpc.Status(StatusCode.Internal, Message)));
 				}
 				
