Horde: Remove accidental console writes

diff --git a/Engine/Source/Programs/Horde/HordeServer/Compute/Impl/ComputeService.cs b/Engine/Source/Programs/Horde/HordeServer/Compute/Impl/ComputeService.cs
--- a/Engine/Source/Programs/Horde/HordeServer/Compute/Impl/ComputeService.cs
+++ b/Engine/Source/Programs/Horde/HordeServer/Compute/Impl/ComputeService.cs
@@ -163,11 +163,9 @@ namespace HordeServer.Compute.Impl
 		/// <inheritdoc/>
 		public async Task<AgentLease?> TryAssignLeaseAsync(IAgent Agent, CancellationToken CancellationToken)
 		{
-			Console.WriteLine("TRY ASSIGN LEASE");
 			TaskSchedulerEntry<ComputeTaskInfo>? Entry = await TaskScheduler.DequeueAsync(Agent, CancellationToken);
 			if(Entry != null)
 			{
-				Console.WriteLine("GOT IT");
 				ComputeTaskMessage ComputeTask = new ComputeTaskMessage();
 				ComputeTask.ChannelId = Entry.Item.ChannelId.ToString();
 				ComputeTask.NamespaceId = NamespaceId.ToString();
