Horde: Lower the default value for concurrent REAPI actions on a single agent

diff --git a/Engine/Source/Programs/Horde/HordeServer/ServerSettings.cs b/Engine/Source/Programs/Horde/HordeServer/ServerSettings.cs
--- a/Engine/Source/Programs/Horde/HordeServer/ServerSettings.cs
+++ b/Engine/Source/Programs/Horde/HordeServer/ServerSettings.cs
@@ -115,7 +115,7 @@ namespace HordeServer
 		/// <summary>
 		/// Max number of concurrent leases per agent
 		/// </summary>
-		public int MaxConcurrentLeasesPerAgent { get; set; } = 2;
+		public int MaxConcurrentLeasesPerAgent { get; set; } = 1;
 	}
 
 	/// <summary>
