Horde: Minor fixes for compute API.

* Fixed stall on agent when a new queue is created, due to message not being posted to the new queue channel.
* Fix exception deserializing task status entries with null leaseid members.
* Other cleanup prep for REST API.

diff --git a/Engine/Source/Programs/Horde/HordeAgent/Execution/ComputeTaskExecutor.cs b/Engine/Source/Programs/Horde/HordeAgent/Execution/ComputeTaskExecutor.cs
--- a/Engine/Source/Programs/Horde/HordeAgent/Execution/ComputeTaskExecutor.cs
+++ b/Engine/Source/Programs/Horde/HordeAgent/Execution/ComputeTaskExecutor.cs
@@ -108,7 +108,7 @@ namespace HordeAgent
 			{
 				throw new RpcException(new Grpc.Core.Status(StatusCode.NotFound, "Unable to find action message"));
 			}
-			Logger.LogInformation("Executing task {Hash} for lease ID {LeaseId}", ComputeTaskMessage.TaskHash, LeaseId);
+			Logger.LogInformation("Executing task {Hash} for lease ID {LeaseId}", ComputeTaskMessage.TaskHash.Hash, LeaseId);
 
 			DirectoryReference.CreateDirectory(SandboxDir);
 			FileUtils.ForceDeleteDirectoryContents(SandboxDir);
diff --git a/Engine/Source/Programs/Horde/HordeServer/Compute/IComputeService.cs b/Engine/Source/Programs/Horde/HordeServer/Compute/IComputeService.cs
--- a/Engine/Source/Programs/Horde/HordeServer/Compute/IComputeService.cs
+++ b/Engine/Source/Programs/Horde/HordeServer/Compute/IComputeService.cs
@@ -10,6 +10,7 @@ using Microsoft.Extensions.Hosting;
 using MongoDB.Bson;
 using System;
 using System.Collections.Generic;
+using System.Diagnostics.CodeAnalysis;
 using System.Linq;
 using System.Threading;
 using System.Threading.Tasks;
@@ -22,14 +23,15 @@ namespace HordeServer.Compute
 	/// <summary>
 	/// Unused, except for ChannelId alias.
 	/// </summary>
-	interface IComputeChannel
+	[SuppressMessage("Design", "CA1040:Avoid empty interfaces")]
+	public interface IComputeChannel
 	{
 	}
 
 	/// <summary>
 	/// Status of a compute task
 	/// </summary>
-	interface IComputeTaskStatus
+	public interface IComputeTaskStatus
 	{
 		/// <summary>
 		/// The input hash
@@ -65,7 +67,7 @@ namespace HordeServer.Compute
 	/// <summary>
 	/// Interface for the compute service
 	/// </summary>
-	interface IComputeService : INewTaskSource
+	public interface IComputeService : INewTaskSource
 	{
 		/// <summary>
 		/// Post tasks to be executed to a channel
@@ -77,13 +79,20 @@ namespace HordeServer.Compute
 		/// <returns>Async task</returns>
 		Task AddTasksAsync(NamespaceId NamespaceId, IoHash RequirementsHash, List<IoHash> TaskHashes, ChannelId ChannelId);
 
+		/// <summary>
+		/// Dequeue completed items from a queue and return immediately
+		/// </summary>
+		/// <param name="ChannelId">Queue to remove items from</param>
+		/// <returns>List of status updates</returns>
+		Task<List<IComputeTaskStatus>> GetTaskUpdatesAsync(ChannelId ChannelId);
+
 		/// <summary>
 		/// Dequeue completed items from a queue
 		/// </summary>
 		/// <param name="ChannelId">Queue to remove items from</param>
 		/// <param name="CancellationToken">Cancellation token to stop waiting for items</param>
 		/// <returns>List of status updates</returns>
-		Task<List<IComputeTaskStatus>> GetTaskUpdatesAsync(ChannelId ChannelId, CancellationToken CancellationToken);
+		Task<List<IComputeTaskStatus>> WaitForTaskUpdatesAsync(ChannelId ChannelId, CancellationToken CancellationToken);
 	}
 
 	/// <summary>
diff --git a/Engine/Source/Programs/Horde/HordeServer/Compute/Impl/ComputeService.cs b/Engine/Source/Programs/Horde/HordeServer/Compute/Impl/ComputeService.cs
--- a/Engine/Source/Programs/Horde/HordeServer/Compute/Impl/ComputeService.cs
+++ b/Engine/Source/Programs/Horde/HordeServer/Compute/Impl/ComputeService.cs
@@ -100,7 +100,7 @@ namespace HordeServer.Compute.Impl
 		/// <inheritdoc/>
 		public ObjectId? LeaseId
 		{
-			get => new ObjectId(LeaseIdBytes);
+			get => (LeaseIdBytes == null)? (ObjectId?)null : new ObjectId(LeaseIdBytes);
 			set => LeaseIdBytes = value?.ToByteArray();
 		}
 
@@ -123,7 +123,7 @@ namespace HordeServer.Compute.Impl
 	{
 		public MessageDescriptor Descriptor => ComputeTaskMessage.Descriptor;
 
-		public NamespaceId NamespaceId { get; } = new NamespaceId("default");
+		public static NamespaceId DefaultNamespaceId { get; } = new NamespaceId("default");
 
 		ITaskScheduler<ComputeTaskInfo> TaskScheduler;
 		RedisMessageQueue<ComputeTaskStatus> MessageQueue;
@@ -137,7 +137,7 @@ namespace HordeServer.Compute.Impl
 		/// <param name="Logger">The logger instance</param>
 		public ComputeService(IDatabase Redis, IObjectCollection ObjectCollection, ILogger<ComputeService> Logger)
 		{
-			this.TaskScheduler = new RedisTaskScheduler<ComputeTaskInfo>(Redis, ObjectCollection, NamespaceId, "compute/tasks/");
+			this.TaskScheduler = new RedisTaskScheduler<ComputeTaskInfo>(Redis, ObjectCollection, DefaultNamespaceId, "compute/tasks/");
 			this.MessageQueue = new RedisMessageQueue<ComputeTaskStatus>(Redis, "compute/messages/");
 			this.Logger = Logger;
 		}
@@ -168,7 +168,7 @@ namespace HordeServer.Compute.Impl
 			{
 				ComputeTaskMessage ComputeTask = new ComputeTaskMessage();
 				ComputeTask.ChannelId = Entry.Item.ChannelId.ToString();
-				ComputeTask.NamespaceId = NamespaceId.ToString();
+				ComputeTask.NamespaceId = DefaultNamespaceId.ToString();
 				ComputeTask.RequirementsHash = Entry.RequirementsHash;
 				ComputeTask.TaskHash = Entry.Item.TaskHash;
 
@@ -188,9 +188,15 @@ namespace HordeServer.Compute.Impl
 		}
 
 		/// <inheritdoc/>
-		public async Task<List<IComputeTaskStatus>> GetTaskUpdatesAsync(ChannelId ChannelId, CancellationToken CancellationToken)
+		public async Task<List<IComputeTaskStatus>> GetTaskUpdatesAsync(ChannelId ChannelId)
 		{
-			List<ComputeTaskStatus> Messages = await MessageQueue.ReadAsync(ChannelId.ToString(), CancellationToken);
+			List<ComputeTaskStatus> Messages = await MessageQueue.ReadMessagesAsync(ChannelId.ToString());
+			return Messages.ConvertAll<IComputeTaskStatus>(x => x);
+		}
+
+		public async Task<List<IComputeTaskStatus>> WaitForTaskUpdatesAsync(ChannelId ChannelId, CancellationToken CancellationToken)
+		{
+			List<ComputeTaskStatus> Messages = await MessageQueue.WaitForMessagesAsync(ChannelId.ToString(), CancellationToken);
 			return Messages.ConvertAll<IComputeTaskStatus>(x => x);
 		}
 
@@ -253,7 +259,7 @@ namespace HordeServer.Compute.Impl
 					MoveNextTask = MoveNextAndCancel(RequestStream, CancellationSource);
 					while (!CancellationSource.IsCancellationRequested)
 					{
-						List<IComputeTaskStatus> Updates = await ComputeService.GetTaskUpdatesAsync(ChannelId, CancellationSource.Token);
+						List<IComputeTaskStatus> Updates = await ComputeService.WaitForTaskUpdatesAsync(ChannelId, CancellationSource.Token);
 						foreach (IComputeTaskStatus Update in Updates)
 						{
 							GetTaskUpdatesRpcResponse Response = new GetTaskUpdatesRpcResponse();
diff --git a/Engine/Source/Programs/Horde/HordeServer/Compute/Impl/RedisMessageQueue.cs b/Engine/Source/Programs/Horde/HordeServer/Compute/Impl/RedisMessageQueue.cs
--- a/Engine/Source/Programs/Horde/HordeServer/Compute/Impl/RedisMessageQueue.cs
+++ b/Engine/Source/Programs/Horde/HordeServer/Compute/Impl/RedisMessageQueue.cs
@@ -26,13 +26,20 @@ namespace HordeServer.Compute.Impl
 		/// <param name="Message">Message to post</param>
 		Task PostAsync(string ChannelId, T Message);
 
+		/// <summary>
+		/// Waits for a message to be available on the given channel
+		/// </summary>
+		/// <param name="ChannelId">The channel to read from</param>
+		/// <returns>True if a message was available, false otherwise</returns>
+		Task<List<T>> ReadMessagesAsync(string ChannelId);
+
 		/// <summary>
 		/// Waits for a message to be available on the given channel
 		/// </summary>
 		/// <param name="ChannelId">The channel to read from</param>
 		/// <param name="CancellationToken">May be signalled to stop the wait, without throwing an exception</param>
 		/// <returns>True if a message was available, false otherwise</returns>
-		Task<List<T>> ReadAsync(string ChannelId, CancellationToken CancellationToken);
+		Task<List<T>> WaitForMessagesAsync(string ChannelId, CancellationToken CancellationToken);
 	}
 
 	/// <summary>
@@ -122,7 +129,15 @@ namespace HordeServer.Compute.Impl
 		}
 
 		/// <inheritdoc/>
-		public async Task<List<T>> ReadAsync(string ChannelId, CancellationToken CancellationToken)
+		public async Task<List<T>> ReadMessagesAsync(string ChannelId)
+		{
+			List<T> Messages = new List<T>();
+			await ReadMessagesAsync(GetChannel(ChannelId), Messages);
+			return Messages;
+		}
+
+		/// <inheritdoc/>
+		public async Task<List<T>> WaitForMessagesAsync(string ChannelId, CancellationToken CancellationToken)
 		{
 			List<T> Messages = new List<T>();
 
diff --git a/Engine/Source/Programs/Horde/HordeServer/Compute/Impl/RedisTaskScheduler.cs b/Engine/Source/Programs/Horde/HordeServer/Compute/Impl/RedisTaskScheduler.cs
--- a/Engine/Source/Programs/Horde/HordeServer/Compute/Impl/RedisTaskScheduler.cs
+++ b/Engine/Source/Programs/Horde/HordeServer/Compute/Impl/RedisTaskScheduler.cs
@@ -181,7 +181,9 @@ namespace HordeServer.Compute.Impl
 			ITransaction Transaction = Redis.CreateTransaction();
 			Transaction.AddCondition(Condition.KeyExists(GetQueueKey(RequirementsHash)));
 			_ = Transaction.SetAddAsync(QueueIndex, RequirementsHash);
-			await Transaction.ExecuteAsync(CommandFlags.FireAndForget);
+			await Transaction.ExecuteAsync();
+
+			await Redis.PublishAsync(NewQueueChannel, RequirementsHash);
 		}
 
 		/// <summary>
@@ -204,7 +206,9 @@ namespace HordeServer.Compute.Impl
 		public async Task EnqueueAsync(T Item, IoHash RequirementsHash)
 		{
 			RedisList<T> List = GetQueue(RequirementsHash);
-			if (await Redis.ListRightPushAsync(List, Item) == 1)
+
+			long NewLength = await Redis.ListRightPushAsync(List, Item);
+			if (NewLength == 1)
 			{
 				await AddQueueToIndexAsync(RequirementsHash);
 			}
@@ -218,7 +222,9 @@ namespace HordeServer.Compute.Impl
 		public async Task RequeueAsync(T Item, IoHash RequirementsHash)
 		{
 			RedisList<T> Queue = GetQueue(RequirementsHash);
-			if (await Redis.ListLeftPushAsync(Queue, Item) == 1)
+
+			long NewLength = await Redis.ListLeftPushAsync(Queue, Item);
+			if (NewLength == 1)
 			{
 				await AddQueueToIndexAsync(RequirementsHash);
 			}
@@ -490,172 +496,5 @@ namespace HordeServer.Compute.Impl
 			}
 			return Requirements;
 		}
-
-#if false
-		class QueueStats
-		{
-			QueueKey Key;
-			double Score;
-			int Length;
-		}
-
-		interface IQueueIndex
-		{
-			public IReadOnlyList<QueueKey> Keys { get; }
-
-			public async Task AddAsync(QueueKey Key);
-
-			public async Task RemoveAsync(QueueKey Key);
-		}
-#if false
-
-
-		DatabaseService DatabaseService;
-		IObjectCollection ObjectCollection;
-		IDatabase Redis;
-		RedisKey BaseKey;
-		RedisSortedSet<QueueKey> Index;
-		RedisChannel<QueueKey> NewQueueChannel;
-		List<Listener> Listeners = new List<Listener>();
-		IMemoryCache RequirementsCache = new MemoryCache(new MemoryCacheOptions());
-		List<QueueKey> CachedQueues = new List<QueueKey>();
-		Task? UpdateQueuesTask;
-		BackgroundTick? BackgroundTick;
-		ILogger Logger;
-
-		public RedisTaskScheduler(DatabaseService DatabaseService, IObjectCollection ObjectCollection, IDatabase Redis, RedisKey BaseKey, ILogger<RedisTaskScheduler<T>> Logger)
-		{
-			this.DatabaseService = DatabaseService;
-			this.ObjectCollection = ObjectCollection;
-			this.Redis = Redis;
-			this.BaseKey = BaseKey;
-			this.Index = new RedisSortedSet<QueueKey>(BaseKey.Append("index"));
-			this.NewQueueChannel = new RedisChannel<QueueKey>(BaseKey.Append("new_queues").ToString());
-			this.Logger = Logger;
-
-			UpdateQueuesTask = Task.Run(() => CheckNewQueues());
-		}
-
-		public async Task StopAsync()
-		{
-			if (UpdateQueuesTask != null)
-			{
-				NewQueues.Writer.TryComplete();
-				await UpdateQueuesTask;
-			}
-		}
-
-		public void Dispose()
-		{
-			StopAsync().Wait();
-			BackgroundTick?.Dispose();
-			RequirementsCache.Dispose();
-		}
-
-		RedisList<T> GetQueue(QueueKey QueueKey)
-		{
-			return new RedisList<T>(QueueKey.GetKey(BaseKey));
-		}
-
-
-		/// <inheritdoc/>
-		public async Task EnqueueAsync(TaskSchedulerEntry<T> Entry, bool AtFront = false)
-		{
-			if (!await TryAssignToLocalListenerAsync(Entry))
-			{
-				await AddToSharedQueue(Entry, AtFront);
-			}
-		}
-
-		async Task<bool> TryAssignToLocalListenerAsync(TaskSchedulerEntry<T> Entry)
-		{
-			// Get the requirements for this entry
-			Requirements? Requirements = await GetRequirementsAsync(Entry.NamespaceId, Entry.RequirementsHash);
-			if (Requirements == null)
-			{
-				throw new Exception($"Unable to find requirements with hash {Entry.RequirementsHash}");
-			}
-
-			// Find a local listener that can execute the work
-			for (; ; )
-			{
-				Listener? Listener;
-				lock (Listeners)
-				{
-					Listener = Listeners.FirstOrDefault(x => !x.CompletionSource.Task.IsCompleted && CheckRequirements(x.Agent, Requirements));
-				}
-
-				if (Listener == null)
-				{
-					return false;
-				}
-				else if (Listener.CompletionSource.TrySetResult(Entry))
-				{
-					return true;
-				}
-			}
-		}
-
-		private async Task AddToSharedQueue(TaskSchedulerEntry<T> Entry, bool AtFront)
-		{
-			// TODO: check if there's already something listening
-
-			QueueKey QueueKey = new QueueKey(Entry);
-			RedisList<T> Queue = new RedisList<T>(QueueKey.GetKey(BaseKey));
-
-			// Insert the item in to the appropriate queue
-			long NewLength;
-			if (AtFront)
-			{
-				NewLength = await Redis.ListLeftPushAsync(Queue, Entry.TaskData);
-			}
-			else
-			{
-				NewLength = await Redis.ListRightPushAsync(Queue, Entry.TaskData);
-			}
-
-			// Check if this was a new queue; if so, we need to add a new entry to the index for it, and notify other pods
-			if (NewLength == 1)
-			{
-				await Redis.SortedSetAddAsync(Index, QueueKey, DateTime.UtcNow.Ticks);
-				await Redis.PublishAsync(NewQueueChannel, QueueKey);
-			}
-		}
-
-
-		static bool CheckRequirements(IAgent Agent, Requirements Requirements)
-		{
-			_ = Agent;
-			_ = Requirements;
-			// TODO: check requirements
-			// TODO: cache requirements
-			return true;
-		}
-
-		List<QueueKey> ModifyCachedQueues(Action<List<QueueKey>> Action)
-		{
-			for (; ; )
-			{
-				List<QueueKey> QueueKeys = CachedQueues;
-				if(TryModifyCachedQueueKeys(QueueKeys, Action))
-				{
-					return QueueKeys;
-				}
-			}
-		}
-
-		bool TryModifyCachedQueueKeys(List<QueueKey> InitialQueues, Action<List<QueueKey>> Action)
-		{
-			if (CachedQueues == InitialQueues)
-			{
-				List<QueueKey> NewQueues = new List<QueueKey>(InitialQueues);
-				Action(NewQueues);
-				return Interlocked.CompareExchange(ref CachedQueues, NewQueues, InitialQueues) == InitialQueues;
-			}
-			return false;
-		}
-
-#endif
-#endif
 	}
 }
