Horde: Add REST API for compute service.

* New tasks can be queued with a POST to /compute/{channelid}
* Status updates can be dequeued with a POST to /compute/{channelid}/updates
* ChannelId can be any series of lowercase letters, numbers, and the characters _ and -.
* Request/response objects are declared in Api.cs.

diff --git a/Engine/Source/Programs/Horde/HordeServer/Compute/Impl/ComputeController.cs b/Engine/Source/Programs/Horde/HordeServer/Compute/Impl/ComputeController.cs
new file mode 100644
--- /dev/null
+++ b/Engine/Source/Programs/Horde/HordeServer/Compute/Impl/ComputeController.cs
@@ -0,0 +1,96 @@
+// Copyright Epic Games, Inc. All Rights Reserved.
+
+using EpicGames.Core;
+using EpicGames.Horde.Compute;
+using HordeServer.Services;
+using HordeServer.Utilities;
+using Microsoft.AspNetCore.Authorization;
+using Microsoft.AspNetCore.Mvc;
+using System;
+using System.Collections.Generic;
+using System.ComponentModel.DataAnnotations;
+using System.Linq;
+using System.Threading;
+using System.Threading.Tasks;
+
+namespace HordeServer.Compute.Impl
+{
+	using ChannelId = StringId<IComputeChannel>;
+
+	/// <summary>
+	/// Controller for the /api/v1/compute endpoint
+	/// </summary>
+	[ApiController]
+	[Authorize]
+	[Route("[controller]")]
+	public class ComputeController : ControllerBase
+	{
+		AclService AclService;
+		IComputeService ComputeService;
+
+		/// <summary>
+		/// Constructor
+		/// </summary>
+		/// <param name="AclService"></param>
+		/// <param name="ComputeService">The compute service singleton</param>
+		public ComputeController(AclService AclService, IComputeService ComputeService)
+		{
+			this.AclService = AclService;
+			this.ComputeService = ComputeService;
+		}
+
+		/// <summary>
+		/// Add tasks to be executed remotely
+		/// </summary>
+		/// <param name="ChannelId">The channel to add to</param>
+		/// <param name="Request">The request parameters</param>
+		/// <returns></returns>
+		[HttpPost]
+		[Authorize]
+		[Route("/api/v1/compute/{ChannelId}")]
+		public async Task AddTasksAsync([FromRoute] string ChannelId, [FromBody] AddTasksRequest Request)
+		{
+			await ComputeService.AddTasksAsync(Impl.ComputeService.DefaultNamespaceId, IoHash.Parse(Request.RequirementsHash), Request.TaskHashes.ConvertAll(x => IoHash.Parse(x)), new ChannelId(ChannelId));
+		}
+
+		/// <summary>
+		/// Read updates for a particular channel
+		/// </summary>
+		/// <param name="ChannelId">The channel to add to</param>
+		/// <param name="Wait">Amount of time to wait before responding, in seconds</param>
+		/// <returns></returns>
+		[HttpPost]
+		[Authorize]
+		[Route("/api/v1/compute/{ChannelId}/updates")]
+		public async Task<GetTaskUpdatesResponse> GetUpdatesAsync([FromRoute] string ChannelId, [FromQuery] int Wait = 0)
+		{
+			List<IComputeTaskStatus> Results;
+			if (Wait == 0)
+			{
+				Results = await ComputeService.GetTaskUpdatesAsync(new ChannelId(ChannelId));
+			}
+			else
+			{
+				using CancellationTokenSource DelaySource = new CancellationTokenSource(Wait * 1000);
+				Results = await ComputeService.WaitForTaskUpdatesAsync(new ChannelId(ChannelId), DelaySource.Token);
+			}
+
+			GetTaskUpdatesResponse Response = new GetTaskUpdatesResponse();
+			foreach (IComputeTaskStatus Result in Results)
+			{
+				GetTaskUpdateResponse Update = new GetTaskUpdateResponse();
+
+				Update.TaskHash = Result.TaskHash.ToString();
+				Update.Time = Result.Time;
+				Update.State = Result.State;
+				Update.ResultHash = Result.ResultHash?.ToString();
+				Update.AgentId = Result.AgentId?.ToString();
+				Update.LeaseId = Result.LeaseId?.ToString();
+
+				Response.Updates.Add(Update);
+			}
+
+			return Response;
+		}
+	}
+}
diff --git a/Engine/Source/Programs/Shared/EpicGames.Horde/Compute/Api.cs b/Engine/Source/Programs/Shared/EpicGames.Horde/Compute/Api.cs
new file mode 100644
--- /dev/null
+++ b/Engine/Source/Programs/Shared/EpicGames.Horde/Compute/Api.cs
@@ -0,0 +1,84 @@
+// Copyright Epic Games, Inc. All Rights Reserved.
+
+using EpicGames.Core;
+using EpicGames.Horde.Common;
+using System;
+using System.Collections.Generic;
+using System.ComponentModel.DataAnnotations;
+using System.Linq;
+using System.Text;
+
+namespace EpicGames.Horde.Compute
+{
+	/// <summary>
+	/// Request to add tasks to the compute queue
+	/// </summary>
+	public class AddTasksRequest
+	{
+		/// <inheritdoc cref="AddTasksRpcRequest.RequirementsHash"/>
+		[Required]
+		public string RequirementsHash { get; set; } = String.Empty;
+
+		/// <inheritdoc cref="AddTasksRpcRequest.TaskHashes"/>
+		[Required]
+		public List<string> TaskHashes { get; set; } = new List<string>();
+
+		/// <inheritdoc cref="AddTasksRpcRequest.DoNotCache"/>
+		public bool DoNotCache { get; set; }
+	}
+
+	/// <summary>
+	/// Specifies updates from the given channel
+	/// </summary>
+	public class GetTaskUpdatesResponse
+	{
+		/// <summary>
+		/// Task updates
+		/// </summary>
+		public List<GetTaskUpdateResponse> Updates { get; set; } = new List<GetTaskUpdateResponse>();
+	}
+
+	/// <summary>
+	/// Supplies information about the current execution state of a task
+	/// </summary>
+	public class GetTaskUpdateResponse
+	{
+		/// <inheritdoc cref="GetTaskUpdatesRpcResponse.TaskHash"/>
+		public string TaskHash { get; set; } = null!;
+
+		/// <inheritdoc cref="GetTaskUpdatesRpcResponse.Time"/>
+		public DateTime Time { get; set; }
+
+		/// <inheritdoc cref="GetTaskUpdatesRpcResponse.State"/>
+		public ComputeTaskState State { get; set; }
+
+		/// <inheritdoc cref="GetTaskUpdatesRpcResponse.ResultHash"/>
+		public string? ResultHash { get; set; }
+
+		/// <inheritdoc cref="GetTaskUpdatesRpcResponse.AgentId"/>
+		public string? AgentId { get; set; }
+
+		/// <inheritdoc cref="GetTaskUpdatesRpcResponse.LeaseId"/>
+		public string? LeaseId { get; set; }
+	}
+
+	partial class AddTasksRpcRequest
+	{
+		public AddTasksRpcRequest(string ChannelId, string NamespaceId, IoHash RequirementsHash, List<IoHash> TaskHashes, bool DoNotCache)
+		{
+			this.ChannelId = ChannelId;
+			this.NamespaceId = NamespaceId;
+			this.RequirementsHash = RequirementsHash;
+			this.TaskHashes.Add(TaskHashes.Select(x => (IoHashWrapper)x));
+			this.DoNotCache = DoNotCache;
+		}
+	}
+
+	partial class GetTaskUpdatesRpcRequest
+	{
+		public GetTaskUpdatesRpcRequest(string ChannelId)
+		{
+			this.ChannelId = ChannelId;
+		}
+	}
+}
diff --git a/Engine/Source/Programs/Shared/EpicGames.Horde/Compute/Types.cs b/Engine/Source/Programs/Shared/EpicGames.Horde/Compute/Types.cs
deleted file mode 100644
