Horde: Disable agent connection lost callbacks for REAPI actions

diff --git a/Engine/Source/Programs/Horde/HordeServer/Tasks/Impl/ActionTaskSource.cs b/Engine/Source/Programs/Horde/HordeServer/Tasks/Impl/ActionTaskSource.cs
--- a/Engine/Source/Programs/Horde/HordeServer/Tasks/Impl/ActionTaskSource.cs
+++ b/Engine/Source/Programs/Horde/HordeServer/Tasks/Impl/ActionTaskSource.cs
@@ -491,9 +491,10 @@ namespace HordeServer.Tasks.Impl
 
 				void OnConnectionLost()
 				{
-					string Message = $"Connection lost to agent. LeaseId={Lease.Id} OperationId={Operation!.Id}.";
-					Logger.LogInformation("Connection lost to agent. LeaseId={LeaseId} OperationId={OperationId} Stacktrade={Stacktrace}.", Lease.Id, Operation!.Id, Environment.StackTrace);
-					Operation.TrySetResult(new ActionExecuteResult(new Google.Rpc.Status(StatusCode.Internal, Message)));
+					// Disabled until behavior of OnConnectionLost is fixed
+					// string Message = $"Connection lost to agent. LeaseId={Lease.Id} OperationId={Operation!.Id}.";
+					// Logger.LogInformation("Connection lost to agent. LeaseId={LeaseId} OperationId={OperationId} Stacktrade={Stacktrace}.", Lease.Id, Operation!.Id, Environment.StackTrace);
+					// Operation.TrySetResult(new ActionExecuteResult(new Google.Rpc.Status(StatusCode.Internal, Message)));
 				}
 				
 				if (Subscription.TrySetLease(Lease, OnConnectionLost))
