Horde: Add more logging for completing an REAPI action

diff --git a/Engine/Source/Programs/Horde/HordeServer/Tasks/Impl/ActionTaskSource.cs b/Engine/Source/Programs/Horde/HordeServer/Tasks/Impl/ActionTaskSource.cs
--- a/Engine/Source/Programs/Horde/HordeServer/Tasks/Impl/ActionTaskSource.cs
+++ b/Engine/Source/Programs/Horde/HordeServer/Tasks/Impl/ActionTaskSource.cs
@@ -501,9 +501,11 @@ namespace HordeServer.Tasks.Impl
 				{
 					Operation.SetStatus(ExecutionStage.Types.Value.Executing, null);
 
+					Logger.LogInformation("Executing remote operation {OperationId}...", Operation.Id);
 					ActionExecuteResult Result = await Operation.ResultTaskSource.Task;
 					if (Result.Result != null)
 					{
+						Logger.LogInformation("Remote execution {OperationId} successful. Retries {NumRetries} of {MaxNumRetries}", Operation.Id, Operation.NumExecResults, ExecuteOperation.MaxRetries);
 						return new ExecuteResponse
 						{
 							Status = new Status(StatusCode.OK, String.Empty),
@@ -516,6 +518,7 @@ namespace HordeServer.Tasks.Impl
 					{
 						if (Operation.NumExecResults > ExecuteOperation.MaxRetries)
 						{
+							Logger.LogError("Remote execution failed for {OperationId}. Retries exhausted, {NumRetries} of {MaxNumRetries}", Operation.Id, Operation.NumExecResults, ExecuteOperation.MaxRetries);
 							return new ExecuteResponse
 							{
 								Status = new Status((StatusCode) Result.Error.Code, Result.Error.Message),
