Horde: Remove Console.WriteLine calls from the new task source.

diff --git a/Engine/Source/Programs/Horde/HordeServer/Tasks/ITaskSource.cs b/Engine/Source/Programs/Horde/HordeServer/Tasks/ITaskSource.cs
--- a/Engine/Source/Programs/Horde/HordeServer/Tasks/ITaskSource.cs
+++ b/Engine/Source/Programs/Horde/HordeServer/Tasks/ITaskSource.cs
@@ -222,12 +222,10 @@ namespace HordeServer.Tasks
 				{
 					if (Accepted)
 					{
-						Console.WriteLine("ACCEPTED");
 						await Outer.TaskSource.OnLeaseStartedAsync(Agent, NewLease.Lease.Id, Any.Parser.ParseFrom(NewLease.Lease.Payload));
 					}
 					else 
 					{
-						Console.WriteLine("CANCELLED");
 						await Outer.TaskSource.CancelLeaseAsync(Agent, NewLease.Lease.Id, Any.Parser.ParseFrom(NewLease.Lease.Payload));
 					}
 				}
