Horde: In-process gRPC invoker to aid testing of gRPC services in-process

diff --git a/Engine/Source/Programs/Horde/HordeServerTests/RpcServiceTest.cs b/Engine/Source/Programs/Horde/HordeServerTests/RpcServiceTest.cs
--- a/Engine/Source/Programs/Horde/HordeServerTests/RpcServiceTest.cs
+++ b/Engine/Source/Programs/Horde/HordeServerTests/RpcServiceTest.cs
@@ -4,8 +4,10 @@ using System;
 using System.Collections.Generic;
 using System.IO;
 using System.IO.Compression;
+using System.Reflection;
 using System.Security.Claims;
 using System.Threading;
+using System.Threading.Channels;
 using System.Threading.Tasks;
 using Google.Protobuf;
 using Grpc.Core;
@@ -53,7 +55,7 @@ namespace HordeServerTests
 	[TestClass]
 	public class RpcServiceTest : DatabaseIntegrationTest
 	{
-		private readonly ServerCallContext AdminContext = new ContextStub("app-horde-admins");
+		private readonly ServerCallContext AdminContext = new ServerCallContextStub("app-horde-admins");
 
 		sealed class HttpContextStub : HttpContext
 		{
@@ -88,19 +90,36 @@ namespace HordeServerTests
 			}
 		}
 
-		class ContextStub : ServerCallContext
+		public class ServerCallContextStub : ServerCallContext
 		{
 			// Copied from ServerCallContextExtensions.cs in Grpc.Core
 			const string HttpContextKey = "__HttpContext";
 
-			public ContextStub(string RoleClaimType)
+			public static ServerCallContext ForAdminWithAgentSessionId(string AgentSessionId)
+			{
+				return new ServerCallContextStub(new ClaimsPrincipal(new ClaimsIdentity(new List<Claim>
+				{
+					new Claim(HordeClaimTypes.Role, "app-horde-admins"),
+					new Claim(HordeClaimTypes.AgentSessionId, AgentSessionId),
+				}, "TestAuthType")));
+			}
+			
+			public static ServerCallContext ForAdmin()
+			{
+				return new ServerCallContextStub(new ClaimsPrincipal(new ClaimsIdentity(new List<Claim>
+				{
+					new Claim(HordeClaimTypes.Role, "app-horde-admins"),
+				}, "TestAuthType")));
+			}
+
+			public ServerCallContextStub(string RoleClaimType)
 			{
 				// The GetHttpContext extension falls back to getting the HttpContext from UserState
 				// We can piggyback on that behavior during tests
 				UserState[HttpContextKey] = new HttpContextStub(RoleClaimType);
 			}
 			
-			public ContextStub(ClaimsPrincipal User)
+			public ServerCallContextStub(ClaimsPrincipal User)
 			{
 				// The GetHttpContext extension falls back to getting the HttpContext from UserState
 				// We can piggyback on that behavior during tests
@@ -128,6 +147,150 @@ namespace HordeServerTests
 			protected override WriteOptions WriteOptionsCore { get; set; } = null!;
 			protected override AuthContext AuthContextCore { get; } = null!;
 		}
+		
+		public class RpcServiceInvoker : CallInvoker
+		{
+			private readonly RpcService RpcService;
+			private readonly ServerCallContext ServerCallContext;
+
+			public RpcServiceInvoker(RpcService RpcService, ServerCallContext ServerCallContext)
+			{
+				this.RpcService = RpcService;
+				this.ServerCallContext = ServerCallContext;
+			}
+
+			public override TResponse BlockingUnaryCall<TRequest, TResponse>(Method<TRequest, TResponse> method, string host, CallOptions options, TRequest request)
+			{
+				throw new NotImplementedException("Blocking calls are not supported! Method " + method.FullName);
+			}
+
+			public override AsyncUnaryCall<TResponse> AsyncUnaryCall<TRequest, TResponse>(Method<TRequest, TResponse> Method, string Host, CallOptions Options, TRequest Request)
+			{
+				MethodInfo MethodInfo = GetMethod(Method.Name);
+				Task<TResponse> Res = (MethodInfo.Invoke(RpcService, new object[] {Request, ServerCallContext}) as Task<TResponse>)!;
+				return new AsyncUnaryCall<TResponse>(Res, null!, null!, null!, null!, null!);
+			}
+
+			public override AsyncServerStreamingCall<TResponse> AsyncServerStreamingCall<TRequest, TResponse>(Method<TRequest, TResponse> Method, string Host, CallOptions Options,
+				TRequest Request)
+			{
+				Console.WriteLine($"RpcServiceInvoker.AsyncServerStreamingCall(method={Method.FullName} request={Request})");
+				throw new NotImplementedException();
+			}
+
+			public override AsyncClientStreamingCall<TRequest, TResponse> AsyncClientStreamingCall<TRequest, TResponse>(Method<TRequest, TResponse> Method, string Host, CallOptions Options)
+			{
+				Console.WriteLine($"RpcServiceInvoker.AsyncClientStreamingCall(method={Method.FullName})");
+				throw new NotImplementedException();
+			}
+
+			public override AsyncDuplexStreamingCall<TRequest, TResponse> AsyncDuplexStreamingCall<TRequest, TResponse>(Method<TRequest, TResponse> Method, string Host, CallOptions Options)
+			{
+				Console.WriteLine($"RpcServiceInvoker.AsyncDuplexStreamingCall(method={Method.FullName})");
+
+				GrpcDuplexStreamHandler<TRequest> RequestStream = new GrpcDuplexStreamHandler<TRequest>(ServerCallContext);
+				GrpcDuplexStreamHandler<TResponse> ResponseStream = new GrpcDuplexStreamHandler<TResponse>(ServerCallContext);
+
+				MethodInfo MethodInfo = GetMethod(Method.Name);
+				Task MethodTask = (MethodInfo.Invoke(RpcService, new object[] {RequestStream, ResponseStream, ServerCallContext}) as Task)!;
+				MethodTask.ContinueWith(t => { Console.Error.WriteLine($"Uncaught exception in {Method.Name}: {t}"); }, TaskContinuationOptions.OnlyOnFaulted);
+				
+				return new AsyncDuplexStreamingCall<TRequest, TResponse>(RequestStream, ResponseStream, null!, null!, null!, null!);
+			}
+			
+			private MethodInfo GetMethod(string MethodName)
+			{
+				MethodInfo? Method = RpcService.GetType().GetMethod(MethodName);
+				if (Method == null)
+				{
+					throw new ArgumentException($"Method {MethodName} not found in RpcService");
+				}
+
+				return Method;
+			}
+		}
+		
+		/// <summary>
+		/// Combines and cross-writes streams for a duplex streaming call in gRPC
+		/// </summary>
+		/// <typeparam name="T"></typeparam>
+		public class GrpcDuplexStreamHandler<T> : IServerStreamWriter<T>, IClientStreamWriter<T>, IAsyncStreamReader<T> where T : class
+		{
+			private readonly ServerCallContext ServerCallContext;
+			private readonly Channel<T> Channel;
+
+			public WriteOptions? WriteOptions { get; set; }
+
+			public GrpcDuplexStreamHandler(ServerCallContext ServerCallContext)
+			{
+				Channel = System.Threading.Channels.Channel.CreateUnbounded<T>();
+
+				this.ServerCallContext = ServerCallContext;
+			}
+
+			public void Complete()
+			{
+				Channel.Writer.Complete();
+			}
+
+			public IAsyncEnumerable<T> ReadAllAsync()
+			{
+				return Channel.Reader.ReadAllAsync();
+			}
+
+			public async Task<T?> ReadNextAsync()
+			{
+				if (await Channel.Reader.WaitToReadAsync())
+				{
+					Channel.Reader.TryRead(out var message);
+					return message;
+				}
+				else
+				{
+					return null;
+				}
+			}
+
+			public Task WriteAsync(T Message)
+			{
+				if (ServerCallContext.CancellationToken.IsCancellationRequested)
+				{
+					return Task.FromCanceled(ServerCallContext.CancellationToken);
+				}
+
+				if (!Channel.Writer.TryWrite(Message))
+				{
+					throw new InvalidOperationException("Unable to write message.");
+				}
+
+				return Task.CompletedTask;
+			}
+
+			public Task CompleteAsync()
+			{
+				Complete();
+				return Task.CompletedTask;
+			}
+
+			public async Task<bool> MoveNext(CancellationToken CancellationToken)
+			{
+				ServerCallContext.CancellationToken.ThrowIfCancellationRequested();
+				
+				if (await Channel.Reader.WaitToReadAsync(CancellationToken))
+				{
+					Channel.Reader.TryRead(out var Message);
+					Current = Message;
+					return true;
+				}
+				else
+				{
+					Current = null!;
+					return false;
+				}
+			}
+
+			public T Current { get; private set; } = null!;
+		}
 
 		[TestMethod]
 		public async Task CreateSessionTest()
@@ -223,7 +386,7 @@ namespace HordeServerTests
 		{
 			TestSetup TestSetup = await GetTestSetup();
 			ObjectId SessionId = ObjectId.GenerateNewId();
-			ServerCallContext Context = new ContextStub(new ClaimsPrincipal(new ClaimsIdentity(new List<Claim>
+			ServerCallContext Context = new ServerCallContextStub(new ClaimsPrincipal(new ClaimsIdentity(new List<Claim>
 			{
 				new Claim(HordeClaimTypes.Role, "app-horde-admins"),
 				new Claim(HordeClaimTypes.AgentSessionId, SessionId.ToString()),
