Horde: Support situations where a Perforce server is not explciitly configured more gracefully. When running from Visual Studio, the server will now use the user's own Perforce connection by default.

diff --git a/Engine/Source/Programs/Horde/HordeServer/Models/Globals.cs b/Engine/Source/Programs/Horde/HordeServer/Models/Globals.cs
--- a/Engine/Source/Programs/Horde/HordeServer/Models/Globals.cs
+++ b/Engine/Source/Programs/Horde/HordeServer/Models/Globals.cs
@@ -269,9 +269,9 @@ namespace HordeServer.Models
 		public string Name { get; set; } = null!;
 
 		/// <summary>
-		/// Username for Horde to log in to this server
+		/// Username for Horde to log in to this server. Will use the default user if not set.
 		/// </summary>
-		public string ServiceAccount { get; set; } = "buildmachine";
+		public string? ServiceAccount { get; set; }
 
 		/// <summary>
 		/// Whether the service account can impersonate other users
@@ -397,9 +397,14 @@ namespace HordeServer.Models
 
 		static List<PerforceCluster> GetDefaultClusters()
 		{
+			PerforceServer Server = new PerforceServer();
+			Server.ServerAndPort = Perforce.P4.P4Server.Get("P4PORT") ?? "perforce:1666";
+
 			PerforceCluster Cluster = new PerforceCluster();
 			Cluster.Name = "Default";
-			Cluster.Servers.Add(new PerforceServer { ServerAndPort = "perforce:1666", Properties = new List<string> { "HordeServer=1" } });
+			Cluster.CanImpersonate = false;
+			Cluster.Servers.Add(Server);
+
 			return new List<PerforceCluster> { Cluster };
 		}
 	}
diff --git a/Engine/Source/Programs/Horde/HordeServer/Services/Impl/PerforceService.cs b/Engine/Source/Programs/Horde/HordeServer/Services/Impl/PerforceService.cs
--- a/Engine/Source/Programs/Horde/HordeServer/Services/Impl/PerforceService.cs
+++ b/Engine/Source/Programs/Horde/HordeServer/Services/Impl/PerforceService.cs
@@ -118,16 +118,22 @@ namespace HordeServer.Services
 			return Server;
 		}
 
-		static P4.Repository CreateConnection(IPerforceServer Server, string UserName, string Ticket)
+		static P4.Repository CreateConnection(IPerforceServer Server, string? UserName, string? Ticket)
 		{
 			P4.Repository Repository = new P4.Repository(new P4.Server(new P4.ServerAddress(Server.ServerAndPort)));
 			try
 			{
 				P4.Connection Connection = Repository.Connection;
-				Connection.UserName = UserName;
+				if (UserName != null)
+				{
+					Connection.UserName = UserName;
+				}
 
 				P4.Options Options = new P4.Options();
-				Options["Ticket"] = Ticket;
+				if (Ticket != null)
+				{
+					Options["Ticket"] = Ticket;
+				}
 
 				// connect to the server
 				if (!Connection.Connect(Options))
@@ -156,12 +162,17 @@ namespace HordeServer.Services
 				return CreateConnection(Server, Settings.P4BridgeServiceUsername, Settings.P4BridgeServicePassword);
 			}
 
-			PerforceCredentials? Credentials = Cluster.Credentials.FirstOrDefault(x => x.UserName.Equals(Cluster.ServiceAccount, StringComparison.OrdinalIgnoreCase));
-			if (Credentials == null)
+			string? UserName = null;
+			string? Password = null;
+			if (Cluster.ServiceAccount != null)
 			{
-				throw new Exception($"No credentials defined for {Cluster.ServiceAccount} on {Cluster.Name}");
+				PerforceCredentials? Credentials = Cluster.Credentials.FirstOrDefault(x => x.UserName.Equals(Cluster.ServiceAccount, StringComparison.OrdinalIgnoreCase));
+				if (Credentials == null)
+				{
+					throw new Exception($"No credentials defined for {Cluster.ServiceAccount} on {Cluster.Name}");
+				}
 			}
-			return CreateConnection(Server, Credentials.UserName, Credentials.Password);
+			return CreateConnection(Server, UserName, Password);
 		}
 
 		async Task<CachedTicketInfo> GetImpersonateCredential(PerforceCluster Cluster, string ImpersonateUser)
@@ -313,9 +324,9 @@ namespace HordeServer.Services
 
 		}
 
-		static string GetClientName(string ServiceUserName, string Stream, bool ReadOnly, bool CreateChange, string? Username = null)
+		static string GetClientName(string? ServiceUserName, string Stream, bool ReadOnly, bool CreateChange, string? Username = null)
 		{
-			string ClientName = $"horde-p4bridge-{Dns.GetHostName()}-{ServiceUserName}-{Stream.Replace("/", "+", StringComparison.OrdinalIgnoreCase)}";
+			string ClientName = $"horde-p4bridge-{Dns.GetHostName()}-{ServiceUserName ?? "default"}-{Stream.Replace("/", "+", StringComparison.OrdinalIgnoreCase)}";
 
 			if (!ReadOnly)
 			{
@@ -336,7 +347,7 @@ namespace HordeServer.Services
 
 		}
 
-		static P4.Client GetOrCreateClient(string ServiceUserName, P4.Repository Repository, string? Stream, string? Username = null, bool ReadOnly = true, bool CreateChange = false, int? ClientFromChange = null, bool UseClientFromChange = false, bool UsePortFromChange = false)
+		static P4.Client GetOrCreateClient(string? ServiceUserName, P4.Repository Repository, string? Stream, string? Username = null, bool ReadOnly = true, bool CreateChange = false, int? ClientFromChange = null, bool UseClientFromChange = false, bool UsePortFromChange = false)
 		{
 			P4.Client? Client = null;
 			P4.Changelist? Changelist = null;
diff --git a/Engine/Source/Programs/Horde/HordeServer/Services/PerforceLoadBalancer.cs b/Engine/Source/Programs/Horde/HordeServer/Services/PerforceLoadBalancer.cs
--- a/Engine/Source/Programs/Horde/HordeServer/Services/PerforceLoadBalancer.cs
+++ b/Engine/Source/Programs/Horde/HordeServer/Services/PerforceLoadBalancer.cs
@@ -259,23 +259,49 @@ namespace HordeServer.Services
 		public async Task<IPerforceServer?> SelectServerAsync(PerforceCluster Cluster, HashSet<string>? Properties)
 		{
 			// Find all the valid servers for this agent
-			HashSet<string> ValidServers = new HashSet<string>(StringComparer.OrdinalIgnoreCase);
+			List<PerforceServer> ValidServers = new List<PerforceServer>();
 			if (Properties != null)
 			{
-				ValidServers.UnionWith(Cluster.Servers.Where(x => x.Properties != null && x.Properties.Count > 0 && x.Properties.All(y => Properties.Contains(y))).Select(x => x.ServerAndPort));
+				ValidServers.AddRange(Cluster.Servers.Where(x => x.Properties != null && x.Properties.Count > 0 && x.Properties.All(y => Properties.Contains(y))));
 			}
 			if (ValidServers.Count == 0)
 			{
-				ValidServers.UnionWith(Cluster.Servers.Where(x => x.Properties == null || x.Properties.Count == 0).Select(x => x.ServerAndPort));
+				ValidServers.AddRange(Cluster.Servers.Where(x => x.Properties == null || x.Properties.Count == 0));
 			}
 
-			// Find all the matching servers
+			HashSet<string> ValidServerNames = new HashSet<string>(ValidServers.Select(x => x.ServerAndPort), StringComparer.OrdinalIgnoreCase);
+
+			// Find all the matching servers.
 			PerforceServerList ServerList = await GetServerListAsync();
-			List<PerforceServerEntry> Candidates = ServerList.Servers.Where(x => x.Cluster == Cluster.Name && x.Status == PerforceServerStatus.Healthy && ValidServers.Contains(x.BaseServerAndPort)).ToList();
+
+			List<PerforceServerEntry> Candidates = ServerList.Servers.Where(x => x.Cluster == Cluster.Name && ValidServerNames.Contains(x.BaseServerAndPort)).ToList();
+			if (Candidates.Count == 0)
+			{
+				foreach (PerforceServer ValidServer in ValidServers)
+				{
+					Logger.LogDebug("Fetching server info for {ServerAndPort}", ValidServer.ServerAndPort);
+					await UpdateServerAsync(Cluster, ValidServer, Candidates);
+				}
+				if (Candidates.Count == 0)
+				{
+					Logger.LogWarning("Unable to resolve any Perforce servers from valid list");
+					return null; 
+				}
+			}
+
+			// Remove any servers that are unhealthy
+			if (Candidates.Any(x => x.Status == PerforceServerStatus.Healthy))
+			{
+				Candidates.RemoveAll(x => x.Status != PerforceServerStatus.Healthy);
+			}
+			else
+			{
+				Candidates.RemoveAll(x => x.Status != PerforceServerStatus.Unknown);
+			}
 
 			if (Candidates.Count == 0)
 			{
-				Logger.LogDebug("Unable to find Perforce server from cluster {ClusterName} matching {ValidServers}", Cluster.Name, String.Join(", ", ValidServers));
+				Logger.LogWarning("Unable to find any healthy Perforce server in cluster {ClusterName}", Cluster.Name);
 				return null;
 			}
 
