Horde: Disable running conform jobs on local agents. Add support for development/production config files.

diff --git a/Engine/Source/Programs/Horde/HordeAgent/HordeAgent.csproj b/Engine/Source/Programs/Horde/HordeAgent/HordeAgent.csproj
--- a/Engine/Source/Programs/Horde/HordeAgent/HordeAgent.csproj
+++ b/Engine/Source/Programs/Horde/HordeAgent/HordeAgent.csproj
@@ -60,11 +60,13 @@
       <Generator></Generator>
       <CopyToOutputDirectory>PreserveNewest</CopyToOutputDirectory>
       <CopyToPublishDirectory>PreserveNewest</CopyToPublishDirectory>
+      <DependentUpon>appsettings.json</DependentUpon>
     </Content>
     <Content Include="appsettings.Production.json" Condition="Exists('appsettings.Production.json')">
       <Generator></Generator>
       <CopyToOutputDirectory>PreserveNewest</CopyToOutputDirectory>
       <CopyToPublishDirectory>PreserveNewest</CopyToPublishDirectory>
+      <DependentUpon>appsettings.json</DependentUpon>
     </Content>
   </ItemGroup>
 
diff --git a/Engine/Source/Programs/Horde/HordeAgent/Properties/launchSettings.json b/Engine/Source/Programs/Horde/HordeAgent/Properties/launchSettings.json
new file mode 100644
--- /dev/null
+++ b/Engine/Source/Programs/Horde/HordeAgent/Properties/launchSettings.json
@@ -0,0 +1,11 @@
+{
+  "profiles": {
+    "HordeAgent": {
+      "commandName": "Project",
+      "commandLineArgs": "service run",
+      "environmentVariables": {
+        "DOTNET_ENVIRONMENT": "Development"
+      }
+    }
+  }
+}
\ No newline at end of file
diff --git a/Engine/Source/Programs/Horde/HordeAgent/appsettings.Development.json b/Engine/Source/Programs/Horde/HordeAgent/appsettings.Development.json
new file mode 100644
--- /dev/null
+++ b/Engine/Source/Programs/Horde/HordeAgent/appsettings.Development.json
@@ -0,0 +1,6 @@
+{
+	"Horde":
+	{
+		"RunConform": false
+	}
+}
diff --git a/Engine/Source/Programs/Horde/HordeAgent/appsettings.Production.json b/Engine/Source/Programs/Horde/HordeAgent/appsettings.Production.json
new file mode 100644
--- /dev/null
+++ b/Engine/Source/Programs/Horde/HordeAgent/appsettings.Production.json
@@ -0,0 +1,2 @@
+{
+}
