Horde: Include file and directory sizes in DirectoryTree objects, as well as file attributes.

diff --git a/Engine/Source/Programs/Horde/HordeAgent/Commands/Execution/ExecuteCommandV2.cs b/Engine/Source/Programs/Horde/HordeAgent/Commands/Execution/ExecuteCommandV2.cs
--- a/Engine/Source/Programs/Horde/HordeAgent/Commands/Execution/ExecuteCommandV2.cs
+++ b/Engine/Source/Programs/Horde/HordeAgent/Commands/Execution/ExecuteCommandV2.cs
@@ -91,14 +91,15 @@ namespace HordeAgent.Commands
 		[CommandLine("-LogLevel")]
 		public string LogLevelStr = "debug";
 
-		static IoHash CreateSandbox(DirectoryInfo BaseDirInfo, Dictionary<IoHash, byte[]> UploadList)
+		static (DirectoryTree, IoHash) CreateSandbox(DirectoryInfo BaseDirInfo, Dictionary<IoHash, byte[]> UploadList)
 		{
 			DirectoryTree Tree = new DirectoryTree();
 
 			foreach (DirectoryInfo SubDirInfo in BaseDirInfo.EnumerateDirectories())
 			{
-				IoHash SubDirHash = CreateSandbox(SubDirInfo, UploadList);
-				Tree.Directories.Add(new DirectoryNode(SubDirInfo.Name, SubDirHash));
+				(DirectoryTree SubTree, IoHash SubDirHash) = CreateSandbox(SubDirInfo, UploadList);
+				long Size = SubTree.Files.Sum(x => x.Size) + SubTree.Directories.Sum(x => x.Size);
+				Tree.Directories.Add(new DirectoryNode(SubDirInfo.Name, SubDirHash, Size));
 			}
 			Tree.Directories.SortBy(x => x.Name, Utf8StringComparer.Ordinal);
 
@@ -107,11 +108,11 @@ namespace HordeAgent.Commands
 				byte[] Data = File.ReadAllBytes(FileInfo.FullName);
 				IoHash Hash = IoHash.Compute(Data);
 				UploadList[Hash] = Data;
-				Tree.Files.Add(new FileNode(FileInfo.Name, Hash));
+				Tree.Files.Add(new FileNode(FileInfo.Name, Hash, FileInfo.Length, (int)FileInfo.Attributes));
 			}
 			Tree.Files.SortBy(x => x.Name, Utf8StringComparer.Ordinal);
 
-			return AddCbObject(UploadList, Tree);
+			return (Tree, AddCbObject(UploadList, Tree));
 		}
 
 		/// <inheritdoc/>
@@ -170,7 +171,7 @@ namespace HordeAgent.Commands
 				}
 
 				Dictionary<IoHash, byte[]> Blobs = new Dictionary<IoHash, byte[]>();
-				IoHash SandboxHash = CreateSandbox(InputDir.ToDirectoryInfo(), Blobs);
+				(_, IoHash SandboxHash) = CreateSandbox(InputDir.ToDirectoryInfo(), Blobs);
 
 				ComputeTask Task = new ComputeTask(JsonComputeTask.Executable, JsonComputeTask.Arguments.ConvertAll<Utf8String>(x => x), JsonComputeTask.WorkingDirectory, SandboxHash);
 				Task.EnvVars = JsonComputeTask.EnvVars.ToDictionary(x => (Utf8String)x.Key, x => (Utf8String)x.Value);
diff --git a/Engine/Source/Programs/Horde/HordeAgent/Execution/ComputeTaskExecutor.cs b/Engine/Source/Programs/Horde/HordeAgent/Execution/ComputeTaskExecutor.cs
--- a/Engine/Source/Programs/Horde/HordeAgent/Execution/ComputeTaskExecutor.cs
+++ b/Engine/Source/Programs/Horde/HordeAgent/Execution/ComputeTaskExecutor.cs
@@ -237,13 +237,14 @@ namespace HordeAgent
 			await Task.WhenAll(Tasks);
 		}
 
-		Task<IoHash> PutOutput(string NamespaceId, DirectoryReference BaseDir, IEnumerable<FileReference> Files)
+		async Task<IoHash> PutOutput(string NamespaceId, DirectoryReference BaseDir, IEnumerable<FileReference> Files)
 		{
 			List<FileReference> SortedFiles = Files.OrderBy(x => x.FullName, StringComparer.Ordinal).ToList();
-			return PutDirectoryTree(NamespaceId, BaseDir.FullName.Length, SortedFiles, 0, SortedFiles.Count);
+			(_, IoHash Hash) = await PutDirectoryTree(NamespaceId, BaseDir.FullName.Length, SortedFiles, 0, SortedFiles.Count);
+			return Hash;
 		}
 
-		async Task<IoHash> PutDirectoryTree(string NamespaceId, int BaseDirLen, List<FileReference> SortedFiles, int MinIdx, int MaxIdx)
+		async Task<(DirectoryTree, IoHash)> PutDirectoryTree(string NamespaceId, int BaseDirLen, List<FileReference> SortedFiles, int MinIdx, int MaxIdx)
 		{
 			List<Task<FileNode>> Files = new List<Task<FileNode>>();
 			List<Task<DirectoryNode>> Trees = new List<Task<DirectoryNode>>();
@@ -281,19 +282,20 @@ namespace HordeAgent
 			DirectoryTree Tree = new DirectoryTree();
 			Tree.Files.AddRange(await Task.WhenAll(Files));
 			Tree.Directories.AddRange(await Task.WhenAll(Trees));
-			return await BlobStore.PutObjectAsync(NamespaceId, Tree);
+			return (Tree, await BlobStore.PutObjectAsync(NamespaceId, Tree));
 		}
 
 		async Task<DirectoryNode> CreateDirectoryNode(string NamespaceId, string Name, int BaseDirLen, List<FileReference> SortedFiles, int MinIdx, int MaxIdx)
 		{
-			IoHash Hash = await PutDirectoryTree(NamespaceId, BaseDirLen, SortedFiles, MinIdx, MaxIdx);
-			return new DirectoryNode(Name, Hash);
+			(DirectoryTree Tree, IoHash Hash) = await PutDirectoryTree(NamespaceId, BaseDirLen, SortedFiles, MinIdx, MaxIdx);
+			return new DirectoryNode(Name, Hash, Tree.Files.Sum(x => x.Size) + Tree.Directories.Sum(x => x.Size));
 		}
 
 		async Task<FileNode> CreateFileNode(string NamespaceId, string Name, FileReference File)
 		{
-			IoHash Hash = await BlobStore.PutBlobAsync(NamespaceId, await FileReference.ReadAllBytesAsync(File));
-			return new FileNode(Name, Hash);
+			byte[] Data = await FileReference.ReadAllBytesAsync(File);
+			IoHash Hash = await BlobStore.PutBlobAsync(NamespaceId, Data);
+			return new FileNode(Name, Hash, Data.Length, (int)FileReference.GetAttributes(File));
 		}
 
 		/// <summary>
diff --git a/Engine/Source/Programs/Shared/EpicGames.Horde/Compute/DirectoryTree.cs b/Engine/Source/Programs/Shared/EpicGames.Horde/Compute/DirectoryTree.cs
--- a/Engine/Source/Programs/Shared/EpicGames.Horde/Compute/DirectoryTree.cs
+++ b/Engine/Source/Programs/Shared/EpicGames.Horde/Compute/DirectoryTree.cs
@@ -43,6 +43,18 @@ namespace EpicGames.Horde.Compute
 		[CbField("h")]
 		public IoHash Hash { get; set; }
 
+		/// <summary>
+		/// Size of the file
+		/// </summary>
+		[CbField("s")]
+		public long Size { get; set; }
+
+		/// <summary>
+		/// The file attributes
+		/// </summary>
+		[CbField("a")]
+		public int Attributes { get; set; }
+
 		/// <summary>
 		/// Private constructor for serialization
 		/// </summary>
@@ -53,10 +65,12 @@ namespace EpicGames.Horde.Compute
 		/// <summary>
 		/// Constructor
 		/// </summary>
-		public FileNode(Utf8String Name, IoHash Hash)
+		public FileNode(Utf8String Name, IoHash Hash, long Size, int Attributes)
 		{
 			this.Name = Name;
 			this.Hash = Hash;
+			this.Size = Size;
+			this.Attributes = Attributes;
 		}
 	}
 
@@ -77,6 +91,12 @@ namespace EpicGames.Horde.Compute
 		[CbField("h")]
 		public IoHash Hash { get; set; }
 
+		/// <summary>
+		/// Total size of the items in this directory
+		/// </summary>
+		[CbField("s")]
+		public long Size { get; set; }
+
 		/// <summary>
 		/// Private constructor for serialization
 		/// </summary>
@@ -87,10 +107,11 @@ namespace EpicGames.Horde.Compute
 		/// <summary>
 		/// Constructor
 		/// </summary>
-		public DirectoryNode(Utf8String Name, IoHash Hash)
+		public DirectoryNode(Utf8String Name, IoHash Hash, long Size)
 		{
 			this.Name = Name;
 			this.Hash = Hash;
+			this.Size = Size;
 		}
 	}
 }
diff --git a/Engine/Source/Programs/Shared/EpicGames.Serialization/CompactBinary.cs b/Engine/Source/Programs/Shared/EpicGames.Serialization/CompactBinary.cs
--- a/Engine/Source/Programs/Shared/EpicGames.Serialization/CompactBinary.cs
+++ b/Engine/Source/Programs/Shared/EpicGames.Serialization/CompactBinary.cs
@@ -793,7 +793,15 @@ namespace EpicGames.Serialization
 		/// <summary>
 		/// Access the field as an int64. Returns the provided default on error.
 		/// </summary>
-		public long AsInt64(long Default = 0)
+		public long AsInt64()
+		{
+			return AsInt64(0);
+		}
+
+		/// <summary>
+		/// Access the field as an int64. Returns the provided default on error.
+		/// </summary>
+		public long AsInt64(long Default)
 		{
 			return (long)AsInteger((ulong)Default, 63, true);
 		}
diff --git a/Engine/Source/Programs/Shared/EpicGames.Serialization/Converters/CbDefaultConverter.cs b/Engine/Source/Programs/Shared/EpicGames.Serialization/Converters/CbDefaultConverter.cs
--- a/Engine/Source/Programs/Shared/EpicGames.Serialization/Converters/CbDefaultConverter.cs
+++ b/Engine/Source/Programs/Shared/EpicGames.Serialization/Converters/CbDefaultConverter.cs
@@ -100,11 +100,14 @@ namespace EpicGames.Serialization.Converters
 
 		static void WriteInt32Value(CbWriter Writer, int Value) => Writer.WriteIntegerValue(Value);
 		static void WriteInt32(CbWriter Writer, Utf8String Name, int Value) => Writer.WriteInteger(Name, Value);
+		static void WriteInt64Value(CbWriter Writer, long Value) => Writer.WriteIntegerValue(Value);
+		static void WriteInt64(CbWriter Writer, Utf8String Name, long Value) => Writer.WriteInteger(Name, Value);
 
 		static Dictionary<Type, CbConverterInfo> TypeToConverterInfo = new Dictionary<Type, CbConverterInfo>(new KeyValuePair<Type, CbConverterInfo>[]
 		{
 			GetPodConverterInfo(x => x.AsBool(), (w, v) => w.WriteBoolValue(v), (w, n, v) => w.WriteBool(n, v), null),
 			GetPodConverterInfo(x => x.AsInt32(), (w, v) => WriteInt32Value(w, v), (w, n, v) => WriteInt32(w, n, v), null),
+			GetPodConverterInfo(x => x.AsInt64(), (w, v) => WriteInt64Value(w, v), (w, n, v) => WriteInt64(w, n, v), null),
 			GetPodConverterInfo(x => x.AsString(), (w, v) => w.WriteStringValue(v), (w, n, v) => w.WriteString(n, v), null),
 			GetPodConverterInfo(x => x.AsHash(), (w, v) => w.WriteHashValue(v), (w, n, v) => w.WriteHash(n, v), null),
 			GetPodConverterInfo(x => x.AsDateTime(), (w, v) => w.WriteDateTimeValue(v), (w, n, v) => w.WriteDateTime(n, v), null)
