Horde: Licensee experience and server configuration
#rnx

diff --git a/Engine/Source/Programs/Horde/HordeDashboard/src/backend/Api.ts b/Engine/Source/Programs/Horde/HordeDashboard/src/backend/Api.ts
--- a/Engine/Source/Programs/Horde/HordeDashboard/src/backend/Api.ts
+++ b/Engine/Source/Programs/Horde/HordeDashboard/src/backend/Api.ts
@@ -2977,5 +2977,622 @@ export type GetDeviceReservationResponse = {
 export type UpdateLeaseRequest = {
 	/** Mark this lease as aborted */
 	aborted?: boolean
+
+}
+
+/** Server settings */
+export type GetServerSettingsResponse = {
+
+	/** The number of live server setting updates */
+	numServerUpdates?: number;
+
+	globalConfigPath: string;
+
+	/** The server settings on local storage */
+	userServerSettingsPath:string;
+
+	/** MongoDB connection string */
+	databaseConnectionString?: string;
+
+	/** MongoDB database name */
+	databaseName: string;
+
+	/** The claim type for administrators */
+	adminClaimType: string;
+
+	/** Value of the claim type for administrators */
+	adminClaimValue: string;
+
+	/** Optional certificate to trust in order to access the database (eg. AWS cert for TLS) */
+	databasePublicCert?: string;
+
+	/** Access the database in read-only mode (avoids creating indices or updating content)
+		Useful for debugging a local instance of HordeServer against a production database. */
+	databaseReadOnlyMode: boolean;
+
+	/** Optional PFX certificate to use for encryting agent SSL traffic. This can be a self-signed certificate, as long as it's trusted by agents.	*/
+	serverPrivateCert?: string;
+
+	/** Issuer for tokens from the auth provider */	
+	oidcAuthority?: string;
+
+	/** Client id for the OIDC authority */
+	oidcClientId?: string;
+
+	/** Optional redirect url provided to OIDC login */
+	oidcSigninRedirect?: string;
+
+	/** Name of the issuer in bearer tokens from the server*/
+	jwtIssuer?: string;
+
+	/** Secret key used to sign JWTs. This setting is typically only used for development. In prod, a unique secret key will be generated and stored in the DB for each unique server instance.*/
+	jwtSecret?: string;
+
+	/** Length of time before JWT tokens expire, in hourse*/
+	jwtExpiryTimeHours: number;
+
+	/** Disable authentication for debugging purposes*/
+	disableAuth: boolean;
+
+	/** Whether to enable Cors, generally for development purposes*/
+	corsEnabled: boolean;
+
+	/** Allowed Cors origin */
+	corsOrigin: string;
+
+	/** Whether to enable a schedule in test data (false by default for development builds)*/
+	enableScheduleInTestData: boolean;
+
+	/** Interval between rebuilding the schedule queue with a DB query.*/
+	schedulePollingInterval: any;
+
+	/** Interval between polling for new jobs*/
+	noResourceBackOffTime: any;
+
+	/** Interval between attempting to assign agents to take on jobs*/
+	initiateJobBackOffTime: any;
+
+	/** Interval between scheduling jobs when an unknown error occurs*/
+	unknownErrorBackOffTime: any;
+
+	/** Config for connecting to Redis server(s).
+		Setting it to null will disable Redis use and connection
+		See format at https://stackexchange.github.io/StackExchange.Redis/Configuration.html */
+	redisConnectionConfig?: string
+
+	/** Type of write cache to use in log service, currently Supported: "InMemory" or "Redis" */
+	logServiceWriteCacheType: string
+
+	/** Provider Type, currently Supported: "S3" or "FileSystem" */
+	externalStorageProviderType: any
+
+	/** Local log/artifact storage directory, if using type filesystem */
+	localLogsDir: string;
+
+	/** Local blob storage directory, if using type filesystem */
+	localBlobsDir: string;
+
+	/** Local artifact storage directory, if using type filesystem */
+	localArtifactsDir: string;
+
+	/** S3 bucket region for logfile storage */
+	s3BucketRegion: string;
+
+	/** Arn to assume for s3.  "Basic", "AssumeRole", "AssumeRoleWebIdentity" only */
+	s3CredentialType: string;
+
+	/** S3 Client username (used in Basic auth type only) */
+	s3ClientKeyId: string;
+
+	/** S3 client password (used in Basic auth type only) */
+	s3ClientSecret: string;
+
+	/** Arn to assume for s3 */
+	s3AssumeArn: string;
+
+	/** S3 log bucket name */
+	s3LogBucketName: string;
+
+	/** S3 artifact bucket name */
+	s3ArtifactBucketName: string;
+
+	/** When using a relay storage provider, specifies the remote server to use */
+	logRelayServer?: string;
+
+	/** Authentication token for using a relay server */
+	logRelayBearerToken?: string;
+
+	/** Whether to log json to stdout */
+	logJsonToStdOut: boolean;
+
+	/** Which fleet manager service to use */
+	fleetManager: any;
+
+	/** Whether to run scheduled jobs. Not wanted for development. */
+	disableSchedules: boolean;
+
+	/** Timezone for evaluating schedules */
+	scheduleTimeZone?: string;
+
+	/** Token for interacting with Slack */
+	slackToken?: string;
+
+	/** Token for opening a socket to slack */
+	slackSocketToken?: string;
+
+	/** Channel to send stream notification update failures to */
+	updateStreamsNotificationChannel?: string;
+
+	/** URI to the SmtpServer to use for sending email notifications */
+	smtpServer?: string;
+
+	/** The email address to send email notifications from */
+	emailSenderAddress?: string;
+
+	/** The name for the sender when sending email notifications */
+	emailSenderName?: string;
+
+	/** The p4 bridge service username */
+	p4BridgeServiceUsername?: string;
+
+	/** The p4 bridge service password */
+	p4BridgeServicePassword?: string;
+
+	/** Whether the p4 bridge service account can impersonate other users */
+	p4BridgeCanImpersonate: boolean;
+
+	/** Set the minimum size of the global thread pool
+		This value has been found in need of tweaking to avoid timeouts with the Redis client during bursts
+		of traffic. Default is 16 for .NET Core CLR. The correct value is dependent on the traffic the Horde Server
+		is receiving. For Epic's internal deployment, this is set to 40. */
+	globalThreadPoolMinSize?: number;
+
+	/** Path to the root config file */
+	configPath?: string;
+
+	/** Lazily computed timezone value */
+	timeZoneInfo: any;
+
+}
+
+export type UpdateServerSettingsRequest = {
+
+	settings: Record<string, string | boolean | number>;
+}
+
+export type ServerUpdateResponse = {
+
+	errors: string[];
+
+	restartRequired: boolean;
+
+}
+
+
+// Config
+
+/// References a project configuration
+export type ProjectConfigRef = {
+
+	/// Unique id for the project
+	id: string;
+
+	/// Config path for the project		
+	path: string;
+}
+
+/// How frequently the maintence window repeats
+export enum ScheduledDowntimeFrequency {
+	/// Once
+	Once,
+
+	/// Every day
+	Daily,
+
+	/// Every week
+	Weekly
+}
+
+
+/// Settings for the maintenance window
+export type ScheduledDowntime = {
+
+	/// Start time
+	dateTimeOffset: any;
+
+	/// Finish time
+	finishTime: any;
+
+	/// Frequency that the window repeats\
+	frequency: ScheduledDowntimeFrequency;
+}
+
+
+
+/// User notice
+export type Notice = {
+
+	/// Unique id for this notice
+	id: string;
+
+	/// Start time to display this message
+	startTime?: Date | string;
+
+	/// Finish time to display this message
+	finishTime?: Date | string;
+
+	/// Message to display
+	message: string;
+}
+
+/// Path to a platform and stream to use for syncing AutoSDK
+export type AutoSdkWorkspace = {
+
+	/// Name of this workspace
+	name?: string;
+
+	/// The agent properties to check (eg. "OSFamily=Windows")
+	properties: string[];
+
+	/// Username for logging in to the server
+	userName?: string;
+
+	/// Stream to use
+	stream?: string;
+}
+
+
+/// Information about an individual Perforce server
+export type PerforceServer = {
+
+	/// The server and port. The server may be a DNS entry with multiple records, in which case it will be actively load balanced.
+	serverAndPort: string;
+
+	/// Whether to query the healthcheck address under each server
+	healthCheck?: boolean;
+
+	/// Whether to resolve the DNS entries and load balance between different hosts
+	resolveDns?: boolean;
+
+	/// Maximum number of simultaneous conforms on this server
+	maxConformCount: number;
+
+	/// List of properties for an agent to be eligable to use this server
+	properties?: string[];
+}
+
+
+/// Credentials for a Perforce user
+export type PerforceCredentials = {
+	/// The username
+	userName: string;
+
+	/// Password for the user
+	password: string;
+}
+
+
+/// Information about a cluster of Perforce servers. 
+export type PerforceCluster = {
+
+	/// The default cluster name
+	defaultName?: string;
+
+	/// Name of the cluster
+	name: string;
+
+	/// Username for Horde to log in to this server
+	serviceAccount: string;
+
+	/// Whether the service account can impersonate other users
+	canImpersonate: boolean;
+
+	/// List of servers
+	servers: PerforceServer[];
+
+	/// List of server credentials
+	credentials: PerforceCredentials[];
+
+	/// List of autosdk streams
+	autoSdk: AutoSdkWorkspace[];
+
+}
+
+
+/// Configuration for storage system
+export type StorageConfig = {
+
+	/// List of storage namespaces
+	namespaces: NamespaceConfig[];
+}
+
+/// Configuration for a storage namespace
+export type NamespaceConfig = {
+
+	/// Identifier for this namespace
+	id: string;
+
+	/// Buckets within this namespace
+	buckets: BucketConfig[];
+
+	/// Access control for this namespace
+	//UpdateAclRequest? Acl { get; set; }
+}
+
+/// Configuration for a bucket
+export type BucketConfig = {
+	/// Identifier for the bucket
+	id: string;
+}
+
+/** Global configuration */
+export type GlobalConfig = {
+
+	/// List of projects
+	projects: ProjectConfigRef[];
+
+	/// Manually added status messages
+	notices?: Notice[];
+
+	/// List of scheduled downtime
+	downtime?: ScheduledDowntime[];
+
+	/// List of Perforce clusters
+	perforceClusters: PerforceCluster[];
+
+	/// Maximum number of conforms to run at once
+	maxConformCount?: number;
+
+	/// List of storage namespaces
+	storage?: StorageConfig;
+
+	/// Access control list
+	// public UpdateAclRequest ? Acl { get; set; }
+}
+
+
+export type CreateProjectCategoryRequest = {
+	/// <summary>
+	/// Name of this category
+	/// </summary>
+
+	name: string;
+
+	/// <summary>
+	/// Index of the row to display this category on
+	/// </summary>
+	row: number;
+
+	/// <summary>
+	/// Whether to show this category on the nav menu
+	/// </summary>
+	showOnNavMenu: boolean;
+
+	/// <summary>
+	/// Patterns for stream names to include
+	/// </summary>
+	includePatterns: string[];
+
+	/// <summary>
+	/// Patterns for stream names to exclude
+	/// </summary>
+	excludePatterns: string[];
+
+}
+
+
+// Project Config
+/// Stores configuration for a project
+export type ProjectConfig = {
+
+	/// Name for the new project
+	name: string;
+
+	/// Path to the project logo
+	logo?: string;
+
+	/// Categories to include in this project
+	categories: CreateProjectCategoryRequest[];
+
+	/// List of streams
+	streams: StreamConfigRef[];
+
+	/// Acl entries
+	// public UpdateAclRequest? Acl { get; set; }
+}
+
+
+/// <summary>
+/// Reference to configuration for a stream
+/// </summary>
+export type StreamConfigRef = {
+
+	/// <summary>
+	/// Unique id for the stream
+	/// </summary>
+	id: string;
+
+	/// <summary>
+	/// Path to the configuration file
+	/// </summary>
+	path: string;
+}
+
+
+
+export type CreateAgentTypeRequest = {
+	/// <summary>
+	/// Pool of agents to use for this agent type
+	/// </summary>
+	pool: string;
+
+	/// <summary>
+	/// Name of the workspace to sync
+	/// </summary>
+	workspace?: string;
+
+	/// <summary>
+	/// Path to the temporary storage dir
+	/// </summary>
+	tempStorageDir?: string;
+
+	/// <summary>
+	/// Environment variables to be set when executing the job
+	/// </summary>
+	environment?: Record<string, string>;
+}
+
+
+/// Information about a workspace type
+export type CreateWorkspaceTypeRequest = {
+
+	/// <summary>
+	/// Name of the Perforce server cluster to use
+	/// </summary>
+	cluster?: string;
+
+	/// <summary>
+	/// The Perforce server and port (eg. perforce:1666)
+	/// </summary>
+	serverAndPort?: string;
+
+	/// <summary>
+	/// User to log into Perforce with (defaults to buildmachine)
+	/// </summary>
+	userName?: string;
+
+	/// <summary>
+	/// Password to use to log into the workspace
+	/// </summary>
+	password?: string;
+
+	/// <summary>
+	/// Identifier to distinguish this workspace from other workspaces. Defaults to the workspace type name.
+	/// </summary>
+	identifier?: string;
+
+	/// <summary>
+	/// Override for the stream to sync
+	/// </summary>
+	stream?: string;
+
+	/// <summary>
+	/// Custom view for the workspace
+	/// </summary>
+	view?: string[];
+
+	/// <summary>
+	/// Whether to use an incrementally synced workspace
+	/// </summary>
+	incremental: boolean;
+}
+
+
+/// Config for a stream
+export type StreamConfig = {
+
+	/// <summary>
+	/// Name of the stream
+	/// </summary>
+	name: string;
+
+	/// <summary>
+	/// The perforce cluster containing the stream
+	/// </summary>
+	clusterName?: string;
+
+	/// <summary>
+	/// Order for this stream
+	/// </summary>
+	order?: number;
+
+	/// Notification channel for all jobs in this stream
+	notificationChannel?: string;
+
+	/// Notification channel filter for this template. Can be Success|Failure|Warnings
+	notificationChannelFilter?: string;
+
+	/// <summary>
+	/// Channel to post issue triage notifications
+	/// </summary>
+	triageChannel?: string;
+
+	/// <summary>
+	/// Default template for running preflights
+	/// </summary>
+	defaultPreflightTemplate?: string;
+
+	/// Default template for running preflights
+	defaultPreflight?: DefaultPreflightRequest;
+
+	/// <summary>
+	/// List of tabs to show for the new stream
+	/// </summary>
+	tabs: any[]/*CreateStreamTabRequest[]*/;
+
+	/// <summary>
+	/// Map of agent name to type
+	/// </summary>
+	agentTypes: Record<string, CreateAgentTypeRequest>;
+
+	/// <summary>
+	/// Map of workspace name to type
+	/// </summary>
+	workspaceTypes: Record<string, CreateWorkspaceTypeRequest>;
+
+	/// <summary>
+	/// List of templates to create
+	/// </summary>
+	templates: any[];
+
+	/// <summary>
+	/// Custom permissions for this object
+	/// </summary>
+	// public UpdateAclRequest ? Acl { get; set; }
+
+	/// <summary>
+	/// Pause stream builds until specified date
+	/// </summary>
+	/// public DateTime ? PausedUntil { get; set; }
+
+	/// <summary>
+	/// Reason for pausing builds of the stream
+	/// </summary>
+	// public string ? PauseComment { get; set; }
+}
+
+/// Parameters to update server settings
+export type UpdateGlobalConfigRequest = {
+
+	/// Delta updates for global config
+	globalsJson?: string;
+
+	/// projects json
+	projectsJson?:Record<string, string>;
+
+	/// streams json
+	streamsJson?:Record<string, string>;
+
+	/// default pool name
+	defaultPoolName?: string;
+
+	/// Delta updates for server settings from dashboard
+	serverSettingJson?: string;
+
+	/// Base64 encoded project logo
+	projectLogo?: string;
+
+}
+
+export type GetServerInfoResponse = {
+
+	/// Server version info
+	serverVersion: string;
+
+	/// The operating system server is hosted on
+	osDescription: string;
+
+	/// whether the server is running in single instance mode
+	singleInstance: boolean;
 }
 
diff --git a/Engine/Source/Programs/Horde/HordeDashboard/src/backend/Backend.ts b/Engine/Source/Programs/Horde/HordeDashboard/src/backend/Backend.ts
--- a/Engine/Source/Programs/Horde/HordeDashboard/src/backend/Backend.ts
+++ b/Engine/Source/Programs/Horde/HordeDashboard/src/backend/Backend.ts
@@ -3,7 +3,7 @@
 import { configure } from "mobx";
 import { isNumber } from 'util';
 import templateCache from '../backend/TemplateCache';
-import { AgentData, ArtifactData, BatchUpdatePoolRequest, ChangeSummaryData, CreateJobRequest, CreateJobResponse, CreatePoolRequest, CreateSoftwareResponse, CreateSubscriptionRequest, CreateSubscriptionResponse, DashboardPreference, EventData, GetArtifactZipRequest, GetUserResponse, GetGraphResponse, GetIssueStreamResponse, GetJobStepRefResponse, GetJobStepTraceResponse, GetJobTimingResponse, GetLogEventResponse, GetNotificationResponse, GetSoftwareResponse, GetSubscriptionResponse, IssueData, IssueQuery, JobData, JobQuery, LeaseData, LogData, LogLineData, PoolData, ProjectData, ScheduleData, ScheduleQuery, SearchLogFileResponse, SessionData, StreamData, TemplateData, TestData, UpdateAgentRequest, UpdateIssueRequest, UpdateJobRequest, UpdateNotificationsRequest, UpdatePoolRequest, UpdateStepRequest, UpdateStepResponse, UpdateUserRequest, GetUtilizationTelemetryResponse, TabType, GetJobsTabResponse, JobsTabColumnType, GetPerforceServerStatusResponse, GetDeviceResponse, GetDevicePlatformResponse, CreateDeviceResponse, CreateDeviceRequest, UpdateDeviceRequest, GetDevicePoolResponse, GetDeviceReservationResponse, UsersQuery, LeaseOutcome, UpdateLeaseRequest } from './Api';
+import { AgentData, ArtifactData, BatchUpdatePoolRequest, ChangeSummaryData, CreateDeviceRequest, CreateDeviceResponse, CreateJobRequest, CreateJobResponse, CreatePoolRequest, CreateSoftwareResponse, CreateSubscriptionRequest, CreateSubscriptionResponse, DashboardPreference, EventData, GetArtifactZipRequest, GetDevicePlatformResponse, GetDevicePoolResponse, GetDeviceReservationResponse, GetDeviceResponse, GetGraphResponse, GetIssueStreamResponse, GetJobsTabResponse, GetJobStepRefResponse, GetJobStepTraceResponse, GetJobTimingResponse, GetLogEventResponse, GetNotificationResponse, GetPerforceServerStatusResponse, GetServerInfoResponse, GetServerSettingsResponse, GetSoftwareResponse, GetSubscriptionResponse, GetUserResponse, GetUtilizationTelemetryResponse, GlobalConfig, IssueData, IssueQuery, JobData, JobQuery, JobsTabColumnType, LeaseData, LogData, LogLineData, PerforceCluster, PoolData, ProjectData, ScheduleData, ScheduleQuery, SearchLogFileResponse, ServerUpdateResponse, SessionData, StreamData, TabType, TemplateData, TestData, UpdateAgentRequest, UpdateDeviceRequest, UpdateGlobalConfigRequest, UpdateIssueRequest, UpdateJobRequest, UpdateLeaseRequest, UpdateNotificationsRequest, UpdatePoolRequest, UpdateServerSettingsRequest, UpdateStepRequest, UpdateStepResponse, UpdateUserRequest, UsersQuery } from './Api';
 import dashboard from './Dashboard';
 import { ChallengeStatus, Fetch } from './Fetch';
 import graphCache, { GraphQuery } from './GraphCache';
@@ -125,24 +125,24 @@ export class Backend {
                 this.backend.get(`/api/v1/agents/${agentId}/leases/${leaseId}`).then((response) => {
                     resolve(response.data as LeaseData);
                 }).catch(reason => { reject(reason); });
-            });    
+            });
         }
 
         return new Promise<LeaseData>((resolve, reject) => {
             this.backend.get(`/api/v1/leases/${leaseId}`).then((response) => {
                 resolve(response.data as LeaseData);
             }).catch(reason => { reject(reason); });
-        });    
+        });
 
     }
 
-    getLeases(agentId: string, index: number, count: number, includeBatches?:boolean): Promise<LeaseData[]> {
+    getLeases(agentId: string, index: number, count: number, includeBatches?: boolean): Promise<LeaseData[]> {
         const params: string[] = [];
         params.push("Index=" + index);
         params.push("Count=" + count);
         const paramString = params.join("&");
         return new Promise<LeaseData[]>((resolve, reject) => {
-            this.backend.get(`/api/v1/agents/${agentId}/leases?${paramString}`).then( async (response) => {
+            this.backend.get(`/api/v1/agents/${agentId}/leases?${paramString}`).then(async (response) => {
                 const leases = response.data as LeaseData[];
                 leases.forEach(function (lease) {
                     lease.startTime = new Date(Date.parse(lease.startTime as string));
@@ -159,18 +159,18 @@ export class Backend {
                 // get batch states
                 if (includeBatches) {
 
-                    const jobIds:Set<string> = new Set();
+                    const jobIds: Set<string> = new Set();
                     leases.filter(lease => lease.jobId && lease.batchId).forEach(lease => jobIds.add(lease.jobId));
                     if (jobIds.size) {
-                        const jobs = await this.getJobsByIds(Array.from(jobIds), {filter:"id,batches"});
+                        const jobs = await this.getJobsByIds(Array.from(jobIds), { filter: "id,batches" });
 
                         leases.forEach(lease => {
                             const job = jobs.find(j => lease.jobId === j.id);
                             if (job) {
-                                lease.batch = job.batches?.find(b => lease.batchId === b.id );
+                                lease.batch = job.batches?.find(b => lease.batchId === b.id);
                             }
                         });
-                    }    
+                    }
                 }
 
                 resolve(leases);
@@ -789,10 +789,10 @@ export class Backend {
         query = query ?? {};
 
         return new Promise<GetUserResponse[]>((resolve, reject) => {
-            this.backend.get("/api/v1/users",{
+            this.backend.get("/api/v1/users", {
                 params: query as any
             }).then((response) => {
-                resolve(response.data as GetUserResponse[] );
+                resolve(response.data as GetUserResponse[]);
             }).catch(reason => reject(reason));
         });
 
@@ -889,7 +889,7 @@ export class Backend {
 
     }
 
-    getIssueEvents(issueId: number, jobId?: string, logIds?: string[], count?:number): Promise<GetLogEventResponse[]> {
+    getIssueEvents(issueId: number, jobId?: string, logIds?: string[], count?: number): Promise<GetLogEventResponse[]> {
 
         if (logIds && logIds.length) {
             jobId = undefined;
@@ -1003,6 +1003,18 @@ export class Backend {
 
     }
 
+    downloadAgentZip() {
+        try {
+            const url = `${this.serverUrl}/api/v1/agentsoftware/latest/download`;
+            const link = document.createElement('a');
+            link.href = url;
+            document.body.appendChild(link);
+            link.click();
+        } catch (reason) {
+            console.error(reason);
+        }
+    }
+
     updateJobStep(jobId: string, batchId: string, stepId: string, request: UpdateStepRequest): Promise<UpdateStepResponse> {
         return new Promise<UpdateStepResponse>((resolve, reject) => {
             this.backend.put(`api/v1/jobs/${jobId}/batches/${batchId}/steps/${stepId}`, request).then((value) => {
@@ -1107,7 +1119,7 @@ export class Backend {
                 reject(reason);
             });
         });
-        
+
     }
 
     getDevices(): Promise<GetDeviceResponse[]> {
@@ -1118,7 +1130,7 @@ export class Backend {
             }).catch(reason => {
                 reject(reason);
             });
-        });        
+        });
     }
 
     addDevice(request: CreateDeviceRequest): Promise<CreateDeviceResponse> {
@@ -1128,7 +1140,7 @@ export class Backend {
             }).catch(reason => {
                 reject(reason);
             });
-        });        
+        });
     }
 
     checkoutDevice(deviceId: string, checkout:boolean): Promise<void> {
@@ -1149,17 +1161,17 @@ export class Backend {
             }).catch(reason => {
                 reject(reason);
             });
-        });        
+        });
     }
 
-    updateLease(leaseId: string, update:UpdateLeaseRequest): Promise<void> {
+    updateLease(leaseId: string, update: UpdateLeaseRequest): Promise<void> {
         return new Promise<void>((resolve, reject) => {
             this.backend.put(`/api/v1/leases/${leaseId}`, update).then(() => {
                 resolve();
             }).catch(reason => {
                 reject(reason);
             });
-        });        
+        });
     }
 
 
@@ -1170,7 +1182,7 @@ export class Backend {
             }).catch(reason => {
                 reject(reason);
             });
-        });        
+        });
     }
 
     getDevicePlatforms(): Promise<GetDevicePlatformResponse[]> {
@@ -1182,7 +1194,7 @@ export class Backend {
                 reject(reason);
             });
         });
-        
+
     }
 
     getDevicePools(): Promise<GetDevicePoolResponse[]> {
@@ -1194,7 +1206,7 @@ export class Backend {
                 reject(reason);
             });
         });
-        
+
     }
 
     getDeviceReservations(): Promise<GetDeviceReservationResponse[]> {
@@ -1206,13 +1218,54 @@ export class Backend {
                 reject(reason);
             });
         });
-        
+
     }
 
-    getServerVersion():Promise<string> {
-        return new Promise<string>((resolve, reject) => {
-            this.backend.get(`/api/v1/server/version`).then((value) => {
-                resolve(value.data as string);
+    getServerSettings(): Promise<GetServerSettingsResponse> {
+        return new Promise<GetServerSettingsResponse>((resolve, reject) => {
+            this.backend.get(`/api/v1/config/serversettings`).then((value) => {
+                resolve(value.data as GetServerSettingsResponse);
+            }).catch(reason => {
+                reject(reason);
+            });
+        });
+    }
+
+    updateServerSettings(request: UpdateServerSettingsRequest): Promise< ServerUpdateResponse> {
+        return new Promise<ServerUpdateResponse>((resolve, reject) => {
+            this.backend.put(`/api/v1/config/serversettings`, request).then((value) => {
+                resolve(value.data as ServerUpdateResponse);
+            }).catch(reason => {
+                reject(reason);
+            });
+        });
+    }
+
+    getGlobalConfig(): Promise<GlobalConfig> {
+        return new Promise<GlobalConfig>((resolve, reject) => {
+            this.backend.get(`/api/v1/config/global`).then((value) => {
+                resolve(value.data as GlobalConfig);
+            }).catch(reason => {
+                reject(reason);
+            });
+        });        
+    }
+
+    // updates global configuation
+    updateGlobalConfig(request: UpdateGlobalConfigRequest): Promise<void> {
+        return new Promise<void>((resolve, reject) => {
+            this.backend.put(`/api/v1/config/global`, request).then(() => {
+                resolve();
+            }).catch(reason => {
+                reject(reason);
+            });
+        });
+    }
+
+    getServerInfo(): Promise<GetServerInfoResponse> {
+        return new Promise<GetServerInfoResponse>((resolve, reject) => {
+            this.backend.get(`/api/v1/server/info`).then((value) => {
+                resolve(value.data as GetServerInfoResponse);
             }).catch(reason => {
                 reject(reason);
             });
@@ -1240,7 +1293,7 @@ export class Backend {
     async serverLogout(redirect: string) {
         try {
             this.backend.logout = true;
-            await this.backend.get("/api/v1/dashboard/logout");            
+            await this.backend.get("/api/v1/dashboard/logout");
             window.location.assign(redirect);
 
         } catch (err) {
@@ -1285,7 +1338,7 @@ export class Backend {
                 return;
             }
 
-            this.serverVersion = await this.getServerVersion();
+            this.serverInfo = await this.getServerInfo();
 
             await dashboard.update();
 
@@ -1311,7 +1364,7 @@ export class Backend {
     updateID?: any;
     logout: boolean = false;
 
-    serverVersion:string = "";
+    serverInfo: GetServerInfoResponse = {serverVersion: "0", osDescription: "Unknown", singleInstance: false};
 
 
     private backend: Fetch;
diff --git a/Engine/Source/Programs/Horde/HordeDashboard/src/backend/Dashboard.ts b/Engine/Source/Programs/Horde/HordeDashboard/src/backend/Dashboard.ts
--- a/Engine/Source/Programs/Horde/HordeDashboard/src/backend/Dashboard.ts
+++ b/Engine/Source/Programs/Horde/HordeDashboard/src/backend/Dashboard.ts
@@ -309,6 +309,10 @@ export class Dashboard {
 
     }
 
+    setServerSettingsChanged(value: boolean | undefined) {
+        this.serverSettingsChanged = value ?? false;
+        this.setUpdated();
+    }
 
     @observable
     updated: number = 0;
@@ -351,6 +355,8 @@ export class Dashboard {
 
     requestLogout = false;
 
+    serverSettingsChanged: boolean = false;
+
     private data: GetUserResponse = { id: "", name: "", enableIssueNotifications: false, enableExperimentalFeatures: false, claims: [], pinnedJobIds: [], dashboardSettings: { preferences: new Map() } };
 
     private updateTimeoutId: any = undefined;
diff --git a/Engine/Source/Programs/Horde/HordeServer/Api/Config.cs b/Engine/Source/Programs/Horde/HordeServer/Api/Config.cs
new file mode 100644
--- /dev/null
+++ b/Engine/Source/Programs/Horde/HordeServer/Api/Config.cs
@@ -0,0 +1,452 @@
+// Copyright Epic Games, Inc. All Rights Reserved.
+
+using System;
+using System.Collections.Generic;
+using HordeServer.Models;
+
+namespace HordeServer.Api
+{
+	/// <summary>
+	/// Server settings 
+	/// </summary>
+	public class GetServerSettingsResponse
+	{
+		/// <summary>
+		/// MongoDB connection string
+		/// </summary>
+		public string? DatabaseConnectionString { get; set; }
+
+		/// <summary>
+		/// MongoDB database name
+		/// </summary>
+		public string DatabaseName { get; set; }
+
+		/// <summary>
+		/// The claim type for administrators
+		/// </summary>
+		public string AdminClaimType { get; set; }
+
+		/// <summary>
+		/// Value of the claim type for administrators
+		/// </summary>
+		public string AdminClaimValue { get; set; }
+
+		/// <summary>
+		/// Optional certificate to trust in order to access the database (eg. AWS public cert for TLS)
+		/// </summary>
+		public string? DatabasePublicCert { get; set; }
+
+		/// <summary>
+		/// Access the database in read-only mode (avoids creating indices or updating content)
+		/// Useful for debugging a local instance of HordeServer against a production database.
+		/// </summary>
+		public bool DatabaseReadOnlyMode { get; set; }
+
+		/// <summary>
+		/// Optional PFX certificate to use for encryting agent SSL traffic. This can be a self-signed certificate, as long as it's trusted by agents.
+		/// </summary>
+		public string? ServerPrivateCert { get; set; }
+
+		/// <summary>
+		/// Issuer for tokens from the auth provider
+		/// </summary>
+		public string? OidcAuthority { get; set; }
+
+		/// <summary>
+		/// Client id for the OIDC authority
+		/// </summary>
+		public string? OidcClientId { get; set; }
+
+		/// <summary>
+		/// Optional redirect url provided to OIDC login
+		/// </summary>
+		public string? OidcSigninRedirect { get; set; }
+
+		/// <summary>
+		/// Name of the issuer in bearer tokens from the server
+		/// </summary>
+		public string? JwtIssuer { get; set; }
+
+		/// <summary>
+		/// Secret key used to sign JWTs. This setting is typically only used for development. In prod, a unique secret key will be generated and stored in the DB for each unique server instance.
+		/// </summary>
+		public string? JwtSecret { get; set; }
+
+		/// <summary>
+		/// Length of time before JWT tokens expire, in hourse
+		/// </summary>
+		public int JwtExpiryTimeHours { get; set; }
+
+		/// <summary>
+		/// Disable authentication for debugging purposes
+		/// </summary>
+		public bool DisableAuth { get; set; }
+
+		/// <summary>
+		/// Whether to enable Cors, generally for development purposes
+		/// </summary>
+		public bool CorsEnabled { get; set; }
+
+		/// <summary>
+		/// Allowed Cors origin 
+		/// </summary>
+		public string CorsOrigin { get; set; }
+
+		/// <summary>
+		/// Whether to enable a schedule in test data (false by default for development builds)
+		/// </summary>
+		public bool EnableScheduleInTestData { get; set; }
+
+		/// <summary>
+		/// Interval between rebuilding the schedule queue with a DB query.
+		/// </summary>
+		public TimeSpan SchedulePollingInterval { get; set; }
+
+		/// <summary>
+		/// Interval between polling for new jobs
+		/// </summary>
+		public TimeSpan NoResourceBackOffTime { get; set; }
+
+		/// <summary>
+		/// Interval between attempting to assign agents to take on jobs
+		/// </summary>
+		public TimeSpan InitiateJobBackOffTime { get; set; }
+
+		/// <summary>
+		/// Interval between scheduling jobs when an unknown error occurs
+		/// </summary>
+		public TimeSpan UnknownErrorBackOffTime { get; set; }
+
+		/// <summary>
+		/// Config for connecting to Redis server(s).
+		/// Setting it to null will disable Redis use and connection
+		/// See format at https://stackexchange.github.io/StackExchange.Redis/Configuration.html
+		/// </summary>
+		public string? RedisConnectionConfig { get; set; }
+
+		/// <summary>
+		/// Type of write cache to use in log service
+		/// Currently Supported: "InMemory" or "Redis"
+		/// </summary>
+		public string LogServiceWriteCacheType { get; set; }
+
+		/// <summary>
+		/// Provider Type
+		/// Currently Supported: "S3" or "FileSystem"
+		/// </summary>
+		public StorageProviderType ExternalStorageProviderType { get; set; }
+
+		/// <summary>
+		/// Local log/artifact storage directory, if using type filesystem
+		/// </summary>
+		public string LocalLogsDir { get; set; } = null!;
+
+
+		/// <summary>
+		/// Local blob storage directory, if using type filesystem
+		/// </summary>
+		public string LocalBlobsDir { get; set; } = null!;
+
+		/// <summary>
+		/// Local artifact storage directory, if using type filesystem
+		/// </summary>
+		public string LocalArtifactsDir { get; set; } = null!;
+
+		/// <summary>
+		/// S3 bucket region for logfile storage
+		/// </summary>
+		public string S3BucketRegion { get; set; } = null!;
+
+		/// <summary>
+		/// Arn to assume for s3.  "Basic", "AssumeRole", "AssumeRoleWebIdentity" only
+		/// </summary>
+		public string? S3CredentialType { get; set; }
+
+		/// <summary>
+		/// S3 Client username (used in Basic auth type only)
+		/// </summary>
+		public string S3ClientKeyId { get; set; } = null!;
+
+		/// <summary>
+		/// S3 client password (used in Basic auth type only)
+		/// </summary>
+		public string S3ClientSecret { get; set; } = null!;
+
+		/// <summary>
+		/// Arn to assume for s3
+		/// </summary>
+		public string S3AssumeArn { get; set; } = null!;
+
+		/// <summary>
+		/// S3 log bucket name
+		/// </summary>
+		public string S3LogBucketName { get; set; } = null!;
+
+		/// <summary>
+		/// S3 artifact bucket name
+		/// </summary>
+		public string S3ArtifactBucketName { get; set; } = null!;
+
+		/// <summary>
+		/// When using a relay storage provider, specifies the remote server to use
+		/// </summary>
+		public string? LogRelayServer { get; set; }
+
+		/// <summary>
+		/// Authentication token for using a relay server
+		/// </summary>
+		public string? LogRelayBearerToken { get; set; }
+
+		/// <summary>
+		/// Whether to log json to stdout
+		/// </summary>
+		public bool LogJsonToStdOut { get; set; }
+
+		/// <summary>
+		/// Which fleet manager service to use
+		/// </summary>
+		public FleetManagerType FleetManager { get; set; }
+
+		/// <summary>
+		/// Whether to run scheduled jobs. Not wanted for development.
+		/// </summary>
+		public bool DisableSchedules { get; set; }
+
+		/// <summary>
+		/// Timezone for evaluating schedules
+		/// </summary>
+		public string? ScheduleTimeZone { get; set; }
+
+		/// <summary>
+		/// Token for interacting with Slack
+		/// </summary>
+		public string? SlackToken { get; set; }
+
+		/// <summary>
+		/// Token for opening a socket to slack
+		/// </summary>
+		public string? SlackSocketToken { get; set; }
+
+		/// <summary>
+		/// Channel to send stream notification update failures to
+		/// </summary>
+		public string? UpdateStreamsNotificationChannel { get; set; }
+
+		/// <summary>
+		/// URI to the SmtpServer to use for sending email notifications
+		/// </summary>
+		public string? SmtpServer { get; set; }
+
+		/// <summary>
+		/// The email address to send email notifications from
+		/// </summary>
+		public string? EmailSenderAddress { get; set; }
+
+		/// <summary>
+		/// The name for the sender when sending email notifications
+		/// </summary>
+		public string? EmailSenderName { get; set; }
+
+		/// <summary>
+		/// The p4 bridge service username
+		/// </summary>
+		public string? P4BridgeServiceUsername { get; set; }
+
+		/// <summary>
+		/// The p4 bridge service password
+		/// </summary>
+		public string? P4BridgeServicePassword { get; set; }
+
+		/// <summary>
+		/// Whether the p4 bridge service account can impersonate other users
+		/// </summary>
+		public bool P4BridgeCanImpersonate { get; set; }
+
+		/// <summary>
+		/// Set the minimum size of the global thread pool
+		/// This value has been found in need of tweaking to avoid timeouts with the Redis client during bursts
+		/// of traffic. Default is 16 for .NET Core CLR. The correct value is dependent on the traffic the Horde Server
+		/// is receiving. For Epic's internal deployment, this is set to 40.
+		/// </summary>
+		public int? GlobalThreadPoolMinSize { get; set; }
+
+		/// <summary>
+		/// The number of live server setting updates 
+		/// </summary>
+		public int NumServerUpdates { get; set; }
+
+		/// <summary>
+		/// Path to the root config file
+		/// </summary>
+		public string? GlobalConfigPath { get; set; }
+
+		/// <summary>
+		/// Path to the user config file
+		/// </summary>
+		public string? UserServerSettingsPath { get; set; }
+
+
+		/// <summary>
+		/// Lazily computed timezone value
+		/// </summary>
+		public TimeZoneInfo TimeZoneInfo { get; set; } = null!;
+
+		/// <summary>
+		/// Constructor
+		/// </summary>		
+		public GetServerSettingsResponse(ServerSettings Settings, string? GlobalConfigPath = null)
+		{
+			this.GlobalConfigPath = GlobalConfigPath;
+
+			if (this.GlobalConfigPath == null)
+			{
+				this.GlobalConfigPath = Settings.ConfigPath;
+			}
+
+			this.NumServerUpdates = 0;
+
+			this.UserServerSettingsPath = Program.UserConfigFile.ToString();
+			this.DatabaseConnectionString = Settings.DatabaseConnectionString;
+			this.DatabaseName = Settings.DatabaseName;
+			this.DatabasePublicCert = Settings.DatabasePublicCert;
+			this.DatabaseReadOnlyMode = Settings.DatabaseReadOnlyMode;
+
+			this.AdminClaimType = Settings.AdminClaimType;
+			this.AdminClaimValue = Settings.AdminClaimValue;
+
+			this.ServerPrivateCert = Settings.ServerPrivateCert;
+			this.OidcAuthority = Settings.OidcAuthority;
+			this.OidcClientId = Settings.OidcClientId;
+			this.OidcSigninRedirect = Settings.OidcSigninRedirect;
+
+			this.JwtIssuer = Settings.JwtIssuer;
+			this.JwtSecret = Settings.JwtSecret;
+			this.JwtExpiryTimeHours = Settings.JwtExpiryTimeHours;
+
+			this.DisableAuth = Settings.DisableAuth;
+			this.CorsEnabled = Settings.CorsEnabled;
+			this.CorsOrigin = Settings.CorsOrigin;
+			this.EnableScheduleInTestData = Settings.EnableScheduleInTestData;
+			this.SchedulePollingInterval = Settings.SchedulePollingInterval;
+
+
+			this.NoResourceBackOffTime = Settings.NoResourceBackOffTime;
+			this.InitiateJobBackOffTime = Settings.InitiateJobBackOffTime;
+			this.UnknownErrorBackOffTime = Settings.UnknownErrorBackOffTime;
+
+
+			this.RedisConnectionConfig = Settings.RedisConnectionConfig;
+			this.LogServiceWriteCacheType = Settings.LogServiceWriteCacheType;
+			this.ExternalStorageProviderType = Settings.ExternalStorageProviderType;
+			this.LocalLogsDir = Settings.LocalLogsDir;
+			this.LocalBlobsDir = Settings.LocalBlobsDir;
+			this.LocalArtifactsDir = Settings.LocalArtifactsDir;
+
+			this.S3BucketRegion = Settings.S3BucketRegion;
+			this.S3CredentialType = Settings.S3CredentialType;
+			this.S3ClientKeyId = Settings.S3ClientKeyId;
+			this.S3ClientSecret = Settings.S3ClientSecret;
+			this.S3AssumeArn = Settings.S3AssumeArn;
+			this.S3LogBucketName = Settings.S3LogBucketName;
+			this.S3ArtifactBucketName = Settings.S3ArtifactBucketName;
+
+			this.LogRelayServer = Settings.LogRelayServer;
+			this.LogRelayBearerToken = Settings.LogRelayBearerToken;
+			this.LogJsonToStdOut = Settings.LogJsonToStdOut;
+
+			this.FleetManager = Settings.FleetManager;
+			this.DisableSchedules = Settings.DisableSchedules;
+			this.ScheduleTimeZone = Settings.ScheduleTimeZone;
+
+			this.SlackToken = Settings.SlackToken;
+			this.SlackSocketToken = Settings.SlackSocketToken;
+
+
+			this.SlackToken = Settings.SlackToken;
+			this.SlackToken = Settings.SlackToken;
+			this.UpdateStreamsNotificationChannel = Settings.UpdateStreamsNotificationChannel;
+
+			this.SmtpServer = Settings.SmtpServer;
+			this.EmailSenderAddress = Settings.EmailSenderAddress;
+			this.EmailSenderName = Settings.EmailSenderName;
+
+			this.UpdateStreamsNotificationChannel = Settings.UpdateStreamsNotificationChannel;
+			this.UpdateStreamsNotificationChannel = Settings.UpdateStreamsNotificationChannel;
+			this.UpdateStreamsNotificationChannel = Settings.UpdateStreamsNotificationChannel;
+
+			this.P4BridgeServiceUsername = Settings.P4BridgeServiceUsername;
+			this.P4BridgeServicePassword = Settings.P4BridgeServicePassword;
+			this.P4BridgeCanImpersonate = Settings.P4BridgeCanImpersonate;
+
+
+			this.GlobalThreadPoolMinSize = Settings.GlobalThreadPoolMinSize;
+			this.TimeZoneInfo = Settings.TimeZoneInfo;
+
+		}
+	}
+
+	/// <summary>
+	/// Parameters to update global settings
+	/// </summary>
+	public class UpdateServerSettingsRequest
+	{
+		/// <summary>
+		/// Delta settings to update
+		/// </summary>
+		public Dictionary<string, object>? Settings { get; set; }
+
+	}
+
+	/// <summary>
+	/// Response for server updates
+	/// </summary>
+	public class ServerUpdateResponse
+	{
+
+		/// <summary>
+		/// List of any error messages when updating
+		/// </summary>
+		public List<string> Errors { get; set; } = new List<string>();
+
+		/// <summary>
+		/// Whether a server restart is required
+		/// </summary>
+		public bool RestartRequired { get; set; } = false;
+
+	}
+
+
+	/// <summary>
+	/// Parameters to update global settings
+	/// </summary>
+	public class UpdateGlobalConfigRequest
+	{
+		/// <summary>
+		/// Delta updates for global config
+		/// </summary>
+		public string? GlobalsJson { get; set; }
+
+		/// <summary>
+		/// Projects json, filename => json string
+		/// </summary>
+		public Dictionary<string, string>? ProjectsJson { get; set; }
+
+		/// <summary>
+		/// Base64 encoded project logo
+		/// </summary>
+		public string? ProjectLogo { get; set; }
+
+		/// <summary>
+		/// Streams json, filename => json string
+		/// </summary>
+		public Dictionary<string, string>? StreamsJson { get; set; }
+		
+		/// <summary>
+		/// Default Pool Name
+		/// </summary>
+		public string? DefaultPoolName {get; set;}		
+
+	}
+
+}
+
diff --git a/Engine/Source/Programs/Horde/HordeServer/Api/Server.cs b/Engine/Source/Programs/Horde/HordeServer/Api/Server.cs
new file mode 100644
--- /dev/null
+++ b/Engine/Source/Programs/Horde/HordeServer/Api/Server.cs
@@ -0,0 +1,43 @@
+// Copyright Epic Games, Inc. All Rights Reserved.
+
+using System;
+using System.Runtime.InteropServices;
+using System.Reflection;
+using System.Diagnostics;
+
+namespace HordeServer.Api
+{
+	/// <summary>
+	/// Server Info
+	/// </summary>
+	public class GetServerInfoResponse
+	{
+        /// <summary>
+		/// Server version info
+		/// </summary>
+        public string ServerVersion { get; set; }
+
+        /// <summary>
+        /// The operating system server is hosted on
+        /// </summary>
+        public string OSDescription { get; set; }
+
+		/// <summary>
+		/// Whether this is an installed Horde build
+		/// </summary>
+		public bool SingleInstance { get; set; }
+
+		/// <summary>
+		/// Constructor
+		/// </summary>
+		public GetServerInfoResponse( bool SingleInstance)
+        {
+
+            FileVersionInfo VersionInfo = FileVersionInfo.GetVersionInfo(Assembly.GetExecutingAssembly().Location);		
+			ServerVersion = VersionInfo.ProductVersion;			
+			OSDescription = RuntimeInformation.OSDescription;
+			this.SingleInstance = SingleInstance;
+		}
+	}
+}
+
diff --git a/Engine/Source/Programs/Horde/HordeServer/Controllers/ConfigController.cs b/Engine/Source/Programs/Horde/HordeServer/Controllers/ConfigController.cs
new file mode 100644
--- /dev/null
+++ b/Engine/Source/Programs/Horde/HordeServer/Controllers/ConfigController.cs
@@ -0,0 +1,171 @@
+// Copyright Epic Games, Inc. All Rights Reserved.
+
+using EpicGames.Core;
+using HordeServer.Api;
+using HordeServer.Models;
+using HordeServer.Services;
+using Microsoft.AspNetCore.Authorization;
+using Microsoft.AspNetCore.Mvc;
+using Microsoft.Extensions.Options;
+using System.Threading.Tasks;
+using System.Text.Json;
+using Microsoft.Extensions.Logging;
+using System.Runtime.InteropServices;
+
+namespace HordeServer.Controllers
+{
+
+	/// <summary>
+	/// Controller managing server configuration
+	/// </summary>
+	[ApiController]
+	[Authorize]
+	[Route("[controller]")]
+	public class ConfigController : ControllerBase
+	{
+		/// <summary>
+		/// The database service singleton
+		/// </summary>
+		DatabaseService DatabaseService;
+
+		/// <summary>
+		/// The acl service singleton
+		/// </summary>
+		AclService AclService;
+
+		/// <summary>
+		/// The config service singleton
+		/// </summary>
+		ConfigService ConfigService;
+		/// <summary>
+		/// Settings for the server
+		/// </summary>
+		IOptionsMonitor<ServerSettings> Settings;
+
+		/// <summary>
+		///  Logger for controller
+		/// </summary>
+		private readonly ILogger<ConfigController> Logger;
+
+		/// <summary>
+		/// Constructor
+		/// </summary>
+		/// <param name="DatabaseService">The database service singleton</param>
+		/// <param name="AclService">The ACL service singleton</param>
+		/// <param name="ConfigService">The Config service singleton</param>
+		/// <param name="Logger">Logger for the controller</param>
+		/// <param name="Settings">Server settings</param>
+		public ConfigController(DatabaseService DatabaseService, AclService AclService, ConfigService ConfigService, IOptionsMonitor<ServerSettings> Settings, ILogger<ConfigController> Logger)
+		{
+			this.DatabaseService = DatabaseService;
+			this.AclService = AclService;
+			this.ConfigService = ConfigService;
+			this.Settings = Settings;
+			this.Logger = Logger;			
+		}
+
+		// Global Configuration
+
+		/// <summary>
+		/// Update global config
+		/// </summary>
+		[HttpPut]
+		[Route("/api/v1/config/global")]
+		public async Task<ActionResult> UpdateGlobalSettings([FromBody] UpdateGlobalConfigRequest Request)
+		{
+			if (!await AclService.AuthorizeAsync(AclAction.AdminWrite, User))
+			{
+				return Forbid();
+			}
+
+			if (!Settings.CurrentValue.SingleInstance)
+			{
+				Logger.LogError("Updating global configuration settings with ConfigController currently only supported in single instance mode");
+				return BadRequest();
+			}
+
+			await ConfigService.UpdateGlobalConfig(Request);
+
+			return Ok();
+		}
+
+		/// <summary>
+		/// Get global config settings
+		/// </summary>
+		/// <returns>Administrative settings for the server</returns>
+		[HttpGet]
+		[Route("/api/v1/config/global")]
+		public async Task<ActionResult<GlobalConfig?>> GetGlobalConfig()
+		{
+			if (!await AclService.AuthorizeAsync(AclAction.AdminRead, User))
+			{
+				return Forbid();
+			}
+
+			if (string.IsNullOrEmpty(Settings.CurrentValue.ConfigPath))
+			{
+				return BadRequest();
+			}
+
+			FileReference GlobalConfigFile = new FileReference(Settings.CurrentValue.ConfigPath);
+
+			byte[] Data = await FileReference.ReadAllBytesAsync(GlobalConfigFile);
+			JsonSerializerOptions Options = new JsonSerializerOptions();
+			Startup.ConfigureJsonSerializer(Options);
+
+			return JsonSerializer.Deserialize<GlobalConfig>(Data, Options);
+		}
+
+		// Administrative Server Settings
+
+		/// <summary>
+		/// Get administrative server settings
+		/// </summary>
+		/// <returns>Administrative settings for the server</returns>
+		[HttpGet]
+		[Route("/api/v1/config/serversettings")]
+		public async Task<ActionResult<GetServerSettingsResponse>> GetServerSettings()
+		{
+			if (!await AclService.AuthorizeAsync(AclAction.AdminRead, User))
+			{
+				return Forbid();
+			}
+
+			GetServerSettingsResponse Response = new GetServerSettingsResponse(Settings.CurrentValue);
+			Response.NumServerUpdates = ConfigService.NumUserConfigUpdates;
+
+			return Response;
+		}
+
+		/// <summary>
+		/// Update administrative server settings
+		/// </summary>
+		/// <returns>Administrative settings for the server</returns>
+		[HttpPut]
+		[Route("/api/v1/config/serversettings")]
+		public async Task<ActionResult<ServerUpdateResponse>> UpdateServerSettings([FromBody] UpdateServerSettingsRequest Request)
+		{
+			if (!await AclService.AuthorizeAsync(AclAction.AdminWrite, User))
+			{
+				return Forbid();
+			}
+
+			if (!Settings.CurrentValue.SingleInstance)
+			{
+				Logger.LogError("Updating server settings from ConfigController currently only supported in single instance mode");
+				return BadRequest();
+			}
+
+			if (Request.Settings == null || Request.Settings.Count == 0)
+			{
+				return new ServerUpdateResponse() { RestartRequired = false };
+			}
+
+			ServerUpdateResponse Response = await ConfigService.UpdateServerSettings(Request);
+
+			return Response;
+		}
+
+	}
+
+}
diff --git a/Engine/Source/Programs/Horde/HordeServer/Controllers/ServerController.cs b/Engine/Source/Programs/Horde/HordeServer/Controllers/ServerController.cs
--- a/Engine/Source/Programs/Horde/HordeServer/Controllers/ServerController.cs
+++ b/Engine/Source/Programs/Horde/HordeServer/Controllers/ServerController.cs
@@ -33,15 +33,22 @@ namespace HordeServer.Controllers
 	[Route("[controller]")]
 	public class ServerController : ControllerBase
 	{
+
+		/// <summary>
+		/// Settings for the server
+		/// </summary>
+		IOptionsMonitor<ServerSettings> Settings;
+
 		/// <summary>
 		/// Constructor
 		/// </summary>
-		public ServerController()
+		public ServerController(IOptionsMonitor<ServerSettings> Settings)
 		{
+			this.Settings = Settings;
 		}
-
+		
 		/// <summary>
-		/// Force a reset on the database
+		/// Get server version
 		/// </summary>
 		[HttpGet]
 		[Route("/api/v1/server/version")]
@@ -49,6 +56,17 @@ namespace HordeServer.Controllers
 		{
 			FileVersionInfo FileVersionInfo = FileVersionInfo.GetVersionInfo(Assembly.GetExecutingAssembly().Location);
 			return Ok(FileVersionInfo.ProductVersion);
+		}		
+
+		/// <summary>
+		/// Get server information
+		/// </summary>
+		[HttpGet]
+		[Route("/api/v1/server/info")]
+		[ProducesResponseType(typeof(GetServerInfoResponse), 200)]
+		public ActionResult<GetServerInfoResponse> GetServerInfo()
+		{
+			return new GetServerInfoResponse(this.Settings.CurrentValue.SingleInstance);
 		}
 	}
 }
diff --git a/Engine/Source/Programs/Horde/HordeServer/HordeServer.csproj b/Engine/Source/Programs/Horde/HordeServer/HordeServer.csproj
--- a/Engine/Source/Programs/Horde/HordeServer/HordeServer.csproj
+++ b/Engine/Source/Programs/Horde/HordeServer/HordeServer.csproj
@@ -29,6 +29,7 @@
     <PackageReference Include="Microsoft.AspNetCore.Authentication.JwtBearer" Version="3.1.5" />
     <PackageReference Include="Microsoft.AspNetCore.Authentication.OpenIdConnect" Version="3.1.5" />
     <PackageReference Include="Microsoft.AspNetCore.DataProtection.StackExchangeRedis" Version="3.1.7" />
+    <PackageReference Include="Microsoft.AspNetCore.SpaServices.Extensions" Version="3.1.17" />
     <PackageReference Include="Microsoft.CodeAnalysis.FxCopAnalyzers" Version="3.0.0">
       <PrivateAssets>all</PrivateAssets>
       <IncludeAssets>runtime; build; native; contentfiles; analyzers; buildtransitive</IncludeAssets>
diff --git a/Engine/Source/Programs/Horde/HordeServer/Program.cs b/Engine/Source/Programs/Horde/HordeServer/Program.cs
--- a/Engine/Source/Programs/Horde/HordeServer/Program.cs
+++ b/Engine/Source/Programs/Horde/HordeServer/Program.cs
@@ -2,6 +2,7 @@
 
 using System;
 using System.Collections.Generic;
+using System.Globalization;
 using System.IO;
 using System.Linq;
 using System.Net;
@@ -146,10 +147,11 @@ namespace HordeServer
 
 			IConfiguration Config = new ConfigurationBuilder()
 				.SetBasePath(AppDir.FullName)
-				.AddJsonFile("appsettings.json", optional: false)
-				.AddJsonFile($"appsettings.{Environment.GetEnvironmentVariable("ASPNETCORE_ENVIRONMENT")}.json", optional: true)
+				.AddJsonFile("appsettings.json", optional: false) 
+				.AddJsonFile("appsettings.Build.json", optional: true) // specific settings for builds (installer/dockerfile)
+				.AddJsonFile($"appsettings.{Environment.GetEnvironmentVariable("ASPNETCORE_ENVIRONMENT")}.json", optional: true) // environment variable overrides, also used in k8s setups with Helm
 				.AddJsonFile("appsettings.User.json", optional: true)
-				.AddJsonFile(UserConfigFile.FullName, optional: true)
+				.AddJsonFile(UserConfigFile.FullName, optional: true, reloadOnChange: true)
 				.AddEnvironmentVariables()
 				.AddCommandLine(Args)
 				.Build();
@@ -157,6 +159,8 @@ namespace HordeServer
 			ServerSettings HordeSettings = new ServerSettings();
 			Config.GetSection("Horde").Bind(HordeSettings);
 
+			InitializeDefaults(HordeSettings);
+
 			DirectoryReference LogDir = AppDir;
 			if (RuntimeInformation.IsOSPlatform(OSPlatform.Windows))
 			{
@@ -281,6 +285,28 @@ namespace HordeServer
 			}
 			return DirectoryReference.Combine(GetAppDir(), "Data");
 		}
+
+		static void InitializeDefaults(ServerSettings Settings)
+		{
+			if (Settings.SingleInstance)
+			{
+				FileReference GlobalConfig = FileReference.Combine(Program.DataDir, "Config/globals.json");
+
+				if (!FileReference.Exists(UserConfigFile))
+				{
+					DirectoryReference.CreateDirectory(UserConfigFile.Directory);
+					FileReference.WriteAllText(UserConfigFile, $"{{\"Horde\": {{ \"configPath\" : \"{GlobalConfig.ToString().Replace("\\", "/", StringComparison.Ordinal)}\"}}}}");
+				}
+				
+				if (!FileReference.Exists(GlobalConfig))
+				{
+					DirectoryReference.CreateDirectory(GlobalConfig.Directory);
+					FileReference.WriteAllText(GlobalConfig, "{}");
+				}
+
+			}
+
+		}
 	}
 }
 
diff --git a/Engine/Source/Programs/Horde/HordeServer/ServerSettings.cs b/Engine/Source/Programs/Horde/HordeServer/ServerSettings.cs
--- a/Engine/Source/Programs/Horde/HordeServer/ServerSettings.cs
+++ b/Engine/Source/Programs/Horde/HordeServer/ServerSettings.cs
@@ -123,6 +123,12 @@ namespace HordeServer
 	/// </summary>
 	public class ServerSettings
 	{
+
+		/// <summary>
+		/// Whether the server is running as a single instance or with multiple instances, such as in k8s
+		/// </summary>
+		public bool SingleInstance { get; set; } = false;
+
 		/// <summary>
 		/// MongoDB connection string
 		/// </summary>
diff --git a/Engine/Source/Programs/Horde/HordeServer/Services/ConfigService.cs b/Engine/Source/Programs/Horde/HordeServer/Services/ConfigService.cs
--- a/Engine/Source/Programs/Horde/HordeServer/Services/ConfigService.cs
+++ b/Engine/Source/Programs/Horde/HordeServer/Services/ConfigService.cs
@@ -8,19 +8,24 @@ using HordeServer.Models;
 using HordeServer.Notifications;
 using HordeServer.Utilities;
 using Microsoft.AspNetCore.StaticFiles;
+using Microsoft.Extensions.Configuration;
 using Microsoft.Extensions.Logging;
 using Microsoft.Extensions.Options;
 using MongoDB.Bson;
 using System;
 using System.Collections.Generic;
 using System.ComponentModel.DataAnnotations;
-using System.IO;
 using System.Linq;
+using System.IO;
 using System.Text;
 using System.Text.Json;
 using System.Text.RegularExpressions;
 using System.Threading;
 using System.Threading.Tasks;
+using System.Reflection;
+
+using PoolId = HordeServer.Utilities.StringId<HordeServer.Models.IPool>;
+using System.Globalization;
 
 namespace HordeServer.Services
 {
@@ -60,15 +65,26 @@ namespace HordeServer.Services
 		/// </summary>
 		IPerforceService PerforceService;
 
+		/// <summary>
+		/// Load balancer instance
+		/// </summary>
+		private PerforceLoadBalancer PerforceLoadBalancer;
+
+
 		/// <summary>
 		/// Instance of the notification service
 		/// </summary>
 		INotificationService NotificationService;
 
+		/// <summary>
+		/// Singleton instance of the pool service
+		/// </summary>
+		private readonly PoolService PoolService;
+
 		/// <summary>
 		/// The server settings
 		/// </summary>
-		ServerSettings Settings;
+		IOptionsMonitor<ServerSettings> Settings;
 
 		/// <summary>
 		/// Logging device
@@ -80,21 +96,28 @@ namespace HordeServer.Services
 		/// </summary>
 		/// <param name="DatabaseService"></param>
 		/// <param name="PerforceService"></param>
+		/// <param name="PerforceLoadBalancer"></param>
 		/// <param name="ProjectService"></param>
 		/// <param name="StreamService"></param>
 		/// <param name="NotificationService"></param>
+		/// <param name="PoolService"></param>
 		/// <param name="Settings"></param>
 		/// <param name="Logger"></param>
-		public ConfigService(DatabaseService DatabaseService, IPerforceService PerforceService, ProjectService ProjectService, StreamService StreamService, INotificationService NotificationService, IOptionsMonitor<ServerSettings> Settings, ILogger<ConfigService> Logger)
+		public ConfigService(DatabaseService DatabaseService, IPerforceService PerforceService, ProjectService ProjectService, StreamService StreamService, INotificationService NotificationService,  PoolService PoolService, PerforceLoadBalancer PerforceLoadBalancer, IOptionsMonitor<ServerSettings> Settings, ILogger<ConfigService> Logger)
 			: base(DatabaseService, new ObjectId("5ff60549e7632b15e64ac2f7"), Logger)
 		{
 			this.DatabaseService = DatabaseService;
 			this.PerforceService = PerforceService;
+			this.PerforceLoadBalancer = PerforceLoadBalancer;
 			this.ProjectService = ProjectService;
 			this.StreamService = StreamService;
 			this.NotificationService = NotificationService;
-			this.Settings = Settings.CurrentValue;
+			this.PoolService = PoolService;
+			this.Settings = Settings;
 			this.Logger = Logger;
+
+			// This will trigger if the local Horde.json user configuration is changed
+			this.Settings.OnChange(OnUserConfigUpdated);
 		}
 
 		GlobalConfig? CachedGlobalConfig;
@@ -420,12 +443,374 @@ namespace HordeServer.Services
 		/// <inheritdoc/>
 		protected override async Task<DateTime> TickLeaderAsync(CancellationToken StoppingToken)
 		{
-			if (Settings.ConfigPath != null)
+
+			Uri? ConfigUri = null;
+			
+			if (Path.IsPathRooted(Settings.CurrentValue.ConfigPath) && !Settings.CurrentValue.ConfigPath.StartsWith("//", StringComparison.Ordinal))
+			{
+				// absolute path to config
+				ConfigUri = new Uri(Settings.CurrentValue.ConfigPath);					
+			}
+			else if (Settings.CurrentValue.ConfigPath != null)
+			{
+				// relative (development) or perforce path
+				ConfigUri = CombinePaths(new Uri(FileReference.Combine(Program.AppDir, "_").FullName), Settings.CurrentValue.ConfigPath);				
+			}
+
+			if (ConfigUri != null)
 			{
-				Uri ConfigUri = CombinePaths(new Uri(FileReference.Combine(Program.AppDir, "_").FullName), Settings.ConfigPath);
 				await UpdateConfigAsync(ConfigUri);
 			}
+		
 			return DateTime.UtcNow + TimeSpan.FromMinutes(1.0);
 		}
+
+		//
+		// On premises configuration handling (aka Horde.json settings, though needs to have perforce global config handled too, checkout/modify/submit)
+		//
+
+		/// <summary>
+		/// Update the global configuration
+		/// </summary>
+		/// <param name="Request"></param>
+		/// <returns></returns>
+		public async Task<bool> UpdateGlobalConfig(UpdateGlobalConfigRequest Request)
+		{
+			if (Settings.CurrentValue.ConfigPath == null || !Path.IsPathRooted(Settings.CurrentValue.ConfigPath) || Settings.CurrentValue.ConfigPath.StartsWith("//", StringComparison.Ordinal))
+			{
+				throw new Exception($"Global config path must be rooted for service updates.  (Perforce paths are not currently supported for live updates), {Settings.CurrentValue.ConfigPath}");
+			}
+			
+			FileReference GlobalConfigFile = new FileReference(Settings.CurrentValue.ConfigPath);
+			DirectoryReference GlobalConfigDirectory = GlobalConfigFile.Directory;
+
+			// make sure the directory exists
+			if (!DirectoryReference.Exists(GlobalConfigDirectory))
+			{
+				DirectoryReference.CreateDirectory(GlobalConfigDirectory);
+			}
+
+			bool ConfigDirty = false;
+
+			// projects
+			if (Request.ProjectsJson != null)
+			{
+				ConfigDirty = true;
+
+				// write out the projects
+				foreach (KeyValuePair<string, string> Project in Request.ProjectsJson)
+				{
+					FileReference.WriteAllText(FileReference.Combine(GlobalConfigDirectory, Project.Key), Project.Value);
+
+					if (Request.ProjectLogo != null)
+					{
+						Byte[] bytes = Convert.FromBase64String(Request.ProjectLogo);
+						FileReference.WriteAllBytes(FileReference.Combine(GlobalConfigDirectory, Project.Key.Replace(".json", ".png", StringComparison.OrdinalIgnoreCase)), bytes);
+					}
+				}
+
+			}
+
+			// streams
+			if (Request.StreamsJson != null)
+			{
+				ConfigDirty = true;
+
+				// write out the streams
+				foreach (KeyValuePair<string, string> Stream in Request.StreamsJson)
+				{
+					FileReference.WriteAllText(FileReference.Combine(GlobalConfigDirectory, Stream.Key), Stream.Value);
+				}
+			}
+
+			// update global config path
+			if (!string.IsNullOrEmpty(Request.GlobalsJson))
+			{
+				ConfigDirty = true;
+
+				FileReference.WriteAllText(GlobalConfigFile, Request.GlobalsJson);
+			}
+
+			if (ConfigDirty == true)
+			{				
+				await UpdateConfigAsync(new Uri("file://" + GlobalConfigFile.ToString()));
+
+				// wait for the perforce clusters to be picked up, this ticks at a minute and might be nice to be able to explicitly update this
+				// the PerforceLoadBalancer has an ElectedTick member and no TickNow method
+				// In any event, this delay helps initial feel as cluster will be available shortly after it
+				int WaitSec = 15;				
+				DateTime Start = DateTime.UtcNow;
+				do
+				{
+					await Task.Delay(1000);
+
+					List<IPerforceServer> Servers = await PerforceLoadBalancer.GetServersAsync();
+					if (Servers.Count > 0)
+					{
+						break;
+					}
+				}
+				while ((DateTime.UtcNow - Start).TotalSeconds < WaitSec);
+
+			}
+
+			// create the default pool 
+			if (Request.DefaultPoolName != null)
+			{
+				PoolId PoolIdValue = PoolId.Sanitize(Request.DefaultPoolName);
+
+				IPool? Pool = await PoolService.GetPoolAsync(PoolIdValue);
+				if (Pool == null)
+				{
+					await PoolService.CreatePoolAsync(Request.DefaultPoolName, Properties: new Dictionary<string, string>() { ["Color"] = "0" });
+				}
+			}
+
+			return true;
+
+		}
+
+		/// <summary>
+		/// Update the server settings
+		/// </summary>
+		/// <param name="Request"></param>
+		/// <returns></returns>
+		public async Task<ServerUpdateResponse> UpdateServerSettings(UpdateServerSettingsRequest Request)
+		{
+
+			ServerUpdateResponse Response = new ServerUpdateResponse();
+
+			Dictionary<string, object> NewUserSettings = new Dictionary<string, object>();
+
+			try
+			{
+
+				if (Request.Settings == null || Request.Settings.Count == 0)
+				{
+					return Response;
+				}
+
+				if (UserConfigUpdated != null)
+				{
+					Response.Errors.Add("User configuation already being updated");
+					return Response;
+				}
+
+				// Load the current user configuration
+				IConfiguration Config = new ConfigurationBuilder()
+					.AddJsonFile(Program.UserConfigFile.FullName, optional: true)
+					.Build();
+
+				ServerSettings UserSettings = new ServerSettings();
+				Config.GetSection("Horde").Bind(UserSettings);
+
+				// Figure out current and new settings
+				HashSet<string> UserProps = new HashSet<string>(StringComparer.OrdinalIgnoreCase);
+
+				if (FileReference.Exists(Program.UserConfigFile))
+				{
+					byte[] Data = await FileReference.ReadAllBytesAsync(Program.UserConfigFile);
+
+					JsonDocument Document = JsonDocument.Parse(Data);
+					JsonElement HordeElement;
+
+					if (Document.RootElement.TryGetProperty("Horde", out HordeElement))
+					{
+						foreach (JsonProperty Property in HordeElement.EnumerateObject())
+						{
+							UserProps.Add(Property.Name);
+						}
+					}
+				}
+
+				foreach (string Name in Request.Settings.Keys)
+				{
+					UserProps.Add(Name);
+				}
+
+				// Apply the changes from the request
+				foreach (KeyValuePair<string, object> Pair in Request.Settings)
+				{
+					PropertyInfo? Property = UserSettings.GetType().GetProperty(Pair.Key, BindingFlags.IgnoreCase | BindingFlags.Public | BindingFlags.Instance);
+					if (Property == null)
+					{
+						Response.Errors.Add($"Horde configuration property {Pair.Key} does not exist when reading server settings");
+						Logger.LogError(Response.Errors.Last());
+						continue;
+					}
+
+					if (Property.SetMethod == null)
+					{
+						Response.Errors.Add($"Horde configuration property {Pair.Key} does not have a set method");
+						Logger.LogError(Response.Errors.Last());
+						continue;
+					}
+
+					// explicit setting removal
+					if (Pair.Value == null)
+					{
+						UserProps.Remove(Pair.Key);
+						continue;
+					}					
+
+
+					JsonElement Element = (JsonElement)Pair.Value;
+
+					object? Value = null;
+
+					switch (Element.ValueKind)
+					{
+						case JsonValueKind.True:
+							Value = true;
+							break;
+						case JsonValueKind.False:
+							Value = false;
+							break;
+						case JsonValueKind.Number:
+							Value = Element.GetDouble();
+							break;
+						case JsonValueKind.String:
+							Value = Element.GetString();
+							break;
+					}
+
+					// unable to map type
+					if (Value == null)
+					{
+						Response.Errors.Add($"Unable to map type for Property {Property.Name}");
+						continue;
+					}					
+
+					// handle common conversions
+					if (Property.GetType() != Value.GetType())
+					{
+						try
+						{
+							Value =	Convert.ChangeType(Value, Property.PropertyType, CultureInfo.CurrentCulture);						
+						}
+						catch (Exception Ex)
+						{
+							Response.Errors.Add($"Property {Property.Name} raised exception during conversion, {Ex.Message}");
+							continue;							
+						}
+						
+						if (Value == null)
+						{
+							Response.Errors.Add($"Property {Property.Name} had null value on conversion");
+							continue;							
+						}
+					}
+					else 
+					{
+						Response.Errors.Add($"Property {Property.Name} is not assignable to {Value.ToString()}");
+						continue;
+					}
+					
+
+					//  Set the value, providing some validation
+					try
+					{
+						Property.SetMethod.Invoke(UserSettings, new object[] { Value });
+					}
+					catch (Exception Ex)
+					{
+						Response.Errors.Add($"Exception updating property {Pair.Key}, {Ex.Message}");
+					}
+					
+				}
+
+				// Construct the new user settings
+				foreach (string Name in UserProps)
+				{
+					PropertyInfo? Property = UserSettings.GetType().GetProperty(Name, BindingFlags.IgnoreCase | BindingFlags.Public | BindingFlags.Instance);
+					if (Property == null)
+					{
+						Response.Errors.Add($"Horde configuration property {Name} does not exist when writing server settings");
+						Logger.LogError(Response.Errors.Last());
+
+						continue;
+					}
+
+					if (Property.GetMethod == null)
+					{
+						Response.Errors.Add($"Horde configuration property {Name} does not have a get method");
+						Logger.LogError(Response.Errors.Last());
+						continue;
+					}
+
+					object? Result = Property.GetMethod.Invoke(UserSettings, null);
+
+					if (Result == null)
+					{
+						Response.Errors.Add($"Horde configuration property {Name} was null while writing and should have already been filtered out");
+						Logger.LogError(Response.Errors.Last());
+						continue;
+					}
+
+					NewUserSettings.Add(Property.Name, Result);
+
+				}
+			}
+			catch (Exception Ex)
+			{
+				Response.Errors.Add($"Exception while updating settings: {Ex.Message}");
+				Logger.LogError("{0}: {1}", Response.Errors.Last(), Ex.StackTrace ?? "");
+			}
+
+			if (Response.Errors.Count != 0)
+			{
+				return Response;
+			}
+
+			Dictionary<string, object> NewLocalSettings = new Dictionary<string, object>();
+			NewLocalSettings["Horde"] = NewUserSettings;
+
+			try
+			{
+				UserConfigUpdated = new TaskCompletionSource<bool>();
+
+				// This will trigger a setting update as the user config json is set to reload on change
+				try
+				{
+					await FileReference.WriteAllBytesAsync(Program.UserConfigFile, JsonSerializer.SerializeToUtf8Bytes(NewLocalSettings, new JsonSerializerOptions { WriteIndented = true }));
+				}
+				catch (Exception Ex)
+				{
+					Response.Errors.Add($"Unable to serialize json settings to {Program.UserConfigFile.ToString()}, {Ex.Message}");
+					Logger.LogError("Unable to serialize json settings to {0}, {1} : {2}", Program.UserConfigFile.ToString(), Ex.Message, Ex.StackTrace ?? "");
+				}
+
+				if (Response.Errors.Count == 0)
+				{
+					if (await Task.WhenAny(UserConfigUpdated.Task, Task.Delay(5000)) != UserConfigUpdated.Task)
+					{
+						Response.Errors.Add("Server update timed out while awaiting write");
+						Logger.LogError("Server update timed out after writing config");
+					}
+				}
+
+			}
+			finally
+			{
+				UserConfigUpdated = null;
+			}
+
+			return Response;
+
+		}
+		void OnUserConfigUpdated(ServerSettings Settings, string Name)
+		{
+			NumUserConfigUpdates++;
+			UserConfigUpdated?.SetResult(true);
+		}
+
+		private TaskCompletionSource<bool>? UserConfigUpdated;
+
+		/// <summary>
+		/// The number of server configuration updates that have been made while the server is running
+		/// This is used to detect whether server may need to be restarted
+		/// </summary>
+		public int NumUserConfigUpdates { get; private set; } = 0;
+
 	}
 }
diff --git a/Engine/Source/Programs/Horde/HordeServer/Startup.cs b/Engine/Source/Programs/Horde/HordeServer/Startup.cs
--- a/Engine/Source/Programs/Horde/HordeServer/Startup.cs
+++ b/Engine/Source/Programs/Horde/HordeServer/Startup.cs
@@ -253,15 +253,15 @@ namespace HordeServer
 
 		public IConfiguration Configuration { get; }
 
-		ServerSettings Settings { get; set; } = new ServerSettings();
-
-		// This method gets called by the runtime. Use this method to add services to the container.
+		// This method gets called *multiple times* by the runtime. Use this method to add services to the container.
 		public void ConfigureServices(IServiceCollection Services)
 		{
-			IConfigurationSection ConfigSection = Configuration.GetSection("Horde");
-			Services.AddOptions<ServerSettings>().Configure(Options => ConfigSection.Bind(Options)).ValidateDataAnnotations();
+			// IOptionsMonitor pattern for live updating of configuration settings
+			Services.Configure<ServerSettings>(Configuration.GetSection("Horde"));
 
-			ConfigSection.Bind(Settings);
+			// Settings used for configuring services
+			ServerSettings Settings = new ServerSettings();
+			Configuration.GetSection("Horde").Bind(Settings);			
 
 			if (Settings.GlobalThreadPoolMinSize != null)
 			{
@@ -350,13 +350,14 @@ namespace HordeServer
 			}
 
 			Services.AddSingleton<AclService>();
-			Services.AddSingleton<AgentService>();
+			Services.AddSingleton<AgentService>();			
 			Services.AddSingleton<AgentSoftwareService>();
 			Services.AddSingleton<ArtifactService>();
 			Services.AddSingleton<ConsistencyService>();
 			Services.AddSingleton<RequestTrackerService>();
 			Services.AddSingleton<CredentialService>();
 			Services.AddSingleton<DatabaseService>();
+			Services.AddSingleton<ConfigService>();
 			Services.AddSingleton<IDogStatsd>(ctx =>
 			{
 				string? DatadogAgentHost = Environment.GetEnvironmentVariable("DD_AGENT_HOST");
@@ -572,6 +573,12 @@ namespace HordeServer
 				Config.IncludeXmlComments(Path.Combine(AppContext.BaseDirectory, $"{Assembly.GetExecutingAssembly().GetName().Name}.xml"));
 			});
 
+			DirectoryReference DashboardDir = DirectoryReference.Combine(Program.AppDir, "DashboardApp");
+			if (DirectoryReference.Exists(DashboardDir)) 
+			{
+				Services.AddSpaStaticFiles(Config => {Config.RootPath = "DashboardApp";});
+			}
+
 			ConfigureMongoDbClient();
 
 			OnAddHealthChecks(Services);
@@ -719,7 +726,7 @@ namespace HordeServer
 		}
 
 		// This method gets called by the runtime. Use this method to configure the HTTP request pipeline.
-		public void Configure(IApplicationBuilder App, IWebHostEnvironment Env, Microsoft.Extensions.Hosting.IHostApplicationLifetime Lifetime)
+		public static void Configure(IApplicationBuilder App, IWebHostEnvironment Env, Microsoft.Extensions.Hosting.IHostApplicationLifetime Lifetime, IOptions<ServerSettings> Settings)
 		{
 			App.UseForwardedHeaders();
 
@@ -728,7 +735,7 @@ namespace HordeServer
 				App.UseDeveloperExceptionPage();
 			}
 
-			if (Settings.CorsEnabled)
+			if (Settings.Value.CorsEnabled)
 			{
 				App.UseCors("CorsPolicy");
 			}
@@ -762,6 +769,14 @@ namespace HordeServer
 
 			App.UseDefaultFiles();
 			App.UseStaticFiles();
+
+			DirectoryReference DashboardDir = DirectoryReference.Combine(Program.AppDir, "DashboardApp");
+
+			if (DirectoryReference.Exists(DashboardDir)) 
+			{
+				App.UseSpaStaticFiles();
+			}
+						
 			App.UseRouting();
 
 			App.UseAuthentication();
@@ -788,7 +803,15 @@ namespace HordeServer
 				Endpoints.MapControllers();
 			});
 
-			if (Settings.OpenBrowser && RuntimeInformation.IsOSPlatform(OSPlatform.Windows))
+			if (DirectoryReference.Exists(DashboardDir)) 
+			{
+				App.UseSpa(Spa =>
+				{ 
+					Spa.Options.SourcePath = "DashboardApp";        
+				});
+			}
+
+			if (Settings.Value.OpenBrowser && RuntimeInformation.IsOSPlatform(OSPlatform.Windows))
 			{
 				Lifetime.ApplicationStarted.Register(() => LaunchBrowser(App));
 			}
diff --git a/Engine/Source/Programs/Horde/Installer/Installer.xml b/Engine/Source/Programs/Horde/Installer/Installer.xml
--- a/Engine/Source/Programs/Horde/Installer/Installer.xml
+++ b/Engine/Source/Programs/Horde/Installer/Installer.xml
@@ -11,11 +11,17 @@
 			<Delete Files="$(StagingDir)/..."/>
 			<DotNet Arguments="publish &quot;$(RootDir)/Engine/Source/Programs/Horde/HordeServer/HordeServer.csproj&quot; --output &quot;$(StagingDir)/Server&quot; --runtime win-x64 $(VersionArguments)"/>
 			<DotNet Arguments="publish &quot;$(RootDir)/Engine/Source/Programs/Horde/HordeAgent/HordeAgent.csproj&quot; --output &quot;$(StagingDir)/Agent&quot; $(VersionArguments)"/>
+
+			<!-- Do not include development defaults which created a ue5 project and stream -->
+			<Delete Files="$(StagingDir)/Server/Defaults/..."/>
+
+			<!-- Create build settings -->
+			<WriteTextFile File="$(StagingDir)/Server/appsettings.Build.json" Text="{ &quot;Horde&quot;: { &quot;SingleInstance&quot;: true } }"/>
 	
 			<Spawn Exe="cmd.exe" Arguments="/C yarn install" WorkingDir="$(RootDir)/Engine/Source/Programs/Horde/HordeDashboard"/>
 			<Spawn Exe="cmd.exe" Arguments="/C yarn run build" WorkingDir="$(RootDir)/Engine/Source/Programs/Horde/HordeDashboard"/>
 	
-			<Copy From="$(RootDir)/Engine/Source/Programs/Horde/HordeDashboard/build/..." To="$(StagingDir)/Server/wwwroot/..."/>
+			<Copy From="$(RootDir)/Engine/Source/Programs/Horde/HordeDashboard/build/..." To="$(StagingDir)/Server/DashboardApp/..."/>
 		</Node>
 
 		<Property Name="WixDir" Value="$(RootDir)/Engine/Source/ThirdParty/WiX/3.8"/>
