Horde: Add per-file data to changelist details.

diff --git a/Engine/Source/Programs/Horde/HordeServer/Api/Changes.cs b/Engine/Source/Programs/Horde/HordeServer/Api/Changes.cs
--- a/Engine/Source/Programs/Horde/HordeServer/Api/Changes.cs
+++ b/Engine/Source/Programs/Horde/HordeServer/Api/Changes.cs
@@ -93,7 +93,7 @@ namespace HordeServer.Api
 			this.Number = ChangeDetails.Number;
 			this.Author = ChangeDetails.Author;
 			this.Description = ChangeDetails.Description;
-			this.Files = ChangeDetails.Files;
+			this.Files = ChangeDetails.Files.ConvertAll(x => x.Path);
 		}
 	}
 }
diff --git a/Engine/Source/Programs/Horde/HordeServer/Commits/Impl/CommitService.cs b/Engine/Source/Programs/Horde/HordeServer/Commits/Impl/CommitService.cs
--- a/Engine/Source/Programs/Horde/HordeServer/Commits/Impl/CommitService.cs
+++ b/Engine/Source/Programs/Horde/HordeServer/Commits/Impl/CommitService.cs
@@ -190,7 +190,7 @@ namespace HordeServer.Commits.Impl
 		/// <returns></returns>
 		async Task AddCommitAsync(IStream Stream, int Number)
 		{
-			ChangeDetails Details = await PerforceService.GetChangeDetailsAsync(Stream.ClusterName, Number);
+			ChangeDetails Details = await PerforceService.GetChangeDetailsAsync(Stream.ClusterName, Stream.Name, Number);
 
 			IUser Author = await UserCollection.FindOrAddUserByLoginAsync(Details.Author);
 			int? OriginalChange = ParseRobomergeSource(Details.Description);
diff --git a/Engine/Source/Programs/Horde/HordeServer/IssueHandlers/IIssueHandler.cs b/Engine/Source/Programs/Horde/HordeServer/IssueHandlers/IIssueHandler.cs
--- a/Engine/Source/Programs/Horde/HordeServer/IssueHandlers/IIssueHandler.cs
+++ b/Engine/Source/Programs/Horde/HordeServer/IssueHandlers/IIssueHandler.cs
@@ -78,12 +78,12 @@ namespace HordeServer.IssueHandlers
 		{
 			this.Details = Details;
 
-			foreach (string Path in Details.Files)
+			foreach (ChangeFile File in Details.Files)
 			{
-				int Idx = Path.LastIndexOf('.');
+				int Idx = File.Path.LastIndexOf('.');
 				if (Idx != -1)
 				{
-					StringView Extension = new StringView(Path, Idx);
+					StringView Extension = new StringView(File.Path, Idx);
 					if (CodeExtensions.Contains(Extension))
 					{
 						ContainsCode = true;
@@ -107,9 +107,9 @@ namespace HordeServer.IssueHandlers
 		/// <returns>True if the change modifies the given file</returns>
 		public bool ModifiesFile(string FileToCheck)
 		{
-			foreach (string File in Details.Files)
+			foreach (ChangeFile File in Details.Files)
 			{
-				if (File.EndsWith(FileToCheck, StringComparison.OrdinalIgnoreCase) && (File.Length == FileToCheck.Length || File[File.Length - FileToCheck.Length - 1] == '/'))
+				if (File.Path.EndsWith(FileToCheck, StringComparison.OrdinalIgnoreCase) && (File.Length == FileToCheck.Length || File.Path[File.Path.Length - FileToCheck.Length - 1] == '/'))
 				{
 					return true;
 				}
diff --git a/Engine/Source/Programs/Horde/HordeServer/IssueHandlers/Impl/ContentIssueHandler.cs b/Engine/Source/Programs/Horde/HordeServer/IssueHandlers/Impl/ContentIssueHandler.cs
--- a/Engine/Source/Programs/Horde/HordeServer/IssueHandlers/Impl/ContentIssueHandler.cs
+++ b/Engine/Source/Programs/Horde/HordeServer/IssueHandlers/Impl/ContentIssueHandler.cs
@@ -90,7 +90,7 @@ namespace HordeServer.IssueHandlers.Impl
 			{
 				if (Suspect.ContainsContent)
 				{
-					if (Suspect.Details.Files.Any(x => Fingerprint.Keys.Any(y => x.Contains(y, StringComparison.OrdinalIgnoreCase))))
+					if (Suspect.Details.Files.Any(x => Fingerprint.Keys.Any(y => x.Path.Contains(y, StringComparison.OrdinalIgnoreCase))))
 					{
 						Suspect.Rank += 20;
 					}
diff --git a/Engine/Source/Programs/Horde/HordeServer/IssueHandlers/Impl/SymbolIssueHandler.cs b/Engine/Source/Programs/Horde/HordeServer/IssueHandlers/Impl/SymbolIssueHandler.cs
--- a/Engine/Source/Programs/Horde/HordeServer/IssueHandlers/Impl/SymbolIssueHandler.cs
+++ b/Engine/Source/Programs/Horde/HordeServer/IssueHandlers/Impl/SymbolIssueHandler.cs
@@ -68,7 +68,7 @@ namespace HordeServer.IssueHandlers.Impl
 			{
 				if (Change.ContainsCode)
 				{
-					int Matches = Names.Count(x => Change.Details.Files.Any(y => y.Contains(x, StringComparison.OrdinalIgnoreCase)));
+					int Matches = Names.Count(x => Change.Details.Files.Any(y => y.Path.Contains(x, StringComparison.OrdinalIgnoreCase)));
 					Change.Rank += 10 + (10 * Matches);
 				}
 			}
diff --git a/Engine/Source/Programs/Horde/HordeServer/Models/Change.cs b/Engine/Source/Programs/Horde/HordeServer/Models/Change.cs
--- a/Engine/Source/Programs/Horde/HordeServer/Models/Change.cs
+++ b/Engine/Source/Programs/Horde/HordeServer/Models/Change.cs
@@ -1,5 +1,6 @@
 // Copyright Epic Games, Inc. All Rights Reserved.
 
+using EpicGames.Core;
 using HordeServer.Api;
 using HordeServer.Utilities;
 using MongoDB.Bson;
@@ -74,6 +75,61 @@ namespace HordeServer.Models
 		ContainsContent = 2,
 	}
 
+	/// <summary>
+	/// Modified file in a changelist
+	/// </summary>
+	public class ChangeFile
+	{
+		/// <summary>
+		/// Path to the file
+		/// </summary>
+		public string Path { get; set; }
+
+		/// <summary>
+		/// Path to the file within the depot
+		/// </summary>
+		public string DepotPath { get; set; }
+
+		/// <summary>
+		/// Revision of the file. A value of -1 indicates that the file was deleted.
+		/// </summary>
+		public int Revision { get; set; }
+
+		/// <summary>
+		/// Length of the file
+		/// </summary>
+		public long Length { get; set; }
+
+		/// <summary>
+		/// MD5 digest of the file
+		/// </summary>
+		public Md5Hash Digest { get; set; }
+
+		/// <summary>
+		/// The file type
+		/// </summary>
+		public string Type { get; set; }
+
+		/// <summary>
+		/// Constructor
+		/// </summary>
+		/// <param name="Path"></param>
+		/// <param name="DepotPath"></param>
+		/// <param name="Revision"></param>
+		/// <param name="Length"></param>
+		/// <param name="Digest"></param>
+		/// <param name="Type"></param>
+		public ChangeFile(string Path, string DepotPath, int Revision, long Length, Md5Hash Digest, string Type)
+		{
+			this.Path = Path;
+			this.DepotPath = DepotPath;
+			this.Revision = Revision;
+			this.Length = Length;
+			this.Digest = Digest;
+			this.Type = Type;
+		}
+	}
+
 	/// <summary>
 	/// Information about a commit
 	/// </summary>
@@ -126,7 +182,7 @@ namespace HordeServer.Models
 		/// <summary>
 		/// List of files that were modified, relative to the stream base
 		/// </summary>
-		public List<string> Files { get; set; }
+		public List<ChangeFile> Files { get; set; }
 
 		/// <summary>
 		/// Date that the change was submitted
@@ -153,7 +209,7 @@ namespace HordeServer.Models
 		/// <param name="Description">Changelist description</param>
 		/// <param name="Files">List of files modified, relative to the stream base</param>
 		/// <param name="Date">Date that the change was submitted</param>
-		public ChangeDetails(int Number, string Author, string Path, string Description, List<string> Files, DateTime Date)
+		public ChangeDetails(int Number, string Author, string Path, string Description, List<ChangeFile> Files, DateTime Date)
 		{
 			this.Number = Number;
 			this.Author = Author;
@@ -172,9 +228,9 @@ namespace HordeServer.Models
 			ChangeContentFlags Scope = 0;
 
 			// Check whether the files are code or content
-			foreach (string File in Files)
+			foreach (ChangeFile File in Files)
 			{
-				if (CodeExtensions.Any(Extension => File.EndsWith(Extension, StringComparison.OrdinalIgnoreCase)))
+				if (CodeExtensions.Any(Extension => File.Path.EndsWith(Extension, StringComparison.OrdinalIgnoreCase)))
 				{
 					Scope |= ChangeContentFlags.ContainsCode;
 				}
diff --git a/Engine/Source/Programs/Horde/HordeServer/Services/ConfigService.cs b/Engine/Source/Programs/Horde/HordeServer/Services/ConfigService.cs
--- a/Engine/Source/Programs/Horde/HordeServer/Services/ConfigService.cs
+++ b/Engine/Source/Programs/Horde/HordeServer/Services/ConfigService.cs
@@ -428,8 +428,11 @@ namespace HordeServer.Services
 					List<FileSummary> Files = await PerforceService.FindFilesAsync(ConfigPath.Host, new[] { FileName });
 					Change = Files[0].Change;
 
-					ChangeDetails Details = await PerforceService.GetChangeDetailsAsync(ConfigPath.Host, Change);
-					(Author, Description) = (Details.Author, Details.Description);
+					List<ChangeSummary> Changes = await PerforceService.GetChangesAsync(ConfigPath.Host, Change, 1);
+					if (Changes.Count > 0 && Changes[0].Number == Change)
+					{
+						(Author, Description) = (Changes[0].Author, Changes[0].Description);
+					}
 				}
 				catch (Exception Ex2)
 				{
diff --git a/Engine/Source/Programs/Horde/HordeServer/Services/IPerforceService.cs b/Engine/Source/Programs/Horde/HordeServer/Services/IPerforceService.cs
--- a/Engine/Source/Programs/Horde/HordeServer/Services/IPerforceService.cs
+++ b/Engine/Source/Programs/Horde/HordeServer/Services/IPerforceService.cs
@@ -61,11 +61,20 @@ namespace HordeServer.Services
 		/// <summary>
 		/// Gets the definition of a stream
 		/// </summary>
-		/// <param name="ClusterName"></param>
+		/// <param name="ClusterName">Name of the Perforce cluster</param>
 		/// <param name="StreamName">The stream name</param>
 		/// <returns></returns>
 		public Task<IStreamView> GetStreamViewAsync(string ClusterName, string StreamName);
 
+		/// <summary>
+		/// Gets the state of all files in a stream
+		/// </summary>
+		/// <param name="ClusterName">Name of the Perforce cluster</param>
+		/// <param name="StreamName">The stream name</param>
+		/// <param name="Change">The changelist to query</param>
+		/// <returns></returns>
+		public Task<List<ChangeFile>> GetStreamSnapshotAsync(string ClusterName, string StreamName, int Change);
+
 		/// <summary>
 		/// Create a new changelist by submitting the given file
 		/// </summary>
@@ -117,9 +126,10 @@ namespace HordeServer.Services
 		/// Gets the latest change for a particular stream
 		/// </summary>
 		/// <param name="ClusterName">Name of the Perforce cluster</param>
+		/// <param name="StreamName">The stream to query</param>
 		/// <param name="ChangeNumber">Change numbers to query</param>
 		/// <returns>Commit details</returns>
-		public Task<ChangeDetails> GetChangeDetailsAsync(string ClusterName, int ChangeNumber);
+		public Task<ChangeDetails> GetChangeDetailsAsync(string ClusterName, string StreamName, int ChangeNumber);
 
 		/// <summary>
 		/// Gets the latest change for a particular stream
diff --git a/Engine/Source/Programs/Horde/HordeServer/Services/Impl/PerforceService.cs b/Engine/Source/Programs/Horde/HordeServer/Services/Impl/PerforceService.cs
--- a/Engine/Source/Programs/Horde/HordeServer/Services/Impl/PerforceService.cs
+++ b/Engine/Source/Programs/Horde/HordeServer/Services/Impl/PerforceService.cs
@@ -1,26 +1,27 @@
 // Copyright Epic Games, Inc. All Rights Reserved.
 
+using EpicGames.Core;
 using HordeServer.Models;
-using System.Linq;
+using HordeServer.Utilities;
 using Microsoft.Extensions.Logging;
+using Microsoft.Extensions.Options;
 using System;
 using System.Collections.Generic;
-using System.Net;
-using System.Text;
-using System.Threading.Tasks;
+using System.Diagnostics.CodeAnalysis;
+using System.Globalization;
 using System.IO;
-
-using P4 = Perforce.P4;
-using Microsoft.Extensions.Options;
+using System.Linq;
+using System.Net;
+using System.Reflection;
 using System.Runtime.InteropServices;
-using System.Globalization;
+using System.Text;
 using System.Text.RegularExpressions;
-using System.Diagnostics.CodeAnalysis;
-using HordeServer.Utilities;
-using System.Reflection;
+using System.Threading.Tasks;
 
 namespace HordeServer.Services
 {
+	using P4 = Perforce.P4;
+
 	static class PerforceExtensions
 	{
 		static FieldInfo Field = typeof(P4.Changelist).GetField("_baseForm", BindingFlags.Instance | BindingFlags.NonPublic)!;
@@ -550,6 +551,81 @@ namespace HordeServer.Services
 			}
 		}
 
+		static int GetSyncRevision(string Path, P4.FileAction HeadAction, int HeadRev)
+		{
+			switch (HeadAction)
+			{
+				case P4.FileAction.None:
+				case P4.FileAction.Add:
+				case P4.FileAction.Branch:
+				case P4.FileAction.MoveAdd:
+				case P4.FileAction.Edit:
+				case P4.FileAction.Integrate:
+					return HeadRev;
+				case P4.FileAction.Delete:
+				case P4.FileAction.MoveDelete:
+					return -1;
+				default:
+					throw new Exception($"Unrecognized P4 file change type '{HeadAction}' for file {Path}#{HeadRev}");
+			}
+		}
+
+		/// <summary>
+		/// Creates a <see cref="ChangeFile"/> from a <see cref="P4.FileMetaData"/>
+		/// </summary>
+		/// <param name="RelativePath"></param>
+		/// <param name="MetaData"></param>
+		/// <returns></returns>
+		static ChangeFile CreateChangeFile(string RelativePath, P4.FileMetaData MetaData)
+		{
+			int Revision = GetSyncRevision(MetaData.DepotPath.Path, MetaData.HeadAction, MetaData.HeadRev);
+			return new ChangeFile(RelativePath, MetaData.DepotPath.Path, Revision, MetaData.FileSize, Md5Hash.Parse(MetaData.Digest ?? String.Empty), MetaData.HeadType.ToString());
+		}
+
+		/// <summary>
+		/// Creates a <see cref="ChangeFile"/> from a <see cref="P4.FileMetaData"/>
+		/// </summary>
+		/// <param name="RelativePath"></param>
+		/// <param name="MetaData"></param>
+		/// <returns></returns>
+		static ChangeFile CreateChangeFile(string RelativePath, P4.ShelvedFile MetaData)
+		{
+			int Revision = GetSyncRevision(MetaData.Path.Path, MetaData.Action, MetaData.Revision);
+			return new ChangeFile(RelativePath, MetaData.Path.Path, Revision, MetaData.Size, Md5Hash.Parse(MetaData.Digest ?? String.Empty), MetaData.Type.ToString());
+		}
+
+		/// <inheritdoc/>
+		public async Task<List<ChangeFile>> GetStreamSnapshotAsync(string ClusterName, string StreamName, int Change)
+		{
+			PerforceCluster Cluster = await GetClusterAsync(ClusterName);
+			using (P4.Repository Repository = await GetConnection(Cluster, Stream: StreamName))
+			{
+				P4.FileSpec FileSpec = new P4.FileSpec(new P4.DepotPath($"//UE5/Main/Engine/Build/..."), new P4.ChangelistIdVersion(Change));
+				Repository.Connection.Client.SyncFiles(new P4.SyncFilesCmdOptions(P4.SyncFilesCmdFlags.ServerOnly | P4.SyncFilesCmdFlags.Quiet), FileSpec);
+
+				string ClientPrefix = $"//{Repository.Connection.Client.Name}/";
+				P4.FileSpec StatSpec = new P4.FileSpec(new P4.ClientPath($"{ClientPrefix}..."));
+
+				IList<P4.FileMetaData> Files = Repository.GetFileMetaData(new P4.GetFileMetaDataCmdOptions(P4.GetFileMetadataCmdFlags.Synced | P4.GetFileMetadataCmdFlags.LocalPath | P4.GetFileMetadataCmdFlags.FileSize, null, null, 0, null, null, null), StatSpec);
+
+				List<ChangeFile> Results = new List<ChangeFile>();
+				foreach (P4.FileMetaData File in Files)
+				{
+					string RelativePath = File.ClientPath.Path;
+					if (RelativePath.StartsWith(ClientPrefix, StringComparison.OrdinalIgnoreCase))
+					{
+						RelativePath = RelativePath.Substring(ClientPrefix.Length);
+					}
+					else
+					{
+						throw new Exception($"Client path does not start with client name: {RelativePath}");
+					}
+					Results.Add(CreateChangeFile(RelativePath, File));
+				}
+				return Results;
+			}
+		}
+
 		/// <inheritdoc/>
 		public async Task<List<ChangeSummary>> GetChangesAsync(string ClusterName, int? MinChange, int MaxResults)
 		{
@@ -610,7 +686,7 @@ namespace HordeServer.Services
 		}
 
 		/// <inheritdoc/>
-		public async Task<ChangeDetails> GetChangeDetailsAsync(string ClusterName, int ChangeNumber)
+		public async Task<ChangeDetails> GetChangeDetailsAsync(string ClusterName, string StreamName, int ChangeNumber)
 		{
 			PerforceCluster Cluster = await GetClusterAsync(ClusterName);
 			using (P4.Repository Repository = await GetConnection(Cluster, NoClient: true))
@@ -642,14 +718,18 @@ namespace HordeServer.Services
 
 					P4.TaggedObject TaggedObject = Result.TaggedOutput[0];
 					{
-						List<string> Files = new List<string>();
+						List<ChangeFile> Files = new List<ChangeFile>();
 
 						P4.Changelist Change = new P4.Changelist();
 						Change.FromChangeCmdTaggedOutput(TaggedObject, true, Offset, DSTMismatch);
 
 						foreach (P4.FileMetaData DescribeFile in Change.Files)
 						{
-							Files.Add(DescribeFile.DepotPath.Path);
+							string? RelativePath;
+							if (TryGetStreamRelativePath(DescribeFile.DepotPath.Path, StreamName, out RelativePath))
+							{
+								Files.Add(CreateChangeFile(RelativePath, DescribeFile));
+							}
 						}
 
 						return new ChangeDetails(Change.Id, Change.OwnerName, Change.GetPath(), Change.Description, Files, Change.ModifiedDate);
@@ -697,7 +777,7 @@ namespace HordeServer.Services
 
 					foreach (P4.TaggedObject TaggedObject in Result.TaggedOutput)
 					{
-						List<string> Files = new List<string>();
+						List<ChangeFile> Files = new List<ChangeFile>();
 
 						P4.Changelist Change = new P4.Changelist();
 						Change.FromChangeCmdTaggedOutput(TaggedObject, true, Offset, DSTMismatch);
@@ -707,7 +787,7 @@ namespace HordeServer.Services
 							string? RelativePath;
 							if (TryGetStreamRelativePath(DescribeFile.DepotPath.Path, StreamName, out RelativePath))
 							{
-								Files.Add(RelativePath);
+								Files.Add(CreateChangeFile(RelativePath, DescribeFile));
 							}
 						}
 
@@ -718,7 +798,7 @@ namespace HordeServer.Services
 								string? RelativePath;
 								if (TryGetStreamRelativePath(ShelvedFile.Path.ToString(), StreamName, out RelativePath))
 								{
-									Files.Add(RelativePath);
+									Files.Add(CreateChangeFile(RelativePath, ShelvedFile));
 								}
 							}
 
diff --git a/Engine/Source/Programs/Horde/HordeServerTests/Stubs/Services/PerforceServiceStub.cs b/Engine/Source/Programs/Horde/HordeServerTests/Stubs/Services/PerforceServiceStub.cs
--- a/Engine/Source/Programs/Horde/HordeServerTests/Stubs/Services/PerforceServiceStub.cs
+++ b/Engine/Source/Programs/Horde/HordeServerTests/Stubs/Services/PerforceServiceStub.cs
@@ -14,6 +14,19 @@ using StreamId = HordeServer.Utilities.StringId<HordeServer.Models.IStream>;
 
 namespace HordeServerTests.Stubs.Services
 {
+	static class PerforceExtensions
+	{
+		public static ChangeFile CreateChangeFile(string Path)
+		{
+			return new ChangeFile(Path, null!, 0, 0, EpicGames.Core.Md5Hash.Zero, null!);
+		}
+
+		public static void Add(this List<ChangeFile> Files, string Path)
+		{
+			Files.Add(CreateChangeFile(Path));
+		}
+	}
+
 	class PerforceServiceStub : IPerforceService
 	{
 		class ChangeComparer : IComparer<int>
@@ -31,7 +44,7 @@ namespace HordeServerTests.Stubs.Services
 				StreamChanges = new SortedDictionary<int, ChangeDetails>(new ChangeComparer());
 				Changes[StreamName] = StreamChanges;
 			}
-			StreamChanges.Add(Number, new ChangeDetails(Number, Author, null!, Description, Files.ToList(), DateTime.Now));
+			StreamChanges.Add(Number, new ChangeDetails(Number, Author, null!, Description, Files.Select(x => PerforceExtensions.CreateChangeFile(x)).ToList(), DateTime.Now));
 		}
 
 		public Task<List<ChangeSummary>> GetChangesAsync(string ClusterName, string StreamName, int? MinChange, int? MaxChange, int NumResults, string? ImpersonateUser)
@@ -90,7 +103,7 @@ namespace HordeServerTests.Stubs.Services
 			{
 				foreach (ChangeDetails Details in StreamChanges.Values)
 				{
-					if (Details.Number <= Change && Details.Files.Any(x => x.EndsWith(".h") || x.EndsWith(".cpp")))
+					if (Details.Number <= Change && Details.Files.Any(x => x.Path.EndsWith(".h") || x.Path.EndsWith(".cpp")))
 					{
 						CodeChange = Details.Number;
 						break;
@@ -103,7 +116,7 @@ namespace HordeServerTests.Stubs.Services
 
 		public Task<int> CreateNewChangeAsync(string ClusterName, string StreamName, string Path)
 		{
-			ChangeDetails NewChange = new ChangeDetails(Changes[StreamName].First().Key + 1, "", null!, "", new List<string> { Path }, DateTime.Now);
+			ChangeDetails NewChange = new ChangeDetails(Changes[StreamName].First().Key + 1, "", null!, "", new List<ChangeFile> { PerforceExtensions.CreateChangeFile(Path) }, DateTime.Now);
 			Changes[StreamName].Add(NewChange.Number, NewChange);
 			return Task.FromResult(NewChange.Number);
 		}
@@ -152,5 +165,15 @@ namespace HordeServerTests.Stubs.Services
 		{
 			throw new NotImplementedException();
 		}
+
+		public Task<List<ChangeFile>> GetStreamSnapshotAsync(string ClusterName, string StreamName, int Change)
+		{
+			throw new NotImplementedException();
+		}
+
+		public Task<ChangeDetails> GetChangeDetailsAsync(string ClusterName, string StreamName, int ChangeNumber)
+		{
+			throw new NotImplementedException();
+		}
 	}
 }
