Horde: Auto register agents to default pool, remove delay waiting for perforce clusters to become available
#rnx

diff --git a/Engine/Source/Programs/Horde/HordeServer/Services/ConfigService.cs b/Engine/Source/Programs/Horde/HordeServer/Services/ConfigService.cs
--- a/Engine/Source/Programs/Horde/HordeServer/Services/ConfigService.cs
+++ b/Engine/Source/Programs/Horde/HordeServer/Services/ConfigService.cs
@@ -536,25 +536,7 @@ namespace HordeServer.Services
 
 			if (ConfigDirty == true)
 			{				
-				await UpdateConfigAsync(new Uri("file://" + GlobalConfigFile.ToString()));
-
-				// wait for the perforce clusters to be picked up, this ticks at a minute and might be nice to be able to explicitly update this
-				// the PerforceLoadBalancer has an ElectedTick member and no TickNow method
-				// In any event, this delay helps initial feel as cluster will be available shortly after it
-				int WaitSec = 15;				
-				DateTime Start = DateTime.UtcNow;
-				do
-				{
-					await Task.Delay(1000);
-
-					List<IPerforceServer> Servers = await PerforceLoadBalancer.GetServersAsync();
-					if (Servers.Count > 0)
-					{
-						break;
-					}
-				}
-				while ((DateTime.UtcNow - Start).TotalSeconds < WaitSec);
-
+				await UpdateConfigAsync(new Uri("file://" + GlobalConfigFile.ToString()));				
 			}
 
 			// create the default pool 
@@ -565,7 +547,8 @@ namespace HordeServer.Services
 				IPool? Pool = await PoolService.GetPoolAsync(PoolIdValue);
 				if (Pool == null)
 				{
-					await PoolService.CreatePoolAsync(Request.DefaultPoolName, Properties: new Dictionary<string, string>() { ["Color"] = "0" });
+					AgentRequirements Requirements = new AgentRequirements(new List<DeviceRequirements>() { new DeviceRequirements(null, new HashSet<string>(){ "OSFamily=Windows"}, null) }, null, null, false);
+					await PoolService.CreatePoolAsync(Request.DefaultPoolName, Requirements: Requirements, Properties: new Dictionary<string, string>() { ["Color"] = "0" });
 				}
 			}
 
