Horde: Change to UE4 pool as default during migration
#rnx

diff --git a/Engine/Source/Programs/Horde/HordeServer/Controllers/DevicesController.cs b/Engine/Source/Programs/Horde/HordeServer/Controllers/DevicesController.cs
--- a/Engine/Source/Programs/Horde/HordeServer/Controllers/DevicesController.cs
+++ b/Engine/Source/Programs/Horde/HordeServer/Controllers/DevicesController.cs
@@ -693,10 +693,10 @@ namespace HordeServer.Controllers
 
 			if (string.IsNullOrEmpty(PoolId))
 			{
-                Message = $"No pool specified for device reservation request, defaulting to FNUE5 (will soon be an error)" + Details;
+                Message = $"No pool specified for device reservation request, defaulting to UE4 (will soon be an error)" + Details;
                 Logger.LogError(Message + $" JobId: {Request.JobId}, StepId: {Request.StepId}");
 				await DeviceService.NotifyDeviceServiceAsync(Message, null, Request.JobId, Request.StepId);
-				PoolId = "fnue5";
+				PoolId = "ue4";
                 //return BadRequest(Message);
             }			
 
