Horde: Configure namespaces and buckets for storage from config file.

diff --git a/Engine/Source/Programs/Horde/HordeServer/Services/ConfigService.cs b/Engine/Source/Programs/Horde/HordeServer/Services/ConfigService.cs
--- a/Engine/Source/Programs/Horde/HordeServer/Services/ConfigService.cs
+++ b/Engine/Source/Programs/Horde/HordeServer/Services/ConfigService.cs
@@ -26,11 +26,14 @@ using System.Reflection;
 
 using PoolId = HordeServer.Utilities.StringId<HordeServer.Models.IPool>;
 using System.Globalization;
+using HordeServer.Storage;
 
 namespace HordeServer.Services
 {
 	using ProjectId = StringId<IProject>;
 	using StreamId = StringId<IStream>;
+	using NamespaceId = StringId<INamespace>;
+	using BucketId = StringId<IBucket>;
 
 	/// <summary>
 	/// Polls Perforce for stream config changes
@@ -70,12 +73,21 @@ namespace HordeServer.Services
 		/// </summary>
 		private PerforceLoadBalancer PerforceLoadBalancer;
 
-
 		/// <summary>
 		/// Instance of the notification service
 		/// </summary>
 		INotificationService NotificationService;
 
+		/// <summary>
+		/// The namespace collection
+		/// </summary>
+		INamespaceCollection NamespaceCollection;
+
+		/// <summary>
+		/// The bucket collection
+		/// </summary>
+		IBucketCollection BucketCollection;
+
 		/// <summary>
 		/// Singleton instance of the pool service
 		/// </summary>
@@ -100,10 +112,12 @@ namespace HordeServer.Services
 		/// <param name="ProjectService"></param>
 		/// <param name="StreamService"></param>
 		/// <param name="NotificationService"></param>
+		/// <param name="NamespaceCollection"></param>
+		/// <param name="BucketCollection"></param>
 		/// <param name="PoolService"></param>
 		/// <param name="Settings"></param>
 		/// <param name="Logger"></param>
-		public ConfigService(DatabaseService DatabaseService, IPerforceService PerforceService, ProjectService ProjectService, StreamService StreamService, INotificationService NotificationService,  PoolService PoolService, PerforceLoadBalancer PerforceLoadBalancer, IOptionsMonitor<ServerSettings> Settings, ILogger<ConfigService> Logger)
+		public ConfigService(DatabaseService DatabaseService, IPerforceService PerforceService, ProjectService ProjectService, StreamService StreamService, INotificationService NotificationService, INamespaceCollection NamespaceCollection, IBucketCollection BucketCollection, PoolService PoolService, PerforceLoadBalancer PerforceLoadBalancer, IOptionsMonitor<ServerSettings> Settings, ILogger<ConfigService> Logger)
 			: base(DatabaseService, new ObjectId("5ff60549e7632b15e64ac2f7"), Logger)
 		{
 			this.DatabaseService = DatabaseService;
@@ -112,6 +126,8 @@ namespace HordeServer.Services
 			this.ProjectService = ProjectService;
 			this.StreamService = StreamService;
 			this.NotificationService = NotificationService;
+			this.NamespaceCollection = NamespaceCollection;
+			this.BucketCollection = BucketCollection;
 			this.PoolService = PoolService;
 			this.Settings = Settings;
 			this.Logger = Logger;
@@ -167,6 +183,8 @@ namespace HordeServer.Services
 				Globals.ScheduledDowntime = CachedGlobalConfig.Downtime;
 				Globals.MaxConformCount = CachedGlobalConfig.MaxConformCount;
 
+				await UpdateStorageConfigAsync(GlobalConfig.Storage);
+
 				if (await DatabaseService.TryUpdateSingletonAsync(Globals))
 				{
 					break;
@@ -318,6 +336,44 @@ namespace HordeServer.Services
 			}
 		}
 
+		async Task UpdateStorageConfigAsync(StorageConfig? Config)
+		{
+			List<INamespace> Namespaces = await NamespaceCollection.FindAsync();
+			List<NamespaceId> RemoveNamespaceIds = Namespaces.ConvertAll(x => x.Id);
+
+			if (Config != null)
+			{
+				foreach (NamespaceConfig NamespaceConfig in Config.Namespaces)
+				{
+					NamespaceId NamespaceId = new NamespaceId(NamespaceConfig.Id);
+					RemoveNamespaceIds.Remove(NamespaceId);
+					await NamespaceCollection.AddOrUpdateAsync(NamespaceId, NamespaceConfig);
+
+					List<IBucket> Buckets = await BucketCollection.FindAsync(NamespaceId);
+					List<BucketId> RemoveBucketIds = Buckets.ConvertAll(x => x.BucketId);
+
+					foreach(BucketConfig BucketConfig in NamespaceConfig.Buckets)
+					{
+						BucketId BucketId = new BucketId(BucketConfig.Id);
+						await BucketCollection.AddOrUpdateAsync(NamespaceId, BucketId, BucketConfig);
+						RemoveBucketIds.Remove(BucketId);
+					}
+
+					foreach (BucketId RemoveBucketId in RemoveBucketIds)
+					{
+						Logger.LogInformation("Removing bucket {BucketId}", RemoveBucketId);
+						await BucketCollection.RemoveAsync(NamespaceId, RemoveBucketId);
+					}
+				}
+			}
+
+			foreach (NamespaceId RemoveNamespaceId in RemoveNamespaceIds)
+			{
+				Logger.LogInformation("Removing namespace {NamespaceId}", RemoveNamespaceId);
+				await NamespaceCollection.RemoveAsync(RemoveNamespaceId);
+			}
+		}
+
 		static FileExtensionContentTypeProvider ContentTypeProvider = new FileExtensionContentTypeProvider();
 
 		static string GetMimeTypeFromPath(Uri Path)
