Horde: Add support for serializing requests and responses as compact binary objects.

diff --git a/Engine/Source/Programs/Horde/HordeServer/Compute/Impl/ComputeController.cs b/Engine/Source/Programs/Horde/HordeServer/Compute/Impl/ComputeController.cs
--- a/Engine/Source/Programs/Horde/HordeServer/Compute/Impl/ComputeController.cs
+++ b/Engine/Source/Programs/Horde/HordeServer/Compute/Impl/ComputeController.cs
@@ -50,7 +50,7 @@ namespace HordeServer.Compute.Impl
 		[Route("/api/v1/compute/{ChannelId}")]
 		public async Task AddTasksAsync([FromRoute] string ChannelId, [FromBody] AddTasksRequest Request)
 		{
-			await ComputeService.AddTasksAsync(Impl.ComputeService.DefaultNamespaceId, IoHash.Parse(Request.RequirementsHash), Request.TaskHashes.ConvertAll(x => IoHash.Parse(x)), new ChannelId(ChannelId));
+			await ComputeService.AddTasksAsync(Impl.ComputeService.DefaultNamespaceId, Request.RequirementsHash, Request.TaskHashes, new ChannelId(ChannelId));
 		}
 
 		/// <summary>
@@ -80,10 +80,10 @@ namespace HordeServer.Compute.Impl
 			{
 				GetTaskUpdateResponse Update = new GetTaskUpdateResponse();
 
-				Update.TaskHash = Result.TaskHash.ToString();
+				Update.TaskHash = Result.TaskHash;
 				Update.Time = Result.Time;
 				Update.State = Result.State;
-				Update.ResultHash = Result.ResultHash?.ToString();
+				Update.ResultHash = Result.ResultHash;
 				Update.AgentId = Result.AgentId?.ToString();
 				Update.LeaseId = Result.LeaseId?.ToString();
 
diff --git a/Engine/Source/Programs/Horde/HordeServer/Startup.cs b/Engine/Source/Programs/Horde/HordeServer/Startup.cs
--- a/Engine/Source/Programs/Horde/HordeServer/Startup.cs
+++ b/Engine/Source/Programs/Horde/HordeServer/Startup.cs
@@ -565,7 +565,11 @@ namespace HordeServer
 
 			Services.AddMvc().AddJsonOptions(Options => ConfigureJsonSerializer(Options.JsonSerializerOptions));
 
-			Services.AddControllers();
+			Services.AddControllers(Options =>
+			{
+				Options.InputFormatters.Insert(0, new CbInputFormatter());
+				Options.OutputFormatters.Insert(0, new CbOutputFormatter());
+			});
 
 			Services.AddSwaggerGen(Config =>
 			{
diff --git a/Engine/Source/Programs/Horde/HordeServer/Utilities/CompactBinaryFormatter.cs b/Engine/Source/Programs/Horde/HordeServer/Utilities/CompactBinaryFormatter.cs
new file mode 100644
--- /dev/null
+++ b/Engine/Source/Programs/Horde/HordeServer/Utilities/CompactBinaryFormatter.cs
@@ -0,0 +1,77 @@
+// Copyright Epic Games, Inc. All Rights Reserved.
+
+using Microsoft.AspNetCore.Mvc.Formatters;
+using System;
+using System.Collections.Generic;
+using System.Linq;
+using Microsoft.Net.Http.Headers;
+using System.Threading.Tasks;
+using System.IO;
+using EpicGames.Serialization;
+using EpicGames.Serialization.Converters;
+using System.Reflection;
+
+namespace HordeServer.Utilities
+{
+	/// <summary>
+	/// Converter to allow reading compact binary objects as request bodies
+	/// </summary>
+	public class CbInputFormatter : InputFormatter
+	{
+		/// <summary>
+		/// Constructor
+		/// </summary>
+		public CbInputFormatter()
+		{
+			SupportedMediaTypes.Add(MediaTypeHeaderValue.Parse("application/x-ue-cb"));
+		}
+
+		/// <inheritdoc/>
+		protected override bool CanReadType(Type Type)
+		{
+			return true;
+		}
+
+		/// <inheritdoc/>
+		public override Task<InputFormatterResult> ReadRequestBodyAsync(InputFormatterContext Context)
+		{
+			// Buffer the data into an array
+			byte[] Data;
+			using (MemoryStream Stream = new MemoryStream())
+			{
+				Context.HttpContext.Request.Body.CopyToAsync(Stream);
+				Data = Stream.ToArray();
+			}
+
+			// Serialize the object
+			return InputFormatterResult.SuccessAsync(CbSerializer.Deserialize(new CbField(Data), Context.ModelType));
+		}
+	}
+
+	/// <summary>
+	/// Converter to allow writing compact binary objects as responses
+	/// </summary>
+	public class CbOutputFormatter : OutputFormatter
+	{
+		/// <summary>
+		/// Constructor
+		/// </summary>
+		public CbOutputFormatter()
+		{
+			SupportedMediaTypes.Add(MediaTypeHeaderValue.Parse("application/x-ue-cb"));
+		}
+
+		/// <inheritdoc/>
+		protected override bool CanWriteType(Type Type)
+		{
+			return true;
+		}
+
+		/// <inheritdoc/>
+		public override async Task WriteResponseBodyAsync(OutputFormatterWriteContext Context)
+		{
+			ReadOnlyMemory<byte> Data = CbSerializer.Serialize(Context.ObjectType, Context.Object).GetView();
+			await Context.HttpContext.Response.BodyWriter.WriteAsync(Data);
+		}
+	}
+}
diff --git a/Engine/Source/Programs/Shared/EpicGames.Horde/Compute/Api.cs b/Engine/Source/Programs/Shared/EpicGames.Horde/Compute/Api.cs
--- a/Engine/Source/Programs/Shared/EpicGames.Horde/Compute/Api.cs
+++ b/Engine/Source/Programs/Shared/EpicGames.Horde/Compute/Api.cs
@@ -2,6 +2,7 @@
 
 using EpicGames.Core;
 using EpicGames.Horde.Common;
+using EpicGames.Serialization;
 using System;
 using System.Collections.Generic;
 using System.ComponentModel.DataAnnotations;
@@ -16,14 +17,15 @@ namespace EpicGames.Horde.Compute
 	public class AddTasksRequest
 	{
 		/// <inheritdoc cref="AddTasksRpcRequest.RequirementsHash"/>
-		[Required]
-		public string RequirementsHash { get; set; } = String.Empty;
+		[Required, CbField("r")]
+		public IoHash RequirementsHash { get; set; }
 
 		/// <inheritdoc cref="AddTasksRpcRequest.TaskHashes"/>
-		[Required]
-		public List<string> TaskHashes { get; set; } = new List<string>();
+		[Required, CbField("t")]
+		public List<IoHash> TaskHashes { get; set; } = new List<IoHash>();
 
 		/// <inheritdoc cref="AddTasksRpcRequest.DoNotCache"/>
+		[CbField("nc")]
 		public bool DoNotCache { get; set; }
 	}
 
@@ -35,6 +37,7 @@ namespace EpicGames.Horde.Compute
 		/// <summary>
 		/// Task updates
 		/// </summary>
+		[CbField("u")]
 		public List<GetTaskUpdateResponse> Updates { get; set; } = new List<GetTaskUpdateResponse>();
 	}
 
@@ -44,21 +47,27 @@ namespace EpicGames.Horde.Compute
 	public class GetTaskUpdateResponse
 	{
 		/// <inheritdoc cref="GetTaskUpdatesRpcResponse.TaskHash"/>
-		public string TaskHash { get; set; } = null!;
+		[CbField("h")]
+		public IoHash TaskHash { get; set; }
 
 		/// <inheritdoc cref="GetTaskUpdatesRpcResponse.Time"/>
+		[CbField("t")]
 		public DateTime Time { get; set; }
 
 		/// <inheritdoc cref="GetTaskUpdatesRpcResponse.State"/>
+		[CbField("s")]
 		public ComputeTaskState State { get; set; }
 
 		/// <inheritdoc cref="GetTaskUpdatesRpcResponse.ResultHash"/>
-		public string? ResultHash { get; set; }
+		[CbField("r")]
+		public IoHash? ResultHash { get; set; }
 
 		/// <inheritdoc cref="GetTaskUpdatesRpcResponse.AgentId"/>
+		[CbField("a")]
 		public string? AgentId { get; set; }
 
 		/// <inheritdoc cref="GetTaskUpdatesRpcResponse.LeaseId"/>
+		[CbField("l")]
 		public string? LeaseId { get; set; }
 	}
 
diff --git a/Engine/Source/Programs/Shared/EpicGames.IoHash/IoHash.cs b/Engine/Source/Programs/Shared/EpicGames.IoHash/IoHash.cs
--- a/Engine/Source/Programs/Shared/EpicGames.IoHash/IoHash.cs
+++ b/Engine/Source/Programs/Shared/EpicGames.IoHash/IoHash.cs
@@ -9,12 +9,15 @@ using System.Buffers.Binary;
 using System.Diagnostics.CodeAnalysis;
 using System.ComponentModel;
 using System.Globalization;
+using System.Text.Json.Serialization;
+using System.Text.Json;
 
 namespace EpicGames.Core
 {
 	/// <summary>
 	/// Struct representing a strongly typed IoHash value (a 20-byte Blake3 hash).
 	/// </summary>
+	[JsonConverter(typeof(IoHashJsonConverter))]
 	[TypeConverter(typeof(IoHashTypeConverter))]
 	public struct IoHash : IEquatable<IoHash>, IComparable<IoHash>
 	{
@@ -90,6 +93,16 @@ namespace EpicGames.Core
 			return new IoHash(StringUtils.ParseHexString(Text));
 		}
 
+		/// <summary>
+		/// Parses a digest from the given hex string
+		/// </summary>
+		/// <param name="Text"></param>
+		/// <returns></returns>
+		public static IoHash Parse(ReadOnlySpan<byte> Text)
+		{
+			return new IoHash(StringUtils.ParseHexString(Text));
+		}
+
 		/// <inheritdoc cref="IComparable{T}.CompareTo(T)"/>
 		public int CompareTo(IoHash Other)
 		{
@@ -116,6 +129,9 @@ namespace EpicGames.Core
 		/// <inheritdoc/>
 		public override int GetHashCode() => BinaryPrimitives.ReadInt32LittleEndian(Span);
 
+		/// <inheritdoc/>
+		public Utf8String ToUtf8String() => StringUtils.FormatUtf8HexString(Memory.Span);
+
 		/// <inheritdoc/>
 		public override string ToString() => StringUtils.FormatHexString(Memory.Span);
 
@@ -185,6 +201,18 @@ namespace EpicGames.Core
 		}
 	}
 
+	/// <summary>
+	/// Type converter for IoHash to and from JSON
+	/// </summary>
+	sealed class IoHashJsonConverter : JsonConverter<IoHash>
+	{
+		/// <inheritdoc/>
+		public override IoHash Read(ref Utf8JsonReader Reader, Type TypeToConvert, JsonSerializerOptions Options) => IoHash.Parse(Reader.ValueSpan);
+
+		/// <inheritdoc/>
+		public override void Write(Utf8JsonWriter Writer, IoHash Value, JsonSerializerOptions Options) => Writer.WriteStringValue(Value.ToUtf8String().Span);
+	}
+
 	/// <summary>
 	/// Type converter from strings to IoHash objects
 	/// </summary>
diff --git a/Engine/Source/Programs/Shared/EpicGames.Serialization/CompactBinarySerializer.cs b/Engine/Source/Programs/Shared/EpicGames.Serialization/CompactBinarySerializer.cs
--- a/Engine/Source/Programs/Shared/EpicGames.Serialization/CompactBinarySerializer.cs
+++ b/Engine/Source/Programs/Shared/EpicGames.Serialization/CompactBinarySerializer.cs
@@ -174,6 +174,19 @@ namespace EpicGames.Serialization
 			return CbConverterCache<T>.Instance;
 		}
 
+		/// <summary>
+		/// Serialize an object
+		/// </summary>
+		/// <param name="Type">Type of the object to serialize</param>
+		/// <param name="Value"></param>
+		/// <returns></returns>
+		public static CbObject Serialize(Type Type, object Value)
+		{
+			CbWriter Writer = new CbWriter();
+			GetConverter(Type).WriteObject(Writer, Value);
+			return Writer.ToObject();
+		}
+
 		/// <summary>
 		/// Serialize an object
 		/// </summary>
@@ -210,6 +223,17 @@ namespace EpicGames.Serialization
 			GetConverter<T>().WriteNamed(Writer, Name, Value);
 		}
 
+		/// <summary>
+		/// Deserialize an object from a <see cref="CbObject"/>
+		/// </summary>
+		/// <param name="Field"></param>
+		/// <param name="Type">Type of the object to read</param>
+		/// <returns></returns>
+		public static object? Deserialize(CbField Field, Type Type)
+		{
+			return GetConverter(Type).ReadObject(Field);
+		}
+
 		/// <summary>
 		/// Deserialize an object from a <see cref="CbObject"/>
 		/// </summary>
diff --git a/Engine/Source/Programs/Shared/EpicGames.Serialization/Converters/CbDefaultConverter.cs b/Engine/Source/Programs/Shared/EpicGames.Serialization/Converters/CbDefaultConverter.cs
--- a/Engine/Source/Programs/Shared/EpicGames.Serialization/Converters/CbDefaultConverter.cs
+++ b/Engine/Source/Programs/Shared/EpicGames.Serialization/Converters/CbDefaultConverter.cs
@@ -103,6 +103,10 @@ namespace EpicGames.Serialization.Converters
 		static void WriteInt64Value(CbWriter Writer, long Value) => Writer.WriteIntegerValue(Value);
 		static void WriteInt64(CbWriter Writer, Utf8String Name, long Value) => Writer.WriteInteger(Name, Value);
 
+		static string ReadString(CbField Field) => Field.AsString().ToString();
+		static void WriteString(CbWriter Writer, Utf8String Name, string Value) => Writer.WriteString(Name, Value);
+		static void WriteStringValue(CbWriter Writer, string Value) => Writer.WriteStringValue(Value);
+
 		static Dictionary<Type, CbConverterInfo> TypeToConverterInfo = new Dictionary<Type, CbConverterInfo>(new KeyValuePair<Type, CbConverterInfo>[]
 		{
 			GetPodConverterInfo(x => x.AsBool(), (w, v) => w.WriteBoolValue(v), (w, n, v) => w.WriteBool(n, v), null),
@@ -110,8 +114,8 @@ namespace EpicGames.Serialization.Converters
 			GetPodConverterInfo(x => x.AsInt64(), (w, v) => WriteInt64Value(w, v), (w, n, v) => WriteInt64(w, n, v), null),
 			GetPodConverterInfo(x => x.AsString(), (w, v) => w.WriteStringValue(v), (w, n, v) => w.WriteString(n, v), null),
 			GetPodConverterInfo(x => x.AsHash(), (w, v) => w.WriteHashValue(v), (w, n, v) => w.WriteHash(n, v), null),
-			GetPodConverterInfo(x => x.AsDateTime(), (w, v) => w.WriteDateTimeValue(v), (w, n, v) => w.WriteDateTime(n, v), null)
-			// TODO: other basic types
+			GetPodConverterInfo(x => x.AsDateTime(), (w, v) => w.WriteDateTimeValue(v), (w, n, v) => w.WriteDateTime(n, v), null),
+			GetPodConverterInfo(x => ReadString(x), (w, v) => WriteStringValue(w, v), (w, n, v) => WriteString(w, n, v), SkipIfNullOrZero)
 		});
 
 		static KeyValuePair<Type, CbConverterInfo> GetPodConverterInfo<T>(Expression<Func<CbField, T>> Read, Expression<Action<CbWriter, T>> Write, Expression<Action<CbWriter, Utf8String, T>> WriteNamed, SkipDefaultDelegate? SkipDefault)
