Horde: Prevent compact binary output being used by default.

diff --git a/Engine/Source/Programs/Horde/HordeServer/Startup.cs b/Engine/Source/Programs/Horde/HordeServer/Startup.cs
--- a/Engine/Source/Programs/Horde/HordeServer/Startup.cs
+++ b/Engine/Source/Programs/Horde/HordeServer/Startup.cs
@@ -567,8 +567,8 @@ namespace HordeServer
 
 			Services.AddControllers(Options =>
 			{
-				Options.InputFormatters.Insert(0, new CbInputFormatter());
-				Options.OutputFormatters.Insert(0, new CbOutputFormatter());
+				Options.InputFormatters.Add(new CbInputFormatter());
+				Options.OutputFormatters.Add(new CbOutputFormatter());
 			});
 
 			Services.AddSwaggerGen(Config =>
