Horde: Remove size from compute directory objects.

diff --git a/Engine/Source/Programs/Horde/HordeAgent/Commands/Execution/ExecuteCommandV2.cs b/Engine/Source/Programs/Horde/HordeAgent/Commands/Execution/ExecuteCommandV2.cs
--- a/Engine/Source/Programs/Horde/HordeAgent/Commands/Execution/ExecuteCommandV2.cs
+++ b/Engine/Source/Programs/Horde/HordeAgent/Commands/Execution/ExecuteCommandV2.cs
@@ -98,8 +98,7 @@ namespace HordeAgent.Commands
 			foreach (DirectoryInfo SubDirInfo in BaseDirInfo.EnumerateDirectories())
 			{
 				(DirectoryTree SubTree, IoHash SubDirHash) = CreateSandbox(SubDirInfo, UploadList);
-				long Size = SubTree.Files.Sum(x => x.Size) + SubTree.Directories.Sum(x => x.Size);
-				Tree.Directories.Add(new DirectoryNode(SubDirInfo.Name, SubDirHash, Size));
+				Tree.Directories.Add(new DirectoryNode(SubDirInfo.Name, SubDirHash));
 			}
 			Tree.Directories.SortBy(x => x.Name, Utf8StringComparer.Ordinal);
 
diff --git a/Engine/Source/Programs/Horde/HordeAgent/Execution/ComputeTaskExecutor.cs b/Engine/Source/Programs/Horde/HordeAgent/Execution/ComputeTaskExecutor.cs
--- a/Engine/Source/Programs/Horde/HordeAgent/Execution/ComputeTaskExecutor.cs
+++ b/Engine/Source/Programs/Horde/HordeAgent/Execution/ComputeTaskExecutor.cs
@@ -287,8 +287,8 @@ namespace HordeAgent
 
 		async Task<DirectoryNode> CreateDirectoryNode(string NamespaceId, string Name, int BaseDirLen, List<FileReference> SortedFiles, int MinIdx, int MaxIdx)
 		{
-			(DirectoryTree Tree, IoHash Hash) = await PutDirectoryTree(NamespaceId, BaseDirLen, SortedFiles, MinIdx, MaxIdx);
-			return new DirectoryNode(Name, Hash, Tree.Files.Sum(x => x.Size) + Tree.Directories.Sum(x => x.Size));
+			(_, IoHash Hash) = await PutDirectoryTree(NamespaceId, BaseDirLen, SortedFiles, MinIdx, MaxIdx);
+			return new DirectoryNode(Name, Hash);
 		}
 
 		async Task<FileNode> CreateFileNode(string NamespaceId, string Name, FileReference File)
diff --git a/Engine/Source/Programs/Shared/EpicGames.Horde/Compute/DirectoryTree.cs b/Engine/Source/Programs/Shared/EpicGames.Horde/Compute/DirectoryTree.cs
--- a/Engine/Source/Programs/Shared/EpicGames.Horde/Compute/DirectoryTree.cs
+++ b/Engine/Source/Programs/Shared/EpicGames.Horde/Compute/DirectoryTree.cs
@@ -91,12 +91,6 @@ namespace EpicGames.Horde.Compute
 		[CbField("h")]
 		public IoHash Hash { get; set; }
 
-		/// <summary>
-		/// Total size of the items in this directory
-		/// </summary>
-		[CbField("s")]
-		public long Size { get; set; }
-
 		/// <summary>
 		/// Private constructor for serialization
 		/// </summary>
@@ -107,11 +101,10 @@ namespace EpicGames.Horde.Compute
 		/// <summary>
 		/// Constructor
 		/// </summary>
-		public DirectoryNode(Utf8String Name, IoHash Hash, long Size)
+		public DirectoryNode(Utf8String Name, IoHash Hash)
 		{
 			this.Name = Name;
 			this.Hash = Hash;
-			this.Size = Size;
 		}
 	}
 }
