Horde: Support for generating Perforce snapshots to the CAS.

diff --git a/Engine/Source/Programs/Horde/HordeServer/Commits/Impl/CommitService.cs b/Engine/Source/Programs/Horde/HordeServer/Commits/Impl/CommitService.cs
--- a/Engine/Source/Programs/Horde/HordeServer/Commits/Impl/CommitService.cs
+++ b/Engine/Source/Programs/Horde/HordeServer/Commits/Impl/CommitService.cs
@@ -21,9 +21,12 @@ using MongoDB.Bson.Serialization.Attributes;
 using System.Collections.Concurrent;
 using System.Globalization;
 using Microsoft.Extensions.Hosting;
+using EpicGames.Perforce.Managed;
 
 namespace HordeServer.Commits.Impl
 {
+	using P4 = Perforce.P4;
+	using NamespaceId = StringId<INamespace>;
 	using StreamId = StringId<IStream>;
 
 	/// <summary>
@@ -49,7 +52,10 @@ namespace HordeServer.Commits.Impl
 			public List<StreamReplicationState> Streams { get; set; } = new List<StreamReplicationState>();
 		}
 
+		NamespaceId NamespaceId { get; } = new NamespaceId("default");
+
 		ICommitCollection CommitCollection;
+		IObjectCollection ObjectCollection;
 		IStreamCollection StreamCollection;
 		IPerforceService PerforceService;
 		IUserCollection UserCollection;
@@ -61,9 +67,10 @@ namespace HordeServer.Commits.Impl
 		/// <summary>
 		/// Constructor
 		/// </summary>
-		public CommitService(ICommitCollection CommitCollection, IStreamCollection StreamCollection, IPerforceService PerforceService, IUserCollection UserCollection, DatabaseService DatabaseService, ILogger<CommitService> Logger)
+		public CommitService(ICommitCollection CommitCollection, IObjectCollection ObjectCollection, IStreamCollection StreamCollection, IPerforceService PerforceService, IUserCollection UserCollection, DatabaseService DatabaseService, ILogger<CommitService> Logger)
 		{
 			this.CommitCollection = CommitCollection;
+			this.ObjectCollection = ObjectCollection;
 			this.StreamCollection = StreamCollection;
 			this.PerforceService = PerforceService;
 			this.UserCollection = UserCollection;
@@ -255,5 +262,234 @@ namespace HordeServer.Commits.Impl
 				return null;
 			}
 		}
+
+		#region Snapshot generation
+
+		/// <summary>
+		/// Number of files to query at a time before recursing down the tree
+		/// </summary>
+		const int BatchSize = 16384;
+
+		/// <summary>
+		/// Creates a snapshot of the repository at a particular changelist
+		/// </summary>
+		/// <param name="Repository">The Perforce repository to mirror</param>
+		/// <param name="Stream">Name of the stream to be mirrored</param>
+		/// <param name="Change">Changelist to replicate</param>
+		/// <param name="Comparison">Comparison to use for names</param>
+		Task CreateSnapshotAsync(P4.Repository Repository, string Stream, int Change, StringComparison Comparison)
+		{
+			// Get the stream spec
+			P4.Stream StreamSpec = Repository.GetStream($"{Stream}@{Change}", new P4.Options { { "-v", "" } });
+			ViewMap ViewMap = new ViewMap(StreamSpec.View);
+			return CreateSnapshotAsync(Repository, ViewMap, Change, Comparison);
+		}
+
+		/// <summary>
+		/// Creates a snapshot of a stream at a particular changelist
+		/// </summary>
+		/// <param name="Repository">The Perforce repository to mirror</param>
+		/// <param name="View">View for the stream</param>
+		/// <param name="Change">Changelist to replicate</param>
+		/// <param name="Comparison">Comparison to use for names</param>
+		async Task<StreamTreeRef> CreateSnapshotAsync(P4.Repository Repository, ViewMap View, int Change, StringComparison Comparison)
+		{
+			// In order to optimize for the common case where we're adding multiple files to the same directory, we keep an open list of
+			// StreamTreeBuilder objects down to the last modified node. Whenever we need to move to a different node, we write the 
+			// previous path to the object store.
+			List<StreamTreeBuilder> TreePath = new List<StreamTreeBuilder>();
+			TreePath.Add(new StreamTreeBuilder());
+
+			// Generate a snapshot using all the root paths
+			List<string> RootPaths = View.GetRootPaths(Comparison);
+			foreach (string RootPath in RootPaths)
+			{
+				ViewMap FilteredView = new ViewMap();
+				foreach (ViewMapEntry Entry in View.Entries)
+				{
+					if (Entry.SourcePrefix.StartsWith(RootPath, Comparison))
+					{
+						FilteredView.Entries.Add(Entry);
+					}
+				}
+				await AddDepotDirToSnapshotAsync(Repository, RootPath, Change, FilteredView, Comparison, TreePath);
+			}
+
+			// Write the root tree
+			return await EncodeTreeAsync(TreePath[0]);
+		}
+
+		/// <summary>
+		/// Adds all files under a directory in the depot to a snapshot
+		/// </summary>
+		/// <param name="Repository">The perforce repository</param>
+		/// <param name="DepotDir">Depot directory to recurse through, with a trailing slash</param>
+		/// <param name="Change">Changelist number to capture</param>
+		/// <param name="View">View for the stream</param>
+		/// <param name="Comparison">Comparison type to use for names</param>
+		/// <param name="TreePath">The last computed path through the tree</param>
+		async Task AddDepotDirToSnapshotAsync(P4.Repository Repository, string DepotDir, int Change, ViewMap View, StringComparison Comparison, List<StreamTreeBuilder> TreePath)
+		{
+			if (View.MayMatchAnyFilesInSubDirectory(DepotDir, Comparison))
+			{
+				P4.GetFileMetaDataCmdOptions Options = new P4.GetFileMetaDataCmdOptions(P4.GetFileMetadataCmdFlags.FileSize, null, null, BatchSize, null, null, null);
+				IList<P4.FileMetaData> Files = Repository.GetFileMetaData(Options, P4.FileSpec.DepotSpec($"{DepotDir}...", new P4.ChangelistIdVersion(Change)));
+				if (Files != null && Files.Count < BatchSize)
+				{
+					await AddDepotFilesToSnapshotAsync(Files, View, Comparison, TreePath);
+					return;
+				}
+
+				IList<string> Dirs = Repository.GetDepotDirs(new P4.GetDepotDirsCmdOptions(P4.GetDepotDirsCmdFlags.None, null), $"{DepotDir}*@{Change}");
+				foreach (string Dir in Dirs)
+				{
+					await AddDepotDirToSnapshotAsync(Repository, Dir + "/", Change, View, Comparison, TreePath);
+				}
+			}
+
+			if (View.MayMatchAnyFilesInDirectory(DepotDir, Comparison))
+			{
+				P4.GetFileMetaDataCmdOptions Options = new P4.GetFileMetaDataCmdOptions(P4.GetFileMetadataCmdFlags.FileSize, null, null, -1, null, null, null);
+				IList<P4.FileMetaData> Files = Repository.GetFileMetaData(Options, P4.FileSpec.DepotSpec($"{DepotDir}*", new P4.ChangelistIdVersion(Change)));
+				await AddDepotFilesToSnapshotAsync(Files, View, Comparison, TreePath);
+			}
+		}
+
+		/// <summary>
+		/// Adds a set of files in the depot to a snapshot
+		/// </summary>
+		/// <param name="Files">The files to add</param>
+		/// <param name="View">View for the stream</param>
+		/// <param name="Comparison">Comparison type to use for names</param>
+		/// <param name="TreePath">The last computed path through the tree</param>
+		async Task AddDepotFilesToSnapshotAsync(IList<P4.FileMetaData>? Files, ViewMap View, StringComparison Comparison, List<StreamTreeBuilder> TreePath)
+		{
+			if (Files != null)
+			{
+				foreach (P4.FileMetaData File in Files)
+				{
+					if (File.Digest != null)
+					{
+						string? TargetFile;
+						if (View.TryMapFile(File.DepotPath.Path, Comparison, out TargetFile))
+						{
+							await AddFileAsync(TargetFile, File, TreePath);
+						}
+					}
+				}
+			}
+		}
+
+		/// <summary>
+		/// Adds a single depot file to a snapshot
+		/// </summary>
+		/// <param name="Path">The stream path for the file</param>
+		/// <param name="MetaData">The file metadata</param>
+		/// <param name="TreePath">The last computed path through the tree</param>
+		async Task AddFileAsync(string Path, P4.FileMetaData MetaData, List<StreamTreeBuilder> TreePath)
+		{
+			string[] Segments = Path.Split('/');
+
+			StreamTreeBuilder Node = TreePath[0];
+			for (int Idx = 0; Idx < Segments.Length - 1; Idx++)
+			{
+				Utf8String Segment = Segments[Idx];
+
+				StreamTreeBuilder NextNode = await GetChildTreeBuilder(Node, Segment);
+				if (Idx + 1 < TreePath.Count && TreePath[Idx + 1] != NextNode)
+				{
+					await CollapseTreeAsync(Node, TreePath[Idx + 1]);
+				}
+				if (Idx + 1 >= TreePath.Count)
+				{
+					TreePath.Add(NextNode);
+				}
+
+				Node = NextNode;
+			}
+
+			StreamFile File = new StreamFile(MetaData.DepotPath.Path, MetaData.FileSize, new FileContentId(Md5Hash.Parse(MetaData.Digest), MetaData.HeadType.ToString()), MetaData.HeadRev);
+			Node.NameToFile[Segments[Segments.Length - 1]] = File;
+		}
+
+		/// <summary>
+		/// Gets a child tree builder for a particular node, deserializing an existing tree from the object store if necessary
+		/// </summary>
+		/// <param name="Node">The parent node</param>
+		/// <param name="Name">Name of the child tree</param>
+		/// <returns>The child tree, or a new one if it doesn't exist</returns>
+		async Task<StreamTreeBuilder> GetChildTreeBuilder(StreamTreeBuilder Node, Utf8String Name)
+		{
+			StreamTreeBuilder? NextNode;
+			if (!Node.NameToTreeBuilder.TryGetValue(Name, out NextNode))
+			{
+				StreamTreeRef? NextRef;
+				if (Node.NameToTree.TryGetValue(Name, out NextRef))
+				{
+					NextNode = await ExpandTreeAsync(Node, Name, NextRef);
+				}
+				else
+				{
+					NextNode = new StreamTreeBuilder();
+				}
+				Node.NameToTreeBuilder.Add(Name, NextNode);
+			}
+			return NextNode;
+		}
+
+		/// <summary>
+		/// Expands one level of the tree that has previously been serialized to the object store
+		/// </summary>
+		/// <param name="Parent">The parent node in the tree</param>
+		/// <param name="Name">Name of the </param>
+		/// <param name="ChildRef"></param>
+		/// <returns></returns>
+		async Task<StreamTreeBuilder> ExpandTreeAsync(StreamTreeBuilder Parent, Utf8String Name, StreamTreeRef ChildRef)
+		{
+			CbObject Object = await ObjectCollection.GetRequiredObjectAsync(NamespaceId, ChildRef.Hash);
+			StreamTreeBuilder Node = new StreamTreeBuilder(new StreamTree(Object, ChildRef.Path));
+			Parent.NameToTree.Remove(Name);
+			return Node;
+		}
+
+		/// <summary>
+		/// Collapse a subtree under a given child
+		/// </summary>
+		/// <param name="Parent"></param>
+		/// <param name="Child"></param>
+		/// <returns></returns>
+		async Task CollapseTreeAsync(StreamTreeBuilder Parent, StreamTreeBuilder Child)
+		{
+			foreach ((Utf8String Name, StreamTreeBuilder Node) in Parent.NameToTreeBuilder)
+			{
+				if (Node == Child)
+				{
+					StreamTreeRef ChildRef = await EncodeTreeAsync(Child);
+					Parent.NameToTreeBuilder.Remove(Name);
+					Parent.NameToTree.Add(Name, ChildRef);
+					break;
+				}
+			}
+		}
+
+		/// <summary>
+		/// Writes a subtree to the object store
+		/// </summary>
+		/// <param name="TreeBuilder">Root of the tree to write</param>
+		/// <returns>New tree reference for the serialized tree</returns>
+		async Task<StreamTreeRef> EncodeTreeAsync(StreamTreeBuilder TreeBuilder)
+		{
+			Dictionary<IoHash, CbObject> Objects = new Dictionary<IoHash, CbObject>();
+			StreamTreeRef Ref = TreeBuilder.Encode(Objects);
+
+			foreach ((IoHash Hash, CbObject Object) in Objects)
+			{
+				await ObjectCollection.AddAsync(NamespaceId, Hash, Object);
+			}
+
+			return Ref;
+		}
+
+		#endregion
 	}
 }
diff --git a/Engine/Source/Programs/Horde/HordeServer/GlobalSuppressions.cs b/Engine/Source/Programs/Horde/HordeServer/GlobalSuppressions.cs
--- a/Engine/Source/Programs/Horde/HordeServer/GlobalSuppressions.cs
+++ b/Engine/Source/Programs/Horde/HordeServer/GlobalSuppressions.cs
@@ -40,3 +40,4 @@ using System.Diagnostics.CodeAnalysis;
 [assembly: SuppressMessage("Design", "CA1031:Do not catch general exception types", Justification = "<Pending>", Scope = "member", Target = "~M:HordeServer.Services.ScheduleService.UpdateSchedulesAsync~System.Threading.Tasks.Task{System.DateTimeOffset}")]
 [assembly: SuppressMessage("Design", "CA1030:Use events where appropriate", Justification = "<Pending>", Scope = "module")]
 [assembly: SuppressMessage("Design", "CA1031:Do not catch general exception types", Justification = "<Pending>", Scope = "module")]
+[assembly: SuppressMessage("Design", "CA1032:Implement standard exception constructors", Justification = "<Pending>", Scope = "module")]
diff --git a/Engine/Source/Programs/Horde/HordeServer/HordeServer.csproj b/Engine/Source/Programs/Horde/HordeServer/HordeServer.csproj
--- a/Engine/Source/Programs/Horde/HordeServer/HordeServer.csproj
+++ b/Engine/Source/Programs/Horde/HordeServer/HordeServer.csproj
@@ -79,6 +79,7 @@
     <ProjectReference Include="..\..\Shared\EpicGames.Horde.Common\EpicGames.Horde.Common.csproj" />
     <ProjectReference Include="..\..\Shared\EpicGames.Horde\EpicGames.Horde.csproj" />
     <ProjectReference Include="..\..\Shared\EpicGames.IoHash\EpicGames.IoHash.csproj" />
+    <ProjectReference Include="..\..\Shared\EpicGames.Perforce.Managed\EpicGames.Perforce.Managed.csproj" />
     <ProjectReference Include="..\..\Shared\EpicGames.Perforce\EpicGames.Perforce.csproj" />
     <ProjectReference Include="..\..\Shared\EpicGames.Redis\EpicGames.Redis.csproj" />
     <ProjectReference Include="..\..\Shared\EpicGames.Serialization\EpicGames.Serialization.csproj" />
diff --git a/Engine/Source/Programs/Horde/HordeServer/Storage/IObjectCollection.cs b/Engine/Source/Programs/Horde/HordeServer/Storage/IObjectCollection.cs
--- a/Engine/Source/Programs/Horde/HordeServer/Storage/IObjectCollection.cs
+++ b/Engine/Source/Programs/Horde/HordeServer/Storage/IObjectCollection.cs
@@ -52,6 +52,34 @@ namespace HordeServer.Storage
 		Task<List<IoHash>> ExistsAsync(NamespaceId NsId, IEnumerable<IoHash> Hashes);
 	}
 
+	/// <summary>
+	/// Exeception thrown when an object is not found
+	/// </summary>
+	public class CbObjectNotFoundException : Exception
+	{
+		/// <summary>
+		/// Namespace of the missing object
+		/// </summary>
+		public NamespaceId NsId { get; }
+
+		/// <summary>
+		/// Hash of the missing object
+		/// </summary>
+		public IoHash Hash { get; }
+
+		/// <summary>
+		/// Constructor
+		/// </summary>
+		/// <param name="NsId"></param>
+		/// <param name="Hash"></param>
+		public CbObjectNotFoundException(NamespaceId NsId, IoHash Hash)
+			: base($"Missing object {Hash} in {NsId}")
+		{
+			this.NsId = NsId;
+			this.Hash = Hash;
+		}
+	}
+
 	/// <summary>
 	/// Extension methods for <see cref="IObjectCollection"/>
 	/// </summary>
@@ -73,6 +101,23 @@ namespace HordeServer.Storage
 			return Hash;
 		}
 
+		/// <summary>
+		/// Gets a typed object from the store
+		/// </summary>
+		/// <param name="ObjectCollection"></param>
+		/// <param name="NsId"></param>
+		/// <param name="Hash"></param>
+		/// <returns></returns>
+		public static async Task<CbObject> GetRequiredObjectAsync(this IObjectCollection ObjectCollection, NamespaceId NsId, IoHash Hash)
+		{
+			CbObject? Object = await ObjectCollection.GetAsync(NsId, Hash);
+			if (Object == null)
+			{
+				throw new CbObjectNotFoundException(NsId, Hash);
+			}
+			return Object;
+		}
+
 		/// <summary>
 		/// Gets a typed object from the store
 		/// </summary>
@@ -91,6 +136,24 @@ namespace HordeServer.Storage
 			return CbSerializer.Deserialize<T>(Object.AsField());
 		}
 
+		/// <summary>
+		/// Gets a typed object from the store
+		/// </summary>
+		/// <typeparam name="T"></typeparam>
+		/// <param name="ObjectCollection"></param>
+		/// <param name="NsId"></param>
+		/// <param name="Hash"></param>
+		/// <returns></returns>
+		public static async Task<T> GetRequiredObjectAsync<T>(this IObjectCollection ObjectCollection, NamespaceId NsId, IoHash Hash) where T : class
+		{
+			CbObject? Object = await ObjectCollection.GetAsync(NsId, Hash);
+			if (Object == null)
+			{
+				throw new CbObjectNotFoundException(NsId, Hash);
+			}
+			return CbSerializer.Deserialize<T>(Object.AsField());
+		}
+
 		/// <summary>
 		/// Test whether a single blob exists
 		/// </summary>
diff --git a/Engine/Source/Programs/Horde/HordeServer/Utilities/ViewMap.cs b/Engine/Source/Programs/Horde/HordeServer/Utilities/ViewMap.cs
new file mode 100644
--- /dev/null
+++ b/Engine/Source/Programs/Horde/HordeServer/Utilities/ViewMap.cs
@@ -0,0 +1,384 @@
+// Copyright Epic Games, Inc. All Rights Reserved.
+
+using System;
+using System.Collections.Generic;
+using System.Diagnostics;
+using System.Diagnostics.CodeAnalysis;
+using System.Text.RegularExpressions;
+
+namespace HordeServer.Utilities
+{
+	using P4 = Perforce.P4;
+
+	/// <summary>
+	/// Stores a mapping from one set of paths to another
+	/// </summary>
+	public class ViewMap
+	{
+		/// <summary>
+		/// List of entries making up the view
+		/// </summary>
+		public List<ViewMapEntry> Entries { get; } = new List<ViewMapEntry>();
+
+		/// <summary>
+		/// Default constructor
+		/// </summary>
+		public ViewMap()
+		{
+		}
+
+		/// <summary>
+		/// Constructor
+		/// </summary>
+		/// <param name="ViewMap">Peforce viewmap definition</param>
+		public ViewMap(P4.ViewMap ViewMap)
+		{
+			foreach (P4.MapEntry Entry in ViewMap)
+			{
+				Entries.Add(new ViewMapEntry(Entry));
+			}
+		}
+
+		/// <summary>
+		/// Determines if a file is included in the view
+		/// </summary>
+		/// <param name="File">The file to test</param>
+		/// <param name="Comparison">The comparison type</param>
+		/// <returns>True if the file is included in the view</returns>
+		public bool MatchFile(string File, StringComparison Comparison)
+		{
+			bool Included = false;
+			foreach (ViewMapEntry Entry in Entries)
+			{
+				if (Entry.MatchFile(File, Comparison))
+				{
+					Included = Entry.Include;
+				}
+			}
+			return Included;
+		}
+
+		/// <summary>
+		/// Attempts to convert a source file to its target path
+		/// </summary>
+		/// <param name="SourceFile"></param>
+		/// <param name="Comparison">The comparison type</param>
+		/// <param name="TargetFile"></param>
+		/// <returns></returns>
+		public bool TryMapFile(string SourceFile, StringComparison Comparison, [NotNullWhen(true)] out string? TargetFile)
+		{
+			ViewMapEntry? MapEntry = null;
+			foreach (ViewMapEntry Entry in Entries)
+			{
+				if (Entry.MatchFile(SourceFile, Comparison))
+				{
+					MapEntry = Entry;
+				}
+			}
+
+			if (MapEntry != null && MapEntry.Include)
+			{
+				TargetFile = MapEntry.MapFile(SourceFile);
+				return true;
+			}
+			else
+			{
+				TargetFile = null;
+				return false;
+			}
+		}
+
+		/// <summary>
+		/// Tests whether the view includes any files directly under the given directory (not including any subdirectories).
+		/// </summary>
+		/// <param name="Dir">The directory to test; should end with a Slash.</param>
+		/// <param name="Comparison">The comparison type</param>
+		/// <returns>True if the view may include any files in the given directory</returns>
+		public bool MayMatchAnyFilesInDirectory(string Dir, StringComparison Comparison)
+		{
+			bool MayInclude = false;
+			foreach (ViewMapEntry Entry in Entries)
+			{
+				if (Entry.Include)
+				{
+					if (Entry.MayMatchAnyFilesInDirectory(Dir, Comparison))
+					{
+						MayInclude = true;
+					}
+				}
+				else
+				{
+					if (Entry.MatchAllFilesInDirectory(Dir, Comparison))
+					{
+						MayInclude = false;
+					}
+				}
+			}
+			return MayInclude;
+		}
+
+		/// <summary>
+		/// Tests whether the view may include any files in subdirectories of the given path
+		/// </summary>
+		/// <param name="Dir">The directory to test; should end with a Slash.</param>
+		/// <param name="Comparison">The comparison type</param>
+		/// <returns>True if the pattern may match files under the given directory</returns>
+		public bool MayMatchAnyFilesInSubDirectory(string Dir, StringComparison Comparison)
+		{
+			bool MayInclude = false;
+			foreach (ViewMapEntry Entry in Entries)
+			{
+				if (Entry.Include)
+				{
+					if (Entry.MayMatchAnyFilesInSubDirectory(Dir, Comparison))
+					{
+						MayInclude = true;
+					}
+				}
+				else
+				{
+					if (Entry.MatchAllFilesInDirectory(Dir, Comparison))
+					{
+						MayInclude = false;
+					}
+				}
+			}
+			return MayInclude;
+		}
+
+		/// <summary>
+		/// Gets the root paths from the view entries
+		/// </summary>
+		/// <returns></returns>
+		public List<string> GetRootPaths(StringComparison Comparison)
+		{
+			List<string> RootPaths = new List<string>();
+			foreach (ViewMapEntry Entry in Entries)
+			{
+				if (Entry.Include)
+				{
+					int LastSlashIdx = Entry.SourcePrefix.LastIndexOf('/');
+					string RootPath = Entry.SourcePrefix.Substring(0, LastSlashIdx + 1);
+
+					for (int Idx = 0; ; Idx++)
+					{
+						if (Idx == RootPaths.Count)
+						{
+							RootPaths.Add(RootPath);
+							break;
+						}
+						else if (RootPaths[Idx].StartsWith(RootPath, Comparison))
+						{
+							RootPaths[Idx] = RootPath;
+							break;
+						}
+						else if (RootPath.StartsWith(RootPaths[Idx], Comparison))
+						{
+							break;
+						}
+					}
+				}
+			}
+			return RootPaths;
+		}
+	}
+
+	/// <summary>
+	/// Entry within a ViewMap
+	/// </summary>
+	public class ViewMapEntry
+	{
+		/// <summary>
+		/// Whether to include files matching this pattern
+		/// </summary>
+		public bool Include { get; }
+
+		/// <summary>
+		/// The wildcard string - either '*' or '...'
+		/// </summary>
+		public string? Wildcard { get; }
+
+		/// <summary>
+		/// The source part of the pattern before the wildcard
+		/// </summary>
+		public string SourcePrefix { get; }
+
+		/// <summary>
+		/// The source part of the pattern after the wildcard. Perforce does not permit a slash to be in this part of the mapping.
+		/// </summary>
+		public string SourceSuffix { get; }
+
+		/// <summary>
+		/// The target mapping for the pattern before the wildcard
+		/// </summary>
+		public string TargetPrefix { get; }
+
+		/// <summary>
+		/// The target mapping for the pattern after the wildcard
+		/// </summary>
+		public string TargetSuffix { get; }
+
+		/// <summary>
+		/// Tests if the entry has a file wildcard ('*')
+		/// </summary>
+		/// <returns>True if the entry has a file wildcard</returns>
+		public bool IsFileWildcard() => Wildcard != null && Wildcard.Length == 1;
+
+		/// <summary>
+		/// Tests if the entry has a path wildcard ('...')
+		/// </summary>
+		/// <returns>True if the entry has a path wildcard</returns>
+		public bool IsPathWildcard() => Wildcard != null && Wildcard.Length == 3;
+
+		/// <summary>
+		/// Constructor
+		/// </summary>
+		/// <param name="Include"></param>
+		/// <param name="Source"></param>
+		/// <param name="Target"></param>
+		public ViewMapEntry(bool Include, string Source, string Target)
+		{
+			this.Include = Include;
+
+			Match Match = Regex.Match(Source, @"^(.*)(\*|\.\.\.)(.*)$");
+			if (Match.Success)
+			{
+				SourcePrefix = Match.Groups[1].Value;
+				SourceSuffix = Match.Groups[3].Value;
+				Wildcard = Match.Groups[2].Value;
+
+				int OtherIdx = Target.IndexOf(Wildcard, StringComparison.Ordinal);
+				TargetPrefix = Target.Substring(0, OtherIdx);
+				TargetSuffix = Target.Substring(OtherIdx + Wildcard.Length);
+			}
+			else
+			{
+				SourcePrefix = Source;
+				SourceSuffix = String.Empty;
+				TargetPrefix = Target;
+				TargetSuffix = String.Empty;
+			}
+		}
+
+		/// <summary>
+		/// Constructor
+		/// </summary>
+		/// <param name="Include"></param>
+		/// <param name="Wildcard"></param>
+		/// <param name="SourcePrefix"></param>
+		/// <param name="SourceSuffix"></param>
+		/// <param name="TargetPrefix"></param>
+		/// <param name="TargetSuffix"></param>
+		public ViewMapEntry(bool Include, string? Wildcard, string SourcePrefix, string SourceSuffix, string TargetPrefix, string TargetSuffix)
+		{
+			this.Include = Include;
+			this.Wildcard = Wildcard;
+			this.SourcePrefix = SourcePrefix;
+			this.SourceSuffix = SourceSuffix;
+			this.TargetPrefix = TargetPrefix;
+			this.TargetSuffix = TargetSuffix;
+		}
+
+		/// <summary>
+		/// Constructor
+		/// </summary>
+		/// <param name="Entry"></param>
+		public ViewMapEntry(P4.MapEntry Entry)
+			: this(Entry.Type == P4.MapType.Include, Entry.Left.Path, Entry.Right.Path)
+		{
+		}
+
+		/// <summary>
+		/// Maps a file to the target path
+		/// </summary>
+		/// <param name="SourceFile"></param>
+		/// <returns></returns>
+		public string MapFile(string SourceFile)
+		{
+			int Count = SourceFile.Length - SourceSuffix.Length - SourcePrefix.Length;
+			return TargetPrefix + SourceFile.Substring(SourcePrefix.Length, Count) + TargetSuffix;
+		}
+
+		/// <summary>
+		/// Determine if a file matches the current entry
+		/// </summary>
+		/// <param name="Path">Path to the file</param>
+		/// <param name="Comparison">The comparison type</param>
+		/// <returns>True if the path matches the entry</returns>
+		public bool MatchFile(string Path, StringComparison Comparison)
+		{
+			if (Wildcard == null)
+			{
+				return Path.Equals(SourcePrefix, Comparison);
+			}
+			else
+			{
+				if (!Path.StartsWith(SourcePrefix, Comparison) || !Path.EndsWith(SourceSuffix, Comparison))
+				{
+					return false;
+				}
+				if (IsFileWildcard() && Path.IndexOf('/', SourcePrefix.Length, Path.Length - SourceSuffix.Length - SourcePrefix.Length) != -1)
+				{
+					return false;
+				}
+				return true;
+			}
+		}
+
+		/// <summary>
+		/// Determine if the pattern may match any files directly under the given directory
+		/// </summary>
+		/// <param name="Dir">The directory to test</param>
+		/// <param name="Comparison">The comparison type</param>
+		/// <returns>True if the pattern may match a file in this directory</returns>
+		public bool MayMatchAnyFilesInDirectory(string Dir, StringComparison Comparison)
+		{
+			Debug.Assert(Dir.EndsWith("/", StringComparison.Ordinal));
+			if (Dir.StartsWith(SourcePrefix, Comparison))
+			{
+				if (IsPathWildcard() || Dir.IndexOf('/', SourcePrefix.Length) == -1)
+				{
+					return true;
+				}
+			}
+			return false;
+		}
+
+		/// <summary>
+		/// Determine if the pattern may match any files in subdirectories of the given directory
+		/// </summary>
+		/// <param name="Dir">The directory to test</param>
+		/// <param name="Comparison">The comparison type</param>
+		/// <returns>True if the pattern may match a file in this directory</returns>
+		public bool MayMatchAnyFilesInSubDirectory(string Dir, StringComparison Comparison)
+		{
+			Debug.Assert(Dir.EndsWith("/", StringComparison.Ordinal));
+			if (Dir.StartsWith(SourcePrefix, Comparison))
+			{
+				if (IsPathWildcard())
+				{
+					return true;
+				}
+			}
+			else
+			{
+				if (SourcePrefix.StartsWith(Dir, Comparison))
+				{
+					return true;
+				}
+			}
+			return false;
+		}
+
+		/// <summary>
+		/// Determine if the pattern matches all files in the given directory
+		/// </summary>
+		/// <param name="Dir">The directory to test</param>
+		/// <param name="Comparison">The comparison type</param>
+		/// <returns>True if the pattern may match a file in this directory</returns>
+		public bool MatchAllFilesInDirectory(string Dir, StringComparison Comparison)
+		{
+			return Wildcard != null && Wildcard.Length == 3 && Dir.StartsWith(SourcePrefix, Comparison);
+		}
+	}
+}
