Horde: Support incremental updates of stream trees.

diff --git a/Engine/Source/Programs/Horde/HordeServer/Commits/Impl/CommitService.cs b/Engine/Source/Programs/Horde/HordeServer/Commits/Impl/CommitService.cs
--- a/Engine/Source/Programs/Horde/HordeServer/Commits/Impl/CommitService.cs
+++ b/Engine/Source/Programs/Horde/HordeServer/Commits/Impl/CommitService.cs
@@ -276,13 +276,12 @@ namespace HordeServer.Commits.Impl
 		/// <param name="Repository">The Perforce repository to mirror</param>
 		/// <param name="Stream">Name of the stream to be mirrored</param>
 		/// <param name="Change">Changelist to replicate</param>
-		/// <param name="Comparison">Comparison to use for names</param>
-		Task CreateSnapshotAsync(P4.Repository Repository, string Stream, int Change, StringComparison Comparison)
+		Task<StreamTreeRef> CreateTreeAsync(P4.Repository Repository, string Stream, int Change)
 		{
 			// Get the stream spec
 			P4.Stream StreamSpec = Repository.GetStream($"{Stream}@{Change}", new P4.Options { { "-v", "" } });
-			ViewMap ViewMap = new ViewMap(StreamSpec.View);
-			return CreateSnapshotAsync(Repository, ViewMap, Change, Comparison);
+			ViewMap View = new ViewMap(StreamSpec.View);
+			return CreateTreeAsync(Repository, View, Change);
 		}
 
 		/// <summary>
@@ -291,9 +290,10 @@ namespace HordeServer.Commits.Impl
 		/// <param name="Repository">The Perforce repository to mirror</param>
 		/// <param name="View">View for the stream</param>
 		/// <param name="Change">Changelist to replicate</param>
-		/// <param name="Comparison">Comparison to use for names</param>
-		async Task<StreamTreeRef> CreateSnapshotAsync(P4.Repository Repository, ViewMap View, int Change, StringComparison Comparison)
+		async Task<StreamTreeRef> CreateTreeAsync(P4.Repository Repository, ViewMap View, int Change)
 		{
+			StringComparison Comparison = Repository.Server.Metadata.CaseSensitive ? StringComparison.Ordinal : StringComparison.OrdinalIgnoreCase;
+
 			// In order to optimize for the common case where we're adding multiple files to the same directory, we keep an open list of
 			// StreamTreeBuilder objects down to the last modified node. Whenever we need to move to a different node, we write the 
 			// previous path to the object store.
@@ -315,7 +315,7 @@ namespace HordeServer.Commits.Impl
 				await AddDepotDirToSnapshotAsync(Repository, RootPath, Change, FilteredView, Comparison, TreePath);
 			}
 
-			// Write the root tree
+			// Encode the tree
 			return await EncodeTreeAsync(TreePath[0]);
 		}
 
@@ -395,7 +395,7 @@ namespace HordeServer.Commits.Impl
 			{
 				Utf8String Segment = Segments[Idx];
 
-				StreamTreeBuilder NextNode = await GetChildTreeBuilder(Node, Segment);
+				StreamTreeBuilder NextNode = await FindOrAddTreeAsync(Node, Segment);
 				if (Idx + 1 < TreePath.Count && TreePath[Idx + 1] != NextNode)
 				{
 					await CollapseTreeAsync(Node, TreePath[Idx + 1]);
@@ -408,48 +408,130 @@ namespace HordeServer.Commits.Impl
 				Node = NextNode;
 			}
 
-			StreamFile File = new StreamFile(MetaData.DepotPath.Path, MetaData.FileSize, new FileContentId(Md5Hash.Parse(MetaData.Digest), MetaData.HeadType.ToString()), MetaData.HeadRev);
-			Node.NameToFile[Segments[Segments.Length - 1]] = File;
+			Node.NameToFile[Segments[Segments.Length - 1]] = CreateStreamFile(MetaData);
+		}
+
+		#endregion
+
+		#region Incremental snapshot generation
+
+		/// <summary>
+		/// Creates a snapshot for the given stream, given an initial state and list of modified files
+		/// </summary>
+		/// <param name="TreeRef">The tree to update</param>
+		/// <param name="Repository"></param>
+		/// <param name="View"></param>
+		/// <param name="Change"></param>
+		/// <returns></returns>
+		Task<StreamTreeRef> UpdateTreeAsync(StreamTreeRef TreeRef, P4.Repository Repository, ViewMap View, int Change)
+		{
+			P4.Changelist Changelist = Repository.GetChangelist(Change);
+			StringComparison Comparison = Repository.Server.Metadata.CaseSensitive ? StringComparison.Ordinal : StringComparison.OrdinalIgnoreCase;
+			return UpdateTreeAsync(Changelist.Files, View, Comparison, TreeRef);
 		}
 
 		/// <summary>
-		/// Gets a child tree builder for a particular node, deserializing an existing tree from the object store if necessary
+		/// Creates a snapshot for the given stream, given an initial state and list of modified files
 		/// </summary>
-		/// <param name="Node">The parent node</param>
-		/// <param name="Name">Name of the child tree</param>
-		/// <returns>The child tree, or a new one if it doesn't exist</returns>
-		async Task<StreamTreeBuilder> GetChildTreeBuilder(StreamTreeBuilder Node, Utf8String Name)
+		/// <param name="Files"></param>
+		/// <param name="View"></param>
+		/// <param name="Comparison"></param>
+		/// <param name="PrevTreeRef"></param>
+		/// <returns></returns>
+		async Task<StreamTreeRef> UpdateTreeAsync(IList<P4.FileMetaData> Files, ViewMap View, StringComparison Comparison, StreamTreeRef PrevTreeRef)
 		{
-			StreamTreeBuilder? NextNode;
-			if (!Node.NameToTreeBuilder.TryGetValue(Name, out NextNode))
+			StreamTreeBuilder Tree = await ReadTreeAsync(PrevTreeRef);
+
+			// Update the tree
+			foreach (P4.FileMetaData File in Files)
 			{
-				StreamTreeRef? NextRef;
-				if (Node.NameToTree.TryGetValue(Name, out NextRef))
-				{
-					NextNode = await ExpandTreeAsync(Node, Name, NextRef);
-				}
-				else
+				string? LocalPathStr;
+				if (View.TryMapFile(File.DepotPath.Path, Comparison, out LocalPathStr))
 				{
-					NextNode = new StreamTreeBuilder();
+					Utf8String LocalPath = LocalPathStr;
+
+					StreamTreeBuilder Node = Tree;
+					for (; ; )
+					{
+						int NextIdx = LocalPath.IndexOf('/');
+						if (NextIdx == -1)
+						{
+							break;
+						}
+
+						Utf8String Name = LocalPath[0..NextIdx];
+						Node = await FindOrAddTreeAsync(Node, Name);
+
+						LocalPath = LocalPath[(NextIdx + 1)..];
+					}
+
+					Utf8String FileName = LocalPath;
+					if (File.HeadRev > 0)
+					{
+						Node.NameToFile[FileName] = CreateStreamFile(File);
+					}
+					else
+					{
+						if (!Node.NameToFile.Remove(FileName))
+						{
+							// TODO: Handle mismatched case?
+							throw new NotImplementedException();
+						}
+					}
 				}
-				Node.NameToTreeBuilder.Add(Name, NextNode);
 			}
-			return NextNode;
+
+			return await EncodeTreeAsync(Tree);
+		}
+
+		#endregion
+
+		#region Misc tree manipulation
+
+		/// <summary>
+		/// Gets a <see cref="StreamTreeBuilder"/> for a subtree with the given name
+		/// </summary>
+		/// <param name="Tree"></param>
+		/// <param name="Name"></param>
+		/// <returns></returns>
+		async Task<StreamTreeBuilder> FindOrAddTreeAsync(StreamTreeBuilder Tree, Utf8String Name)
+		{
+			StreamTreeBuilder? Result;
+			if (Tree.NameToTreeBuilder.TryGetValue(Name, out Result))
+			{
+				return Result;
+			}
+
+			StreamTreeRef? TreeRef;
+			if (Tree.NameToTree.TryGetValue(Name, out TreeRef))
+			{
+				Result = await ReadTreeAsync(TreeRef);
+
+				Tree.NameToTree.Remove(Name);
+				Tree.NameToTreeBuilder.Add(Name, Result);
+
+				return Result;
+			}
+
+			Result = new StreamTreeBuilder();
+			Tree.NameToTreeBuilder[Name] = Result;
+
+			return Result;
 		}
 
 		/// <summary>
-		/// Expands one level of the tree that has previously been serialized to the object store
+		/// Reads a tree builder for the given tree reference
 		/// </summary>
-		/// <param name="Parent">The parent node in the tree</param>
-		/// <param name="Name">Name of the </param>
-		/// <param name="ChildRef"></param>
+		/// <param name="TreeRef"></param>
 		/// <returns></returns>
-		async Task<StreamTreeBuilder> ExpandTreeAsync(StreamTreeBuilder Parent, Utf8String Name, StreamTreeRef ChildRef)
+		async Task<StreamTreeBuilder> ReadTreeAsync(StreamTreeRef TreeRef)
 		{
-			CbObject Object = await ObjectCollection.GetRequiredObjectAsync(NamespaceId, ChildRef.Hash);
-			StreamTreeBuilder Node = new StreamTreeBuilder(new StreamTree(Object, ChildRef.Path));
-			Parent.NameToTree.Remove(Name);
-			return Node;
+			CbObject? Object = await ObjectCollection.GetAsync(NamespaceId, TreeRef.Hash);
+			if (Object == null)
+			{
+				throw new Exception($"Missing object {TreeRef.Hash} referenced by tree for {TreeRef.Path}");
+			}
+			return new StreamTreeBuilder(new StreamTree(Object, TreeRef.Path + "/"));
 		}
 
 		/// <summary>
@@ -490,6 +572,47 @@ namespace HordeServer.Commits.Impl
 			return Ref;
 		}
 
+		/// <summary>
+		/// Creates a StreamFile from a FileMetaData object
+		/// </summary>
+		/// <param name="MetaData"></param>
+		/// <returns></returns>
+		static StreamFile CreateStreamFile(P4.FileMetaData MetaData)
+		{
+			P4.FileType FileType = MetaData.Type ?? MetaData.HeadType;
+			return new StreamFile(MetaData.DepotPath.Path, MetaData.FileSize, new FileContentId(Md5Hash.Parse(MetaData.Digest), FileType.ToString()), MetaData.HeadRev);
+		}
+
+		/// <summary>
+		/// Print the contents of a snapshot
+		/// </summary>
+		/// <param name="TreeRef"></param>
+		/// <returns></returns>
+		Task PrintSnapshotAsync(StreamTreeRef TreeRef)
+		{
+			return PrintSnapshotInternalAsync(TreeRef, "/");
+		}
+
+		/// <summary>
+		/// Print the contents of a snapshot
+		/// </summary>
+		/// <param name="TreeRef"></param>
+		/// <param name="Prefix"></param>
+		/// <returns></returns>
+		async Task PrintSnapshotInternalAsync(StreamTreeRef TreeRef, string Prefix)
+		{
+			StreamTree Tree = new StreamTree(await ObjectCollection.GetRequiredObjectAsync(NamespaceId, TreeRef.Hash), TreeRef.Path);
+
+			foreach ((Utf8String Name, StreamFile File) in Tree.NameToFile)
+			{
+				Logger.LogInformation($"{Prefix}/{Name} = {File.Path}#{File.Revision}");
+			}
+			foreach ((Utf8String Name, StreamTreeRef ChildTreeRef) in Tree.NameToTree)
+			{
+				await PrintSnapshotInternalAsync(ChildTreeRef, $"{Prefix}/{Name}");
+			}
+		}
+
 		#endregion
 	}
 }
