Horde: Separate out functionality for deriving a stream view from a stream spec.

diff --git a/Engine/Source/Programs/Horde/HordeServer/Commits/Impl/CommitService.cs b/Engine/Source/Programs/Horde/HordeServer/Commits/Impl/CommitService.cs
--- a/Engine/Source/Programs/Horde/HordeServer/Commits/Impl/CommitService.cs
+++ b/Engine/Source/Programs/Horde/HordeServer/Commits/Impl/CommitService.cs
@@ -270,20 +270,6 @@ namespace HordeServer.Commits.Impl
 		/// </summary>
 		const int BatchSize = 16384;
 
-		/// <summary>
-		/// Creates a snapshot of the repository at a particular changelist
-		/// </summary>
-		/// <param name="Repository">The Perforce repository to mirror</param>
-		/// <param name="Stream">Name of the stream to be mirrored</param>
-		/// <param name="Change">Changelist to replicate</param>
-		Task<StreamTreeRef> CreateTreeAsync(P4.Repository Repository, string Stream, int Change)
-		{
-			// Get the stream spec
-			P4.Stream StreamSpec = Repository.GetStream($"{Stream}@{Change}", new P4.Options { { "-v", "" } });
-			ViewMap View = new ViewMap(StreamSpec.View);
-			return CreateTreeAsync(Repository, View, Change);
-		}
-
 		/// <summary>
 		/// Creates a snapshot of a stream at a particular changelist
 		/// </summary>
@@ -488,6 +474,19 @@ namespace HordeServer.Commits.Impl
 
 		#region Misc tree manipulation
 
+		/// <summary>
+		/// Gets the view for a stream
+		/// </summary>
+		/// <param name="Repository">The Perforce repository to mirror</param>
+		/// <param name="Stream">Name of the stream to be mirrored</param>
+		/// <param name="Change">Changelist to replicate</param>
+		static ViewMap GetStreamView(P4.Repository Repository, string Stream, int Change)
+		{
+			// Get the stream spec
+			P4.Stream StreamSpec = Repository.GetStream($"{Stream}@{Change}", new P4.Options { { "-v", "" } });
+			return new ViewMap(StreamSpec.View);
+		}
+
 		/// <summary>
 		/// Gets a <see cref="StreamTreeBuilder"/> for a subtree with the given name
 		/// </summary>
