Horde: Remove HordeAgent component from installer, register it with softwate service at startup, and fix up download logic
#rnx

diff --git a/Engine/Source/Programs/Horde/HordeAgent/Program.cs b/Engine/Source/Programs/Horde/HordeAgent/Program.cs
--- a/Engine/Source/Programs/Horde/HordeAgent/Program.cs
+++ b/Engine/Source/Programs/Horde/HordeAgent/Program.cs
@@ -135,6 +135,11 @@ namespace HordeAgent
 			// Check if there's a matching command
 			Type? CommandType = null;
 			CommandAttribute? CommandAttribute = null;
+
+			string[] DefaultArgs = new string[] { "Service", "Run" };
+;			Type? DefaultCommandType = null;
+			CommandAttribute? DefaultCommandAttribute = null;
+
 			foreach ((CommandAttribute Attribute, Type Type) in Commands.OrderBy(x => x.Item1.Names.Length))
 			{
 				if (MatchCommand(Args, Attribute))
@@ -142,6 +147,13 @@ namespace HordeAgent
 					CommandType = Type;
 					CommandAttribute = Attribute;
 				}
+
+				if (MatchCommand(DefaultArgs, Attribute))
+				{
+					DefaultCommandType = Type;
+					DefaultCommandAttribute = Attribute;
+				}
+
 			}
 
 			// Check if there are any commands specified on the command line.
@@ -156,7 +168,7 @@ namespace HordeAgent
 					PrintCommands(Commands.Select(x => x.Item1), Logger);
 					return 1;
 				}
-				else
+				else if (Args.Length > 0 && Args[0].Equals("-Help", StringComparison.OrdinalIgnoreCase))
 				{
 					Logger.LogInformation("HordeAgent");
 					Logger.LogInformation("");
@@ -173,6 +185,17 @@ namespace HordeAgent
 					Logger.LogInformation("Specify \"Command -Help\" for command-specific help");
 					return 0;
 				}
+				else
+				{
+					CommandType = DefaultCommandType;
+					CommandAttribute = DefaultCommandAttribute;
+
+					if (CommandType == null || CommandAttribute == null)
+					{
+						Logger.LogError("Invalid default command");
+						return 1;
+					}
+				}
 			}
 
 			// Extract the 
diff --git a/Engine/Source/Programs/Horde/HordeDashboard/src/backend/Backend.ts b/Engine/Source/Programs/Horde/HordeDashboard/src/backend/Backend.ts
--- a/Engine/Source/Programs/Horde/HordeDashboard/src/backend/Backend.ts
+++ b/Engine/Source/Programs/Horde/HordeDashboard/src/backend/Backend.ts
@@ -1005,7 +1005,7 @@ export class Backend {
 
     downloadAgentZip() {
         try {
-            const url = `${this.serverUrl}/api/v1/agentsoftware/latest/download`;
+            const url = `${this.serverUrl}/api/v1/software/latest/download`;
             const link = document.createElement('a');
             link.href = url;
             document.body.appendChild(link);
diff --git a/Engine/Source/Programs/Horde/HordeServer/Controllers/SoftwareController.cs b/Engine/Source/Programs/Horde/HordeServer/Controllers/SoftwareController.cs
--- a/Engine/Source/Programs/Horde/HordeServer/Controllers/SoftwareController.cs
+++ b/Engine/Source/Programs/Horde/HordeServer/Controllers/SoftwareController.cs
@@ -144,7 +144,7 @@ namespace HordeServer.Controllers
 				return NotFound();
 			}
 
-			return new FileStreamResult(new MemoryStream(Software.Data!), new MediaTypeHeaderValue("application/octet-stream"));
+			return new FileStreamResult(new MemoryStream(Software.Data!), new MediaTypeHeaderValue("application/octet-stream")) { FileDownloadName = $"HordeAgent-{Software.Id.ToString()}.zip" };
 		}
 
 		/// <summary>
diff --git a/Engine/Source/Programs/Horde/HordeServer/Services/AgentSoftwareService.cs b/Engine/Source/Programs/Horde/HordeServer/Services/AgentSoftwareService.cs
--- a/Engine/Source/Programs/Horde/HordeServer/Services/AgentSoftwareService.cs
+++ b/Engine/Source/Programs/Horde/HordeServer/Services/AgentSoftwareService.cs
@@ -125,17 +125,42 @@ namespace HordeServer.Services
 		/// </summary>
 		LazyCachedValue<Task<AgentSoftwareChannels>> ChannelsDocument;
 
+		/// <summary>
+		/// Singleton instance of the client service
+		/// </summary>
+		private readonly SoftwareService SoftwareService;
+
+
+		/// <summary>
+		/// The server settings
+		/// </summary>
+		IOptionsMonitor<ServerSettings> Settings;
+
+		/// <summary>
+		///  Logger for controller
+		/// </summary>
+		private readonly ILogger<AgentSoftwareService> Logger;
+
+
 		/// <summary>
 		/// Constructor
 		/// </summary>
 		/// <param name="Collection">The software collection</param>
 		/// <param name="Singleton">The channels singleton</param>
-		public AgentSoftwareService(IAgentSoftwareCollection Collection, ISingletonDocument<AgentSoftwareChannels> Singleton)
+		/// <param name="SoftwareService">The software service</param>
+		/// <param name="Settings">The settings monitor</param>		
+		/// <param name="Logger">The logger instance</param>
+		public AgentSoftwareService(IAgentSoftwareCollection Collection, ISingletonDocument<AgentSoftwareChannels> Singleton, SoftwareService SoftwareService, IOptionsMonitor<ServerSettings> Settings, ILogger<AgentSoftwareService> Logger)
 		{
 			this.Collection = Collection;
 			this.Singleton = Singleton;
+			this.SoftwareService = SoftwareService;
+			this.Settings = Settings;
+			this.Logger = Logger;
 
 			ChannelsDocument = new LazyCachedValue<Task<AgentSoftwareChannels>>(() => Singleton.GetAsync(), TimeSpan.FromSeconds(20.0));
+
+			Task.Run(() => RegisterDefaultAgent(5000));
 		}
 
 		/// <summary>
@@ -289,5 +314,52 @@ namespace HordeServer.Services
 		{
 			return Collection.GetAsync(Version);
 		}
+
+		async Task RegisterDefaultAgent(int DelayMs)
+		{
+			await Task.Delay(DelayMs);			
+
+			// Check whether we have an installed agent zip
+			FileReference AgentZip = FileReference.Combine(Program.AppDir, "DefaultAgent/Agent.zip");
+			if (!Settings.CurrentValue.SingleInstance || !FileReference.Exists(AgentZip))
+			{
+				return;
+			}
+
+			Logger.LogInformation("Checking for default agent software update");
+
+			string AgentHash = ContentHash.MD5(AgentZip).ToString();
+
+			FileReference AgentHashFile = FileReference.Combine(Program.DataDir, "Agent/DefaultAgentHash");
+			
+			try 
+			{
+				
+				if (FileReference.Exists(AgentHashFile) && FileReference.ReadAllText(AgentHashFile).Trim() == AgentHash)
+				{
+					Logger.LogInformation("Default agent software is up to date");
+					return;
+				}
+			}
+			catch (Exception Ex)
+			{
+				Logger.LogError(Ex, "Error checking default agent software, {Message}", Ex.Message);
+			}
+
+			using (System.IO.Stream Stream = File.OpenRead(AgentZip.ToString()))
+			{
+				await SoftwareService.CreateSoftwareAsync(Stream, null, true);
+			}
+
+			if (!DirectoryReference.Exists(AgentHashFile.Directory))
+			{
+				DirectoryReference.CreateDirectory(AgentHashFile.Directory);
+			}
+
+			FileReference.WriteAllText(AgentHashFile, AgentHash);
+
+			Logger.LogInformation($"Updated default agent software to {AgentHash}");
+
+		}
 	}
 }
diff --git a/Engine/Source/Programs/Horde/Installer/Installer.xml b/Engine/Source/Programs/Horde/Installer/Installer.xml
--- a/Engine/Source/Programs/Horde/Installer/Installer.xml
+++ b/Engine/Source/Programs/Horde/Installer/Installer.xml
@@ -12,6 +12,8 @@
 			<DotNet Arguments="publish &quot;$(RootDir)/Engine/Source/Programs/Horde/HordeServer/HordeServer.csproj&quot; --output &quot;$(StagingDir)/Server&quot; --runtime win-x64 $(VersionArguments)"/>
 			<DotNet Arguments="publish &quot;$(RootDir)/Engine/Source/Programs/Horde/HordeAgent/HordeAgent.csproj&quot; --output &quot;$(StagingDir)/Agent&quot; $(VersionArguments)"/>
 
+			<Zip FromDir="$(StagingDir)/Agent" ZipFile="$(StagingDir)/Server/DefaultAgent/Agent.zip"/>
+
 			<!-- Do not include development defaults which created a ue5 project and stream -->
 			<Delete Files="$(StagingDir)/Server/Defaults/..."/>
 
@@ -38,14 +40,8 @@
 			<Spawn Exe="$(WixDir)/heat.exe" Arguments="dir &quot;$(StagingDir)/Server-Bulk&quot; -cg HordeServerFiles -scom -sreg -gg -dr InstallDir -srd -var var.SourceDir -out &quot;$(SourceDir)/HordeServerFiles.wxs&quot;"/>
 			<Spawn Exe="$(WixDir)/candle.exe" Arguments="$(CommonArgs) HordeServerFiles.wxs -dSourceDir=&quot;$(StagingDir)&quot; -out $(ObjDir)/HordeServerFiles.wixobj" WorkingDir="$(SourceDir)"/>
 
-			<Delete Files="$(StagingDir)/Agent-Bulk/..."/>
-			<Tag Files="$(StagingDir)/Agent/..." Except="HordeAgent.exe" With="#AgentFiles"/>
-			<Copy Files="#AgentFiles" From="$(StagingDir)/Agent" To="$(StagingDir)/Agent-Bulk/Agent"/>
-			<Spawn Exe="$(WixDir)/heat.exe" Arguments="dir &quot;$(StagingDir)/Agent-Bulk&quot; -cg HordeAgentFiles -scom -sreg -gg -dr InstallDir -srd -var var.SourceDir -out &quot;$(SourceDir)/HordeAgentFiles.wxs&quot;"/>
-			<Spawn Exe="$(WixDir)/candle.exe" Arguments="$(CommonArgs) HordeAgentFiles.wxs -dSourceDir=&quot;$(StagingDir)&quot; -out $(ObjDir)/HordeAgentFiles.wixobj" WorkingDir="$(SourceDir)"/>
-
 			<Spawn Exe="$(WixDir)/candle.exe" Arguments="$(CommonArgs) Installer.wxs -dStagingDir=&quot;$(StagingDir)&quot; -out $(ObjDir)/Installer.wixobj" WorkingDir="$(SourceDir)"/>
-			<Spawn Exe="$(WixDir)/light.exe" Arguments="-out $(BinDir)/Horde.msi -sw1076 -pdbout $(BinDir)/Horde.pdb $(ObjDir)/Installer.wixobj $(ObjDir)/HordeServerFiles.wixobj $(ObjDir)/HordeAgentFiles.wixobj" WorkingDir="$(SourceDir)"/>
+			<Spawn Exe="$(WixDir)/light.exe" Arguments="-out $(BinDir)/Horde.msi -sw1076 -pdbout $(BinDir)/Horde.pdb $(ObjDir)/Installer.wixobj $(ObjDir)/HordeServerFiles.wixobj" WorkingDir="$(SourceDir)"/>
 		</Node>
 		
 		<Node Name="Publish HordeInstaller" Requires="Build HordeInstaller">
