Horde: Add compact binary attributes to PutBlobResponse and ExistsResponse.

diff --git a/Engine/Source/Programs/Horde/HordeServer/Storage/Controllers/BlobsController.cs b/Engine/Source/Programs/Horde/HordeServer/Storage/Controllers/BlobsController.cs
--- a/Engine/Source/Programs/Horde/HordeServer/Storage/Controllers/BlobsController.cs
+++ b/Engine/Source/Programs/Horde/HordeServer/Storage/Controllers/BlobsController.cs
@@ -16,6 +16,7 @@ using MongoDB.Bson;
 using HordeServer.Utilities;
 using EpicGames.Core;
 using HordeServer.Storage.Collections;
+using EpicGames.Serialization;
 
 namespace HordeServer.Storage.Controllers
 {
@@ -29,6 +30,7 @@ namespace HordeServer.Storage.Controllers
 		/// <summary>
 		/// Hash of the blob that was added
 		/// </summary>
+		[CbField("h")]
 		public IoHash Hash { get; set; }
 	}
 
@@ -40,7 +42,8 @@ namespace HordeServer.Storage.Controllers
 		/// <summary>
 		/// List of blobs that exist
 		/// </summary>
-		public IoHash[]? Id { get; set; } = null;
+		[CbField("id")]
+		public List<IoHash>? Id { get; set; } = null;
 	}
 
 	/// <summary>
@@ -161,7 +164,7 @@ namespace HordeServer.Storage.Controllers
 			}
 
 			ExistsResponse Exists = new ExistsResponse();
-			Exists.Id = (await BlobCollection.ExistsAsync(NamespaceId, Hashes)).ToArray();
+			Exists.Id = await BlobCollection.ExistsAsync(NamespaceId, Hashes);
 			return Exists;
 		}
 	}
diff --git a/Engine/Source/Programs/Horde/HordeServer/Storage/Controllers/ObjectsController.cs b/Engine/Source/Programs/Horde/HordeServer/Storage/Controllers/ObjectsController.cs
--- a/Engine/Source/Programs/Horde/HordeServer/Storage/Controllers/ObjectsController.cs
+++ b/Engine/Source/Programs/Horde/HordeServer/Storage/Controllers/ObjectsController.cs
@@ -203,7 +203,7 @@ namespace HordeServer.Storage.Controllers
 			}
 
 			ExistsResponse Response = new ExistsResponse();
-			Response.Id = (await ObjectCollection.ExistsAsync(NamespaceId, Hashes)).ToArray();
+			Response.Id = await ObjectCollection.ExistsAsync(NamespaceId, Hashes);
 			return Response;
 		}
 	}
