Horde - UnrealAutoamtatedTests plugin - Better support unknown test state + fix log display in simplereport
#rnx

diff --git a/Engine/Source/Programs/Horde/HordeDashboard/src/Plugins/UnrealAutomatedTests/models/UnrealAutomatedTests.ts b/Engine/Source/Programs/Horde/HordeDashboard/src/Plugins/UnrealAutomatedTests/models/UnrealAutomatedTests.ts
--- a/Engine/Source/Programs/Horde/HordeDashboard/src/Plugins/UnrealAutomatedTests/models/UnrealAutomatedTests.ts
+++ b/Engine/Source/Programs/Horde/HordeDashboard/src/Plugins/UnrealAutomatedTests/models/UnrealAutomatedTests.ts
@@ -6,6 +6,7 @@ export enum TestState {
     SuccessWithWarnings = "SuccessWithWarnings",
     NotRun = "NotRun",
     InProcess = "InProcess",
+    Unknown = "Unknown",
 }
 
 export enum EventType {
