Horde: Accept a raw byte stream from the request body when adding a new blob.

diff --git a/Engine/Source/Programs/Horde/HordeServer/Storage/Controllers/BlobsController.cs b/Engine/Source/Programs/Horde/HordeServer/Storage/Controllers/BlobsController.cs
--- a/Engine/Source/Programs/Horde/HordeServer/Storage/Controllers/BlobsController.cs
+++ b/Engine/Source/Programs/Horde/HordeServer/Storage/Controllers/BlobsController.cs
@@ -76,25 +76,37 @@ namespace HordeServer.Storage.Controllers
 		}
 
 		/// <summary>
-		/// Submit a blob to the storage api. The payload hash will be verified against the received payload. 
-		/// The hash of the object is returned for a successful submission.
+		/// Submit a blob to the storage api. The request body should be a raw byte stream containing the blob data.
+		/// The payload hash will be verified against the received payload. The hash of the object is returned for a successful submission.
 		/// </summary>
 		/// <param name="NamespaceId">Namespace for the operation</param>
 		/// <param name="Hash">Hash of the blob. Must match the submitted data.</param>
-		/// <param name="File">The data to add</param>
 		/// <returns></returns>
 		[HttpPut]
 		[Route("/api/v1/blobs/{NamespaceId}/{Hash}")]
-		public async Task<ActionResult<PutBlobResponse>> PutBlobAsync(NamespaceId NamespaceId, IoHash Hash, IFormFile File)
+		public async Task<ActionResult<PutBlobResponse>> PutBlobAsync(NamespaceId NamespaceId, IoHash Hash)
 		{
 			if (!await NamespaceCollection.AuthorizeAsync(NamespaceId, User, AclAction.WriteBlobs))
 			{
 				return Forbid();
 			}
 
-			using (Stream Stream = File.OpenReadStream())
+			byte[] Data;
+			using (MemoryStream MemoryStream = new MemoryStream())
 			{
-				await BlobCollection.WriteAsync(NamespaceId, Hash, Stream);
+				await Request.Body.CopyToAsync(MemoryStream);
+				Data = MemoryStream.ToArray();
+			}
+
+			IoHash ActualHash = IoHash.Compute(Data);
+			if (Hash != ActualHash)
+			{
+				return BadRequest("Incorrect hash");
+			}
+
+			using (MemoryStream MemoryStream = new MemoryStream(Data))
+			{
+				await BlobCollection.WriteAsync(NamespaceId, Hash, MemoryStream);
 			}
 
 			return new PutBlobResponse { Hash = Hash };
