Horde: Fix exception trying to parse digest fields for deleted files.

diff --git a/Engine/Source/Programs/Horde/HordeServer/Models/Change.cs b/Engine/Source/Programs/Horde/HordeServer/Models/Change.cs
--- a/Engine/Source/Programs/Horde/HordeServer/Models/Change.cs
+++ b/Engine/Source/Programs/Horde/HordeServer/Models/Change.cs
@@ -103,7 +103,7 @@ namespace HordeServer.Models
 		/// <summary>
 		/// MD5 digest of the file
 		/// </summary>
-		public Md5Hash Digest { get; set; }
+		public Md5Hash? Digest { get; set; }
 
 		/// <summary>
 		/// The file type
@@ -119,7 +119,7 @@ namespace HordeServer.Models
 		/// <param name="Length"></param>
 		/// <param name="Digest"></param>
 		/// <param name="Type"></param>
-		public ChangeFile(string Path, string DepotPath, int Revision, long Length, Md5Hash Digest, string Type)
+		public ChangeFile(string Path, string DepotPath, int Revision, long Length, Md5Hash? Digest, string Type)
 		{
 			this.Path = Path;
 			this.DepotPath = DepotPath;
diff --git a/Engine/Source/Programs/Horde/HordeServer/Services/Impl/PerforceService.cs b/Engine/Source/Programs/Horde/HordeServer/Services/Impl/PerforceService.cs
--- a/Engine/Source/Programs/Horde/HordeServer/Services/Impl/PerforceService.cs
+++ b/Engine/Source/Programs/Horde/HordeServer/Services/Impl/PerforceService.cs
@@ -579,7 +579,8 @@ namespace HordeServer.Services
 		static ChangeFile CreateChangeFile(string RelativePath, P4.FileMetaData MetaData)
 		{
 			int Revision = GetSyncRevision(MetaData.DepotPath.Path, MetaData.HeadAction, MetaData.HeadRev);
-			return new ChangeFile(RelativePath, MetaData.DepotPath.Path, Revision, MetaData.FileSize, Md5Hash.Parse(MetaData.Digest ?? String.Empty), MetaData.HeadType.ToString());
+			Md5Hash? Digest = String.IsNullOrEmpty(MetaData.Digest) ? (Md5Hash?)null : Md5Hash.Parse(MetaData.Digest);
+			return new ChangeFile(RelativePath, MetaData.DepotPath.Path, Revision, MetaData.FileSize, Digest, MetaData.HeadType.ToString());
 		}
 
 		/// <summary>
@@ -591,7 +592,8 @@ namespace HordeServer.Services
 		static ChangeFile CreateChangeFile(string RelativePath, P4.ShelvedFile MetaData)
 		{
 			int Revision = GetSyncRevision(MetaData.Path.Path, MetaData.Action, MetaData.Revision);
-			return new ChangeFile(RelativePath, MetaData.Path.Path, Revision, MetaData.Size, Md5Hash.Parse(MetaData.Digest ?? String.Empty), MetaData.Type.ToString());
+			Md5Hash? Digest = String.IsNullOrEmpty(MetaData.Digest) ? (Md5Hash?)null : Md5Hash.Parse(MetaData.Digest);
+			return new ChangeFile(RelativePath, MetaData.Path.Path, Revision, MetaData.Size, Digest, MetaData.Type.ToString());
 		}
 
 		/// <inheritdoc/>
