Horde: Truncate device slack notification if over limit
#rnx

diff --git a/Engine/Source/Programs/Horde/HordeServer/Notifications/Impl/SlackNotificationSink.cs b/Engine/Source/Programs/Horde/HordeServer/Notifications/Impl/SlackNotificationSink.cs
--- a/Engine/Source/Programs/Horde/HordeServer/Notifications/Impl/SlackNotificationSink.cs
+++ b/Engine/Source/Programs/Horde/HordeServer/Notifications/Impl/SlackNotificationSink.cs
@@ -986,7 +986,13 @@ namespace HordeServer.Notifications.Impl
 		{
 			BlockKitAttachment Attachment = new BlockKitAttachment();
 
-			Attachment.FallbackText = $"{Message}";
+			// truncate message to avoid slack error on message length
+			if (Message.Length > 150)
+			{
+				Message = Message.Substring(0, 146) + "...";
+			}
+            
+            Attachment.FallbackText = $"{Message}";
 
 			if (Device != null && Pool != null)
 			{
