Horde: Catch and report exceptions thrown in ElectedTick.

diff --git a/Engine/Source/Programs/Horde/HordeServer/Utilities/ElectedTick.cs b/Engine/Source/Programs/Horde/HordeServer/Utilities/ElectedTick.cs
--- a/Engine/Source/Programs/Horde/HordeServer/Utilities/ElectedTick.cs
+++ b/Engine/Source/Programs/Horde/HordeServer/Utilities/ElectedTick.cs
@@ -163,17 +163,15 @@ namespace HordeServer.Utilities
 				}
 
 				// Start the tick function, and keep prolonging the next update time until it finishes
-				Task<DateTime> TickTask = Task.Run(() => TickAsync(StoppingToken));
-				while (await Task.WhenAny(TickTask, Task.Delay(ElectionTimeout / 2)) != TickTask)
-				{
-					State.NextUpdateTime = DateTime.UtcNow + ElectionTimeout;
-					await StateAccessor.TryUpdateAsync(State);
-				}
-
-				// Wait for the task to complete
 				DateTime NextUpdateTime;
 				try
 				{
+					Task<DateTime> TickTask = Task.Run(() => TickAsync(StoppingToken));
+					while (await Task.WhenAny(TickTask, Task.Delay(ElectionTimeout / 2)) != TickTask)
+					{
+						State.NextUpdateTime = DateTime.UtcNow + ElectionTimeout;
+						await StateAccessor.TryUpdateAsync(State);
+					}
 					NextUpdateTime = await TickTask;
 				}
 				catch (Exception Ex)
