Horde: Allow passing uppercase characters in strings to StringId constructor; just normalize them to lowercase.

diff --git a/Engine/Source/Programs/Horde/HordeServer/Compute/Impl/ComputeController.cs b/Engine/Source/Programs/Horde/HordeServer/Compute/Impl/ComputeController.cs
--- a/Engine/Source/Programs/Horde/HordeServer/Compute/Impl/ComputeController.cs
+++ b/Engine/Source/Programs/Horde/HordeServer/Compute/Impl/ComputeController.cs
@@ -48,9 +48,9 @@ namespace HordeServer.Compute.Impl
 		[HttpPost]
 		[Authorize]
 		[Route("/api/v1/compute/{ChannelId}")]
-		public async Task AddTasksAsync([FromRoute] string ChannelId, [FromBody] AddTasksRequest Request)
+		public async Task AddTasksAsync([FromRoute] ChannelId ChannelId, [FromBody] AddTasksRequest Request)
 		{
-			await ComputeService.AddTasksAsync(Impl.ComputeService.DefaultNamespaceId, Request.RequirementsHash, Request.TaskHashes, new ChannelId(ChannelId));
+			await ComputeService.AddTasksAsync(Impl.ComputeService.DefaultNamespaceId, Request.RequirementsHash, Request.TaskHashes, ChannelId);
 		}
 
 		/// <summary>
@@ -62,17 +62,17 @@ namespace HordeServer.Compute.Impl
 		[HttpPost]
 		[Authorize]
 		[Route("/api/v1/compute/{ChannelId}/updates")]
-		public async Task<GetTaskUpdatesResponse> GetUpdatesAsync([FromRoute] string ChannelId, [FromQuery] int Wait = 0)
+		public async Task<GetTaskUpdatesResponse> GetUpdatesAsync([FromRoute] ChannelId ChannelId, [FromQuery] int Wait = 0)
 		{
 			List<IComputeTaskStatus> Results;
 			if (Wait == 0)
 			{
-				Results = await ComputeService.GetTaskUpdatesAsync(new ChannelId(ChannelId));
+				Results = await ComputeService.GetTaskUpdatesAsync(ChannelId);
 			}
 			else
 			{
 				using CancellationTokenSource DelaySource = new CancellationTokenSource(Wait * 1000);
-				Results = await ComputeService.WaitForTaskUpdatesAsync(new ChannelId(ChannelId), DelaySource.Token);
+				Results = await ComputeService.WaitForTaskUpdatesAsync(ChannelId, DelaySource.Token);
 			}
 
 			GetTaskUpdatesResponse Response = new GetTaskUpdatesResponse();
diff --git a/Engine/Source/Programs/Horde/HordeServer/Utilities/StringId.cs b/Engine/Source/Programs/Horde/HordeServer/Utilities/StringId.cs
--- a/Engine/Source/Programs/Horde/HordeServer/Utilities/StringId.cs
+++ b/Engine/Source/Programs/Horde/HordeServer/Utilities/StringId.cs
@@ -42,10 +42,11 @@ namespace HordeServer.Utilities
 		/// <summary>
 		/// Constructor
 		/// </summary>
-		/// <param name="Text">Unique id for the string</param>
-		public StringId(string Text)
+		/// <param name="Input">Unique id for the string</param>
+		[SuppressMessage("Globalization", "CA1308:Normalize strings to uppercase", Justification = "<Pending>")]
+		public StringId(string Input)
 		{
-			this.Text = Text;
+			this.Text = Input;
 
 			if (Text.Length == 0)
 			{
@@ -60,9 +61,17 @@ namespace HordeServer.Utilities
 
 			for (int Idx = 0; Idx < Text.Length; Idx++)
 			{
-				if (!IsValidCharacter(Text[Idx]))
+				char Character = Text[Idx];
+				if (!IsValidCharacter(Character))
 				{
-					throw new ArgumentException($"{Text} is not a valid string id");
+					if (Character >= 'A' && Character <= 'Z')
+					{
+						Text = Text.ToLowerInvariant();
+					}
+					else
+					{
+						throw new ArgumentException($"{Text} is not a valid string id");
+					}
 				}
 			}
 		}
