Horde: Prevent build running when all files are already up-to-date.

diff --git a/Engine/Source/Programs/Horde/HordeServer/HordeServer.csproj b/Engine/Source/Programs/Horde/HordeServer/HordeServer.csproj
--- a/Engine/Source/Programs/Horde/HordeServer/HordeServer.csproj
+++ b/Engine/Source/Programs/Horde/HordeServer/HordeServer.csproj
@@ -89,10 +89,10 @@
 
   <ItemGroup>
     <Content Include="..\..\..\ThirdParty\Perforce\P4Api.Net\native\windows\release\p4bridge.*" Condition=" '$(OS)' == 'Windows_NT' ">
-      <CopyToOutputDirectory>Always</CopyToOutputDirectory>
+      <CopyToOutputDirectory>PreserveNewest</CopyToOutputDirectory>
     </Content>
     <Content Include="..\..\..\ThirdParty\Perforce\P4Api.Net\native\linux\release\p4bridge.*" Condition="$([MSBuild]::IsOSPlatform('Linux'))">
-      <CopyToOutputDirectory>Always</CopyToOutputDirectory>
+      <CopyToOutputDirectory>PreserveNewest</CopyToOutputDirectory>
     </Content>
     <Content Include="..\ThirdParty\**" Condition="$([MSBuild]::IsOSPlatform('Windows'))">
       <CopyToOutputDirectory>PreserveNewest</CopyToOutputDirectory>
