Horde: Explicitly clear model only when setting to base
#rnx

diff --git a/Engine/Source/Programs/Horde/HordeServer/Collections/Impl/DeviceCollection.cs b/Engine/Source/Programs/Horde/HordeServer/Collections/Impl/DeviceCollection.cs
--- a/Engine/Source/Programs/Horde/HordeServer/Collections/Impl/DeviceCollection.cs
+++ b/Engine/Source/Programs/Horde/HordeServer/Collections/Impl/DeviceCollection.cs
@@ -438,13 +438,16 @@ namespace HordeServer.Collections.Impl
 				Updates.Add(UpdateBuilder.Set(x => x.Address, NewAddress));
 			}
 
-			if (string.IsNullOrEmpty(NewModelId))
+			if (!string.IsNullOrEmpty(NewModelId))
 			{
-				Updates.Add(UpdateBuilder.Set(x => x.ModelId, null));
-			}
-			else
-			{
-				Updates.Add(UpdateBuilder.Set(x => x.ModelId, NewModelId));
+				if (NewModelId == "Base")
+				{
+					Updates.Add(UpdateBuilder.Set(x => x.ModelId, null));
+				}
+				else
+				{
+					Updates.Add(UpdateBuilder.Set(x => x.ModelId, NewModelId));
+				}				
 			}
 
 			if (NewNotes != null)
