Horde: Fix device utilization range
#rnx

diff --git a/Engine/Source/Programs/Horde/HordeServer/Collections/Impl/DeviceCollection.cs b/Engine/Source/Programs/Horde/HordeServer/Collections/Impl/DeviceCollection.cs
--- a/Engine/Source/Programs/Horde/HordeServer/Collections/Impl/DeviceCollection.cs
+++ b/Engine/Source/Programs/Horde/HordeServer/Collections/Impl/DeviceCollection.cs
@@ -605,7 +605,7 @@ namespace HordeServer.Collections.Impl
 				// keep up to 100, maintaining order
 				if (Utilization.Count > 99)
 				{
-					Utilization.Take(99);
+					Utilization = Utilization.GetRange(0, 99);
 				}
 
 				Utilization.Insert(0, new DeviceUtilizationTelemetry(ReservationTimeUtc) { JobId = JobId, StepId = StepId });
