Horde: Make device slack notification more succinct to avoid limit
#rnx

diff --git a/Engine/Source/Programs/Horde/HordeServer/Controllers/DevicesController.cs b/Engine/Source/Programs/Horde/HordeServer/Controllers/DevicesController.cs
--- a/Engine/Source/Programs/Horde/HordeServer/Controllers/DevicesController.cs
+++ b/Engine/Source/Programs/Horde/HordeServer/Controllers/DevicesController.cs
@@ -687,13 +687,13 @@ namespace HordeServer.Controllers
 				 }
 				 else
 				 {
-					Details = $" - No reservation information available from host {Request.Hostname}";
+					Details = $" - Host {Request.Hostname}";
 				 }							 
             }
 
 			if (string.IsNullOrEmpty(PoolId))
 			{
-                Message = $"No pool specified for device reservation request, defaulting to UE4 (will soon be an error)" + Details;
+                Message = $"No pool specified, defaulting to UE4" + Details;
                 Logger.LogError(Message + $" JobId: {Request.JobId}, StepId: {Request.StepId}");
 				await DeviceService.NotifyDeviceServiceAsync(Message, null, Request.JobId, Request.StepId);
 				PoolId = "ue4";
@@ -704,7 +704,7 @@ namespace HordeServer.Controllers
 			IDevicePool? Pool = Pools.FirstOrDefault(x => x.Id == PoolIdValue);
 			if (Pool == null)
 			{
-                Message = $"Unknown pool {PoolId} on device reservation request" + Details;
+                Message = $"Unknown pool {PoolId} " + Details;
 				Logger.LogError(Message);
 				await DeviceService.NotifyDeviceServiceAsync(Message, null, Request.JobId, Request.StepId);
                 return BadRequest(Message);
@@ -733,7 +733,7 @@ namespace HordeServer.Controllers
 
 				if (!PlatformMap.TryGetValue(PlatformName, out PlatformId))
 				{
-                    Message = $"Unknown platform for device type {PlatformName} on device reservation request" + Details;
+                    Message = $"Unknown platform {PlatformName}" + Details;
 					Logger.LogError(Message);
 					await DeviceService.NotifyDeviceServiceAsync(Message, null, Request.JobId, Request.StepId);
                     return BadRequest(Message);
@@ -743,7 +743,7 @@ namespace HordeServer.Controllers
 
 				if (Platform == null)
 				{
-                    Message = $"Unknown platform {PlatformId} on device reservation request" + Details;
+                    Message = $"Unknown platform {PlatformId}" + Details;
 					Logger.LogError(Message);
 					await DeviceService.NotifyDeviceServiceAsync(Message, null, Request.JobId, Request.StepId);
                     return BadRequest(Message);
@@ -927,7 +927,7 @@ namespace HordeServer.Controllers
 
 			await DeviceService.UpdateDeviceAsync(Device.Id, NewProblem: true);
 
-			string Message = $"Device reported an error, {Device.Name} : {Device.PoolId.ToString().ToUpperInvariant()}";
+			string Message = $"Device problem, {Device.Name} : {Device.PoolId.ToString().ToUpperInvariant()}";
 
             IDeviceReservation? Reservation = await DeviceService.TryGetDeviceReservation(Device.Id);
 
@@ -946,7 +946,7 @@ namespace HordeServer.Controllers
 					}
 					else
 					{
-						Message += $" - No reservation information available from host {Reservation.Hostname}";
+						Message += $" - Host {Reservation.Hostname}";
 					}                    
                 }
             }
