Horde: Fix agent compile error.

diff --git a/Engine/Source/Programs/Horde/HordeAgent/Utility/WorkspaceInfo.cs b/Engine/Source/Programs/Horde/HordeAgent/Utility/WorkspaceInfo.cs
--- a/Engine/Source/Programs/Horde/HordeAgent/Utility/WorkspaceInfo.cs
+++ b/Engine/Source/Programs/Horde/HordeAgent/Utility/WorkspaceInfo.cs
@@ -376,7 +376,7 @@ namespace HordeAgent.Utility
 			}
 
 			// enumerate all the pending changelists
-			List<ChangesRecord> PendingChanges = await Perforce.GetChangesAsync(ChangesOptions.None, Perforce.ClientName, -1, ChangeStatus.Pending, null, null, CancellationToken);
+			List<ChangesRecord> PendingChanges = await Perforce.GetChangesAsync(ChangesOptions.None, Perforce.ClientName, -1, ChangeStatus.Pending, null, FileSpecList.Empty, CancellationToken);
 			foreach (ChangesRecord PendingChange in PendingChanges)
 			{
 				// delete any shelved files if there are any
