Horde: Add gRPC-based health check to HordeRpc

Required to make the server play nice with a gRPC-based load balancer (in our case AWS ALB).

diff --git a/Engine/Source/Programs/Horde/HordeServer/Services/RpcService.cs b/Engine/Source/Programs/Horde/HordeServer/Services/RpcService.cs
--- a/Engine/Source/Programs/Horde/HordeServer/Services/RpcService.cs
+++ b/Engine/Source/Programs/Horde/HordeServer/Services/RpcService.cs
@@ -162,6 +162,25 @@ namespace HordeServer.Services
 			this.Logger = Logger;
 		}
 
+		/// <summary>
+		/// Return server status, primarily intended for load balancers to decide if traffic should be routed to this process
+		///
+		/// For example, the gRPC healthcheck for AWS ALB will pick up and react to the gRPC status code returned.
+		/// </summary>
+		/// <param name="Request">Empty placeholder request (for now)</param>
+		/// <param name="Context">Context for the call</param>
+		/// <returns>Return status code 'unavailable' if stopping</returns>
+		public override Task<HealthCheckResponse> HealthCheck(HealthCheckRequest Request, ServerCallContext Context)
+		{
+			bool IsStopping = LifetimeService.IsPreStopping || LifetimeService.IsStopping;
+			if (IsStopping)
+			{
+				Context.Status = new Status(StatusCode.Unavailable, "Server is stopping");	
+			}
+
+			return Task.FromResult(new HealthCheckResponse {Stopping = IsStopping});
+		}
+
 		/// <summary>
 		/// Waits until the server is terminating
 		/// </summary>
