Horde: Fix/improve logging for Slack notifications. Now consistently uses {SlackUser} properties for messages, and adds additional context for issue ids.

diff --git a/Engine/Source/Programs/Horde/HordeServer/Notifications/Impl/SlackNotificationSink.cs b/Engine/Source/Programs/Horde/HordeServer/Notifications/Impl/SlackNotificationSink.cs
--- a/Engine/Source/Programs/Horde/HordeServer/Notifications/Impl/SlackNotificationSink.cs
+++ b/Engine/Source/Programs/Horde/HordeServer/Notifications/Impl/SlackNotificationSink.cs
@@ -319,7 +319,7 @@ namespace HordeServer.Notifications.Impl
 		private Task SendJobCompleteMessageAsync(string Recipient, IStream Stream, IJob Job, IGraph Graph)
 		{
 			JobStepOutcome JobOutcome = Job.Batches.SelectMany(x => x.Steps).Min(x => x.Outcome);
-			Logger.LogInformation("Sending Slack notification for job {JobId} outcome {Outcome} to {User}", Job.Id, JobOutcome, Recipient);
+			Logger.LogInformation("Sending Slack notification for job {JobId} outcome {Outcome} to {SlackUser}", Job.Id, JobOutcome, Recipient);
 
 			Uri JobLink = new Uri($"{Settings.DashboardUrl}job/{Job.Id}");
 
@@ -396,7 +396,7 @@ namespace HordeServer.Notifications.Impl
 		/// <inheritdoc/>
 		public async Task NotifyJobStepCompleteAsync(IUser SlackUser, IStream JobStream, IJob Job, IJobStepBatch Batch, IJobStep Step, INode Node, List<ILogEventData> JobStepEventData)
 		{
-			Logger.LogInformation("Sending Slack notification for job {JobId}, batch {BatchId}, step {StepId}, outcome {Outcome} to {User} ({UserId})", Job.Id, Batch.Id, Step.Id, Step.Outcome, SlackUser.Name, SlackUser.Id);
+			Logger.LogInformation("Sending Slack notification for job {JobId}, batch {BatchId}, step {StepId}, outcome {Outcome} to {SlackUser} ({UserId})", Job.Id, Batch.Id, Step.Id, Step.Outcome, SlackUser.Name, SlackUser.Id);
 
 			string? SlackUserId = await GetSlackUserId(SlackUser);
 			if (SlackUserId != null)
@@ -646,6 +646,8 @@ namespace HordeServer.Notifications.Impl
 
 		async Task SendIssueMessageAsync(string Recipient, IIssue Issue, IIssueDetails Details, ObjectId? UserId)
 		{
+			using IDisposable Scope = Logger.BeginScope("SendIssueMessageAsync (User: {SlackUser}, Issue: {IssueId})", Recipient, Issue.Id);
+
 			BlockKitAttachment Attachment = new BlockKitAttachment();
 			Attachment.Color = Color.Red;
 
@@ -831,7 +833,7 @@ namespace HordeServer.Notifications.Impl
 			int IssueId = int.Parse(Match.Groups[1].Value, NumberStyles.Integer, CultureInfo.InvariantCulture);
 			string Verb = Match.Groups[2].Value;
 			ObjectId UserId = Match.Groups[3].Value.ToObjectId();
-			Logger.LogInformation("Issue {IssueId}: {Action} from {User} ({UserId})", IssueId, Verb, UserName, UserId);
+			Logger.LogInformation("Issue {IssueId}: {Action} from {SlackUser} ({UserId})", IssueId, Verb, UserName, UserId);
 
 			if (String.Equals(Verb, "ack", StringComparison.Ordinal))
 			{
@@ -1167,7 +1169,7 @@ namespace HordeServer.Notifications.Impl
 			MessageStateDocument? CurrentState = await GetMessageStateAsync(Recipient, EventId);
 			if (CurrentState == null)
 			{
-				Logger.LogInformation("Sending new slack message to {Recipient}", Recipient);
+				Logger.LogInformation("Sending new slack message to {SlackUser}", Recipient);
 				Message.Recipient = Recipient;
 				Message.Ts = null;
 				Response = await SendRequestAsync<PostMessageResponse>(PostMessageUrl, Message);
@@ -1178,7 +1180,7 @@ namespace HordeServer.Notifications.Impl
 			}
 			else
 			{
-				Logger.LogInformation("Updating existing slack message for user {Recipient} ({Channel}, (MessageTs))", Recipient, CurrentState.Channel, CurrentState.Ts);
+				Logger.LogInformation("Updating existing slack message for user {SlackUser} ({Channel}, {MessageTs})", Recipient, CurrentState.Channel, CurrentState.Ts);
 				Message.Recipient = CurrentState.Channel;
 				Message.Ts = CurrentState.Ts;
 				Response = await SendRequestAsync<PostMessageResponse>(UpdateMessageUrl, Message);
@@ -1188,6 +1190,10 @@ namespace HordeServer.Notifications.Impl
 			{
 				await SetMessageStateAsync(Recipient, EventId, UserId, Response.Channel, Response.Ts, RequestDigest);
 			}
+			else
+			{
+				Logger.LogWarning("Missing fields on slack response (channel: {Channel}, ts: {MessageTs})", Response.Channel, Response.Ts);
+			}
 		}
 
 		private async Task<TResponse> SendRequestAsync<TResponse>(string RequestUrl, object Request) where TResponse : SlackResponse
