Horde: Only enable Datadog when explicitly enabled through config file. Prevents a bunch of HTTP exceptions as it fails to connect when running locally.

diff --git a/Engine/Source/Programs/Horde/HordeServer/Collections/DatadogTracingCollection.cs b/Engine/Source/Programs/Horde/HordeServer/Collections/DatadogTracingCollection.cs
deleted file mode 100644
diff --git a/Engine/Source/Programs/Horde/HordeServer/Program.cs b/Engine/Source/Programs/Horde/HordeServer/Program.cs
--- a/Engine/Source/Programs/Horde/HordeServer/Program.cs
+++ b/Engine/Source/Programs/Horde/HordeServer/Program.cs
@@ -1,7 +1,9 @@
 // Copyright Epic Games, Inc. All Rights Reserved.
 
 using System;
+using System.Collections.Concurrent;
 using System.Collections.Generic;
+using System.Diagnostics;
 using System.Globalization;
 using System.IO;
 using System.Linq;
@@ -10,9 +12,13 @@ using System.Reflection;
 using System.Runtime.InteropServices;
 using System.Security.Cryptography;
 using System.Security.Cryptography.X509Certificates;
+using System.Text;
 using System.Text.Json;
+using System.Threading;
+using System.Threading.Channels;
 using System.Threading.Tasks;
 using EpicGames.Core;
+using EpicGames.Perforce;
 using HordeServer.Models;
 using HordeServer.Utilities;
 using Microsoft.AspNetCore.Hosting;
@@ -20,6 +26,7 @@ using Microsoft.AspNetCore.Server.Kestrel.Core;
 using Microsoft.Extensions.Configuration;
 using Microsoft.Extensions.Hosting;
 using Microsoft.Extensions.Logging;
+using Microsoft.Extensions.Logging.Abstractions;
 using MongoDB.Bson;
 using MongoDB.Bson.Serialization;
 using MongoDB.Driver;
@@ -61,6 +68,10 @@ namespace HordeServer
 
 		public static LoggerConfiguration WithHordeConfig(this LoggerConfiguration Configuration, ServerSettings Settings)
 		{
+			if (Settings.WithDatadog)
+			{
+				Configuration = Configuration.Enrich.With<DatadogLogEnricher>();
+			}
 			if (!Settings.LogSessionRequests)
 			{
 				Configuration = Configuration.Filter.ByExcluding(IsSessionLogEvent);
@@ -167,12 +178,14 @@ namespace HordeServer
 				LogDir = DirectoryReference.Combine(DataDir);
 			}
 
-			GlobalTracer.Register(Datadog.Trace.OpenTracing.OpenTracingTracerFactory.CreateTracer());
+			if (HordeSettings.WithDatadog)
+			{
+				GlobalTracer.Register(Datadog.Trace.OpenTracing.OpenTracingTracerFactory.CreateTracer());
+			}
 
 			Serilog.Log.Logger = new LoggerConfiguration()
 				.WithHordeConfig(HordeSettings)
 				.Enrich.FromLogContext()
-				.Enrich.With<DatadogLogEnricher>()
 				.WriteTo.Console(HordeSettings)
 				.WriteTo.File(Path.Combine(LogDir.FullName, "Log.txt"), outputTemplate: "[{Timestamp:HH:mm:ss} {Level:w3}] {Indent}{Message:l}{NewLine}{Exception} [{SourceContext}]", rollingInterval: RollingInterval.Day, rollOnFileSizeLimit: true, fileSizeLimitBytes: 20 * 1024 * 1024, retainedFileCountLimit: 10)
 				.WriteTo.File(new JsonFormatter(renderMessage: true), Path.Combine(LogDir.FullName, "Log.json"), rollingInterval: RollingInterval.Day, rollOnFileSizeLimit: true, fileSizeLimitBytes: 20 * 1024 * 1024, retainedFileCountLimit: 10)
diff --git a/Engine/Source/Programs/Horde/HordeServer/ServerSettings.cs b/Engine/Source/Programs/Horde/HordeServer/ServerSettings.cs
--- a/Engine/Source/Programs/Horde/HordeServer/ServerSettings.cs
+++ b/Engine/Source/Programs/Horde/HordeServer/ServerSettings.cs
@@ -412,6 +412,11 @@ namespace HordeServer
 		/// </summary>
 		public int? GlobalThreadPoolMinSize { get; set; }
 
+		/// <summary>
+		/// Whether to enable Datadog integration for tracing
+		/// </summary>
+		public bool WithDatadog { get; set; }
+
 		/// <summary>
 		/// Path to the root config file
 		/// </summary>
diff --git a/Engine/Source/Programs/Horde/HordeServer/Utilities/MongoTracingCollection.cs b/Engine/Source/Programs/Horde/HordeServer/Utilities/MongoTracingCollection.cs
new file mode 100644
--- /dev/null
+++ b/Engine/Source/Programs/Horde/HordeServer/Utilities/MongoTracingCollection.cs
@@ -0,0 +1,696 @@
+// Copyright Epic Games, Inc. All Rights Reserved.
+
+using System.Collections.Generic;
+using System.Threading;
+using System.Threading.Tasks;
+using MongoDB.Bson;
+using MongoDB.Bson.Serialization;
+using MongoDB.Driver;
+using OpenTracing;
+using OpenTracing.Util;
+
+namespace HordeServer.Utiltiies
+{
+	/// <summary>
+	/// Wrap a IMongoCollection with trace scopes
+	///
+	/// Will capture the entire invocation of MongoDB calls, including serialization.
+	/// The other command logging (not in this file) only deals with queries sent to the server at the protocol level.
+	/// </summary>
+	/// <typeparam name="T">A MongoDB document</typeparam>
+	public class MongoTracingCollection<T> : IMongoCollection<T>
+	{
+		private readonly IMongoCollection<T> Collection;
+		private const string Prefix = "MongoCollection.";
+
+		/// <summary>
+		/// Constructor
+		/// </summary>
+		/// <param name="Collection">Collection to wrap with tracing</param>
+		public MongoTracingCollection(IMongoCollection<T> Collection)
+		{
+			this.Collection = Collection;
+		}
+
+#pragma warning disable CS0618		
+		/// <inheritdoc />
+		public IAsyncCursor<TResult> Aggregate<TResult>(PipelineDefinition<T, TResult> pipeline, AggregateOptions Options = null!,
+			CancellationToken  CancellationToken = new CancellationToken())
+		{
+			return Collection.Aggregate(pipeline, Options,  CancellationToken);
+		}
+
+		/// <inheritdoc />
+		public IAsyncCursor<TResult> Aggregate<TResult>(IClientSessionHandle Session, PipelineDefinition<T, TResult> pipeline,
+			AggregateOptions Options = null!, CancellationToken  CancellationToken = new CancellationToken())
+		{
+			return Collection.Aggregate(Session, pipeline, Options,  CancellationToken);
+		}
+
+		/// <inheritdoc />
+		public async Task<IAsyncCursor<TResult>> AggregateAsync<TResult>(PipelineDefinition<T, TResult> pipeline, AggregateOptions Options = null!,
+			CancellationToken  CancellationToken = new CancellationToken())
+		{
+			using IScope _ = GlobalTracer.Instance.BuildSpan(Prefix + "AggregateAsync").StartActive();
+			return await Collection.AggregateAsync(pipeline, Options,  CancellationToken);
+		}
+
+		/// <inheritdoc />
+		public async Task<IAsyncCursor<TResult>> AggregateAsync<TResult>(IClientSessionHandle Session, PipelineDefinition<T, TResult> pipeline, AggregateOptions Options = null!,
+			CancellationToken  CancellationToken = new CancellationToken())
+		{
+			using IScope _ = GlobalTracer.Instance.BuildSpan(Prefix + "AggregateAsync").StartActive();
+			return await Collection.AggregateAsync(Session, pipeline, Options, CancellationToken);
+		}
+
+		/// <inheritdoc />
+		public BulkWriteResult<T> BulkWrite(IEnumerable<WriteModel<T>> requests, BulkWriteOptions Options = null!,
+			CancellationToken  CancellationToken = new CancellationToken())
+		{
+			return Collection.BulkWrite(requests, Options,  CancellationToken);
+		}
+
+		/// <inheritdoc />
+		public BulkWriteResult<T> BulkWrite(IClientSessionHandle Session, IEnumerable<WriteModel<T>> requests, BulkWriteOptions Options = null!,
+			CancellationToken  CancellationToken = new CancellationToken())
+		{
+			return Collection.BulkWrite(Session, requests, Options,  CancellationToken);
+		}
+
+		/// <inheritdoc />
+		public async Task<BulkWriteResult<T>> BulkWriteAsync(IEnumerable<WriteModel<T>> requests, BulkWriteOptions Options = null!,
+			CancellationToken  CancellationToken = new CancellationToken())
+		{
+			using IScope _ = GlobalTracer.Instance.BuildSpan(Prefix + "BulkWriteAsync").StartActive();
+			return await Collection.BulkWriteAsync(requests, Options,  CancellationToken);
+		}
+
+		/// <inheritdoc />
+		public async Task<BulkWriteResult<T>> BulkWriteAsync(IClientSessionHandle Session, IEnumerable<WriteModel<T>> requests, BulkWriteOptions Options = null!,
+			CancellationToken  CancellationToken = new CancellationToken())
+		{
+			using IScope _ = GlobalTracer.Instance.BuildSpan(Prefix + "BulkWriteAsync").StartActive();
+			return await Collection.BulkWriteAsync(Session, requests, Options,  CancellationToken);
+		}
+
+		/// <inheritdoc />
+		public long Count(FilterDefinition<T> filter, CountOptions Options = null!,
+			CancellationToken  CancellationToken = new CancellationToken())
+		{
+			return Collection.Count(filter, Options,  CancellationToken);
+		}
+
+		/// <inheritdoc />
+		public long Count(IClientSessionHandle Session, FilterDefinition<T> filter, CountOptions Options = null!,
+			CancellationToken  CancellationToken = new CancellationToken())
+		{
+			return Collection.Count(Session, filter, Options,  CancellationToken);
+		}
+
+		/// <inheritdoc />
+		public async Task<long> CountAsync(FilterDefinition<T> filter, CountOptions Options = null!,
+			CancellationToken  CancellationToken = new CancellationToken())
+		{
+			using IScope _ = GlobalTracer.Instance.BuildSpan(Prefix + "CountAsync").StartActive();
+			return await Collection.CountAsync(filter, Options,  CancellationToken);
+		}
+
+		/// <inheritdoc />
+		public async Task<long> CountAsync(IClientSessionHandle Session, FilterDefinition<T> filter, CountOptions Options = null!,
+			CancellationToken  CancellationToken = new CancellationToken())
+		{
+			using IScope _ = GlobalTracer.Instance.BuildSpan(Prefix + "CountAsync").StartActive();
+			return await Collection.CountAsync(Session, filter, Options,  CancellationToken);
+		}
+
+		/// <inheritdoc />
+		public long CountDocuments(FilterDefinition<T> filter, CountOptions Options = null!,
+			CancellationToken  CancellationToken = new CancellationToken())
+		{
+			return Collection.CountDocuments(filter, Options,  CancellationToken);
+		}
+
+		/// <inheritdoc />
+		public long CountDocuments(IClientSessionHandle Session, FilterDefinition<T> filter, CountOptions Options = null!,
+			CancellationToken  CancellationToken = new CancellationToken())
+		{
+			using IScope _ = GlobalTracer.Instance.BuildSpan(Prefix + "CountDocumentsAsync").StartActive();
+			return Collection.CountDocuments(Session, filter, Options,  CancellationToken);
+		}
+
+		/// <inheritdoc />
+		public async Task<long> CountDocumentsAsync(FilterDefinition<T> filter, CountOptions Options = null!,
+			CancellationToken  CancellationToken = new CancellationToken())
+		{
+			using IScope _ = GlobalTracer.Instance.BuildSpan(Prefix + "CountDocumentsAsync").StartActive();
+			return await Collection.CountDocumentsAsync(filter, Options,  CancellationToken);
+		}
+
+		/// <inheritdoc />
+		public async Task<long> CountDocumentsAsync(IClientSessionHandle Session, FilterDefinition<T> filter, CountOptions Options = null!,
+			CancellationToken  CancellationToken = new CancellationToken())
+		{
+			using IScope _ = GlobalTracer.Instance.BuildSpan(Prefix + "CountDocumentsAsync").StartActive();
+			return await Collection.CountDocumentsAsync(Session, filter, Options,  CancellationToken);
+		}
+
+		/// <inheritdoc />
+		public DeleteResult DeleteMany(FilterDefinition<T> filter, CancellationToken  CancellationToken = new CancellationToken())
+		{
+			return Collection.DeleteMany(filter,  CancellationToken);
+		}
+
+		/// <inheritdoc />
+		public DeleteResult DeleteMany(FilterDefinition<T> filter, DeleteOptions Options,
+			CancellationToken  CancellationToken = new CancellationToken())
+		{
+			return Collection.DeleteMany(filter, Options,  CancellationToken);
+		}
+
+		/// <inheritdoc />
+		public DeleteResult DeleteMany(IClientSessionHandle Session, FilterDefinition<T> filter, DeleteOptions Options = null!,
+			CancellationToken  CancellationToken = new CancellationToken())
+		{
+			return Collection.DeleteMany(Session, filter, Options,  CancellationToken);
+		}
+
+		/// <inheritdoc />
+		public async Task<DeleteResult> DeleteManyAsync(FilterDefinition<T> filter, CancellationToken  CancellationToken = new CancellationToken())
+		{
+			using IScope _ = GlobalTracer.Instance.BuildSpan(Prefix + "DeleteManyAsync").StartActive();
+			return await Collection.DeleteManyAsync(filter,  CancellationToken);
+		}
+
+		/// <inheritdoc />
+		public async Task<DeleteResult> DeleteManyAsync(FilterDefinition<T> filter, DeleteOptions Options,
+			CancellationToken  CancellationToken = new CancellationToken())
+		{
+			using IScope _ = GlobalTracer.Instance.BuildSpan(Prefix + "DeleteManyAsync").StartActive();
+			return await Collection.DeleteManyAsync(filter, Options,  CancellationToken);
+		}
+
+		/// <inheritdoc />
+		public async Task<DeleteResult> DeleteManyAsync(IClientSessionHandle Session, FilterDefinition<T> filter, DeleteOptions Options = null!,
+			CancellationToken  CancellationToken = new CancellationToken())
+		{
+			using IScope _ = GlobalTracer.Instance.BuildSpan(Prefix + "DeleteManyAsync").StartActive();
+			return await Collection.DeleteManyAsync(Session, filter, Options,  CancellationToken);
+		}
+
+		/// <inheritdoc />
+		public DeleteResult DeleteOne(FilterDefinition<T> filter, CancellationToken  CancellationToken = new CancellationToken())
+		{
+			return Collection.DeleteOne(filter,  CancellationToken);
+		}
+
+		/// <inheritdoc />
+		public DeleteResult DeleteOne(FilterDefinition<T> filter, DeleteOptions Options,
+			CancellationToken  CancellationToken = new CancellationToken())
+		{
+			return Collection.DeleteOne(filter, Options,  CancellationToken);
+		}
+
+		/// <inheritdoc />
+		public DeleteResult DeleteOne(IClientSessionHandle Session, FilterDefinition<T> filter, DeleteOptions Options = null!,
+			CancellationToken  CancellationToken = new CancellationToken())
+		{
+			return Collection.DeleteOne(Session, filter, Options,  CancellationToken);
+		}
+
+		/// <inheritdoc />
+		public async Task<DeleteResult> DeleteOneAsync(FilterDefinition<T> filter, CancellationToken  CancellationToken = new CancellationToken())
+		{
+			using IScope _ = GlobalTracer.Instance.BuildSpan(Prefix + "DeleteOneAsync").StartActive();
+			return await Collection.DeleteOneAsync(filter,  CancellationToken);
+		}
+
+		/// <inheritdoc />
+		public async Task<DeleteResult> DeleteOneAsync(FilterDefinition<T> filter, DeleteOptions Options,
+			CancellationToken  CancellationToken = new CancellationToken())
+		{
+			using IScope _ = GlobalTracer.Instance.BuildSpan(Prefix + "DeleteOneAsync").StartActive();
+			return await Collection.DeleteOneAsync(filter, Options,  CancellationToken);
+		}
+
+		/// <inheritdoc />
+		public Task<DeleteResult> DeleteOneAsync(IClientSessionHandle Session, FilterDefinition<T> filter, DeleteOptions Options = null!,
+			CancellationToken  CancellationToken = new CancellationToken())
+		{
+			using IScope _ = GlobalTracer.Instance.BuildSpan(Prefix + "DeleteOneAsync").StartActive();
+			return Collection.DeleteOneAsync(Session, filter, Options,  CancellationToken);
+		}
+
+		/// <inheritdoc />
+		public IAsyncCursor<TField> Distinct<TField>(FieldDefinition<T, TField> field, FilterDefinition<T> filter, DistinctOptions Options = null!,
+			CancellationToken  CancellationToken = new CancellationToken())
+		{
+			return Collection.Distinct(field, filter, Options,  CancellationToken);
+		}
+
+		/// <inheritdoc />
+		public IAsyncCursor<TField> Distinct<TField>(IClientSessionHandle Session, FieldDefinition<T, TField> field, FilterDefinition<T> filter,
+			DistinctOptions Options = null!, CancellationToken  CancellationToken = new CancellationToken())
+		{
+			return Collection.Distinct(Session, field, filter, Options,  CancellationToken);
+		}
+
+		/// <inheritdoc />
+		public async Task<IAsyncCursor<TField>> DistinctAsync<TField>(FieldDefinition<T, TField> field, FilterDefinition<T> filter, DistinctOptions Options = null!,
+			CancellationToken  CancellationToken = new CancellationToken())
+		{
+			using IScope _ = GlobalTracer.Instance.BuildSpan(Prefix + "DistinctAsync").StartActive();
+			return await Collection.DistinctAsync(field, filter, Options,  CancellationToken);
+		}
+
+		/// <inheritdoc />
+		public async Task<IAsyncCursor<TField>> DistinctAsync<TField>(IClientSessionHandle Session, FieldDefinition<T, TField> field, FilterDefinition<T> filter,
+			DistinctOptions Options = null!, CancellationToken  CancellationToken = new CancellationToken())
+		{
+			using IScope _ = GlobalTracer.Instance.BuildSpan(Prefix + "DistinctAsync").StartActive();
+			return await Collection.DistinctAsync(Session, field, filter, Options,  CancellationToken);
+		}
+
+		/// <inheritdoc />
+		public long EstimatedDocumentCount(EstimatedDocumentCountOptions Options = null!,
+			CancellationToken  CancellationToken = new CancellationToken())
+		{
+			return Collection.EstimatedDocumentCount(Options,  CancellationToken);
+		}
+
+		/// <inheritdoc />
+		public async Task<long> EstimatedDocumentCountAsync(EstimatedDocumentCountOptions Options = null!,
+			CancellationToken  CancellationToken = new CancellationToken())
+		{
+			using IScope _ = GlobalTracer.Instance.BuildSpan(Prefix + "EstimatedDocumentCountAsync").StartActive();
+			return await Collection.EstimatedDocumentCountAsync(Options,  CancellationToken);
+		}
+
+		/// <inheritdoc />
+		public IAsyncCursor<TProjection> FindSync<TProjection>(FilterDefinition<T> filter, FindOptions<T, TProjection> Options = null!,
+			CancellationToken  CancellationToken = new CancellationToken())
+		{
+			return Collection.FindSync(filter, Options,  CancellationToken);
+		}
+
+		/// <inheritdoc />
+		public IAsyncCursor<TProjection> FindSync<TProjection>(IClientSessionHandle Session, FilterDefinition<T> filter, FindOptions<T, TProjection> Options = null!,
+			CancellationToken  CancellationToken = new CancellationToken())
+		{
+			return Collection.FindSync(Session, filter, Options,  CancellationToken);
+		}
+
+		/// <inheritdoc />
+		public async Task<IAsyncCursor<TProjection>> FindAsync<TProjection>(FilterDefinition<T> filter, FindOptions<T, TProjection> Options = null!,
+			CancellationToken  CancellationToken = new CancellationToken())
+		{
+			using IScope _ = GlobalTracer.Instance.BuildSpan(Prefix + "FindAsync").StartActive();
+			return await Collection.FindAsync(filter, Options,  CancellationToken);
+		}
+
+		/// <inheritdoc />
+		public async Task<IAsyncCursor<TProjection>> FindAsync<TProjection>(IClientSessionHandle Session, FilterDefinition<T> filter, FindOptions<T, TProjection> Options = null!,
+			CancellationToken  CancellationToken = new CancellationToken())
+		{
+			using IScope _ = GlobalTracer.Instance.BuildSpan(Prefix + "FindAsync").StartActive();
+			return await Collection.FindAsync(Session, filter, Options,  CancellationToken);
+		}
+
+		/// <inheritdoc />
+		public TProjection FindOneAndDelete<TProjection>(FilterDefinition<T> filter, FindOneAndDeleteOptions<T, TProjection> Options = null!,
+			CancellationToken  CancellationToken = new CancellationToken())
+		{
+			return Collection.FindOneAndDelete(filter, Options,  CancellationToken);
+		}
+
+		/// <inheritdoc />
+		public TProjection FindOneAndDelete<TProjection>(IClientSessionHandle Session, FilterDefinition<T> filter,
+			FindOneAndDeleteOptions<T, TProjection> Options = null!, CancellationToken  CancellationToken = new CancellationToken())
+		{
+			return Collection.FindOneAndDelete(Session, filter, Options,  CancellationToken);
+		}
+
+		/// <inheritdoc />
+		public async Task<TProjection> FindOneAndDeleteAsync<TProjection>(FilterDefinition<T> filter, FindOneAndDeleteOptions<T, TProjection> Options = null!,
+			CancellationToken  CancellationToken = new CancellationToken())
+		{
+			using IScope _ = GlobalTracer.Instance.BuildSpan(Prefix + "FindOneAndDeleteAsync").StartActive();
+			return await Collection.FindOneAndDeleteAsync(filter, Options,  CancellationToken);
+		}
+
+		/// <inheritdoc />
+		public async Task<TProjection> FindOneAndDeleteAsync<TProjection>(IClientSessionHandle Session, FilterDefinition<T> filter,
+			FindOneAndDeleteOptions<T, TProjection> Options = null!, CancellationToken  CancellationToken = new CancellationToken())
+		{
+			using IScope _ = GlobalTracer.Instance.BuildSpan(Prefix + "FindOneAndDeleteAsync").StartActive();
+			return await Collection.FindOneAndDeleteAsync(Session, filter, Options,  CancellationToken);
+		}
+
+		/// <inheritdoc />
+		public TProjection FindOneAndReplace<TProjection>(FilterDefinition<T> filter, T replacement,
+			FindOneAndReplaceOptions<T, TProjection> Options = null!, CancellationToken  CancellationToken = new CancellationToken())
+		{
+			return Collection.FindOneAndReplace(filter, replacement, Options,  CancellationToken);
+		}
+
+		/// <inheritdoc />
+		public TProjection FindOneAndReplace<TProjection>(IClientSessionHandle Session, FilterDefinition<T> filter, T replacement,
+			FindOneAndReplaceOptions<T, TProjection> Options = null!, CancellationToken  CancellationToken = new CancellationToken())
+		{
+			return Collection.FindOneAndReplace(Session, filter, replacement, Options,  CancellationToken);
+		}
+
+		/// <inheritdoc />
+		public async Task<TProjection> FindOneAndReplaceAsync<TProjection>(FilterDefinition<T> filter, T replacement,
+			FindOneAndReplaceOptions<T, TProjection> Options = null!, CancellationToken  CancellationToken = new CancellationToken())
+		{
+			using IScope _ = GlobalTracer.Instance.BuildSpan(Prefix + "FindOneAndReplaceAsync").StartActive();
+			return await Collection.FindOneAndReplaceAsync(filter, replacement, Options,  CancellationToken);
+		}
+
+		/// <inheritdoc />
+		public async Task<TProjection> FindOneAndReplaceAsync<TProjection>(IClientSessionHandle Session, FilterDefinition<T> filter, T replacement,
+			FindOneAndReplaceOptions<T, TProjection> Options = null!, CancellationToken  CancellationToken = new CancellationToken())
+		{
+			using IScope _ = GlobalTracer.Instance.BuildSpan(Prefix + "FindOneAndReplaceAsync").StartActive();
+			return await Collection.FindOneAndReplaceAsync(Session, filter, replacement, Options,  CancellationToken);
+		}
+
+		/// <inheritdoc />
+		public TProjection FindOneAndUpdate<TProjection>(FilterDefinition<T> filter, UpdateDefinition<T> update,
+			FindOneAndUpdateOptions<T, TProjection> Options = null!, CancellationToken  CancellationToken = new CancellationToken())
+		{
+			return Collection.FindOneAndUpdate(filter, update, Options,  CancellationToken);
+		}
+
+		/// <inheritdoc />
+		public TProjection FindOneAndUpdate<TProjection>(IClientSessionHandle Session, FilterDefinition<T> filter,
+			UpdateDefinition<T> update, FindOneAndUpdateOptions<T, TProjection> Options = null!,
+			CancellationToken  CancellationToken = new CancellationToken())
+		{
+			return Collection.FindOneAndUpdate(Session, filter, update, Options,  CancellationToken);
+		}
+
+		/// <inheritdoc />
+		public async Task<TProjection> FindOneAndUpdateAsync<TProjection>(FilterDefinition<T> filter, UpdateDefinition<T> update,
+			FindOneAndUpdateOptions<T, TProjection> Options = null!, CancellationToken  CancellationToken = new CancellationToken())
+		{
+			using IScope _ = GlobalTracer.Instance.BuildSpan(Prefix + "FindOneAndUpdateAsync").StartActive();
+			return await Collection.FindOneAndUpdateAsync(filter, update, Options,  CancellationToken);
+		}
+
+		/// <inheritdoc />
+		public async Task<TProjection> FindOneAndUpdateAsync<TProjection>(IClientSessionHandle Session, FilterDefinition<T> filter, UpdateDefinition<T> update,
+			FindOneAndUpdateOptions<T, TProjection> Options = null!, CancellationToken  CancellationToken = new CancellationToken())
+		{
+			using IScope _ = GlobalTracer.Instance.BuildSpan(Prefix + "FindOneAndUpdateAsync").StartActive();
+			return await Collection.FindOneAndUpdateAsync(Session, filter, update, Options,  CancellationToken);
+		}
+
+		/// <inheritdoc />
+		public void InsertOne( T Document, InsertOneOptions Options = null!,
+			CancellationToken  CancellationToken = new CancellationToken())
+		{
+			Collection.InsertOne(Document, Options,  CancellationToken);
+		}
+
+		/// <inheritdoc />
+		public void InsertOne(IClientSessionHandle Session, T Document, InsertOneOptions Options = null!,
+			CancellationToken  CancellationToken = new CancellationToken())
+		{
+			Collection.InsertOne(Session, Document, Options,  CancellationToken);
+		}
+
+		/// <inheritdoc />
+		public async Task InsertOneAsync( T Document, CancellationToken CancellationToken)
+		{
+			using IScope _ = GlobalTracer.Instance.BuildSpan(Prefix + "InsertOneAsync").StartActive();
+			await Collection.InsertOneAsync(Document, CancellationToken);
+		}
+
+		/// <inheritdoc />
+		public async Task InsertOneAsync( T Document, InsertOneOptions Options = null!,
+			CancellationToken  CancellationToken = new CancellationToken())
+		{
+			using IScope _ = GlobalTracer.Instance.BuildSpan(Prefix + "InsertOneAsync").StartActive();
+			await Collection.InsertOneAsync(Document, Options,  CancellationToken);
+		}
+
+		/// <inheritdoc />
+		public async Task InsertOneAsync(IClientSessionHandle Session, T Document, InsertOneOptions Options = null!,
+			CancellationToken CancellationToken = new CancellationToken())
+		{
+			using IScope _ = GlobalTracer.Instance.BuildSpan(Prefix + "InsertOneAsync").StartActive();
+			await Collection.InsertOneAsync(Session, Document, Options,  CancellationToken);
+		}
+
+		/// <inheritdoc />
+		public void InsertMany(IEnumerable<T> Documents, InsertManyOptions Options = null!,
+			CancellationToken  CancellationToken = new CancellationToken())
+		{
+			Collection.InsertMany(Documents, Options,  CancellationToken);
+		}
+
+		/// <inheritdoc />
+		public void InsertMany(IClientSessionHandle Session, IEnumerable<T> Documents, InsertManyOptions Options = null!,
+			CancellationToken  CancellationToken = new CancellationToken())
+		{
+			Collection.InsertMany(Session, Documents, Options,  CancellationToken);
+		}
+
+		/// <inheritdoc />
+		public async Task InsertManyAsync(IEnumerable<T> Documents, InsertManyOptions Options = null!,
+			CancellationToken  CancellationToken = new CancellationToken())
+		{
+			using IScope _ = GlobalTracer.Instance.BuildSpan(Prefix + "InsertManyAsync").StartActive();
+			await Collection.InsertManyAsync(Documents, Options,  CancellationToken);
+		}
+
+		/// <inheritdoc />
+		public async Task InsertManyAsync(IClientSessionHandle Session, IEnumerable<T> documents, InsertManyOptions Options = null!,
+			CancellationToken  CancellationToken = new CancellationToken())
+		{
+			using IScope _ = GlobalTracer.Instance.BuildSpan(Prefix + "InsertManyAsync").StartActive();
+			await Collection.InsertManyAsync(Session, documents, Options,  CancellationToken);
+		}
+
+		/// <inheritdoc />
+		public IAsyncCursor<TResult> MapReduce<TResult>(BsonJavaScript map, BsonJavaScript reduce, MapReduceOptions<T, TResult> Options = null!,
+			CancellationToken  CancellationToken = new CancellationToken())
+		{
+			return Collection.MapReduce(map, reduce, Options,  CancellationToken);
+		}
+
+		/// <inheritdoc />
+		public IAsyncCursor<TResult> MapReduce<TResult>(IClientSessionHandle Session, BsonJavaScript map, BsonJavaScript reduce,
+			MapReduceOptions<T, TResult> Options = null!, CancellationToken  CancellationToken = new CancellationToken())
+		{
+			return Collection.MapReduce(Session, map, reduce, Options,  CancellationToken);
+		}
+
+		/// <inheritdoc />
+		public async Task<IAsyncCursor<TResult>> MapReduceAsync<TResult>(BsonJavaScript map, BsonJavaScript reduce, MapReduceOptions<T, TResult> Options = null!,
+			CancellationToken  CancellationToken = new CancellationToken())
+		{
+			using IScope _ = GlobalTracer.Instance.BuildSpan(Prefix + "MapReduceAsync").StartActive();
+			return await Collection.MapReduceAsync(map, reduce, Options,  CancellationToken);
+		}
+
+		/// <inheritdoc />
+		public async Task<IAsyncCursor<TResult>> MapReduceAsync<TResult>(IClientSessionHandle Session, BsonJavaScript map, BsonJavaScript reduce,
+			MapReduceOptions<T, TResult> Options = null!, CancellationToken  CancellationToken = new CancellationToken())
+		{
+			using IScope _ = GlobalTracer.Instance.BuildSpan(Prefix + "MapReduceAsync").StartActive();
+			return await Collection.MapReduceAsync(Session, map, reduce, Options,  CancellationToken);
+		}
+
+		/// <inheritdoc />
+		public IFilteredMongoCollection<TDerivedDocument> OfType<TDerivedDocument>() where TDerivedDocument : T
+		{
+			return Collection.OfType<TDerivedDocument>();
+		}
+
+		/// <inheritdoc />
+		public ReplaceOneResult ReplaceOne(FilterDefinition<T> filter, T replacement, ReplaceOptions Options = null!,
+			CancellationToken  CancellationToken = new CancellationToken())
+		{
+			return Collection.ReplaceOne(filter, replacement, Options,  CancellationToken);
+		}
+
+		/// <inheritdoc />
+		public ReplaceOneResult ReplaceOne(FilterDefinition<T> filter, T replacement, UpdateOptions Options,
+			CancellationToken  CancellationToken = new CancellationToken())
+		{
+			return Collection.ReplaceOne(filter, replacement, Options,  CancellationToken);
+		}
+
+		/// <inheritdoc />
+		public ReplaceOneResult ReplaceOne(IClientSessionHandle Session, FilterDefinition<T> filter, T replacement,
+			ReplaceOptions Options = null!, CancellationToken  CancellationToken = new CancellationToken())
+		{
+			return Collection.ReplaceOne(Session, filter, replacement, Options,  CancellationToken);
+		}
+
+		/// <inheritdoc />
+		public ReplaceOneResult ReplaceOne(IClientSessionHandle Session, FilterDefinition<T> filter, T replacement, UpdateOptions Options,
+			CancellationToken  CancellationToken = new CancellationToken())
+		{
+			return Collection.ReplaceOne(Session, filter, replacement, Options,  CancellationToken);
+		}
+
+		/// <inheritdoc />
+		public async Task<ReplaceOneResult> ReplaceOneAsync(FilterDefinition<T> filter, T replacement, ReplaceOptions Options = null!,
+			CancellationToken  CancellationToken = new CancellationToken())
+		{
+			using IScope _ = GlobalTracer.Instance.BuildSpan(Prefix + "ReplaceOneAsync").StartActive();
+			return await Collection.ReplaceOneAsync(filter, replacement, Options,  CancellationToken);
+		}
+
+		/// <inheritdoc />
+		public async Task<ReplaceOneResult> ReplaceOneAsync(FilterDefinition<T> filter, T replacement, UpdateOptions Options,
+			CancellationToken  CancellationToken = new CancellationToken())
+		{
+			using IScope _ = GlobalTracer.Instance.BuildSpan(Prefix + "ReplaceOneAsync").StartActive();
+			return await Collection.ReplaceOneAsync(filter, replacement, Options,  CancellationToken);
+		}
+
+		/// <inheritdoc />
+		public async Task<ReplaceOneResult> ReplaceOneAsync(IClientSessionHandle Session, FilterDefinition<T> filter, T replacement,
+			ReplaceOptions Options = null!, CancellationToken  CancellationToken = new CancellationToken())
+		{
+			using IScope _ = GlobalTracer.Instance.BuildSpan(Prefix + "ReplaceOneAsync").StartActive();
+			return await Collection.ReplaceOneAsync(Session, filter, replacement, Options,  CancellationToken);
+		}
+
+		/// <inheritdoc />
+		public async Task<ReplaceOneResult> ReplaceOneAsync(IClientSessionHandle Session, FilterDefinition<T> filter, T replacement, UpdateOptions Options,
+			CancellationToken  CancellationToken = new CancellationToken())
+		{
+			using IScope _ = GlobalTracer.Instance.BuildSpan(Prefix + "ReplaceOneAsync").StartActive();
+			return await Collection.ReplaceOneAsync(Session, filter, replacement, Options,  CancellationToken);
+		}
+
+		/// <inheritdoc />
+		public UpdateResult UpdateMany(FilterDefinition<T> filter, UpdateDefinition<T> update, UpdateOptions Options = null!,
+			CancellationToken  CancellationToken = new CancellationToken())
+		{
+			return Collection.UpdateMany(filter, update, Options,  CancellationToken);
+		}
+
+		/// <inheritdoc />
+		public UpdateResult UpdateMany(IClientSessionHandle Session, FilterDefinition<T> filter, UpdateDefinition<T> update,
+			UpdateOptions Options = null!, CancellationToken  CancellationToken = new CancellationToken())
+		{
+			return Collection.UpdateMany(Session, filter, update, Options,  CancellationToken);
+		}
+
+		/// <inheritdoc />
+		public async Task<UpdateResult> UpdateManyAsync(FilterDefinition<T> filter, UpdateDefinition<T> update, UpdateOptions Options = null!,
+			CancellationToken  CancellationToken = new CancellationToken())
+		{
+			using IScope _ = GlobalTracer.Instance.BuildSpan(Prefix + "UpdateManyAsync").StartActive();
+			return await Collection.UpdateManyAsync(filter, update, Options,  CancellationToken);
+		}
+
+		/// <inheritdoc />
+		public async Task<UpdateResult> UpdateManyAsync(IClientSessionHandle Session, FilterDefinition<T> filter, UpdateDefinition<T> update,
+			UpdateOptions Options = null!, CancellationToken  CancellationToken = new CancellationToken())
+		{
+			using IScope _ = GlobalTracer.Instance.BuildSpan(Prefix + "UpdateManyAsync").StartActive();
+			return await Collection.UpdateManyAsync(Session, filter, update, Options,  CancellationToken);
+		}
+
+		/// <inheritdoc />
+		public UpdateResult UpdateOne(FilterDefinition<T> filter, UpdateDefinition<T> update, UpdateOptions Options = null!,
+			CancellationToken  CancellationToken = new CancellationToken())
+		{
+			return Collection.UpdateOne(filter, update, Options,  CancellationToken);
+		}
+
+		/// <inheritdoc />
+		public UpdateResult UpdateOne(IClientSessionHandle Session, FilterDefinition<T> filter, UpdateDefinition<T> update,
+			UpdateOptions Options = null!, CancellationToken  CancellationToken = new CancellationToken())
+		{
+			return Collection.UpdateOne(Session, filter, update, Options,  CancellationToken);
+		}
+
+		/// <inheritdoc />
+		public async Task<UpdateResult> UpdateOneAsync(FilterDefinition<T> filter, UpdateDefinition<T> update, UpdateOptions Options = null!,
+			CancellationToken  CancellationToken = new CancellationToken())
+		{
+			using IScope _ = GlobalTracer.Instance.BuildSpan(Prefix + "UpdateOneAsync").StartActive();
+			return await Collection.UpdateOneAsync(filter, update, Options,  CancellationToken);
+		}
+
+		/// <inheritdoc />
+		public async Task<UpdateResult> UpdateOneAsync(IClientSessionHandle Session, FilterDefinition<T> filter, UpdateDefinition<T> update,
+			UpdateOptions Options = null!, CancellationToken  CancellationToken = new CancellationToken())
+		{
+			using IScope _ = GlobalTracer.Instance.BuildSpan(Prefix + "UpdateOneAsync").StartActive();
+			return await Collection.UpdateOneAsync(Session, filter, update, Options,  CancellationToken);
+		}
+
+		/// <inheritdoc />
+		public IChangeStreamCursor<TResult> Watch<TResult>(PipelineDefinition<ChangeStreamDocument<T>, TResult> pipeline, ChangeStreamOptions Options = null!,
+			CancellationToken  CancellationToken = new CancellationToken())
+		{
+			return Collection.Watch(pipeline, Options,  CancellationToken);
+		}
+
+		/// <inheritdoc />
+		public IChangeStreamCursor<TResult> Watch<TResult>(IClientSessionHandle Session, PipelineDefinition<ChangeStreamDocument<T>, TResult> pipeline,
+			ChangeStreamOptions Options = null!, CancellationToken  CancellationToken = new CancellationToken())
+		{
+			return Collection.Watch(Session, pipeline, Options,  CancellationToken);
+		}
+
+		/// <inheritdoc />
+		public async Task<IChangeStreamCursor<TResult>> WatchAsync<TResult>(PipelineDefinition<ChangeStreamDocument<T>, TResult> pipeline, ChangeStreamOptions Options = null!,
+			CancellationToken  CancellationToken = new CancellationToken())
+		{
+			using IScope _ = GlobalTracer.Instance.BuildSpan(Prefix + "WatchAsync").StartActive();
+			return await Collection.WatchAsync(pipeline, Options,  CancellationToken);
+		}
+
+		/// <inheritdoc />
+		public async Task<IChangeStreamCursor<TResult>> WatchAsync<TResult>(IClientSessionHandle Session, PipelineDefinition<ChangeStreamDocument<T>, TResult> pipeline, ChangeStreamOptions Options = null!,
+			CancellationToken  CancellationToken = new CancellationToken())
+		{
+			using IScope _ = GlobalTracer.Instance.BuildSpan(Prefix + "WatchAsync").StartActive();
+			return await Collection.WatchAsync(Session, pipeline, Options,  CancellationToken);
+		}
+
+		/// <inheritdoc />
+		public IMongoCollection<T> WithReadConcern(ReadConcern ReadConcern)
+		{
+			return Collection.WithReadConcern(ReadConcern);
+		}
+
+		/// <inheritdoc />
+		public IMongoCollection<T> WithReadPreference(ReadPreference ReadPreference)
+		{
+			return Collection.WithReadPreference(ReadPreference);
+		}
+
+		/// <inheritdoc />
+		public IMongoCollection<T> WithWriteConcern(WriteConcern WriteConcern)
+		{
+			return Collection.WithWriteConcern(WriteConcern);
+		}
+#pragma warning restore CS0618
+		
+		/// <inheritdoc />
+		public CollectionNamespace CollectionNamespace => Collection.CollectionNamespace;
+
+		/// <inheritdoc />
+		public IMongoDatabase Database => Collection.Database;
+
+		/// <inheritdoc />
+		public IBsonSerializer<T> DocumentSerializer => Collection.DocumentSerializer;
+
+		/// <inheritdoc />
+		public IMongoIndexManager<T> Indexes => Collection.Indexes;
+
+		/// <inheritdoc />
+		public MongoCollectionSettings Settings => Collection.Settings;
+	}
+}
\ No newline at end of file
