Horde: Missing file.

diff --git a/Engine/Source/Programs/Horde/HordeServer/Services/DatabaseService.cs b/Engine/Source/Programs/Horde/HordeServer/Services/DatabaseService.cs
--- a/Engine/Source/Programs/Horde/HordeServer/Services/DatabaseService.cs
+++ b/Engine/Source/Programs/Horde/HordeServer/Services/DatabaseService.cs
@@ -32,6 +32,7 @@ using System.Text.RegularExpressions;
 using System.Net.NetworkInformation;
 using OpenTracing.Util;
 using OpenTracing;
+using HordeServer.Utiltiies;
 
 namespace HordeServer.Services
 {
@@ -116,15 +117,15 @@ namespace HordeServer.Services
 		const int DefaultMongoPort = 27017;
 
 		/// <summary>
-		/// Trace MongoDB commands in Datadog for observability
+		/// Trace MongoDB commands for observability
 		/// The built-in version from Datadog tracing library doesn't work for us (missing queries).
 		/// </summary>
-		class MongoDatadogTracer : IEventSubscriber
+		class MongoTracer : IEventSubscriber
 		{
 			readonly ConcurrentDictionary<int, IScope> Scopes = new ConcurrentDictionary<int, IScope>();
 			private readonly ReflectionEventSubscriber _subscriber;
 
-			public MongoDatadogTracer()
+			public MongoTracer()
 			{
 				_subscriber = new ReflectionEventSubscriber(this);
 			}
@@ -247,7 +248,7 @@ namespace HordeServer.Services
 					{
 						ClusterBuilder.Subscribe<CommandStartedEvent>(Event => TraceMongoCommand(Event.Command));
 					}
-					ClusterBuilder.Subscribe(new MongoDatadogTracer());
+					ClusterBuilder.Subscribe(new MongoTracer());
 				};
 				
 				MongoSettings.SslSettings = new SslSettings();
@@ -633,7 +634,7 @@ namespace HordeServer.Services
 		/// <returns></returns>
 		public IMongoCollection<T> GetCollection<T>(string Name)
 		{
-			return new DatadogTracingCollection<T>(Database.GetCollection<T>(Name));
+			return new MongoTracingCollection<T>(Database.GetCollection<T>(Name));
 		}
 
 		/// <summary>
