Horde: Prevent job getting stuck in limbo if it fails after retrying steps.

diff --git a/Engine/Source/Programs/Horde/HordeServer/Collections/Impl/JobCollection.cs b/Engine/Source/Programs/Horde/HordeServer/Collections/Impl/JobCollection.cs
--- a/Engine/Source/Programs/Horde/HordeServer/Collections/Impl/JobCollection.cs
+++ b/Engine/Source/Programs/Horde/HordeServer/Collections/Impl/JobCollection.cs
@@ -646,7 +646,7 @@ namespace HordeServer.Collections.Impl
 				List<NodeRef> RetriedNodes = JobDocument.RetriedNodes ?? new List<NodeRef>();
 
 				// Check if there are any steps that need to be run again
-				int InitialRetriedNodes = RetriedNodes.Count;
+				bool bUpdateState = false;
 				foreach (JobStepDocument Step in Batch.Steps)
 				{
 					if (Step.State == JobStepState.Ready || Step.State == JobStepState.Waiting)
@@ -660,11 +660,12 @@ namespace HordeServer.Collections.Impl
 						{
 							RetriedNodes.Add(NodeRef);
 						}
+						bUpdateState = true;
 					}
 				}
 
 				// Update the steps
-				if (RetriedNodes.Count > InitialRetriedNodes)
+				if (bUpdateState)
 				{
 					Updates.Clear();
 					UpdateBatches(JobDocument, Graph, Updates, Logger);
