Horde: Add support for agent chunking to exceed Mongo's document limit
#rnx

diff --git a/Engine/Source/Programs/Horde/HordeServer/Collections/Impl/AgentSoftwareCollection.cs b/Engine/Source/Programs/Horde/HordeServer/Collections/Impl/AgentSoftwareCollection.cs
--- a/Engine/Source/Programs/Horde/HordeServer/Collections/Impl/AgentSoftwareCollection.cs
+++ b/Engine/Source/Programs/Horde/HordeServer/Collections/Impl/AgentSoftwareCollection.cs
@@ -33,27 +33,69 @@ namespace HordeServer.Collections.Impl
 		{
 			[BsonId]
 			public string Id { get; set; }
+
 			public byte[] Data { get; set; } = null!;
 
-			[BsonConstructor]
+            public List<string>? ChunkIds { get; set; } = null;
+
+            [BsonConstructor]
 			private AgentSoftwareDocument()
 			{
 				this.Id = null!;
 				this.Data = null!;
-			}
+                this.ChunkIds = null;
+            }
 
-			public AgentSoftwareDocument(string Id, byte[] Data)
+			public AgentSoftwareDocument(string Id, byte[] Data, List<string>? ChunkIds = null)
 			{
 				this.Id = Id;
 				this.Data = Data;
-			}
+                this.ChunkIds = ChunkIds;
+            }
+		}
+
+		class ChunkDocument
+		{
+			[BsonId]
+			public string Id { get; set; }
+
+			public byte[] Data { get; set; } = null!;
+
+			public string Version { get; set; }
+
+
+            [BsonConstructor]
+			private ChunkDocument()
+			{
+				this.Id = null!;
+				this.Data = null!;
+                this.Version = null!;
+            }
+
+			public ChunkDocument(string Id, string Version, byte[] Data)
+			{
+				this.Id = Id;
+				this.Data = Data;
+                this.Version = Version;
+            }
 		}
 
+		// MongoDB document size limit is 16 megs, we chunk at this with a 1k buffer for anything else in document
+		readonly int ChunkSize = 16 * 1024 * 1024 - 1024;
+
+
+
 		/// <summary>
 		/// Template documents
 		/// </summary>
 		IMongoCollection<AgentSoftwareDocument> Collection;
 
+		/// <summary>
+		/// Agent template chunks
+		/// </summary>
+		IMongoCollection<ChunkDocument> AgentChunks;
+
+
 		/// <summary>
 		/// Constructor
 		/// </summary>
@@ -61,14 +103,52 @@ namespace HordeServer.Collections.Impl
 		public AgentSoftwareCollection(DatabaseService DatabaseService)
 		{
 			Collection = DatabaseService.GetCollection<AgentSoftwareDocument>("AgentSoftware");
+			AgentChunks = DatabaseService.GetCollection<ChunkDocument>("AgentSoftwareChunks");
 		}
 
 		/// <inheritdoc/>
 		public async Task<bool> AddAsync(string Version, byte[] Data)
 		{
-			AgentSoftwareDocument Document = new AgentSoftwareDocument(Version, Data);
-			return await Collection.InsertOneIgnoreDuplicatesAsync(Document);
-		}
+
+			if (await ExistsAsync(Version))
+			{
+                return false;
+            }
+
+            AgentSoftwareDocument Document;
+
+            // Check if agent will fit in document, otherwise we need to chunk
+            if (Data.Length <= ChunkSize)
+			{
+				Document = new AgentSoftwareDocument(Version, Data);
+				await Collection.InsertOneAsync(Document);
+                return true;
+            }
+
+			int TotalRead = 0;            
+            byte[] Buffer = new byte[ChunkSize];            
+			List<ChunkDocument> Chunks = new List<ChunkDocument>();
+            using (MemoryStream ChunkStream = new MemoryStream(Data))
+            {
+                for (; ; )
+                {
+					int BytesRead = ChunkStream.Read(Buffer, 0, ChunkSize);
+					Chunks.Add(new ChunkDocument(ObjectId.GenerateNewId().ToString(), Version, Buffer.Take(BytesRead).ToArray()));
+					TotalRead += BytesRead;
+					if (TotalRead == Data.Length)
+					{
+						break;
+					}
+                }
+            }
+
+            // insert the chunks
+            await AgentChunks.InsertManyAsync(Chunks);
+
+            Document = new AgentSoftwareDocument(Version, Array.Empty<byte>(), Chunks.Select( Chunk => Chunk.Id).ToList());
+			await Collection.InsertOneAsync(Document);
+            return true;
+        }
 
 		/// <inheritdoc/>
 		public async Task<bool> ExistsAsync(string Version)
@@ -79,16 +159,50 @@ namespace HordeServer.Collections.Impl
 		/// <inheritdoc/>
 		public async Task<bool> RemoveAsync(string Version)
 		{
-			DeleteResult Result = await Collection.DeleteOneAsync(x => x.Id == Version);
+			// Delete any chunks
+			await AgentChunks.DeleteManyAsync(x => x.Version == Version);
+
+			DeleteResult Result = await Collection.DeleteOneAsync(x => x.Id == Version);			
 			return Result.DeletedCount > 0;
 		}
 
 		/// <inheritdoc/>
 		public async Task<byte[]?> GetAsync(string Version)
-		{
+		{			
 			IFindFluent<AgentSoftwareDocument, AgentSoftwareDocument> Query = Collection.Find(x => x.Id == Version);
 			AgentSoftwareDocument? Software = await Query.FirstOrDefaultAsync();
-			return Software?.Data;
-		}
+
+			if (Software == null)
+			{
+                return null;
+            }
+
+			if (Software.Data.Length > 0)
+			{
+                return Software.Data;
+            }
+
+			if (Software.ChunkIds != null)
+			{
+				FilterDefinitionBuilder<ChunkDocument> FilterBuilder = Builders<ChunkDocument>.Filter;
+				FilterDefinition<ChunkDocument> Filter = FilterBuilder.Empty;
+				Filter &= FilterBuilder.In(x => x.Id, Software.ChunkIds);
+
+				IFindFluent<ChunkDocument, ChunkDocument> ChunkQuery = AgentChunks.Find(x => Software.ChunkIds.Contains(x.Id));				
+				List<ChunkDocument> Chunks = await ChunkQuery.ToListAsync();
+
+				using (MemoryStream DataStream = new MemoryStream())
+				{
+					foreach (string ChunkId in Software.ChunkIds)
+					{
+                        DataStream.Write(Chunks.First(x => x.Id == ChunkId).Data);
+                    }
+
+                    return DataStream.ToArray();
+                }
+			}
+
+            return null;
+        }
 	}
 }
diff --git a/Engine/Source/Programs/Horde/HordeServer/Controllers/AgentSoftwareController.cs b/Engine/Source/Programs/Horde/HordeServer/Controllers/AgentSoftwareController.cs
--- a/Engine/Source/Programs/Horde/HordeServer/Controllers/AgentSoftwareController.cs
+++ b/Engine/Source/Programs/Horde/HordeServer/Controllers/AgentSoftwareController.cs
@@ -147,7 +147,8 @@ namespace HordeServer.Controllers
 				return NotFound();
 			}
 
-			return new FileStreamResult(new MemoryStream(Data), new MediaTypeHeaderValue("application/octet-stream"));
+			return new FileStreamResult(new MemoryStream(Data), new MediaTypeHeaderValue("application/octet-stream")) { FileDownloadName = $"HordeAgent.zip" };
+			
 		}
 	}
 }
diff --git a/Engine/Source/Programs/Horde/HordeServer/Controllers/SoftwareController.cs b/Engine/Source/Programs/Horde/HordeServer/Controllers/SoftwareController.cs
--- a/Engine/Source/Programs/Horde/HordeServer/Controllers/SoftwareController.cs
+++ b/Engine/Source/Programs/Horde/HordeServer/Controllers/SoftwareController.cs
@@ -144,7 +144,7 @@ namespace HordeServer.Controllers
 				return NotFound();
 			}
 
-			return new FileStreamResult(new MemoryStream(Software.Data!), new MediaTypeHeaderValue("application/octet-stream")) { FileDownloadName = $"HordeAgent-{Software.Id.ToString()}.zip" };
+			return new FileStreamResult(new MemoryStream(Software.Data!), new MediaTypeHeaderValue("application/octet-stream")) { FileDownloadName = $"Software-{Software.Id.ToString()}.zip" };
 		}
 
 		/// <summary>
diff --git a/Engine/Source/Programs/Horde/HordeServer/Services/AgentSoftwareService.cs b/Engine/Source/Programs/Horde/HordeServer/Services/AgentSoftwareService.cs
--- a/Engine/Source/Programs/Horde/HordeServer/Services/AgentSoftwareService.cs
+++ b/Engine/Source/Programs/Horde/HordeServer/Services/AgentSoftwareService.cs
@@ -125,12 +125,6 @@ namespace HordeServer.Services
 		/// </summary>
 		LazyCachedValue<Task<AgentSoftwareChannels>> ChannelsDocument;
 
-		/// <summary>
-		/// Singleton instance of the client service
-		/// </summary>
-		private readonly SoftwareService SoftwareService;
-
-
 		/// <summary>
 		/// The server settings
 		/// </summary>
@@ -146,15 +140,13 @@ namespace HordeServer.Services
 		/// Constructor
 		/// </summary>
 		/// <param name="Collection">The software collection</param>
-		/// <param name="Singleton">The channels singleton</param>
-		/// <param name="SoftwareService">The software service</param>
+		/// <param name="Singleton">The channels singleton</param>		
 		/// <param name="Settings">The settings monitor</param>		
 		/// <param name="Logger">The logger instance</param>
-		public AgentSoftwareService(IAgentSoftwareCollection Collection, ISingletonDocument<AgentSoftwareChannels> Singleton, SoftwareService SoftwareService, IOptionsMonitor<ServerSettings> Settings, ILogger<AgentSoftwareService> Logger)
+		public AgentSoftwareService(IAgentSoftwareCollection Collection, ISingletonDocument<AgentSoftwareChannels> Singleton, IOptionsMonitor<ServerSettings> Settings, ILogger<AgentSoftwareService> Logger)
 		{
 			this.Collection = Collection;
 			this.Singleton = Singleton;
-			this.SoftwareService = SoftwareService;
 			this.Settings = Settings;
 			this.Logger = Logger;
 
@@ -346,10 +338,9 @@ namespace HordeServer.Services
 				Logger.LogError(Ex, "Error checking default agent software, {Message}", Ex.Message);
 			}
 
-			using (System.IO.Stream Stream = File.OpenRead(AgentZip.ToString()))
-			{
-				await SoftwareService.CreateSoftwareAsync(Stream, null, true);
-			}
+            byte[] Bytes = await File.ReadAllBytesAsync(AgentZip.ToString());
+
+			string Version = await SetArchiveAsync(new AgentSoftwareChannelName("default"), null, Bytes);
 
 			if (!DirectoryReference.Exists(AgentHashFile.Directory))
 			{
