Horde: Fix error building agent with native P4 library.

diff --git a/Engine/Source/Programs/Horde/HordeAgent/HordeAgent.csproj b/Engine/Source/Programs/Horde/HordeAgent/HordeAgent.csproj
--- a/Engine/Source/Programs/Horde/HordeAgent/HordeAgent.csproj
+++ b/Engine/Source/Programs/Horde/HordeAgent/HordeAgent.csproj
@@ -124,7 +124,7 @@
     <ItemGroup>
       <OutputFiles Include="$(OutDir)**" />
       <BinaryFilesBase Include="$(OutDir)\*.dll" />
-      <BinaryFiles Include="@(BinaryFilesBase)" Condition="Exists('%(RelativeDir)%(FileName).pdb')" />
+      <BinaryFiles Include="@(BinaryFilesBase)" Condition="Exists('%(RelativeDir)%(FileName).pdb') And '%(FileName)' != 'EpicGames.Perforce.Native'" />
     </ItemGroup>
 
     <MakeDir Directories="$(FullSymbolsDir)" />
