Horde: Prevent job steps being retried infinitely if a batch fails; allow it to retry two times, then skip execution permanently.

diff --git a/Engine/Source/Programs/Horde/HordeServer/Collections/Impl/JobCollection.cs b/Engine/Source/Programs/Horde/HordeServer/Collections/Impl/JobCollection.cs
--- a/Engine/Source/Programs/Horde/HordeServer/Collections/Impl/JobCollection.cs
+++ b/Engine/Source/Programs/Horde/HordeServer/Collections/Impl/JobCollection.cs
@@ -321,6 +321,11 @@ namespace HordeServer.Collections.Impl
 			public StreamId StreamId { get; set; }
 		}
 
+		/// <summary>
+		/// Maximum number of times a step can be retried (after the original run)
+		/// </summary>
+		const int MaxRetries = 2;
+
 		/// <summary>
 		/// The jobs collection
 		/// </summary>
@@ -651,14 +656,13 @@ namespace HordeServer.Collections.Impl
 				{
 					if (Step.State == JobStepState.Ready || Step.State == JobStepState.Waiting)
 					{
-						NodeRef NodeRef = new NodeRef(Batch.GroupIdx, Step.NodeIdx);
-						if (JobDocument.RetriedNodes != null && JobDocument.RetriedNodes.Contains(NodeRef))
+						if (CanRetryNode(JobDocument, Batch.GroupIdx, Step.NodeIdx))
 						{
-							Step.State = JobStepState.Skipped;
+							RetriedNodes.Add(new NodeRef(Batch.GroupIdx, Step.NodeIdx));
 						}
 						else
 						{
-							RetriedNodes.Add(NodeRef);
+							Step.State = JobStepState.Skipped;
 						}
 						bUpdateState = true;
 					}
@@ -1198,7 +1202,7 @@ namespace HordeServer.Collections.Impl
 					{
 						FailedNodes.Remove(Node);
 					}
-					else if (Step.State == JobStepState.Skipped && Node.InputDependencies.Any(x => FailedNodes.Contains(Graph.GetNode(x))))
+					else if (Step.State == JobStepState.Skipped && (Node.InputDependencies.Any(x => FailedNodes.Contains(Graph.GetNode(x))) || !CanRetryNode(Job, Batch.GroupIdx, Step.NodeIdx)))
 					{
 						FailedNodes.Add(Node);
 					}
@@ -1439,6 +1443,14 @@ namespace HordeServer.Collections.Impl
 			}
 		}
 
+		/// <summary>
+		/// Tests whether a node can be retried again
+		/// </summary>
+		static bool CanRetryNode(JobDocument Job, int GroupIdx, int NodeIdx)
+		{
+			return Job.RetriedNodes == null || Job.RetriedNodes.Count(x => x.GroupIdx == GroupIdx && x.NodeIdx == NodeIdx) < MaxRetries;
+		}
+
 		/// <summary>
 		/// Gets the scheduling priority of this job
 		/// </summary>
diff --git a/Engine/Source/Programs/Horde/HordeServerTests/JobCollectionTests.cs b/Engine/Source/Programs/Horde/HordeServerTests/JobCollectionTests.cs
--- a/Engine/Source/Programs/Horde/HordeServerTests/JobCollectionTests.cs
+++ b/Engine/Source/Programs/Horde/HordeServerTests/JobCollectionTests.cs
@@ -200,5 +200,54 @@ namespace HordeServerTests
 				Assert.AreEqual(2, NewJob.Batches[2].Steps[1].NodeIdx);
 			}
 		}
+
+		[TestMethod]
+		public async Task IncompleteBatchAsync()
+		{
+			Mock<ITemplate> TemplateMock = new Mock<ITemplate>(MockBehavior.Strict);
+			TemplateMock.SetupGet(x => x.InitialAgentType).Returns((string?)null);
+
+			IGraph BaseGraph = await GraphCollection.AddAsync(TemplateMock.Object);
+
+			List<string> Arguments = new List<string>();
+			Arguments.Add("-Target=Step 1");
+			Arguments.Add("-Target=Step 3");
+
+			IJob Job = await JobCollection.AddAsync(ObjectId.GenerateNewId(), new StreamId("ue4-main"), new TemplateRefId("test-build"), ContentHash.SHA1("hello"), BaseGraph, "Test job", 123, 123, null, null, null, "Ben", null, null, null, null, false, false, null, null, null, Arguments);
+			Assert.AreEqual(1, Job.Batches.Count);
+
+			await StartBatch(Job, BaseGraph, 0);
+			Assert.IsTrue(await JobCollection.TryUpdateBatchAsync(Job, BaseGraph, Job.Batches[0].Id, null, JobStepBatchState.Complete, JobStepBatchError.Incomplete));
+
+			Job = (await JobCollection.GetAsync(Job.Id))!;
+			Assert.AreEqual(2, Job.Batches.Count);
+			Assert.AreEqual(JobStepBatchState.Complete, Job.Batches[0].State);
+			Assert.AreEqual(0, Job.Batches[0].Steps.Count);
+			Assert.AreEqual(JobStepBatchState.Ready, Job.Batches[1].State);
+			Assert.AreEqual(1, Job.Batches[1].Steps.Count);
+
+			Assert.IsTrue(await JobCollection.TryUpdateBatchAsync(Job, BaseGraph, Job.Batches[1].Id, null, JobStepBatchState.Complete, JobStepBatchError.Incomplete));
+
+			Job = (await JobCollection.GetAsync(Job.Id))!;
+			Assert.AreEqual(3, Job.Batches.Count);
+			Assert.AreEqual(JobStepBatchState.Complete, Job.Batches[0].State);
+			Assert.AreEqual(0, Job.Batches[0].Steps.Count);
+			Assert.AreEqual(JobStepBatchState.Complete, Job.Batches[1].State);
+			Assert.AreEqual(0, Job.Batches[1].Steps.Count);
+			Assert.AreEqual(JobStepBatchState.Ready, Job.Batches[2].State);
+			Assert.AreEqual(1, Job.Batches[2].Steps.Count);
+
+			Assert.IsTrue(await JobCollection.TryUpdateBatchAsync(Job, BaseGraph, Job.Batches[2].Id, null, JobStepBatchState.Complete, JobStepBatchError.Incomplete));
+
+			Job = (await JobCollection.GetAsync(Job.Id))!;
+			Assert.AreEqual(3, Job.Batches.Count);
+			Assert.AreEqual(JobStepBatchState.Complete, Job.Batches[0].State);
+			Assert.AreEqual(0, Job.Batches[0].Steps.Count);
+			Assert.AreEqual(JobStepBatchState.Complete, Job.Batches[1].State);
+			Assert.AreEqual(0, Job.Batches[1].Steps.Count);
+			Assert.AreEqual(JobStepBatchState.Complete, Job.Batches[2].State);
+			Assert.AreEqual(1, Job.Batches[2].Steps.Count);
+			Assert.AreEqual(JobStepState.Skipped, Job.Batches[2].Steps[0].State);
+		}
 	}
 }
