Horde: Read listening ports from config

diff --git a/Engine/Source/Programs/Horde/HordeServer/Program.cs b/Engine/Source/Programs/Horde/HordeServer/Program.cs
--- a/Engine/Source/Programs/Horde/HordeServer/Program.cs
+++ b/Engine/Source/Programs/Horde/HordeServer/Program.cs
@@ -243,13 +243,13 @@ namespace HordeServer
 					{
 						if (SslCert != null)
 						{
-							const int HttpPort = 80;
+							int HttpPort = Config.GetSection("Horde").GetValue("HttpPort", 80);
 							Options.ListenAnyIP(HttpPort, Configure => { });
 							
-							const int Http2Port = 52103;
+							int Http2Port = Config.GetSection("Horde").GetValue("Http2Port", 52103);
 							Options.ListenAnyIP(Http2Port, Configure => { Configure.Protocols = HttpProtocols.Http2; });
 
-							const int HttpsPort = 443;
+							int HttpsPort = Config.GetSection("Horde").GetValue("HttpsPort", 443);
 							Options.ListenAnyIP(HttpsPort, Configure => { Configure.UseHttps(SslCert); });
 						}
 					});
