Horde: Implement health checking as a separate service

Having it in HordeRpc service will not work since it's always authenticated.
This will instead put it as a separate service, following the convention of gRPC servers (see HealthService.cs docs).

diff --git a/Engine/Source/Programs/Horde/HordeServer/HordeServer.csproj b/Engine/Source/Programs/Horde/HordeServer/HordeServer.csproj
--- a/Engine/Source/Programs/Horde/HordeServer/HordeServer.csproj
+++ b/Engine/Source/Programs/Horde/HordeServer/HordeServer.csproj
@@ -72,6 +72,7 @@
     <Protobuf Include="..\..\Shared\EpicGames.Horde.Common\Protos\horde\storage\ref_service.proto" ProtoRoot="..\..\Shared\EpicGames.Horde.Common\Protos" GrpcServices="Server" />
     <Protobuf Include="..\..\Shared\EpicGames.Horde.Common\Protos\horde\storage\blob_service.proto" ProtoRoot="..\..\Shared\EpicGames.Horde.Common\Protos" GrpcServices="Server" />
     <Protobuf Include="..\HordeCommon\Protos\horde\horde_rpc.proto" ProtoRoot="..\HordeCommon\Protos" GrpcServices="Server" />
+    <Protobuf Include="..\HordeCommon\Protos\grpc\health\v1\health.proto" ProtoRoot="..\HordeCommon\Protos" GrpcServices="Server" />
   </ItemGroup>
 
   <ItemGroup>
diff --git a/Engine/Source/Programs/Horde/HordeServer/Services/HealthService.cs b/Engine/Source/Programs/Horde/HordeServer/Services/HealthService.cs
new file mode 100644
--- /dev/null
+++ b/Engine/Source/Programs/Horde/HordeServer/Services/HealthService.cs
@@ -0,0 +1,66 @@
+// Copyright Epic Games, Inc. All Rights Reserved.
+
+using Grpc.Core;
+using Microsoft.Extensions.Logging;
+using System.Threading.Tasks;
+using Grpc.Health.V1;
+using static Grpc.Health.V1.HealthCheckResponse.Types;
+
+namespace HordeServer.Services
+{
+	/// <summary>
+	/// Implements the gRPC health checking protocol
+	/// See [REDACTED_PROFILE_URL]/grpc/blob/master/doc/health-checking.md for details
+	/// </summary>
+	public class HealthService : Health.HealthBase
+	{
+		/// <summary>
+		/// The application lifetime interface
+		/// </summary>
+		LifetimeService LifetimeService;
+
+		/// <summary>
+		/// Writer for log output
+		/// </summary>
+		ILogger<HealthService> Logger;
+
+		/// <summary>
+		/// Constructor
+		/// </summary>
+		/// <param name="LifetimeService">The application lifetime</param>
+		/// <param name="Logger">Log writer</param>
+		public HealthService(LifetimeService LifetimeService, ILogger<HealthService> Logger)
+		{
+			this.LifetimeService = LifetimeService;
+			this.Logger = Logger;
+		}
+
+		/// <summary>
+		/// Return server status, primarily intended for load balancers to decide if traffic should be routed to this process
+		/// The supplied service name is currently ignored.
+		///
+		/// For example, the gRPC health check for AWS ALB will pick up and react to the gRPC status code returned.
+		/// </summary>
+		/// <param name="Request">Empty placeholder request (for now)</param>
+		/// <param name="Context">Context for the call</param>
+		/// <returns>Return status code 'unavailable' if stopping</returns>
+		public override Task<HealthCheckResponse> Check(HealthCheckRequest Request, ServerCallContext Context)
+		{
+			ServingStatus Status = ServingStatus.Serving;
+			
+			bool IsStopping = LifetimeService.IsPreStopping || LifetimeService.IsStopping;
+			if (IsStopping)
+			{
+				Context.Status = new Status(StatusCode.Unavailable, "Server is stopping");
+				Status = ServingStatus.NotServing;
+			}
+			 
+			return Task.FromResult(new HealthCheckResponse {Status = Status});
+		}
+
+		public override Task Watch(HealthCheckRequest request, IServerStreamWriter<HealthCheckResponse> ResponseStream, ServerCallContext Context)
+		{
+			return Task.FromException(new RpcException(new Status(StatusCode.Unimplemented, "Watch() not implemented")));
+		}
+	}
+}
diff --git a/Engine/Source/Programs/Horde/HordeServer/Services/RpcService.cs b/Engine/Source/Programs/Horde/HordeServer/Services/RpcService.cs
--- a/Engine/Source/Programs/Horde/HordeServer/Services/RpcService.cs
+++ b/Engine/Source/Programs/Horde/HordeServer/Services/RpcService.cs
@@ -162,25 +162,6 @@ namespace HordeServer.Services
 			this.Logger = Logger;
 		}
 
-		/// <summary>
-		/// Return server status, primarily intended for load balancers to decide if traffic should be routed to this process
-		///
-		/// For example, the gRPC healthcheck for AWS ALB will pick up and react to the gRPC status code returned.
-		/// </summary>
-		/// <param name="Request">Empty placeholder request (for now)</param>
-		/// <param name="Context">Context for the call</param>
-		/// <returns>Return status code 'unavailable' if stopping</returns>
-		public override Task<HealthCheckResponse> HealthCheck(HealthCheckRequest Request, ServerCallContext Context)
-		{
-			bool IsStopping = LifetimeService.IsPreStopping || LifetimeService.IsStopping;
-			if (IsStopping)
-			{
-				Context.Status = new Status(StatusCode.Unavailable, "Server is stopping");	
-			}
-
-			return Task.FromResult(new HealthCheckResponse {Stopping = IsStopping});
-		}
-
 		/// <summary>
 		/// Waits until the server is terminating
 		/// </summary>
diff --git a/Engine/Source/Programs/Horde/HordeServer/Startup.cs b/Engine/Source/Programs/Horde/HordeServer/Startup.cs
--- a/Engine/Source/Programs/Horde/HordeServer/Startup.cs
+++ b/Engine/Source/Programs/Horde/HordeServer/Startup.cs
@@ -788,7 +788,10 @@ namespace HordeServer
 
 			App.UseEndpoints(Endpoints =>
 			{
+				Endpoints.MapGrpcService<HealthService>();
 				Endpoints.MapGrpcService<RpcService>();
+				
+				// Google Remote Execution API
 				Endpoints.MapGrpcService<ActionRpcService>();
 				Endpoints.MapGrpcService<ActionCacheService>();
 				Endpoints.MapGrpcService<ByteStreamService>();
