Horde: Add missing method summary

diff --git a/Engine/Source/Programs/Horde/HordeServer/Services/HealthService.cs b/Engine/Source/Programs/Horde/HordeServer/Services/HealthService.cs
--- a/Engine/Source/Programs/Horde/HordeServer/Services/HealthService.cs
+++ b/Engine/Source/Programs/Horde/HordeServer/Services/HealthService.cs
@@ -58,7 +58,15 @@ namespace HordeServer.Services
 			return Task.FromResult(new HealthCheckResponse {Status = Status});
 		}
 
-		public override Task Watch(HealthCheckRequest request, IServerStreamWriter<HealthCheckResponse> ResponseStream, ServerCallContext Context)
+		
+		/// <summary>
+		/// Stream the server health status (not implemented)
+		/// </summary>
+		/// <param name="Request"></param>
+		/// <param name="ResponseStream"></param>
+		/// <param name="Context"></param>
+		/// <returns></returns>
+		public override Task Watch(HealthCheckRequest Request, IServerStreamWriter<HealthCheckResponse> ResponseStream, ServerCallContext Context)
 		{
 			return Task.FromException(new RpcException(new Status(StatusCode.Unimplemented, "Watch() not implemented")));
 		}
