Horde: Force HTTP 1 and 2 on default port

diff --git a/Engine/Source/Programs/Horde/HordeServer/Program.cs b/Engine/Source/Programs/Horde/HordeServer/Program.cs
--- a/Engine/Source/Programs/Horde/HordeServer/Program.cs
+++ b/Engine/Source/Programs/Horde/HordeServer/Program.cs
@@ -244,7 +244,7 @@ namespace HordeServer
 						if (SslCert != null)
 						{
 							int HttpPort = Config.GetSection("Horde").GetValue("HttpPort", 80);
-							Options.ListenAnyIP(HttpPort, Configure => { });
+							Options.ListenAnyIP(HttpPort, Configure => { Configure.Protocols = HttpProtocols.Http1AndHttp2; });
 							
 							int Http2Port = Config.GetSection("Horde").GetValue("Http2Port", 52103);
 							Options.ListenAnyIP(Http2Port, Configure => { Configure.Protocols = HttpProtocols.Http2; });
