Horde: Ensure Kestrel port config is in ServerSettings

diff --git a/Engine/Source/Programs/Horde/HordeServer/Program.cs b/Engine/Source/Programs/Horde/HordeServer/Program.cs
--- a/Engine/Source/Programs/Horde/HordeServer/Program.cs
+++ b/Engine/Source/Programs/Horde/HordeServer/Program.cs
@@ -215,7 +215,7 @@ namespace HordeServer
 			{
 				using (X509Certificate2? GrpcCertificate = ReadGrpcCertificate(HordeSettings))
 				{
-					CreateHostBuilderWithCert(Args, Config, GrpcCertificate).Build().Run();
+					CreateHostBuilderWithCert(Args, Config, HordeSettings, GrpcCertificate).Build().Run();
 				}
 			}
 #pragma warning disable CA1031 // Do not catch general exception types
@@ -229,10 +229,11 @@ namespace HordeServer
 		// Used by WebApplicationFactory in controller tests. Uses reflection to call this exact function signature.
 		public static IHostBuilder CreateHostBuilder(string[] Args)
 		{
-			return CreateHostBuilderWithCert(Args, new ConfigurationBuilder().Build(), null);
+			ServerSettings HordeSettings = new ServerSettings();
+			return CreateHostBuilderWithCert(Args, new ConfigurationBuilder().Build(), HordeSettings, null);
 		}
 
-		public static IHostBuilder CreateHostBuilderWithCert(string[] Args, IConfiguration Config, X509Certificate2? SslCert)
+		public static IHostBuilder CreateHostBuilderWithCert(string[] Args, IConfiguration Config, ServerSettings ServerSettings, X509Certificate2? SslCert)
 		{
 			return Host.CreateDefaultBuilder(Args)
 				.UseSerilog()
@@ -243,14 +244,9 @@ namespace HordeServer
 					{
 						if (SslCert != null)
 						{
-							int HttpPort = Config.GetSection("Horde").GetValue("HttpPort", 80);
-							Options.ListenAnyIP(HttpPort, Configure => { Configure.Protocols = HttpProtocols.Http1AndHttp2; });
-							
-							int Http2Port = Config.GetSection("Horde").GetValue("Http2Port", 52103);
-							Options.ListenAnyIP(Http2Port, Configure => { Configure.Protocols = HttpProtocols.Http2; });
-
-							int HttpsPort = Config.GetSection("Horde").GetValue("HttpsPort", 443);
-							Options.ListenAnyIP(HttpsPort, Configure => { Configure.UseHttps(SslCert); });
+							Options.ListenAnyIP(ServerSettings.HttpPort, Configure => { Configure.Protocols = HttpProtocols.Http1AndHttp2; });
+							Options.ListenAnyIP(ServerSettings.Http2Port, Configure => { Configure.Protocols = HttpProtocols.Http2; });
+							Options.ListenAnyIP(ServerSettings.HttpsPort, Configure => { Configure.UseHttps(SslCert); });
 						}
 					});
 					WebBuilder.UseStartup<Startup>(); 
diff --git a/Engine/Source/Programs/Horde/HordeServer/ServerSettings.cs b/Engine/Source/Programs/Horde/HordeServer/ServerSettings.cs
--- a/Engine/Source/Programs/Horde/HordeServer/ServerSettings.cs
+++ b/Engine/Source/Programs/Horde/HordeServer/ServerSettings.cs
@@ -123,6 +123,20 @@ namespace HordeServer
 	/// </summary>
 	public class ServerSettings
 	{
+		/// <summary>
+		/// Main port for serving HTTP
+		/// </summary>
+		public int HttpPort { get; set; } = 80;
+		
+		/// <summary>
+		/// Port for serving only HTTP/2
+		/// </summary>
+		public int Http2Port { get; set; } = 52103;
+		
+		/// <summary>
+		/// Port for serving HTTP with TLS enabled
+		/// </summary>
+		public int HttpsPort { get; set; } = 443;
 
 		/// <summary>
 		/// Whether the server is running as a single instance or with multiple instances, such as in k8s
