Horde: Clarify the reason why a separate port for HTTP/2 is required

diff --git a/Engine/Source/Programs/Horde/HordeServer/Program.cs b/Engine/Source/Programs/Horde/HordeServer/Program.cs
--- a/Engine/Source/Programs/Horde/HordeServer/Program.cs
+++ b/Engine/Source/Programs/Horde/HordeServer/Program.cs
@@ -245,7 +245,11 @@ namespace HordeServer
 						if (SslCert != null)
 						{
 							Options.ListenAnyIP(ServerSettings.HttpPort, Configure => { Configure.Protocols = HttpProtocols.Http1AndHttp2; });
+							
+							// To serve HTTP/2 with gRPC *without* TLS enabled, a separate port for HTTP/2 must be used.
+							// This is useful when having a load balancer in front that terminates TLS.
 							Options.ListenAnyIP(ServerSettings.Http2Port, Configure => { Configure.Protocols = HttpProtocols.Http2; });
+							
 							Options.ListenAnyIP(ServerSettings.HttpsPort, Configure => { Configure.UseHttps(SslCert); });
 						}
 					});
