Horde: Support agent server profiles in installations
#rnx

diff --git a/Engine/Source/Programs/Horde/HordeAgent/Commands/Utilities/CreateCertCommand.cs b/Engine/Source/Programs/Horde/HordeAgent/Commands/Utilities/CreateCertCommand.cs
--- a/Engine/Source/Programs/Horde/HordeAgent/Commands/Utilities/CreateCertCommand.cs
+++ b/Engine/Source/Programs/Horde/HordeAgent/Commands/Utilities/CreateCertCommand.cs
@@ -1,5 +1,6 @@
 // Copyright Epic Games, Inc. All Rights Reserved.
 
+using HordeCommon;
 using Microsoft.Extensions.Configuration;
 using Microsoft.Extensions.Logging;
 using System;
@@ -75,32 +76,16 @@ namespace HordeAgent.Commands.Certs
 			}
 
 			Logger.LogInformation("Creating certificate for {DnsName}", DnsName);
-			using (RSA Algorithm = RSA.Create(2048))
-			{
-				X500DistinguishedName distinguishedName = new X500DistinguishedName($"CN={DnsName}");
-				CertificateRequest Request = new CertificateRequest(distinguishedName, Algorithm, HashAlgorithmName.SHA256, RSASignaturePadding.Pkcs1);
-
-				Request.CertificateExtensions.Add(new X509BasicConstraintsExtension(false, false, 0, true));
-				Request.CertificateExtensions.Add(new X509KeyUsageExtension(X509KeyUsageFlags.KeyEncipherment | X509KeyUsageFlags.DigitalSignature, true));
-				Request.CertificateExtensions.Add(new X509EnhancedKeyUsageExtension(new OidCollection { new Oid("1.3.6.1.5.5.7.3.1") }, true));
 
-				SubjectAlternativeNameBuilder AlternativeNameBuilder = new SubjectAlternativeNameBuilder();
-				AlternativeNameBuilder.AddDnsName(DnsName);
-				Request.CertificateExtensions.Add(AlternativeNameBuilder.Build(true));
+			byte[] PrivateCertData = AgentUtilities.CreateAgentCert(DnsName);
 
-				// NB: MacOS requires 825 days or fewer (https://support.apple.com/en-us/HT210176)
-				using (X509Certificate2 Certificate = Request.CreateSelfSigned(new DateTimeOffset(DateTime.UtcNow.AddDays(-1)), new DateTimeOffset(DateTime.UtcNow.AddDays(800))))
-				{
-					Certificate.FriendlyName = "Horde Server";
+			Logger.LogInformation("Writing private cert: {PrivateCert}", new FileReference(PrivateCertFile).FullName);
+			File.WriteAllBytes(PrivateCertFile, PrivateCertData);
 
-					byte[] PrivateCertData = Certificate.Export(X509ContentType.Pkcs12); // Note: Need to reimport this to use immediately, otherwise key is ephemeral
-					Logger.LogInformation("Writing private cert: {PrivateCert}", new FileReference(PrivateCertFile).FullName);
-					File.WriteAllBytes(PrivateCertFile, PrivateCertData);
+			using X509Certificate2 Certificate = new X509Certificate2(PrivateCertData);
+			Logger.LogInformation("Certificate thumbprint is {Thumbprint}", Certificate.Thumbprint);
+			Logger.LogInformation("Add this thumbprint to list of trusted servers in appsettings.json to trust this server.");
 
-					Logger.LogInformation("Certificate thumbprint is {Thumbprint}", Certificate.Thumbprint);
-					Logger.LogInformation("Add this thumbprint to list of trusted servers in appsettings.json to trust this server.");
-				}
-			}
 			return Task.FromResult(0);
 		}
 	}
diff --git a/Engine/Source/Programs/Horde/HordeCommon/AgentUtilities.cs b/Engine/Source/Programs/Horde/HordeCommon/AgentUtilities.cs
new file mode 100644
--- /dev/null
+++ b/Engine/Source/Programs/Horde/HordeCommon/AgentUtilities.cs
@@ -0,0 +1,138 @@
+// Copyright Epic Games, Inc. All Rights Reserved.
+
+using System;
+using System.Collections.Generic;
+using System.IO;
+using System.IO.Compression;
+using System.Security.Cryptography;
+using System.Security.Cryptography.X509Certificates;
+using System.Text.RegularExpressions;
+using System.Text.Json;
+using System.Diagnostics;
+
+namespace HordeCommon
+{
+	public static class AgentUtilities
+	{
+		/// <summary>
+		/// Create a self signed cert for the agent
+		/// </summary>
+		/// <param name="DnsName"></param>
+		/// <returns></returns>
+		public static byte[] CreateAgentCert(string DnsName)
+		{
+			DnsName = Regex.Replace(DnsName, @"^[a-zA-Z]+://", "");
+			DnsName = Regex.Replace(DnsName, "/.*$", "");
+
+			using (RSA Algorithm = RSA.Create(2048))
+			{
+				X500DistinguishedName distinguishedName = new X500DistinguishedName($"CN={DnsName}");
+				CertificateRequest Request = new CertificateRequest(distinguishedName, Algorithm, HashAlgorithmName.SHA256, RSASignaturePadding.Pkcs1);
+
+				Request.CertificateExtensions.Add(new X509BasicConstraintsExtension(false, false, 0, true));
+				Request.CertificateExtensions.Add(new X509KeyUsageExtension(X509KeyUsageFlags.KeyEncipherment | X509KeyUsageFlags.DigitalSignature, true));
+				Request.CertificateExtensions.Add(new X509EnhancedKeyUsageExtension(new OidCollection { new Oid("1.3.6.1.5.5.7.3.1") }, true));
+
+				SubjectAlternativeNameBuilder AlternativeNameBuilder = new SubjectAlternativeNameBuilder();
+				AlternativeNameBuilder.AddDnsName(DnsName);
+				Request.CertificateExtensions.Add(AlternativeNameBuilder.Build(true));
+
+				// NB: MacOS requires 825 days or fewer (https://support.apple.com/en-us/HT210176)
+				using (X509Certificate2 Certificate = Request.CreateSelfSigned(new DateTimeOffset(DateTime.UtcNow.AddDays(-1)), new DateTimeOffset(DateTime.UtcNow.AddDays(800))))
+				{
+					Certificate.FriendlyName = "Horde Server";
+					byte[] PrivateCertData = Certificate.Export(X509ContentType.Pkcs12); // Note: Need to reimport this to use immediately, otherwise key is ephemeral
+					return PrivateCertData;
+
+				}
+
+			}
+		}
+
+		/// <summary>
+		/// Reads the version number from an archive
+		/// </summary>
+		/// <param name="Data">The archive data</param>
+		/// <returns></returns>
+		public static string ReadVersion(byte[] Data)
+		{
+			MemoryStream InputStream = new MemoryStream(Data);
+			using (ZipArchive InputArchive = new ZipArchive(InputStream, ZipArchiveMode.Read, true))
+			{
+				foreach (ZipArchiveEntry InputEntry in InputArchive.Entries)
+				{
+					if (InputEntry.FullName.Equals("HordeAgent.dll", StringComparison.OrdinalIgnoreCase))
+					{
+						string TempFile = Path.GetTempFileName();
+						try
+						{
+							InputEntry.ExtractToFile(TempFile, true);
+							return FileVersionInfo.GetVersionInfo(TempFile).ProductVersion;
+						}
+						finally
+						{
+							File.Delete(TempFile);
+						}
+					}
+				}
+			}
+			throw new Exception("Unable to find HordeAgent.dll in archive");
+		}
+
+
+		/// <summary>
+		/// Updates the agent app settings within the archive data
+		/// </summary>
+		/// <param name="Data">Data for the zip archive</param>
+		/// <param name="Settings">The settings to update</param>
+		/// <returns>New agent app data</returns>
+		public static byte[] UpdateAppSettings(byte[] Data, Dictionary<string, object> Settings)
+		{
+			bool bWrittenClientId = false;
+
+			MemoryStream OutputStream = new MemoryStream();
+			using (ZipArchive OutputArchive = new ZipArchive(OutputStream, ZipArchiveMode.Create, true))
+			{
+				MemoryStream InputStream = new MemoryStream(Data);
+				using (ZipArchive InputArchive = new ZipArchive(InputStream, ZipArchiveMode.Read, true))
+				{
+					foreach (ZipArchiveEntry InputEntry in InputArchive.Entries)
+					{
+						ZipArchiveEntry OutputEntry = OutputArchive.CreateEntry(InputEntry.FullName);
+
+						using System.IO.Stream InputEntryStream = InputEntry.Open();
+						using System.IO.Stream OutputEntryStream = OutputEntry.Open();
+
+						if (InputEntry.FullName.Equals("appsettings.json", StringComparison.OrdinalIgnoreCase))
+						{
+							using MemoryStream MemoryStream = new MemoryStream();
+							InputEntryStream.CopyTo(MemoryStream);
+
+							Dictionary<string, Dictionary<string, object>> Document = JsonSerializer.Deserialize<Dictionary<string, Dictionary<string, object>>>(MemoryStream.ToArray());
+							foreach (KeyValuePair<string, object> Pair in Settings) {								
+								Document["Horde"][Pair.Key] = Pair.Value;
+							}
+
+							using Utf8JsonWriter Writer = new Utf8JsonWriter(OutputEntryStream, new JsonWriterOptions { Indented = true });
+							JsonSerializer.Serialize<Dictionary<string, Dictionary<string, object>>>(Writer, Document, new JsonSerializerOptions { WriteIndented = true });
+
+							bWrittenClientId = true;
+						}
+						else
+						{
+							InputEntryStream.CopyTo(OutputEntryStream);
+						}
+					}
+				}
+			}
+
+			if (!bWrittenClientId)
+			{
+				throw new InvalidDataException("Missing appsettings.json file from zip archive");
+			}
+
+			return OutputStream.ToArray();
+		}
+
+	}
+}
diff --git a/Engine/Source/Programs/Horde/HordeDashboard/src/backend/Backend.ts b/Engine/Source/Programs/Horde/HordeDashboard/src/backend/Backend.ts
--- a/Engine/Source/Programs/Horde/HordeDashboard/src/backend/Backend.ts
+++ b/Engine/Source/Programs/Horde/HordeDashboard/src/backend/Backend.ts
@@ -1004,8 +1004,8 @@ export class Backend {
     }
 
     downloadAgentZip() {
-        try {
-            const url = `${this.serverUrl}/api/v1/software/latest/download`;
+        try {            
+            const url = `${this.serverUrl}/api/v1/agentsoftware/default/zip`;            
             const link = document.createElement('a');
             link.href = url;
             document.body.appendChild(link);
diff --git a/Engine/Source/Programs/Horde/HordeServer/Collections/Impl/AgentSoftwareCollection.cs b/Engine/Source/Programs/Horde/HordeServer/Collections/Impl/AgentSoftwareCollection.cs
--- a/Engine/Source/Programs/Horde/HordeServer/Collections/Impl/AgentSoftwareCollection.cs
+++ b/Engine/Source/Programs/Horde/HordeServer/Collections/Impl/AgentSoftwareCollection.cs
@@ -34,19 +34,21 @@ namespace HordeServer.Collections.Impl
 			[BsonId]
 			public string Id { get; set; }
 
-			public byte[] Data { get; set; } = null!;
+			[BsonIgnoreIfNull]
+			public byte[]? Data { get; set; }
 
-            public List<string>? ChunkIds { get; set; } = null;
+			[BsonIgnoreIfNull]
+            public List<string>? ChunkIds { get; set; }
 
             [BsonConstructor]
 			private AgentSoftwareDocument()
 			{
 				this.Id = null!;
-				this.Data = null!;
-                this.ChunkIds = null;
+				this.Data = null;
+                this.ChunkIds = null;			
             }
 
-			public AgentSoftwareDocument(string Id, byte[] Data, List<string>? ChunkIds = null)
+			public AgentSoftwareDocument(string Id, byte[]? Data, List<string>? ChunkIds = null)
 			{
 				this.Id = Id;
 				this.Data = Data;
@@ -145,7 +147,7 @@ namespace HordeServer.Collections.Impl
             // insert the chunks
             await AgentChunks.InsertManyAsync(Chunks);
 
-            Document = new AgentSoftwareDocument(Version, Array.Empty<byte>(), Chunks.Select( Chunk => Chunk.Id).ToList());
+            Document = new AgentSoftwareDocument(Version, null, Chunks.Select( Chunk => Chunk.Id).ToList());
 			await Collection.InsertOneAsync(Document);
             return true;
         }
@@ -177,7 +179,7 @@ namespace HordeServer.Collections.Impl
                 return null;
             }
 
-			if (Software.Data.Length > 0)
+			if (Software.Data != null && Software.Data.Length > 0)
 			{
                 return Software.Data;
             }
diff --git a/Engine/Source/Programs/Horde/HordeServer/Program.cs b/Engine/Source/Programs/Horde/HordeServer/Program.cs
--- a/Engine/Source/Programs/Horde/HordeServer/Program.cs
+++ b/Engine/Source/Programs/Horde/HordeServer/Program.cs
@@ -16,9 +16,11 @@ using System.Text;
 using System.Text.Json;
 using System.Threading;
 using System.Threading.Channels;
+using System.Text.RegularExpressions;
 using System.Threading.Tasks;
 using EpicGames.Core;
 using EpicGames.Perforce;
+using HordeCommon;
 using HordeServer.Models;
 using HordeServer.Utilities;
 using Microsoft.AspNetCore.Hosting;
@@ -261,7 +263,7 @@ namespace HordeServer
 		/// Gets the certificate to use for Grpc endpoints
 		/// </summary>
 		/// <returns>Custom certificate to use for Grpc endpoints, or null for the default.</returns>
-		static X509Certificate2? ReadGrpcCertificate(ServerSettings HordeSettings)
+		public static X509Certificate2? ReadGrpcCertificate(ServerSettings HordeSettings)
 		{
 			if (HordeSettings.ServerPrivateCert == null)
 			{
@@ -269,7 +271,18 @@ namespace HordeServer
 			}
 			else
 			{
-				return new X509Certificate2(FileReference.ReadAllBytes(FileReference.Combine(AppDir, HordeSettings.ServerPrivateCert)));
+				FileReference? ServerPrivateCert = null;
+
+				if (!Path.IsPathRooted(HordeSettings.ServerPrivateCert))
+				{
+					ServerPrivateCert = FileReference.Combine(AppDir, HordeSettings.ServerPrivateCert);
+				}
+				else
+				{
+					ServerPrivateCert = new FileReference(HordeSettings.ServerPrivateCert);
+				}
+				
+				return new X509Certificate2(FileReference.ReadAllBytes(ServerPrivateCert));
 			}
 		}
 
@@ -299,26 +312,57 @@ namespace HordeServer
 			return DirectoryReference.Combine(GetAppDir(), "Data");
 		}
 
+		/// <summary>
+		/// Handles bootstrapping of defaults for local servers, which can't be generated during build/installation process (or are better handled here where they can be updated)
+		/// This stuff will change as we get settings into database and could be considered discovery for installer/dockerfile builds 
+		/// </summary>		
 		static void InitializeDefaults(ServerSettings Settings)
-		{
+		{			
 			if (Settings.SingleInstance)
 			{
 				FileReference GlobalConfig = FileReference.Combine(Program.DataDir, "Config/globals.json");
 
+				if (!FileReference.Exists(GlobalConfig))
+				{
+					DirectoryReference.CreateDirectory(GlobalConfig.Directory);
+					FileReference.WriteAllText(GlobalConfig, "{}");
+				}
+
+				FileReference PrivateCertFile = FileReference.Combine(Program.DataDir, "Agent/ServerToAgent.pfx");
+				string PrivateCertFileJsonPath = PrivateCertFile.ToString().Replace("\\", "/", StringComparison.Ordinal);
+
 				if (!FileReference.Exists(UserConfigFile))
 				{
+					// create new user configuration
 					DirectoryReference.CreateDirectory(UserConfigFile.Directory);
-					FileReference.WriteAllText(UserConfigFile, $"{{\"Horde\": {{ \"configPath\" : \"{GlobalConfig.ToString().Replace("\\", "/", StringComparison.Ordinal)}\"}}}}");
+					FileReference.WriteAllText(UserConfigFile, $"{{\"Horde\": {{ \"ConfigPath\" : \"{GlobalConfig.ToString().Replace("\\", "/", StringComparison.Ordinal)}\", \"ServerPrivateCert\" : \"{PrivateCertFileJsonPath}\", \"HttpPort\": 8080}}}}");
 				}
-				
-				if (!FileReference.Exists(GlobalConfig))
+
+				// make sure the cert exists
+				if (!FileReference.Exists(PrivateCertFile))
 				{
-					DirectoryReference.CreateDirectory(GlobalConfig.Directory);
-					FileReference.WriteAllText(GlobalConfig, "{}");
+					string DnsName = System.Net.Dns.GetHostName();
+					Serilog.Log.Logger.Information("Creating certificate for {DnsName}", DnsName);
+
+					byte[] PrivateCertData = AgentUtilities.CreateAgentCert(DnsName);
+					
+					Serilog.Log.Logger.Information("Writing private cert: {PrivateCert}", PrivateCertFile.FullName);
+
+					if (!DirectoryReference.Exists(PrivateCertFile.Directory))
+					{
+						DirectoryReference.CreateDirectory(PrivateCertFile.Directory);
+					}
+
+					FileReference.WriteAllBytes(PrivateCertFile, PrivateCertData);
 				}
 
-			}
+				// note: this isn't great, though we need it early in server startup, and this is only hit on first server boot where the grpc cert isn't generated/set 
+				if (Settings.ServerPrivateCert == null)
+				{
+					Settings.ServerPrivateCert = PrivateCertFile.ToString();
+				}				
 
+			}
 		}
 	}
 }
diff --git a/Engine/Source/Programs/Horde/HordeServer/Services/AgentSoftwareService.cs b/Engine/Source/Programs/Horde/HordeServer/Services/AgentSoftwareService.cs
--- a/Engine/Source/Programs/Horde/HordeServer/Services/AgentSoftwareService.cs
+++ b/Engine/Source/Programs/Horde/HordeServer/Services/AgentSoftwareService.cs
@@ -1,5 +1,6 @@
 // Copyright Epic Games, Inc. All Rights Reserved.
 
+using HordeCommon;
 using HordeServer.Collections;
 using HordeServer.Models;
 using HordeServer.Utilities;
@@ -19,6 +20,7 @@ using Microsoft.Extensions.Logging;
 using System.Runtime.InteropServices;
 using Microsoft.Extensions.Options;
 using System.IO.Compression;
+using System.Security.Cryptography.X509Certificates;
 
 namespace HordeServer.Services
 {
@@ -216,36 +218,6 @@ namespace HordeServer.Services
 			}
 		}
 
-		/// <summary>
-		/// Reads the version number from an archive
-		/// </summary>
-		/// <param name="Data">The archive data</param>
-		/// <returns></returns>
-		static string ReadVersion(byte[] Data)
-		{
-			MemoryStream InputStream = new MemoryStream(Data);
-			using (ZipArchive InputArchive = new ZipArchive(InputStream, ZipArchiveMode.Read, true))
-			{
-				foreach (ZipArchiveEntry InputEntry in InputArchive.Entries)
-				{
-					if (InputEntry.FullName.Equals("HordeAgent.dll", StringComparison.OrdinalIgnoreCase))
-					{
-						string TempFile = Path.GetTempFileName();
-						try
-						{
-							InputEntry.ExtractToFile(TempFile, true);
-							return FileVersionInfo.GetVersionInfo(TempFile).ProductVersion;
-						}
-						finally
-						{
-							File.Delete(TempFile);
-						}
-					}
-				}
-			}
-			throw new Exception("Unable to find HordeAgent.dll in archive");
-		}
-
 		/// <summary>
 		/// Updates a new software revision
 		/// </summary>
@@ -256,7 +228,7 @@ namespace HordeServer.Services
 		public async Task<string> SetArchiveAsync(AgentSoftwareChannelName Name, string? Author, byte[] Data)
 		{
 			// Upload the software
-			string Version = ReadVersion(Data);
+			string Version = AgentUtilities.ReadVersion(Data);
 			await Collection.AddAsync(Version, Data);
 
 			// Update the channel
@@ -306,10 +278,9 @@ namespace HordeServer.Services
 		{
 			return Collection.GetAsync(Version);
 		}
-
 		async Task RegisterDefaultAgent(int DelayMs)
 		{
-			await Task.Delay(DelayMs);			
+			await Task.Delay(DelayMs);
 
 			// Check whether we have an installed agent zip
 			FileReference AgentZip = FileReference.Combine(Program.AppDir, "DefaultAgent/Agent.zip");
@@ -323,10 +294,10 @@ namespace HordeServer.Services
 			string AgentHash = ContentHash.MD5(AgentZip).ToString();
 
 			FileReference AgentHashFile = FileReference.Combine(Program.DataDir, "Agent/DefaultAgentHash");
-			
-			try 
+
+			try
 			{
-				
+
 				if (FileReference.Exists(AgentHashFile) && FileReference.ReadAllText(AgentHashFile).Trim() == AgentHash)
 				{
 					Logger.LogInformation("Default agent software is up to date");
@@ -340,6 +311,27 @@ namespace HordeServer.Services
 
             byte[] Bytes = await File.ReadAllBytesAsync(AgentZip.ToString());
 
+			using X509Certificate2? GrpcCertificate = Program.ReadGrpcCertificate(Settings.CurrentValue);
+
+			if (GrpcCertificate == null)
+			{
+				throw new Exception("Unable to register default agent without valid grpc certicate");
+			}
+
+			// construct agent app settings
+			
+			string ProfileName = "Default";
+
+			string ServerUrl = $"https://{System.Net.Dns.GetHostName()}:{Settings.CurrentValue.HttpsPort}";			
+
+			Dictionary<string, object> DefaultProfile = new Dictionary<string, object>() { { "Name", ProfileName }, { "Environment", "prod" }, { "Thumbprint", GrpcCertificate.Thumbprint }, { "Url", ServerUrl } };
+
+			List<object> ServerProfiles = new List<object>() { DefaultProfile };
+
+			Dictionary<string, object> AgentSettings = new Dictionary<string, object>() { { "ServerProfiles", ServerProfiles }, { "Server", ProfileName }};
+
+			Bytes = AgentUtilities.UpdateAppSettings(Bytes, AgentSettings);
+
 			string Version = await SetArchiveAsync(new AgentSoftwareChannelName("default"), null, Bytes);
 
 			if (!DirectoryReference.Exists(AgentHashFile.Directory))
diff --git a/Engine/Source/Programs/Horde/HordeServer/Services/ConfigService.cs b/Engine/Source/Programs/Horde/HordeServer/Services/ConfigService.cs
--- a/Engine/Source/Programs/Horde/HordeServer/Services/ConfigService.cs
+++ b/Engine/Source/Programs/Horde/HordeServer/Services/ConfigService.cs
@@ -796,7 +796,7 @@ namespace HordeServer.Services
 			catch (Exception Ex)
 			{
 				Response.Errors.Add($"Exception while updating settings: {Ex.Message}");
-				Logger.LogError("{0}: {1}", Response.Errors.Last(), Ex.StackTrace ?? "");
+				Logger.LogError(Ex, "{Error}", Response.Errors.Last());
 			}
 
 			if (Response.Errors.Count != 0)
@@ -814,12 +814,12 @@ namespace HordeServer.Services
 				// This will trigger a setting update as the user config json is set to reload on change
 				try
 				{
-					await FileReference.WriteAllBytesAsync(Program.UserConfigFile, JsonSerializer.SerializeToUtf8Bytes(NewLocalSettings, new JsonSerializerOptions { WriteIndented = true }));
+					await FileReference.WriteAllBytesAsync(Program.UserConfigFile, JsonSerializer.SerializeToUtf8Bytes(NewLocalSettings, new JsonSerializerOptions { WriteIndented = true, IgnoreNullValues  = true }));
 				}
 				catch (Exception Ex)
 				{
 					Response.Errors.Add($"Unable to serialize json settings to {Program.UserConfigFile.ToString()}, {Ex.Message}");
-					Logger.LogError("Unable to serialize json settings to {0}, {1} : {2}", Program.UserConfigFile.ToString(), Ex.Message, Ex.StackTrace ?? "");
+					Logger.LogError(Ex, "Unable to serialize json settings to {ConfigFile}, {Message}", Program.UserConfigFile.ToString(), Ex.Message);
 				}
 
 				if (Response.Errors.Count == 0)
diff --git a/Engine/Source/Programs/Horde/HordeServer/Startup.cs b/Engine/Source/Programs/Horde/HordeServer/Startup.cs
--- a/Engine/Source/Programs/Horde/HordeServer/Startup.cs
+++ b/Engine/Source/Programs/Horde/HordeServer/Startup.cs
@@ -829,7 +829,9 @@ namespace HordeServer
 			IServerAddressesFeature Feature = App.ServerFeatures.Get<IServerAddressesFeature>();
 			if (Feature.Addresses.Count > 0)
 			{
-				Process.Start(new ProcessStartInfo { FileName = Feature.Addresses.First(), UseShellExecute = true });
+				// with a development cert, host will be set by default to localhost, otherwise there will be no host in address
+				string Address = Feature.Addresses.First().Replace("[::]", System.Net.Dns.GetHostName(), StringComparison.OrdinalIgnoreCase);
+				Process.Start(new ProcessStartInfo { FileName = Address, UseShellExecute = true });
 			}
 		}
 	}
