Horde: Add support for querying all objects underneath a given root (eg. /api/v1/objects/NS/HASH/tree). Objects are returned as a stream of compact binary object attachment fields (followed by an actual object) and binary attachment fields.

Optional request body specifies a list of objects that the client already has, and the server does not need to recurse through.

diff --git a/Engine/Source/Programs/Horde/HordeServer/Storage/Controllers/ObjectsController.cs b/Engine/Source/Programs/Horde/HordeServer/Storage/Controllers/ObjectsController.cs
--- a/Engine/Source/Programs/Horde/HordeServer/Storage/Controllers/ObjectsController.cs
+++ b/Engine/Source/Programs/Horde/HordeServer/Storage/Controllers/ObjectsController.cs
@@ -17,6 +17,7 @@ using HordeServer.Utilities;
 using EpicGames.Core;
 using HordeServer.Storage.Collections;
 using EpicGames.Serialization;
+using System.Net;
 
 namespace HordeServer.Storage.Controllers
 {
@@ -65,6 +66,17 @@ namespace HordeServer.Storage.Controllers
 		public List<IoHash> Need { get; set; } = new List<IoHash>();
 	}
 
+	/// <summary>
+	/// Request for retrieving a tree of objects
+	/// </summary>
+	public class GetObjectTreeRequest
+	{
+		/// <summary>
+		/// List of objects that the client already has
+		/// </summary>
+		public List<IoHash> Have { get; set; } = new List<IoHash>();
+	}
+
 	/// <summary>
 	/// Response for putting an object
 	/// </summary>
@@ -140,6 +152,80 @@ namespace HordeServer.Storage.Controllers
 			throw new NotImplementedException();
 		}
 
+		/// <summary>
+		/// Gets an object tree underneath a given root
+		/// </summary>
+		/// <param name="NamespaceId">Namespace for the operation</param>
+		/// <param name="Hash">The root object hash</param>
+		/// <param name="Request">Request object</param>
+		[HttpPost]
+		[Route("/api/v1/objects/{NamespaceId}/{Hash}/tree")]
+		public async Task GetObjectTreeAsync(NamespaceId NamespaceId, IoHash Hash, [FromBody] GetObjectTreeRequest Request)
+		{
+			if (!await NamespaceCollection.AuthorizeAsync(NamespaceId, User, AclAction.ReadBlobs))
+			{
+				Response.StatusCode = (int)HttpStatusCode.Forbidden;
+				return;
+			}
+
+			CbObject? Object = await ObjectCollection.GetAsync(NamespaceId, Hash);
+			if (Object == null)
+			{
+				Response.StatusCode = (int)HttpStatusCode.NotFound;
+				return;
+			}
+
+			Response.StatusCode = 200;
+			await Response.StartAsync(HttpContext.RequestAborted);
+
+			HashSet<IoHash> HaveHashes = new HashSet<IoHash>(Request.Have);
+			if (!HaveHashes.Contains(Hash))
+			{
+				List<IoHash> RequestHashes = new List<IoHash>();
+				for (; ; )
+				{
+					CbWriter Writer = new CbWriter();
+					Writer.WriteObjectAttachmentValue(Hash);
+					await Response.BodyWriter.WriteAsync(Writer.ToByteArray());
+
+					if (Object != null)
+					{
+						await Response.BodyWriter.WriteAsync(Object.GetView());
+
+						List<CbBinaryAttachment> BinaryAttachments = new List<CbBinaryAttachment>();
+						Object.IterateAttachments(Field =>
+						{
+							if (Field.IsBinaryAttachment())
+							{
+								BinaryAttachments.Add(Field.AsBinaryAttachment());
+							}
+							else if (Field.IsObjectAttachment() && HaveHashes.Add(Field.AsObjectAttachment()))
+							{
+								RequestHashes.Add(Field.AsObjectAttachment());
+							}
+						});
+
+						foreach (CbBinaryAttachment BinaryAttachment in BinaryAttachments)
+						{
+							CbWriter BinaryWriter = new CbWriter();
+							BinaryWriter.WriteBinaryAttachmentValue(BinaryAttachment);
+							await Response.BodyWriter.WriteAsync(BinaryWriter.ToByteArray());
+						}
+					}
+
+					if (RequestHashes.Count == 0)
+					{
+						break;
+					}
+
+					Hash = RequestHashes[RequestHashes.Count - 1];
+					RequestHashes.RemoveAt(RequestHashes.Count - 1);
+
+					Object = await ObjectCollection.GetAsync(NamespaceId, Hash);
+				}
+			}
+		}
+
 		/// <summary>
 		/// Does payload exist, HTTP 200 returned if successful, otherwise a HTTP 404 if missing..
 		/// </summary>
diff --git a/Engine/Source/Programs/Shared/EpicGames.Serialization.Tests/JupiterTests.cs b/Engine/Source/Programs/Shared/EpicGames.Serialization.Tests/JupiterTests.cs
--- a/Engine/Source/Programs/Shared/EpicGames.Serialization.Tests/JupiterTests.cs
+++ b/Engine/Source/Programs/Shared/EpicGames.Serialization.Tests/JupiterTests.cs
@@ -45,7 +45,7 @@ namespace EpicGames.Serialization.Tests
 			Writer2.EndObject();
 
 			CbObject Object2 = Writer2.ToObject();
-			Assert.AreEqual(Object2.GetSize(), 22);
+			Assert.AreEqual(21, Object2.GetSize());
 		}
 
 		[TestMethod]
