Horde: Implement the /api/v1/objects/NS/HASH/delta endpoint. Posting to this endpoint with a list of hashes will determine any object hashes under the given object which are missing.

diff --git a/Engine/Source/Programs/Horde/HordeServer/Storage/Controllers/ObjectsController.cs b/Engine/Source/Programs/Horde/HordeServer/Storage/Controllers/ObjectsController.cs
--- a/Engine/Source/Programs/Horde/HordeServer/Storage/Controllers/ObjectsController.cs
+++ b/Engine/Source/Programs/Horde/HordeServer/Storage/Controllers/ObjectsController.cs
@@ -48,11 +48,6 @@ namespace HordeServer.Storage.Controllers
 		/// List of objects that the client has
 		/// </summary>
 		public List<IoHash> Have { get; set; } = new List<IoHash>();
-
-		/// <summary>
-		/// List of objects that the client wants
-		/// </summary>
-		public List<IoHash> Want { get; set; } = new List<IoHash>();
 	}
 
 	/// <summary>
@@ -144,12 +139,62 @@ namespace HordeServer.Storage.Controllers
 		/// <param name="NamespaceId">Namespace for the operation</param>
 		/// <param name="Hash">The root object hash</param>
 		/// <param name="Request">Request object</param>
+		/// <param name="Depth">Maximum depth in the tree to scan</param>
+		/// <param name="Count">Maximum number of items to return</param>
 		/// <returns></returns>
 		[HttpPost]
 		[Route("/api/v1/objects/{NamespaceId}/{Hash}/delta")]
-		public Task<ActionResult<FindObjectDeltaResponse>> FindObjectDeltaAsync(NamespaceId NamespaceId, IoHash Hash, [FromBody] FindObjectDeltaRequest Request)
+		public async Task<ActionResult<FindObjectDeltaResponse>> FindObjectDeltaAsync(NamespaceId NamespaceId, IoHash Hash, [FromBody] FindObjectDeltaRequest Request, int Depth = -1, int Count = -1)
 		{
-			throw new NotImplementedException();
+			if (!await NamespaceCollection.AuthorizeAsync(NamespaceId, User, AclAction.ReadBlobs))
+			{
+				return Forbid();
+			}
+
+			CbObject? Object = await ObjectCollection.GetAsync(NamespaceId, Hash);
+			if (Object == null)
+			{
+				return NotFound();
+			}
+
+			FindObjectDeltaResponse Response = new FindObjectDeltaResponse();
+
+			HashSet<IoHash> SeenHashes = new HashSet<IoHash>(Request.Have);
+			if (!SeenHashes.Contains(Hash))
+			{
+				SeenHashes.Add(IoHash.Zero); // Don't return empty fields
+
+				Queue<(IoHash, int)> Queue = new Queue<(IoHash, int)>();
+				Queue.Enqueue((Hash, Depth));
+
+				while (Queue.Count > 0 && (Count == -1 || Response.Need.Count < Count))
+				{
+					(IoHash CurrentHash, int CurrentDepth) = Queue.Dequeue();
+					Response.Need.Add(CurrentHash);
+
+					if (CurrentDepth != 0)
+					{
+						CbObject? CurrentObject = await ObjectCollection.GetAsync(NamespaceId, CurrentHash);
+						if (CurrentObject != null)
+						{
+							int NextDepth = (CurrentDepth == -1) ? -1 : (CurrentDepth - 1);
+							CurrentObject.IterateAttachments(Field =>
+							{
+								if (Field.IsObjectAttachment())
+								{
+									CbObjectAttachment Attachment = Field.AsObjectAttachment();
+									if (SeenHashes.Add(Field.AsObjectAttachment()))
+									{
+										Queue.Enqueue((Field.AsObjectAttachment().Hash, NextDepth));
+									}
+								}
+							});
+						}
+					}
+				}
+			}
+
+			return Response;
 		}
 
 		/// <summary>
