Horde: Temp change output level to debug.

diff --git a/Engine/Source/Programs/Horde/HordeServer/appsettings.json b/Engine/Source/Programs/Horde/HordeServer/appsettings.json
--- a/Engine/Source/Programs/Horde/HordeServer/appsettings.json
+++ b/Engine/Source/Programs/Horde/HordeServer/appsettings.json
@@ -11,7 +11,7 @@
 	{
 		"MinimumLevel": 
 		{
-			"Default": "Information",
+			"Default": "Debug",
 			"Override": 
 			{
 				"MongoDB": "Warning", // For bundled MongoDB output
