Horde: Change default output level back to "information".

diff --git a/Engine/Source/Programs/Horde/HordeServer/appsettings.Development.json b/Engine/Source/Programs/Horde/HordeServer/appsettings.Development.json
--- a/Engine/Source/Programs/Horde/HordeServer/appsettings.Development.json
+++ b/Engine/Source/Programs/Horde/HordeServer/appsettings.Development.json
@@ -2,5 +2,12 @@
 	"Horde":
 	{
 		"DisableSchedules": true
+	},
+	"Serilog": 
+	{
+		"MinimumLevel": 
+		{
+			"Default": "Debug"
+		}
 	}
 }
diff --git a/Engine/Source/Programs/Horde/HordeServer/appsettings.json b/Engine/Source/Programs/Horde/HordeServer/appsettings.json
--- a/Engine/Source/Programs/Horde/HordeServer/appsettings.json
+++ b/Engine/Source/Programs/Horde/HordeServer/appsettings.json
@@ -11,7 +11,7 @@
 	{
 		"MinimumLevel": 
 		{
-			"Default": "Debug",
+			"Default": "Information",
 			"Override": 
 			{
 				"MongoDB": "Warning", // For bundled MongoDB output
