Horde: Add compact binary request/response support to objects controller.

diff --git a/Engine/Source/Programs/Horde/HordeServer/Storage/Controllers/ObjectsController.cs b/Engine/Source/Programs/Horde/HordeServer/Storage/Controllers/ObjectsController.cs
--- a/Engine/Source/Programs/Horde/HordeServer/Storage/Controllers/ObjectsController.cs
+++ b/Engine/Source/Programs/Horde/HordeServer/Storage/Controllers/ObjectsController.cs
@@ -47,7 +47,8 @@ namespace HordeServer.Storage.Controllers
 		/// <summary>
 		/// List of objects that the client has
 		/// </summary>
-		public List<IoHash> Have { get; set; } = new List<IoHash>();
+		[CbField("h")]
+		public List<CbObjectAttachment> Have { get; set; } = new List<CbObjectAttachment>();
 	}
 
 	/// <summary>
@@ -58,7 +59,8 @@ namespace HordeServer.Storage.Controllers
 		/// <summary>
 		/// List of of objects that the client needs
 		/// </summary>
-		public List<IoHash> Need { get; set; } = new List<IoHash>();
+		[CbField("n")]
+		public List<CbObjectAttachment> Need { get; set; } = new List<CbObjectAttachment>();
 	}
 
 	/// <summary>
@@ -69,7 +71,8 @@ namespace HordeServer.Storage.Controllers
 		/// <summary>
 		/// List of objects that the client already has
 		/// </summary>
-		public List<IoHash> Have { get; set; } = new List<IoHash>();
+		[CbField("h")]
+		public List<CbObjectAttachment> Have { get; set; } = new List<CbObjectAttachment>();
 	}
 
 	/// <summary>
@@ -80,7 +83,8 @@ namespace HordeServer.Storage.Controllers
 		/// <summary>
 		/// List of missing hashes
 		/// </summary>
-		public IoHash Id { get; set; }
+		[CbField("id")]
+		public CbObjectAttachment Id { get; set; }
 	}
 
 	/// <summary>
@@ -159,7 +163,7 @@ namespace HordeServer.Storage.Controllers
 
 			FindObjectDeltaResponse Response = new FindObjectDeltaResponse();
 
-			HashSet<IoHash> SeenHashes = new HashSet<IoHash>(Request.Have);
+			HashSet<IoHash> SeenHashes = new HashSet<IoHash>(Request.Have.Select(x => x.Hash));
 			if (!SeenHashes.Contains(Hash))
 			{
 				SeenHashes.Add(IoHash.Zero); // Don't return empty fields
@@ -223,7 +227,7 @@ namespace HordeServer.Storage.Controllers
 			Response.StatusCode = 200;
 			await Response.StartAsync(HttpContext.RequestAborted);
 
-			HashSet<IoHash> HaveHashes = new HashSet<IoHash>(Request.Have);
+			HashSet<IoHash> HaveHashes = new HashSet<IoHash>(Request.Have.Select(x => x.Hash));
 			if (!HaveHashes.Contains(Hash))
 			{
 				List<IoHash> RequestHashes = new List<IoHash>();
diff --git a/Engine/Source/Programs/Shared/EpicGames.Serialization/CompactBinary.cs b/Engine/Source/Programs/Shared/EpicGames.Serialization/CompactBinary.cs
--- a/Engine/Source/Programs/Shared/EpicGames.Serialization/CompactBinary.cs
+++ b/Engine/Source/Programs/Shared/EpicGames.Serialization/CompactBinary.cs
@@ -11,6 +11,7 @@ using System.Linq;
 using System.Numerics;
 using System.Text;
 using System.Text.Json;
+using System.Text.Json.Serialization;
 
 namespace EpicGames.Serialization
 {
@@ -192,6 +193,7 @@ namespace EpicGames.Serialization
 	/// A binary attachment, referenced by <see cref="IoHash"/>
 	/// </summary>
 	[DebuggerDisplay("{Hash}")]
+	[JsonConverter(typeof(CbBinaryAttachmentJsonConverter))]
 	public struct CbBinaryAttachment
 	{
 		/// <summary>
@@ -232,10 +234,23 @@ namespace EpicGames.Serialization
 		}
 	}
 
+	/// <summary>
+	/// Type converter for IoHash to and from JSON
+	/// </summary>
+	sealed class CbBinaryAttachmentJsonConverter : JsonConverter<CbBinaryAttachment>
+	{
+		/// <inheritdoc/>
+		public override CbBinaryAttachment Read(ref Utf8JsonReader Reader, Type TypeToConvert, JsonSerializerOptions Options) => IoHash.Parse(Reader.ValueSpan);
+
+		/// <inheritdoc/>
+		public override void Write(Utf8JsonWriter Writer, CbBinaryAttachment Value, JsonSerializerOptions Options) => Writer.WriteStringValue(Value.Hash.ToUtf8String().Span);
+	}
+
 	/// <summary>
 	/// An object attachment, referenced by <see cref="IoHash"/>
 	/// </summary>
 	[DebuggerDisplay("{Hash}")]
+	[JsonConverter(typeof(CbObjectAttachmentJsonConverter))]
 	public struct CbObjectAttachment
 	{
 		/// <summary>
@@ -276,6 +291,18 @@ namespace EpicGames.Serialization
 		}
 	}
 
+	/// <summary>
+	/// Type converter for IoHash to and from JSON
+	/// </summary>
+	sealed class CbObjectAttachmentJsonConverter : JsonConverter<CbObjectAttachment>
+	{
+		/// <inheritdoc/>
+		public override CbObjectAttachment Read(ref Utf8JsonReader Reader, Type TypeToConvert, JsonSerializerOptions Options) => IoHash.Parse(Reader.ValueSpan);
+
+		/// <inheritdoc/>
+		public override void Write(Utf8JsonWriter Writer, CbObjectAttachment Value, JsonSerializerOptions Options) => Writer.WriteStringValue(Value.Hash.ToUtf8String().Span);
+	}
+
 	/// <summary>
 	/// Methods that operate on <see cref="CbFieldType"/>.
 	/// </summary>
