Horde: Output an error if trying to add zero requests.

diff --git a/Engine/Source/Programs/Horde/HordeServer/Compute/Impl/ComputeController.cs b/Engine/Source/Programs/Horde/HordeServer/Compute/Impl/ComputeController.cs
--- a/Engine/Source/Programs/Horde/HordeServer/Compute/Impl/ComputeController.cs
+++ b/Engine/Source/Programs/Horde/HordeServer/Compute/Impl/ComputeController.cs
@@ -48,9 +48,15 @@ namespace HordeServer.Compute.Impl
 		[HttpPost]
 		[Authorize]
 		[Route("/api/v1/compute/{ChannelId}")]
-		public async Task AddTasksAsync([FromRoute] ChannelId ChannelId, [FromBody] AddTasksRequest Request)
+		public async Task<ActionResult> AddTasksAsync([FromRoute] ChannelId ChannelId, [FromBody] AddTasksRequest Request)
 		{
+			if (Request.TaskHashes.Count == 0)
+			{
+				return BadRequest("No task hashes specified");
+			}
+
 			await ComputeService.AddTasksAsync(Impl.ComputeService.DefaultNamespaceId, Request.RequirementsHash, Request.TaskHashes, ChannelId);
+			return Ok();
 		}
 
 		/// <summary>
