Horde: Add logging for start and end of each compute lease.

diff --git a/Engine/Source/Programs/Horde/HordeServer/Compute/Impl/ComputeService.cs b/Engine/Source/Programs/Horde/HordeServer/Compute/Impl/ComputeService.cs
--- a/Engine/Source/Programs/Horde/HordeServer/Compute/Impl/ComputeService.cs
+++ b/Engine/Source/Programs/Horde/HordeServer/Compute/Impl/ComputeService.cs
@@ -203,6 +203,7 @@ namespace HordeServer.Compute.Impl
 		public Task OnLeaseStartedAsync(IAgent Agent, ObjectId LeaseId, Any Payload)
 		{
 			ComputeTaskMessage ComputeTask = Payload.Unpack<ComputeTaskMessage>();
+			Logger.LogInformation("Compute lease started (lease: {LeaseId}, task: {TaskHash}, agent: {AgentId}, channel: {ChannelId})", LeaseId, ComputeTask.TaskHash, Agent.Id, ComputeTask.ChannelId);
 
 			ComputeTaskStatus Status = new ComputeTaskStatus(ComputeTask.TaskHash, ComputeTaskState.Executing);
 			return MessageQueue.PostAsync(ComputeTask.ChannelId, Status);
@@ -213,7 +214,6 @@ namespace HordeServer.Compute.Impl
 			ComputeTaskMessage ComputeTask = Payload.Unpack<ComputeTaskMessage>();
 
 			ComputeTaskResultMessage Result = ComputeTaskResultMessage.Parser.ParseFrom(Output.ToArray());
-			Logger.LogInformation("Task complete: {Hash} -> {OutputHash}", (IoHash)(CbObjectAttachment)ComputeTask.TaskHash, (Result.OutputHash == null)? IoHash.Zero : (IoHash)(CbObjectAttachment)Result.OutputHash);
 
 			ComputeTaskStatus Status = new ComputeTaskStatus(ComputeTask.TaskHash, ComputeTaskState.Complete);
 			Status.LeaseId = LeaseId;
@@ -222,6 +222,7 @@ namespace HordeServer.Compute.Impl
 				Status.ResultHash = Result.OutputHash;
 			}
 
+			Logger.LogInformation("Compute lease finished (lease: {LeaseId}, task: {TaskHash}, agent: {AgentId}, channel: {ChannelId}, result: {ResultHash})", LeaseId, ComputeTask.TaskHash, Agent.Id, ComputeTask.ChannelId, Status.ResultHash?.Hash ?? IoHash.Zero);
 			return MessageQueue.PostAsync(ComputeTask.ChannelId, Status);
 		}
 	}
