Horde: Additional logging for compute task handling.

diff --git a/Engine/Source/Programs/Horde/HordeServer/Compute/Impl/ComputeService.cs b/Engine/Source/Programs/Horde/HordeServer/Compute/Impl/ComputeService.cs
--- a/Engine/Source/Programs/Horde/HordeServer/Compute/Impl/ComputeService.cs
+++ b/Engine/Source/Programs/Horde/HordeServer/Compute/Impl/ComputeService.cs
@@ -155,6 +155,7 @@ namespace HordeServer.Compute.Impl
 			foreach (CbObjectAttachment TaskHash in TaskHashes)
 			{
 				ComputeTaskInfo TaskInfo = new ComputeTaskInfo(TaskHash, ChannelId);
+				Logger.LogDebug("Adding task {TaskHash} to queue {RequirementsHash}", TaskInfo.TaskHash.Hash, RequirementsHash);
 				Tasks.Add(TaskScheduler.EnqueueAsync(TaskInfo, RequirementsHash));
 			}
 			await Task.WhenAll(Tasks);
@@ -174,7 +175,10 @@ namespace HordeServer.Compute.Impl
 
 				string LeaseName = $"Remote action ({Entry.Item.TaskHash})";
 				byte[] Payload = Any.Pack(ComputeTask).ToByteArray();
-				return new AgentLease(ObjectId.GenerateNewId(), LeaseName, null, null, null, LeaseState.Pending, Payload, new AgentRequirements(), null);
+
+				AgentLease Lease = new AgentLease(ObjectId.GenerateNewId(), LeaseName, null, null, null, LeaseState.Pending, Payload, new AgentRequirements(), null);
+				Logger.LogDebug("Created lease {LeaseId} for channel {ChannelId} task {TaskHash}", Lease.Id, ComputeTask.ChannelId, (CbObjectAttachment)ComputeTask.TaskHash);
+				return Lease;
 			}
 			return null;
 		}
