Horde: Re-enable scheduled jobs in Dev.

diff --git a/Engine/Source/Programs/Horde/HordeServer/appsettings.Development.json b/Engine/Source/Programs/Horde/HordeServer/appsettings.Development.json
--- a/Engine/Source/Programs/Horde/HordeServer/appsettings.Development.json
+++ b/Engine/Source/Programs/Horde/HordeServer/appsettings.Development.json
@@ -1,7 +1,6 @@
 {
 	"Horde":
 	{
-		"DisableSchedules": true
 	},
 	"Serilog": 
 	{
