Horde: Fix logging during update requests being masked out. RequestPath property is propagated within the scope of a request, causing functionality to ignore standard request logging to mask all nested log events out.

diff --git a/Engine/Source/Programs/Horde/HordeServer/Program.cs b/Engine/Source/Programs/Horde/HordeServer/Program.cs
--- a/Engine/Source/Programs/Horde/HordeServer/Program.cs
+++ b/Engine/Source/Programs/Horde/HordeServer/Program.cs
@@ -74,48 +74,8 @@ namespace HordeServer
 			{
 				Configuration = Configuration.Enrich.With<DatadogLogEnricher>();
 			}
-			if (!Settings.LogSessionRequests)
-			{
-				Configuration = Configuration.Filter.ByExcluding(IsSessionLogEvent);
-			}
 			return Configuration;
 		}
-
-		static bool IsSessionLogEvent(LogEvent Event)
-		{
-			if (Event.Level <= LogEventLevel.Information)
-			{
-				LogEventPropertyValue? PropertyValue;
-				if (Event.Properties.TryGetValue("RequestPath", out PropertyValue))
-				{
-					ScalarValue? ScalarPropertyValue = PropertyValue as ScalarValue;
-					if (ScalarPropertyValue != null)
-					{
-						string? RequestPath = ScalarPropertyValue.Value as string;
-						if (RequestPath != null)
-						{
-							if (RequestPath.Equals("/Horde.HordeRpc/QueryServerStateV2", StringComparison.OrdinalIgnoreCase))
-							{
-								return true;
-							}
-							if (RequestPath.Equals("/Horde.HordeRpc/UpdateSession", StringComparison.OrdinalIgnoreCase))
-							{
-								return true;
-							}
-							if (RequestPath.Equals("/Horde.HordeRpc/CreateEvents", StringComparison.OrdinalIgnoreCase))
-							{
-								return true;
-							}
-							if (RequestPath.Equals("/Horde.HordeRpc/WriteOutput", StringComparison.OrdinalIgnoreCase))
-							{
-								return true;
-							}
-						}
-					}
-				}
-			}
-			return false;
-		}
 	}
 
 	class DatadogLogEnricher : ILogEventEnricher
diff --git a/Engine/Source/Programs/Horde/HordeServer/Startup.cs b/Engine/Source/Programs/Horde/HordeServer/Startup.cs
--- a/Engine/Source/Programs/Horde/HordeServer/Startup.cs
+++ b/Engine/Source/Programs/Horde/HordeServer/Startup.cs
@@ -82,6 +82,7 @@ using EpicGames.Horde.Compute;
 using HordeServer.Compute.Impl;
 using HordeServer.Compute;
 using System.Net.Http.Headers;
+using Serilog.Events;
 
 namespace HordeServer
 {
@@ -750,7 +751,7 @@ namespace HordeServer
 			App.UseSwagger();
 
 			// Enable serilog request logging
-			App.UseSerilogRequestLogging();
+			App.UseSerilogRequestLogging(Options => Options.GetLevel = GetRequestLoggingLevel);
 
 			// Include the source IP address with requests
 			App.Use(async (Context, Next) => {
@@ -836,5 +837,30 @@ namespace HordeServer
 				Process.Start(new ProcessStartInfo { FileName = Address, UseShellExecute = true });
 			}
 		}
+
+		static LogEventLevel GetRequestLoggingLevel(HttpContext Context, double ElapsedMs, Exception Ex)
+		{
+			if (Context.Request != null && Context.Request.Path.HasValue)
+			{
+				string RequestPath = Context.Request.Path;
+				if (RequestPath.Equals("/Horde.HordeRpc/QueryServerStateV2", StringComparison.OrdinalIgnoreCase))
+				{
+					return LogEventLevel.Verbose;
+				}
+				if (RequestPath.Equals("/Horde.HordeRpc/UpdateSession", StringComparison.OrdinalIgnoreCase))
+				{
+					return LogEventLevel.Verbose;
+				}
+				if (RequestPath.Equals("/Horde.HordeRpc/CreateEvents", StringComparison.OrdinalIgnoreCase))
+				{
+					return LogEventLevel.Verbose;
+				}
+				if (RequestPath.Equals("/Horde.HordeRpc/WriteOutput", StringComparison.OrdinalIgnoreCase))
+				{
+					return LogEventLevel.Verbose;
+				}
+			}
+			return LogEventLevel.Information;
+		}
 	}
 }
