Horde: Fix unimplemented method in stub causing test failure.

diff --git a/Engine/Source/Programs/Horde/HordeServerTests/Stubs/Services/PerforceServiceStub.cs b/Engine/Source/Programs/Horde/HordeServerTests/Stubs/Services/PerforceServiceStub.cs
--- a/Engine/Source/Programs/Horde/HordeServerTests/Stubs/Services/PerforceServiceStub.cs
+++ b/Engine/Source/Programs/Horde/HordeServerTests/Stubs/Services/PerforceServiceStub.cs
@@ -173,7 +173,7 @@ namespace HordeServerTests.Stubs.Services
 
 		public Task<ChangeDetails> GetChangeDetailsAsync(string ClusterName, string StreamName, int ChangeNumber)
 		{
-			throw new NotImplementedException();
+			return Task.FromResult(Changes[StreamName][ChangeNumber]);
 		}
 	}
 }
