Horde: Prevent compute tasks being assigned to an agent that is disabled.

diff --git a/Engine/Source/Programs/Horde/HordeServer/Compute/Impl/ComputeService.cs b/Engine/Source/Programs/Horde/HordeServer/Compute/Impl/ComputeService.cs
--- a/Engine/Source/Programs/Horde/HordeServer/Compute/Impl/ComputeService.cs
+++ b/Engine/Source/Programs/Horde/HordeServer/Compute/Impl/ComputeService.cs
@@ -164,6 +164,12 @@ namespace HordeServer.Compute.Impl
 		/// <inheritdoc/>
 		public async Task<AgentLease?> TryAssignLeaseAsync(IAgent Agent, CancellationToken CancellationToken)
 		{
+			if (!Agent.Enabled)
+			{
+				await Task.Delay(-1, CancellationToken);
+				return null;
+			}
+
 			TaskSchedulerEntry<ComputeTaskInfo>? Entry = await TaskScheduler.DequeueAsync(Agent, CancellationToken);
 			if(Entry != null)
 			{
