Horde: Do not send remote execution jobs to REMOTE-EXEC instances. Unable to get any diagnostics from them (not in datadog, cannot RDP).

diff --git a/Engine/Source/Programs/Horde/HordeServer/Compute/Impl/ComputeService.cs b/Engine/Source/Programs/Horde/HordeServer/Compute/Impl/ComputeService.cs
--- a/Engine/Source/Programs/Horde/HordeServer/Compute/Impl/ComputeService.cs
+++ b/Engine/Source/Programs/Horde/HordeServer/Compute/Impl/ComputeService.cs
@@ -164,7 +164,7 @@ namespace HordeServer.Compute.Impl
 		/// <inheritdoc/>
 		public async Task<AgentLease?> TryAssignLeaseAsync(IAgent Agent, CancellationToken CancellationToken)
 		{
-			if (!Agent.Enabled)
+			if (!Agent.Enabled || Agent.Id.ToString().StartsWith("REMOTE-EXEC-", StringComparison.OrdinalIgnoreCase))
 			{
 				await Task.Delay(-1, CancellationToken);
 				return null;
