Horde: Include the compact binary object in an exception if it fails to parse.

diff --git a/Engine/Source/Programs/Horde/HordeServer/Utilities/CompactBinaryFormatter.cs b/Engine/Source/Programs/Horde/HordeServer/Utilities/CompactBinaryFormatter.cs
--- a/Engine/Source/Programs/Horde/HordeServer/Utilities/CompactBinaryFormatter.cs
+++ b/Engine/Source/Programs/Horde/HordeServer/Utilities/CompactBinaryFormatter.cs
@@ -11,6 +11,7 @@ using EpicGames.Serialization;
 using EpicGames.Serialization.Converters;
 using System.Reflection;
 using System.Net.Mime;
+using EpicGames.Core;
 
 namespace HordeServer.Utilities
 {
@@ -56,8 +57,32 @@ namespace HordeServer.Utilities
 			}
 
 			// Serialize the object
+			CbField Field;
+			try
+			{
+				Field = new CbField(Data);
+			}
+			catch (Exception Ex)
+			{
+				throw new Exception($"Unable to parse compact binary data: {FormatHexDump(Data, 256)}", Ex);
+			}
 			return InputFormatterResult.SuccessAsync(CbSerializer.Deserialize(new CbField(Data), Context.ModelType));
 		}
+
+		static string FormatHexDump(byte[] Data, int MaxLength)
+		{
+			ReadOnlySpan<byte> Span = Data.AsSpan(0, Math.Min(Data.Length, MaxLength));
+			string HexString = StringUtils.FormatHexString(Span);
+
+			char[] HexDump = new char[Span.Length * 3];
+			for (int Idx = 0; Idx < Span.Length; Idx++)
+			{
+				HexDump[(Idx * 3) + 0] = ((Idx & 15) == 0) ? '\n' : ' ';
+				HexDump[(Idx * 3) + 1] = HexString[(Idx * 2) + 0];
+				HexDump[(Idx * 3) + 2] = HexString[(Idx * 2) + 1];
+			}
+			return new string(HexDump);
+		}
 	}
 
 	/// <summary>
