Horde: Output hashes for compute requests and responses.

diff --git a/Engine/Source/Programs/Horde/HordeAgent/Commands/Execution/ExecuteCommandV2.cs b/Engine/Source/Programs/Horde/HordeAgent/Commands/Execution/ExecuteCommandV2.cs
--- a/Engine/Source/Programs/Horde/HordeAgent/Commands/Execution/ExecuteCommandV2.cs
+++ b/Engine/Source/Programs/Horde/HordeAgent/Commands/Execution/ExecuteCommandV2.cs
@@ -200,6 +200,9 @@ namespace HordeAgent.Commands
 		{
 			using GrpcChannel Channel = GrpcService.CreateGrpcChannel();
 
+			Logger.LogInformation("task: {TaskHash}", TaskHash);
+			Logger.LogInformation("requirements: {RequirementsHash}", RequirementsHash);
+
 			const string NamespaceId = "default";
 
 			using GrpcChannel CasChannel = Channel;
@@ -250,6 +253,8 @@ namespace HordeAgent.Commands
 
 		async Task HandleCompleteTask(BlobStore.BlobStoreClient Client, string NamespaceId, CbObjectAttachment Hash, ILogger Logger)
 		{
+			Logger.LogInformation("result: {ResultHash}", Hash);
+
 			ComputeTaskResult Result = await Client.GetObjectAsync<ComputeTaskResult>(NamespaceId, Hash);
 			Logger.LogInformation("exit: {ExitCode}", Result.ExitCode);
 
