Horde: Use an explicit port when connecting to a Redis server on localhost.

diff --git a/Engine/Source/Programs/Horde/HordeServer/Services/RedisService.cs b/Engine/Source/Programs/Horde/HordeServer/Services/RedisService.cs
--- a/Engine/Source/Programs/Horde/HordeServer/Services/RedisService.cs
+++ b/Engine/Source/Programs/Horde/HordeServer/Services/RedisService.cs
@@ -93,11 +93,11 @@ namespace HordeServer.Services
 			{
 				if (IsRunningOnDefaultPort())
 				{
-					ConnectionString = "localhost";
+					ConnectionString = $"localhost:{RedisPort}";
 				}
 				else if (TryStartRedisServer())
 				{
-					ConnectionString = "localhost";
+					ConnectionString = $"localhost:{RedisPort}";
 				}
 				else
 				{
