Horde: Add support for different client protocols in ExecuteCommandV2. Passing the -Http argument now uses compact-binary encoded requests over HTTP 1 endpoints.

diff --git a/Engine/Source/Programs/Horde/HordeAgent/Commands/Execution/ExecuteCommandV2.cs b/Engine/Source/Programs/Horde/HordeAgent/Commands/Execution/ExecuteCommandV2.cs
--- a/Engine/Source/Programs/Horde/HordeAgent/Commands/Execution/ExecuteCommandV2.cs
+++ b/Engine/Source/Programs/Horde/HordeAgent/Commands/Execution/ExecuteCommandV2.cs
@@ -23,6 +23,7 @@ using Serilog.Events;
 using EpicGames.Serialization;
 using EpicGames.Horde.Compute;
 using EpicGames.Horde.Storage;
+using System.Net.Http;
 
 namespace HordeAgent.Commands
 {
@@ -91,6 +92,200 @@ namespace HordeAgent.Commands
 		[CommandLine("-LogLevel")]
 		public string LogLevelStr = "debug";
 
+		/// <summary>
+		/// Use the HTTP server endpoints
+		/// </summary>
+		[CommandLine("-Http")]
+		public bool UseHttp;
+
+		/// <summary>
+		/// The server url
+		/// </summary>
+		[CommandLine("-Server=")]
+		public string Server = Environment.GetEnvironmentVariable("HORDE_SERVER") ?? "https://localhost:5001";
+
+		/// <summary>
+		/// The token to use to connect with
+		/// </summary>
+		[CommandLine("-Token=")]
+		public string? Token = Environment.GetEnvironmentVariable("HORDE_TOKEN");
+
+		interface IComputeClient : IDisposable
+		{
+			Task AddBlobAsync(string NamespaceId, IoHash Hash, byte[] Data);
+			Task<byte[]> GetBlobAsync(string NamespaceId, IoHash Hash);
+
+			Task AddObjectAsync(string NamespaceId, IoHash Hash, byte[] Data);
+			Task<byte[]> GetObjectAsync(string NamespaceId, IoHash Hash);
+
+			Task AddTaskAsync(string ChannelId, string NamespaceId, IoHash RequirementsHash, IoHash TaskHash, bool SkipTaskLookup);
+			Task GetTaskUpdatesAsync(string ChannelId, string NamespaceId, Func<GetTaskUpdateResponse, Task<bool>> OnUpdate);
+		}
+
+		class HttpComputeClient : IComputeClient
+		{
+			HttpClient HttpClient;
+
+			public HttpComputeClient(Uri Server, string? BearerToken)
+			{
+				HttpClient = new HttpClient();
+				HttpClient.BaseAddress = Server;
+				if (BearerToken != null)
+				{
+					HttpClient.DefaultRequestHeaders.Authorization = new AuthenticationHeaderValue("Bearer", BearerToken);
+				}
+			}
+
+			public void Dispose()
+			{
+				HttpClient.Dispose();
+			}
+
+			public async Task AddBlobAsync(string NamespaceId, IoHash Hash, byte[] Data)
+			{
+				HttpRequestMessage Request = new HttpRequestMessage(HttpMethod.Put, $"api/v1/blobs/{NamespaceId}/{Hash}");
+				Request.Content = new ByteArrayContent(Data);
+
+				HttpResponseMessage Response = await HttpClient.SendAsync(Request);
+				Response.EnsureSuccessStatusCode();
+			}
+
+			public async Task<byte[]> GetBlobAsync(string NamespaceId, IoHash Hash)
+			{
+				HttpRequestMessage Request = new HttpRequestMessage(HttpMethod.Get, $"api/v1/blobs/{NamespaceId}/{Hash}");
+				HttpResponseMessage Response = await HttpClient.SendAsync(Request);
+				Response.EnsureSuccessStatusCode();
+				return await Response.Content.ReadAsByteArrayAsync();
+			}
+
+			public Task AddObjectAsync(string NamespaceId, IoHash Hash, byte[] Data)
+			{
+				return AddBlobAsync(NamespaceId, Hash, Data);
+			}
+
+			public Task<byte[]> GetObjectAsync(string NamespaceId, IoHash Hash)
+			{
+				return GetBlobAsync(NamespaceId, Hash);
+			}
+
+			public async Task AddTaskAsync(string ChannelId, string NamespaceId, IoHash RequirementsHash, IoHash TaskHash, bool SkipCacheLookup)
+			{
+				AddTasksRequest AddTasks = new AddTasksRequest();
+				AddTasks.RequirementsHash = RequirementsHash;
+				AddTasks.TaskHashes.Add(TaskHash);
+				AddTasks.DoNotCache = SkipCacheLookup;
+
+				ReadOnlyMemoryContent Content = new ReadOnlyMemoryContent(CbSerializer.Serialize(AddTasks).GetView());
+				Content.Headers.Add("Content-Type", "application/x-ue-cb");
+
+				HttpRequestMessage Request = new HttpRequestMessage(HttpMethod.Post, $"api/v1/compute/{ChannelId}");
+				Request.Content = Content;
+
+				HttpResponseMessage Response = await HttpClient.SendAsync(Request);
+				Response.EnsureSuccessStatusCode();
+			}
+
+			public async Task GetTaskUpdatesAsync(string ChannelId, string NamespaceId, Func<GetTaskUpdateResponse, Task<bool>> OnUpdateAsync)
+			{
+				for (; ; )
+				{
+					HttpRequestMessage Request = new HttpRequestMessage(HttpMethod.Post, $"api/v1/compute/{ChannelId}/updates?wait=10");
+					Request.Headers.Add("Accept", "application/x-ue-cb");
+
+					HttpResponseMessage Response = await HttpClient.SendAsync(Request);
+					Response.EnsureSuccessStatusCode();
+
+					byte[] Data = await Response.Content.ReadAsByteArrayAsync();
+					GetTaskUpdatesResponse ParsedResponse = CbSerializer.Deserialize<GetTaskUpdatesResponse>(new CbField(Data));
+					
+					foreach (GetTaskUpdateResponse Update in ParsedResponse.Updates)
+					{
+						if (!await OnUpdateAsync(Update))
+						{
+							return;
+						}
+					}
+				}
+			}
+		}
+
+		class GrpcComputeClient : IComputeClient
+		{
+			GrpcChannel Channel;
+			BlobStore.BlobStoreClient BlobStore;
+			ComputeRpc.ComputeRpcClient ComputeClient;
+
+			public GrpcComputeClient(GrpcService GrpcService)
+			{
+				Channel = GrpcService.CreateGrpcChannel();
+				BlobStore = new BlobStore.BlobStoreClient(Channel);
+				ComputeClient = new ComputeRpc.ComputeRpcClient(Channel);
+			}
+
+			public void Dispose()
+			{
+				Channel.Dispose();
+			}
+
+			public async Task AddBlobAsync(string NamespaceId, IoHash Hash, byte[] Data)
+			{
+				AddBlobsRequest Request = new AddBlobsRequest();
+				Request.NamespaceId = NamespaceId;
+				Request.Blobs.Add(new AddBlobRequest(Hash, Data));
+				await BlobStore.AddAsync(Request);
+			}
+
+			public async Task<byte[]> GetBlobAsync(string NamespaceId, IoHash Hash)
+			{
+				return await BlobStore.GetBlobAsync(NamespaceId, Hash);
+			}
+
+			public Task AddObjectAsync(string NamespaceId, IoHash Hash, byte[] Data)
+			{
+				return AddBlobAsync(NamespaceId, Hash, Data);
+			}
+
+			public Task<byte[]> GetObjectAsync(string NamespaceId, IoHash Hash)
+			{
+				return BlobStore.GetBlobAsync(NamespaceId, Hash);
+			}
+
+			public async Task AddTaskAsync(string ChannelId, string NamespaceId, IoHash RequirementsHash, IoHash TaskHash, bool SkipCacheLookup)
+			{
+				AddTasksRpcRequest AddTasksRequest = new AddTasksRpcRequest(ChannelId, NamespaceId, RequirementsHash, new List<CbObjectAttachment> { TaskHash }, SkipCacheLookup);
+				await ComputeClient.AddTasksAsync(AddTasksRequest);
+			}
+
+			public async Task GetTaskUpdatesAsync(string ChannelId, string NamespaceId, Func<GetTaskUpdateResponse, Task<bool>> OnUpdateAsync)
+			{
+				using (AsyncDuplexStreamingCall<GetTaskUpdatesRpcRequest, GetTaskUpdatesRpcResponse> Call = ComputeClient.GetTaskUpdates())
+				{
+					await Call.RequestStream.WriteAsync(new GetTaskUpdatesRpcRequest(ChannelId));
+
+					//					TaskCompletionSource<bool> CompleteNowTask = new TaskCompletionSource<bool>();
+					//					Task CompleteTask = Task.Run(async () => { await CompleteNowTask.Task; await Call.RequestStream.CompleteAsync(); });
+
+					while (await Call.ResponseStream.MoveNext())
+					{
+						GetTaskUpdatesRpcResponse RpcResponse = Call.ResponseStream.Current;
+
+						GetTaskUpdateResponse Response = new GetTaskUpdateResponse();
+						Response.AgentId = RpcResponse.AgentId;
+						Response.LeaseId = RpcResponse.LeaseId;
+						Response.ResultHash = (RpcResponse.ResultHash == null)? (CbObjectAttachment?)null : (CbObjectAttachment)RpcResponse.ResultHash;
+						Response.State = RpcResponse.State;
+						Response.TaskHash = RpcResponse.TaskHash;
+						Response.Time = RpcResponse.Time.ToDateTime();
+
+						if (!await OnUpdateAsync(Response))
+						{
+							await Call.RequestStream.CompleteAsync();
+						}
+					}
+				}
+			}
+		}
+
 		static (DirectoryTree, IoHash) CreateSandbox(DirectoryInfo BaseDirInfo, Dictionary<IoHash, byte[]> UploadList)
 		{
 			DirectoryTree Tree = new DirectoryTree();
@@ -114,6 +309,18 @@ namespace HordeAgent.Commands
 			return (Tree, AddCbObject(UploadList, Tree));
 		}
 
+		IComputeClient CreateComputeClient(IServiceProvider Services)
+		{
+			if (UseHttp)
+			{
+				return new HttpComputeClient(new Uri(Server), Token);
+			}
+			else
+			{
+				return new GrpcComputeClient(Services.GetRequiredService<GrpcService>());
+			}
+		}
+
 		/// <inheritdoc/>
 		public override async Task<int> ExecuteAsync(ILogger Logger)
 		{
@@ -149,7 +356,7 @@ namespace HordeAgent.Commands
 
 			using (IHost Host = HostBuilder.Build())
 			{
-				GrpcService GrpcService = Host.Services.GetRequiredService<GrpcService>();
+				using IComputeClient ComputeClient = CreateComputeClient(Host.Services);
 				
 				byte[] Json = await FileReference.ReadAllBytesAsync(Input);
 				JsonComputeTask JsonComputeTask = JsonSerializer.Deserialize<JsonComputeTask>(Json, new JsonSerializerOptions { PropertyNameCaseInsensitive = true });
@@ -183,7 +390,7 @@ namespace HordeAgent.Commands
 				Requirements.Exclusive = JsonRequirements.Exclusive;
 				IoHash RequirementsHash = AddCbObject(Blobs, Requirements);
 
-				await ExecuteAction(Logger, GrpcService, TaskHash, RequirementsHash, Blobs);	
+				await ExecuteAction(Logger, ComputeClient, TaskHash, RequirementsHash, Blobs);	
 			}
 			return 0;
 		}
@@ -196,82 +403,58 @@ namespace HordeAgent.Commands
 			return Hash;
 		}
 
-		private async Task ExecuteAction(ILogger Logger, GrpcService GrpcService, CbObjectAttachment TaskHash, CbObjectAttachment RequirementsHash, Dictionary<IoHash, byte[]> UploadList)
+		private async Task ExecuteAction(ILogger Logger, IComputeClient ComputeClient, CbObjectAttachment TaskHash, CbObjectAttachment RequirementsHash, Dictionary<IoHash, byte[]> UploadList)
 		{
-			using GrpcChannel Channel = GrpcService.CreateGrpcChannel();
-
 			Logger.LogInformation("task: {TaskHash}", TaskHash);
 			Logger.LogInformation("requirements: {RequirementsHash}", RequirementsHash);
 
 			const string NamespaceId = "default";
 
-			using GrpcChannel CasChannel = Channel;
-
-			BlobStore.BlobStoreClient BlobStore = new BlobStore.BlobStoreClient(CasChannel);
 			foreach (KeyValuePair<IoHash, byte[]> Pair in UploadList)
 			{
-				AddBlobsRequest Request = new AddBlobsRequest();
-				Request.NamespaceId = NamespaceId;
-				Request.Blobs.Add(new AddBlobRequest(Pair.Key, Pair.Value));
-				await BlobStore.AddAsync(Request);
+				await ComputeClient.AddBlobAsync(NamespaceId, Pair.Key, Pair.Value);
 			}
 
 			// Execute the action
-			ComputeRpc.ComputeRpcClient ComputeClient = new ComputeRpc.ComputeRpcClient(Channel);
+			string ChannelId = Guid.NewGuid().ToString();
+			await ComputeClient.AddTaskAsync(ChannelId, NamespaceId, RequirementsHash, TaskHash, SkipCacheLookup);
 
-			AddTasksRpcRequest AddTasksRequest = new AddTasksRpcRequest(Guid.NewGuid().ToString(), NamespaceId, RequirementsHash, new List<CbObjectAttachment> { TaskHash }, SkipCacheLookup);
-			await ComputeClient.AddTasksAsync(AddTasksRequest);
-
-			using (AsyncDuplexStreamingCall<GetTaskUpdatesRpcRequest, GetTaskUpdatesRpcResponse> Call = ComputeClient.GetTaskUpdates())
+			await ComputeClient.GetTaskUpdatesAsync(ChannelId, NamespaceId, async Response =>
 			{
-				await Call.RequestStream.WriteAsync(new GetTaskUpdatesRpcRequest(AddTasksRequest.ChannelId));
-
-				TaskCompletionSource<bool> CompleteNowTask = new TaskCompletionSource<bool>();
-				Task CompleteTask = Task.Run(async () => { await CompleteNowTask.Task; await Call.RequestStream.CompleteAsync(); } );
-
-				while (await Call.ResponseStream.MoveNext())
+				Logger.LogInformation("{OperationName}: Execution state: {State}", (IoHash)(CbObjectAttachment)Response.TaskHash, Response.State.ToString());
+				if (!String.IsNullOrEmpty(Response.AgentId) || !String.IsNullOrEmpty(Response.LeaseId))
 				{
-					GetTaskUpdatesRpcResponse Response = Call.ResponseStream.Current;
-					Logger.LogInformation("{OperationName}: Execution state: {State}", (IoHash)(CbObjectAttachment)Response.TaskHash, Response.State.ToString());
-					if (!String.IsNullOrEmpty(Response.AgentId) || !String.IsNullOrEmpty(Response.LeaseId))
-					{
-						Logger.LogInformation("{OperationName}: Running on agent {AgentId} under lease {LeaseId}", (IoHash)(CbObjectAttachment)Response.TaskHash, Response.AgentId, Response.LeaseId);
-					}
-					if (Response.ResultHash != null)
-					{
-						await HandleCompleteTask(BlobStore, NamespaceId, Response.ResultHash, Logger);
-					}
-					if (Response.State == ComputeTaskState.Complete)
-					{
-						CompleteNowTask.TrySetResult(true);
-					}
+					Logger.LogInformation("{OperationName}: Running on agent {AgentId} under lease {LeaseId}", (IoHash)(CbObjectAttachment)Response.TaskHash, Response.AgentId, Response.LeaseId);
 				}
-
-				await CompleteTask;
-			}
+				if (Response.ResultHash != null)
+				{
+					await HandleCompleteTask(ComputeClient, NamespaceId, Response.ResultHash.Value, Logger);
+				}
+				return Response.State != ComputeTaskState.Complete;
+			});
 		}
 
-		async Task HandleCompleteTask(BlobStore.BlobStoreClient Client, string NamespaceId, CbObjectAttachment Hash, ILogger Logger)
+		async Task HandleCompleteTask(IComputeClient ComputeClient, string NamespaceId, CbObjectAttachment Hash, ILogger Logger)
 		{
 			Logger.LogInformation("result: {ResultHash}", Hash);
 
-			ComputeTaskResult Result = await Client.GetObjectAsync<ComputeTaskResult>(NamespaceId, Hash);
+			ComputeTaskResult Result = await GetObjectAsync<ComputeTaskResult>(ComputeClient, NamespaceId, Hash);
 			Logger.LogInformation("exit: {ExitCode}", Result.ExitCode);
 
-			await LogTaskOutputAsync(Client, "stdout", NamespaceId, Result.StdOutHash, Logger);
-			await LogTaskOutputAsync(Client, "stderr", NamespaceId, Result.StdErrHash, Logger);
+			await LogTaskOutputAsync(ComputeClient, "stdout", NamespaceId, Result.StdOutHash, Logger);
+			await LogTaskOutputAsync(ComputeClient, "stderr", NamespaceId, Result.StdErrHash, Logger);
 
 			if (Result.OutputHash != null && OutputDir != null)
 			{
-				await WriteOutputAsync(Client, NamespaceId, Result.OutputHash.Value, new DirectoryReference(OutputDir));
+				await WriteOutputAsync(ComputeClient, NamespaceId, Result.OutputHash.Value, new DirectoryReference(OutputDir));
 			}
 		}
 
-		async Task LogTaskOutputAsync(BlobStore.BlobStoreClient Client, string Channel, string NamespaceId, IoHash? LogHash, ILogger Logger)
+		async Task LogTaskOutputAsync(IComputeClient ComputeClient, string Channel, string NamespaceId, IoHash? LogHash, ILogger Logger)
 		{
 			if (LogHash != null)
 			{
-				byte[] StdOutData = await Client.GetBlobAsync(NamespaceId, LogHash.Value);
+				byte[] StdOutData = await ComputeClient.GetBlobAsync(NamespaceId, LogHash.Value);
 				if (StdOutData.Length > 0)
 				{
 					foreach (string Line in Encoding.UTF8.GetString(StdOutData).Split('\n'))
@@ -282,23 +465,37 @@ namespace HordeAgent.Commands
 			}
 		}
 
-		async Task WriteOutputAsync(BlobStore.BlobStoreClient Client, string NamespaceId, IoHash TreeHash, DirectoryReference OutputDir)
+		async Task WriteOutputAsync(IComputeClient ComputeClient, string NamespaceId, IoHash TreeHash, DirectoryReference OutputDir)
 		{
-			DirectoryTree Tree = await Client.GetObjectAsync<DirectoryTree>(NamespaceId, TreeHash);
+			DirectoryTree Tree = await GetObjectAsync<DirectoryTree>(ComputeClient, NamespaceId, TreeHash);
 
 			List<Task> Tasks = new List<Task>();
 			foreach (FileNode File in Tree.Files)
 			{
 				FileReference OutputFile = FileReference.Combine(OutputDir, File.Name.ToString());
-				Tasks.Add(Client.GetFileAsync(NamespaceId, File.Hash, OutputFile));
+				Tasks.Add(WriteObjectToFileAsync(ComputeClient, NamespaceId, File.Hash, OutputFile));
 			}
 			foreach (DirectoryNode Directory in Tree.Directories)
 			{
 				DirectoryReference NextOutputDir = DirectoryReference.Combine(OutputDir, Directory.Name.ToString());
-				Tasks.Add(WriteOutputAsync(Client, NamespaceId, Directory.Hash, NextOutputDir));
+				Tasks.Add(WriteOutputAsync(ComputeClient, NamespaceId, Directory.Hash, NextOutputDir));
 			}
 
 			await Task.WhenAll(Tasks);
 		}
+
+		static async Task<T> GetObjectAsync<T>(IComputeClient ComputeClient, string NamespaceId, IoHash Hash)
+		{
+			byte[] Data = await ComputeClient.GetObjectAsync(NamespaceId, Hash);
+			return CbSerializer.Deserialize<T>(new CbField(Data));
+		}
+
+		static async Task WriteObjectToFileAsync(IComputeClient ComputeClient, string NamespaceId, IoHash Hash, FileReference OutputFile)
+		{
+			DirectoryReference.CreateDirectory(OutputFile.Directory);
+
+			byte[] Data = await ComputeClient.GetBlobAsync(NamespaceId, Hash);
+			await FileReference.WriteAllBytesAsync(OutputFile, Data);
+		}
 	}
 }
diff --git a/Engine/Source/Programs/Horde/HordeServer/Compute/Impl/ComputeService.cs b/Engine/Source/Programs/Horde/HordeServer/Compute/Impl/ComputeService.cs
--- a/Engine/Source/Programs/Horde/HordeServer/Compute/Impl/ComputeService.cs
+++ b/Engine/Source/Programs/Horde/HordeServer/Compute/Impl/ComputeService.cs
@@ -108,11 +108,13 @@ namespace HordeServer.Compute.Impl
 		{
 		}
 
-		public ComputeTaskStatus(CbObjectAttachment TaskHash, ComputeTaskState State)
+		public ComputeTaskStatus(CbObjectAttachment TaskHash, ComputeTaskState State, AgentId AgentId, ObjectId LeaseId)
 		{
 			this.TaskHash = TaskHash;
 			this.Time = DateTime.UtcNow;
 			this.State = State;
+			this.AgentId = AgentId.ToString();
+			this.LeaseIdBytes = LeaseId.ToByteArray();
 		}
 	}
 
@@ -215,7 +217,7 @@ namespace HordeServer.Compute.Impl
 			ComputeTaskMessage ComputeTask = Payload.Unpack<ComputeTaskMessage>();
 			Logger.LogInformation("Compute lease started (lease: {LeaseId}, task: {TaskHash}, agent: {AgentId}, channel: {ChannelId})", LeaseId, ComputeTask.TaskHash, Agent.Id, ComputeTask.ChannelId);
 
-			ComputeTaskStatus Status = new ComputeTaskStatus(ComputeTask.TaskHash, ComputeTaskState.Executing);
+			ComputeTaskStatus Status = new ComputeTaskStatus(ComputeTask.TaskHash, ComputeTaskState.Executing, Agent.Id, LeaseId);
 			return MessageQueue.PostAsync(ComputeTask.ChannelId, Status);
 		}
 
@@ -225,8 +227,7 @@ namespace HordeServer.Compute.Impl
 
 			ComputeTaskResultMessage Result = ComputeTaskResultMessage.Parser.ParseFrom(Output.ToArray());
 
-			ComputeTaskStatus Status = new ComputeTaskStatus(ComputeTask.TaskHash, ComputeTaskState.Complete);
-			Status.LeaseId = LeaseId;
+			ComputeTaskStatus Status = new ComputeTaskStatus(ComputeTask.TaskHash, ComputeTaskState.Complete, Agent.Id, LeaseId);
 			if (Result.OutputHash != null)
 			{
 				Status.ResultHash = Result.OutputHash;
