Horde: Allow namespace permissions to propagate down from the default ACL. Create a "default" namespace using the default config file, allowing locally running compute tasks without additional configuration.

diff --git a/Engine/Source/Programs/Horde/HordeServer/Defaults/globals.json b/Engine/Source/Programs/Horde/HordeServer/Defaults/globals.json
--- a/Engine/Source/Programs/Horde/HordeServer/Defaults/globals.json
+++ b/Engine/Source/Programs/Horde/HordeServer/Defaults/globals.json
@@ -4,5 +4,12 @@
 			"id": "ue5",
 			"path": "ue5.project.json"
 		}
-	]
+	],
+	"storage": {
+		"namespaces": [
+			{
+				"id": "default"
+			}
+		]
+	}
 }
diff --git a/Engine/Source/Programs/Horde/HordeServer/Storage/Collections/NamespaceCollection.cs b/Engine/Source/Programs/Horde/HordeServer/Storage/Collections/NamespaceCollection.cs
--- a/Engine/Source/Programs/Horde/HordeServer/Storage/Collections/NamespaceCollection.cs
+++ b/Engine/Source/Programs/Horde/HordeServer/Storage/Collections/NamespaceCollection.cs
@@ -36,9 +36,11 @@ namespace HordeServer.Storage
 
 		IMongoCollection<Namespace> Namespaces;
 		LazyCache<NamespaceId, INamespace?> NamespaceCache;
+		AclService AclService;
 
-		public NamespaceCollection(DatabaseService DatabaseService)
+		public NamespaceCollection(DatabaseService DatabaseService, AclService AclService)
 		{
+			this.AclService = AclService;
 			this.Namespaces = DatabaseService.GetCollection<Namespace>("Storage.Namespaces");
 			this.NamespaceCache = new LazyCache<NamespaceId, INamespace?>(GetAsync, new LazyCacheOptions { });
 		}
@@ -80,12 +82,12 @@ namespace HordeServer.Storage
 			}
 
 			bool? Result = Namespace.Acl?.Authorize(Action, User);
-			if (Result == null)
+			if (Result != null)
 			{
-				return false;
+				return Result.Value;
 			}
 
-			return Result.Value;
+			return await AclService.AuthorizeAsync(Action, User);
 		}
 	}
 }
