Horde: Evaluate conditions on agent properties when determining if they can execute a compute task.

diff --git a/Engine/Source/Programs/Horde/HordeServer/Compute/Impl/ComputeService.cs b/Engine/Source/Programs/Horde/HordeServer/Compute/Impl/ComputeService.cs
--- a/Engine/Source/Programs/Horde/HordeServer/Compute/Impl/ComputeService.cs
+++ b/Engine/Source/Programs/Horde/HordeServer/Compute/Impl/ComputeService.cs
@@ -139,7 +139,7 @@ namespace HordeServer.Compute.Impl
 		/// <param name="Logger">The logger instance</param>
 		public ComputeService(IDatabase Redis, IObjectCollection ObjectCollection, ILogger<ComputeService> Logger)
 		{
-			this.TaskScheduler = new RedisTaskScheduler<ComputeTaskInfo>(Redis, ObjectCollection, DefaultNamespaceId, "compute/tasks/");
+			this.TaskScheduler = new RedisTaskScheduler<ComputeTaskInfo>(Redis, ObjectCollection, DefaultNamespaceId, "compute/tasks/", Logger);
 			this.MessageQueue = new RedisMessageQueue<ComputeTaskStatus>(Redis, "compute/messages/");
 			this.Logger = Logger;
 		}
diff --git a/Engine/Source/Programs/Horde/HordeServer/Utilities/Condition.cs b/Engine/Source/Programs/Horde/HordeServer/Utilities/Condition.cs
new file mode 100644
--- /dev/null
+++ b/Engine/Source/Programs/Horde/HordeServer/Utilities/Condition.cs
@@ -0,0 +1,520 @@
+// Copyright Epic Games, Inc. All Rights Reserved.
+
+using EpicGames.Core;
+using System;
+using System.Buffers.Text;
+using System.Collections.Generic;
+using System.Diagnostics;
+using System.Globalization;
+using System.Linq;
+using System.Text.RegularExpressions;
+using System.Threading.Tasks;
+
+namespace HordeServer.Utilities
+{
+	/// <summary>
+	/// Represents a scalar value for use by a condition expression
+	/// </summary>
+	readonly struct Scalar
+	{
+		public bool IsInteger => !IsString;
+		public bool IsString => String != null;
+
+		public readonly ulong Integer;
+		public readonly string? String;
+
+		public Scalar(ulong Integer)
+		{
+			this.Integer = Integer;
+			this.String = null;
+		}
+
+		public Scalar(string String)
+		{
+			this.Integer = 0;
+			this.String = String;
+		}
+
+		public ulong AsInteger()
+		{
+			if (!IsInteger)
+			{
+				throw new ConditionException("Value is not an integer");
+			}
+			return Integer;
+		}
+		
+		public string AsString()
+		{
+			if (!IsString)
+			{
+				throw new ConditionException("Value is not a string");
+			}
+			return String!;
+		}
+
+		public static implicit operator Scalar(bool Value)
+		{
+			return new Scalar(Value ? 1UL : 0UL);
+		}
+
+		public static implicit operator Scalar(ulong Integer)
+		{
+			return new Scalar(Integer);
+		}
+
+		public static implicit operator Scalar(string String)
+		{
+			return new Scalar(String);
+		}
+	}
+
+	enum TokenType : byte
+	{
+		End,
+		Identifier,
+		Number,
+		String,
+		Not,
+		Eq,
+		Neq,
+		LogicalOr,
+		LogicalAnd,
+		Lt,
+		Lte,
+		Gt,
+		Gte,
+		Lparen,
+		Rparen,
+		Regex
+	}
+
+	class TokenReader
+	{
+		public string Input { get; private set; }
+		public int Offset { get; private set; }
+		public int Length { get; private set; }
+
+		public StringView Token => new StringView(Input, Offset, Length);
+		public TokenType Type { get; private set; }
+		public ulong IntegerValue { get; private set; }
+		public string? StringValue { get; private set; }
+
+		public TokenReader(string Input)
+		{
+			this.Input = Input;
+			MoveNext();
+		}
+
+		private void SetCurrent(int Length, TokenType Type, ulong IntegerValue = 0, string? StringValue = null)
+		{
+			this.Length = Length;
+			this.Type = Type;
+			this.IntegerValue = IntegerValue;
+			this.StringValue = StringValue;
+		}
+
+		public void MoveNext()
+		{
+			Offset += Length;
+
+			while (Offset < Input.Length)
+			{
+				switch (Input[Offset])
+				{
+					case ' ':
+					case '\t':
+						Offset++;
+						break;
+					case '(':
+						SetCurrent(1, TokenType.Lparen);
+						return;
+					case ')':
+						SetCurrent(1, TokenType.Rparen);
+						return;
+					case '!':
+						if (Offset + 1 < Input.Length && Input[Offset + 1] == '=')
+						{
+							SetCurrent(2, TokenType.Neq);
+						}
+						else
+						{
+							SetCurrent(1, TokenType.Not);
+						}
+						return;
+					case '&':
+						RequireCharacter(Input, Offset + 1, '&');
+						SetCurrent(2, TokenType.LogicalAnd);
+						return;
+					case '|':
+						RequireCharacter(Input, Offset + 1, '|');
+						SetCurrent(2, TokenType.LogicalOr);
+						return;
+					case '=':
+						RequireCharacter(Input, Offset + 1, '=');
+						SetCurrent(2, TokenType.Eq);
+						return;
+					case '~':
+						RequireCharacter(Input, Offset + 1, '=');
+						SetCurrent(2, TokenType.Regex);
+						return;
+					case '>':
+						if (Offset + 1 < Input.Length && Input[Offset + 1] == '=')
+						{
+							SetCurrent(2, TokenType.Gte);
+						}
+						else
+						{
+							SetCurrent(1, TokenType.Gt);
+						}
+						return;
+					case '<':
+						if (Offset + 1 < Input.Length && Input[Offset + 1] == '=')
+						{
+							SetCurrent(2, TokenType.Lte);
+						}
+						else
+						{
+							SetCurrent(1, TokenType.Lt);
+						}
+						return;
+					case '\"':
+					case '\'':
+						Utf8String String = ParseString(Input, Offset, out int StringLength);
+						SetCurrent(StringLength, TokenType.String, StringValue: String.ToString());
+						return;
+					case char Char when IsNumber(Char):
+						ulong Integer = ParseNumber(Input, Offset, out int IntegerLength);
+						SetCurrent(IntegerLength, TokenType.Number, IntegerValue: Integer);
+						return;
+					case char Char when IsIdentifier(Char):
+						int EndIdx = Offset + 1;
+						while (EndIdx < Input.Length && IsIdentifierTail(Input[EndIdx]))
+						{
+							EndIdx++;
+						}
+						SetCurrent(EndIdx - Offset, TokenType.Identifier);
+						return;
+					default:
+						throw new ConditionException($"Invalid character at offset {Offset}: '{Input[Offset]}'");
+				}
+			}
+
+			SetCurrent(0, TokenType.End);
+		}
+
+		static string ParseString(string Text, int Idx, out int BytesConsumed)
+		{
+			char QuoteChar = Text[Idx];
+			if (QuoteChar != '\'' && QuoteChar != '\"')
+			{
+				throw new ConditionException($"Invalid quote character '{(char)QuoteChar}' at offset {Idx}");
+			}
+
+			int NumEscapeChars = 0;
+
+			int EndIdx = Idx + 1;
+			for (; ; EndIdx++)
+			{
+				if (EndIdx >= Text.Length)
+				{
+					throw new ConditionException("Unterminated string in expression");
+				}
+				else if (Text[EndIdx] == '\\')
+				{
+					NumEscapeChars++;
+					EndIdx++;
+				}
+				else if (Text[EndIdx] == QuoteChar)
+				{
+					break;
+				}
+			}
+
+
+			char[] Copy = new char[EndIdx - (Idx + 1) - NumEscapeChars];
+
+			int InputIdx = Idx + 1;
+			int OutputIdx = 0;
+			while (OutputIdx < Copy.Length)
+			{
+				if (Text[InputIdx] == '\\')
+				{
+					InputIdx++;
+				}
+				Copy[OutputIdx++] = Text[InputIdx++];
+			}
+
+			BytesConsumed = EndIdx + 1 - Idx;
+			return new string(Copy);
+		}
+
+		static Dictionary<StringView, ulong> SizeSuffixes = new Dictionary<StringView, ulong>(StringViewComparer.OrdinalIgnoreCase)
+		{
+			["kb"] = 1024UL,
+			["Mb"] = 1024UL * 1024,
+			["Gb"] = 1024UL * 1024 * 1024,
+			["Tb"] = 1024UL * 1024 * 1024 * 1024,
+		};
+
+		static ulong ParseNumber(string Text, int Idx, out int BytesConsumed)
+		{
+			ulong Value = 0;
+
+			int EndIdx = Idx;
+			while (EndIdx < Text.Length && IsNumber(Text[EndIdx]))
+			{
+				Value = (Value * 10) + (uint)(Text[EndIdx] - '0');
+				EndIdx++;
+			}
+
+			if (EndIdx < Text.Length && IsIdentifier(Text[EndIdx]))
+			{
+				int Offset = EndIdx++;
+				while (EndIdx < Text.Length && IsIdentifierTail(Text[EndIdx]))
+				{
+					EndIdx++;
+				}
+
+				StringView Suffix = new StringView(Text, Offset, EndIdx - Offset);
+
+				ulong Size;
+				if (!SizeSuffixes.TryGetValue(Suffix, out Size))
+				{
+					throw new ConditionException($"'{Suffix}' is not a valid numerical suffix");
+				}
+				Value *= Size;
+			}
+
+			BytesConsumed = EndIdx - Idx;
+			return Value;
+		}
+
+		static void RequireCharacter(string Text, int Idx, char Character)
+		{
+			if (Idx == Text.Length || Text[Idx] != Character)
+			{
+				throw new ConditionException($"Invalid character at position {Idx}; expected '{Character}'.");
+			}
+		}
+
+		static bool IsNumber(char Character)
+		{
+			return Character >= '0' && Character <= '9';
+		}
+
+		static bool IsIdentifier(char Character)
+		{
+			return (Character >= 'a' && Character <= 'z') || (Character >= 'A' && Character <= 'Z') || Character == '_';
+		}
+
+		static bool IsIdentifierTail(char Character)
+		{
+			return IsIdentifier(Character) || IsNumber(Character) || Character == '.';
+		}
+	}
+
+	/// <summary>
+	/// Exception thrown when parsing or evaluating a condition
+	/// </summary>
+	public class ConditionException : Exception
+	{
+		internal ConditionException(string Text) : base(Text)
+		{
+		}
+	}
+
+	/// <summary>
+	/// A conditional expression that can be evaluated against a particular object
+	/// </summary>
+	class Condition
+	{
+		[DebuggerDisplay("{Type}")]
+		readonly struct Token
+		{
+			public readonly TokenType Type;
+			public readonly byte Index;
+
+			public Token(TokenType Type, int Index)
+			{
+				this.Type = Type;
+				this.Index = (byte)Index;
+			}
+		}
+
+		List<Token> Tokens = new List<Token>();
+		List<string> Strings = new List<string>();
+		List<ulong> Integers = new List<ulong>();
+
+		private Condition(string Text)
+		{
+			TokenReader Reader = new TokenReader(Text);
+			ParseOr(Reader);
+			if (Reader.Type != TokenType.End)
+			{
+				throw new ConditionException($"Unexpected token at offset {Reader.Offset}: {Reader.Token}");
+			}
+		}
+
+		public static Condition Parse(string Text)
+		{
+			return new Condition(Text);
+		}
+
+		public bool Evaluate(Func<string, Scalar> ResolveIdentifier)
+		{
+			int Idx = 0;
+			Scalar Scalar = Evaluate(ref Idx, ResolveIdentifier);
+			return Scalar.AsInteger() != 0;
+		}
+
+		Scalar Evaluate(ref int Idx, Func<string, Scalar> ResolveIdentifier)
+		{
+			Scalar Lhs;
+			Scalar Rhs;
+
+			Token Token = Tokens[Idx++];
+			switch (Token.Type)
+			{
+				case TokenType.Identifier:
+					return ResolveIdentifier(Strings[Token.Index]);
+				case TokenType.Number:
+					return Integers[Token.Index];
+				case TokenType.String:
+					return Strings[Token.Index].ToString();
+				case TokenType.Not:
+					Lhs = Evaluate(ref Idx, ResolveIdentifier);
+					return Lhs.AsInteger() == 0UL;
+				case TokenType.Eq:
+					Lhs = Evaluate(ref Idx, ResolveIdentifier);
+					Rhs = Evaluate(ref Idx, ResolveIdentifier);
+					return (Lhs.IsString || Rhs.IsString) ? String.Equals(Lhs.AsString(), Rhs.AsString(), StringComparison.OrdinalIgnoreCase) : (Lhs.AsInteger() == Rhs.AsInteger());
+				case TokenType.Neq:
+					Lhs = Evaluate(ref Idx, ResolveIdentifier);
+					Rhs = Evaluate(ref Idx, ResolveIdentifier);
+					return (Lhs.IsString || Rhs.IsString) ? !String.Equals(Lhs.AsString(), Rhs.AsString(), StringComparison.OrdinalIgnoreCase) : (Lhs.AsInteger() != Rhs.AsInteger());
+				case TokenType.LogicalOr:
+					Lhs = Evaluate(ref Idx, ResolveIdentifier);
+					Rhs = Evaluate(ref Idx, ResolveIdentifier);
+					return (Lhs.AsInteger() != 0) || (Rhs.AsInteger() != 0);
+				case TokenType.LogicalAnd:
+					Lhs = Evaluate(ref Idx, ResolveIdentifier);
+					Rhs = Evaluate(ref Idx, ResolveIdentifier);
+					return (Lhs.AsInteger() != 0) && (Rhs.AsInteger() != 0);
+				case TokenType.Lt:
+					Lhs = Evaluate(ref Idx, ResolveIdentifier);
+					Rhs = Evaluate(ref Idx, ResolveIdentifier);
+					return Lhs.AsInteger() < Rhs.AsInteger();
+				case TokenType.Lte:
+					Lhs = Evaluate(ref Idx, ResolveIdentifier);
+					Rhs = Evaluate(ref Idx, ResolveIdentifier);
+					return Lhs.AsInteger() <= Rhs.AsInteger();
+				case TokenType.Gt:
+					Lhs = Evaluate(ref Idx, ResolveIdentifier);
+					Rhs = Evaluate(ref Idx, ResolveIdentifier);
+					return Lhs.AsInteger() > Rhs.AsInteger();
+				case TokenType.Gte:
+					Lhs = Evaluate(ref Idx, ResolveIdentifier);
+					Rhs = Evaluate(ref Idx, ResolveIdentifier);
+					return Lhs.AsInteger() >= Rhs.AsInteger();
+				case TokenType.Regex:
+					Lhs = Evaluate(ref Idx, ResolveIdentifier);
+					Rhs = Evaluate(ref Idx, ResolveIdentifier);
+					return Regex.IsMatch(Lhs.AsString(), Rhs.AsString());
+				default:
+					throw new ConditionException("Invalid token type");
+			}
+		}
+
+		void ParseOr(TokenReader Reader)
+		{
+			int StartCount = Tokens.Count;
+			ParseAnd(Reader);
+
+			if (Reader.Type == TokenType.LogicalOr)
+			{
+				Tokens.Insert(StartCount, new Token(TokenType.LogicalOr, 0));
+				Reader.MoveNext();
+				ParseOr(Reader);
+			}
+		}
+
+		void ParseAnd(TokenReader Reader)
+		{
+			int StartCount = Tokens.Count;
+			ParseComparison(Reader);
+
+			if (Reader.Type == TokenType.LogicalAnd)
+			{
+				Tokens.Insert(StartCount, new Token(TokenType.LogicalAnd, 0));
+				Reader.MoveNext();
+				ParseAnd(Reader);
+			}
+		}
+
+		void ParseComparison(TokenReader Reader)
+		{
+			int StartCount = Tokens.Count;
+			ParseUnary(Reader);
+
+			if (Reader.Type == TokenType.Lt 
+				|| Reader.Type == TokenType.Lte 
+				|| Reader.Type == TokenType.Gt 
+				|| Reader.Type == TokenType.Gte 
+				|| Reader.Type == TokenType.Eq 
+				|| Reader.Type == TokenType.Neq
+				|| Reader.Type == TokenType.Regex)
+			{
+				Tokens.Insert(StartCount, new Token(Reader.Type, 0));
+				Reader.MoveNext();
+				ParseUnary(Reader);
+			}
+		}
+
+		void ParseUnary(TokenReader Reader)
+		{
+			while (Reader.Type == TokenType.Not)
+			{
+				Tokens.Add(new Token(Reader.Type, 0));
+				Reader.MoveNext();
+			}
+			ParseScalar(Reader);
+		}
+
+		void ParseScalar(TokenReader Reader)
+		{
+			if (Reader.Type == TokenType.Identifier)
+			{
+				Strings.Add(Reader.Token.ToString());
+				Tokens.Add(new Token(TokenType.Identifier, Strings.Count - 1));
+				Reader.MoveNext();
+			}
+			else if (Reader.Type == TokenType.Number)
+			{
+				Integers.Add(Reader.IntegerValue);
+				Tokens.Add(new Token(TokenType.Number, Integers.Count - 1));
+				Reader.MoveNext();
+			}
+			else if (Reader.Type == TokenType.String)
+			{
+				Strings.Add(Reader.StringValue!);
+				Tokens.Add(new Token(TokenType.String, Strings.Count - 1));
+				Reader.MoveNext();
+			}
+			else if (Reader.Type == TokenType.Lparen)
+			{
+				Reader.MoveNext();
+				ParseOr(Reader);
+				if (Reader.Type != TokenType.Rparen)
+				{
+					throw new ConditionException($"Missing ')' at offset {Reader.Offset}");
+				}
+				Reader.MoveNext();
+			}
+			else
+			{
+				throw new ConditionException($"Unexpected token '{Reader.Token}' at offset {Reader.Offset}");
+			}
+		}
+	}
+}
