Horde: Output CB hex dump to the serilog logger as well as throwing an exception.

diff --git a/Engine/Source/Programs/Horde/HordeServer/Utilities/CompactBinaryFormatter.cs b/Engine/Source/Programs/Horde/HordeServer/Utilities/CompactBinaryFormatter.cs
--- a/Engine/Source/Programs/Horde/HordeServer/Utilities/CompactBinaryFormatter.cs
+++ b/Engine/Source/Programs/Horde/HordeServer/Utilities/CompactBinaryFormatter.cs
@@ -64,7 +64,8 @@ namespace HordeServer.Utilities
 			}
 			catch (Exception Ex)
 			{
-				throw new Exception($"Unable to parse compact binary data: {FormatHexDump(Data, 256)}", Ex);
+				Serilog.Log.Logger.Error("Unable to parse compact binary: {Dump}", FormatHexDump(Data, 256));
+				throw new Exception($"Unable to parse compact binary request: {FormatHexDump(Data, 256)}", Ex);
 			}
 			return InputFormatterResult.SuccessAsync(CbSerializer.Deserialize(new CbField(Data), Context.ModelType));
 		}
