Horde: Missing files for compute task conditions.

diff --git a/Engine/Source/Programs/Horde/HordeServer/Compute/Impl/RedisTaskScheduler.cs b/Engine/Source/Programs/Horde/HordeServer/Compute/Impl/RedisTaskScheduler.cs
--- a/Engine/Source/Programs/Horde/HordeServer/Compute/Impl/RedisTaskScheduler.cs
+++ b/Engine/Source/Programs/Horde/HordeServer/Compute/Impl/RedisTaskScheduler.cs
@@ -98,6 +98,12 @@ namespace HordeServer.Compute.Impl
 			}
 		}
 
+		class QueueInfo
+		{
+			public Requirements? Requirements;
+			public Utilities.Condition? Condition;
+		}
+
 		IDatabase Redis;
 		IObjectCollection ObjectCollection;
 		RedisKey BaseKey;
@@ -112,6 +118,7 @@ namespace HordeServer.Compute.Impl
 		Task QueueUpdateTask;
 		CancellationTokenSource CancellationSource = new CancellationTokenSource();
 		IMemoryCache RequirementsCache;
+		ILogger Logger;
 
 		/// <summary>
 		/// Constructor
@@ -120,7 +127,8 @@ namespace HordeServer.Compute.Impl
 		/// <param name="ObjectCollection">Interface to the CAS object store</param>
 		/// <param name="BaseKey">Base key for all keys used by this scheduler</param>
 		/// <param name="NamespaceId">The namespace identifier</param>
-		public RedisTaskScheduler(IDatabase Redis, IObjectCollection ObjectCollection, NamespaceId NamespaceId, RedisKey BaseKey)
+		/// <param name="Logger"></param>
+		public RedisTaskScheduler(IDatabase Redis, IObjectCollection ObjectCollection, NamespaceId NamespaceId, RedisKey BaseKey, ILogger Logger)
 		{
 			this.Redis = Redis;
 			this.ObjectCollection = ObjectCollection;
@@ -130,6 +138,7 @@ namespace HordeServer.Compute.Impl
 			this.ActiveQueues = new RedisHash<IoHash, DateTime>(BaseKey.Append("active"));
 			this.NewQueueChannel = new RedisChannel<IoHash>(BaseKey.Append("new_queues").ToString());
 			this.RequirementsCache = new MemoryCache(new MemoryCacheOptions());
+			this.Logger = Logger;
 
 			QueueUpdateTask = Task.Run(() => UpdateQueuesAsync(CancellationSource.Token));
 		}
@@ -179,7 +188,7 @@ namespace HordeServer.Compute.Impl
 		async ValueTask AddQueueToIndexAsync(CbObjectAttachment RequirementsHash)
 		{
 			ITransaction Transaction = Redis.CreateTransaction();
-			Transaction.AddCondition(Condition.KeyExists(GetQueueKey(RequirementsHash)));
+			Transaction.AddCondition(StackExchange.Redis.Condition.KeyExists(GetQueueKey(RequirementsHash)));
 			_ = Transaction.SetAddAsync(QueueIndex, RequirementsHash);
 			await Transaction.ExecuteAsync();
 
@@ -193,7 +202,7 @@ namespace HordeServer.Compute.Impl
 		async ValueTask RemoveQueueFromIndexAsync(CbObjectAttachment RequirementsHash)
 		{
 			ITransaction Transaction = Redis.CreateTransaction();
-			Transaction.AddCondition(Condition.KeyNotExists(GetQueueKey(RequirementsHash)));
+			Transaction.AddCondition(StackExchange.Redis.Condition.KeyNotExists(GetQueueKey(RequirementsHash)));
 			_ = Transaction.SetRemoveAsync(QueueIndex, RequirementsHash);
 			await Transaction.ExecuteAsync(CommandFlags.FireAndForget);
 		}
@@ -287,7 +296,7 @@ namespace HordeServer.Compute.Impl
 		{
 			foreach (IoHash RequirementsHash in Queues)
 			{
-				Requirements? Requirements = await GetRequirementsAsync(RequirementsHash);
+				QueueInfo? Requirements = await GetQueueInfoAsync(RequirementsHash);
 				if (Requirements != null && CheckRequirements(Agent, Requirements))
 				{
 					T? Item = await DequeueAsync(RequirementsHash);
@@ -418,7 +427,7 @@ namespace HordeServer.Compute.Impl
 			RedisList<T> Queue = GetQueue(RequirementsHash);
 
 			// Get the requirements for this queue
-			Requirements? Requirements = await GetRequirementsAsync(RequirementsHash);
+			QueueInfo? Requirements = await GetQueueInfoAsync(RequirementsHash);
 			if (Requirements == null)
 			{
 				throw new Exception($"Unable to find requirements with hash {RequirementsHash}");
@@ -468,33 +477,88 @@ namespace HordeServer.Compute.Impl
 			}
 		}
 
-		static bool CheckRequirements(IAgent Agent, Requirements Requirements)
+		/// <summary>
+		/// Checks that an agent matches the necessary criteria to execute a task
+		/// </summary>
+		/// <param name="Agent"></param>
+		/// <param name="QueueInfo"></param>
+		/// <returns></returns>
+		static bool CheckRequirements(IAgent Agent, QueueInfo QueueInfo)
 		{
-			_ = Agent;
-			_ = Requirements;
-			// TODO: check requirements
-			// TODO: cache requirements
+			if (QueueInfo.Requirements == null)
+			{
+				return false;
+			}
+
+			if (!QueueInfo.Requirements.Condition.IsEmpty)
+			{
+				if (QueueInfo.Condition == null)
+				{
+					return false;
+				}
+				else if (!QueueInfo.Condition.Evaluate(x => GetAgentProperty(Agent, x)))
+				{
+					return false;
+				}
+			}
+
 			return true;
 		}
 
+		/// <summary>
+		/// Gets a named property for an agent
+		/// </summary>
+		/// <param name="Agent">The agent instance</param>
+		/// <param name="Name">Name of the property to read</param>
+		/// <returns>The property value, or an empty string if it does not eixst</returns>
+		static Scalar GetAgentProperty(IAgent Agent, string Name)
+		{
+			HashSet<string>? Properties = Agent.Capabilities.PrimaryDevice.Properties;
+			if (Properties != null)
+			{
+				foreach (string Property in Properties)
+				{
+					if (Property.Length > Name.Length && Property[Name.Length] == '=' && Property.StartsWith(Name, StringComparison.OrdinalIgnoreCase))
+					{
+						return Property.Substring(Name.Length + 1);
+					}
+				}
+			}
+			return new Scalar("");
+		}
+
 		/// <summary>
 		/// Gets the requirements object from the CAS
 		/// </summary>
 		/// <param name="RequirementsHash"></param>
 		/// <returns></returns>
-		async Task<Requirements?> GetRequirementsAsync(IoHash RequirementsHash)
+		async Task<QueueInfo?> GetQueueInfoAsync(IoHash RequirementsHash)
 		{
-			Requirements? Requirements;
-			if (!RequirementsCache.TryGetValue(RequirementsHash, out Requirements))
+			QueueInfo? Info;
+			if (!RequirementsCache.TryGetValue(RequirementsHash, out Info))
 			{
-				Requirements = await ObjectCollection.GetAsync<Requirements>(NamespaceId, RequirementsHash);
+				Info = new QueueInfo();
+				Info.Requirements = await ObjectCollection.GetAsync<Requirements>(NamespaceId, RequirementsHash);
+
+				if (Info.Requirements != null && !Info.Requirements.Condition.IsEmpty)
+				{
+					try
+					{
+						Info.Condition = Utilities.Condition.Parse(Info.Requirements.Condition.ToString());
+					}
+					catch (Exception Ex)
+					{
+						Logger.LogError(Ex, "Unable to parse condition '{Condition}': {Message}", Ex.Message);
+					}
+				}
+
 				using (ICacheEntry Entry = RequirementsCache.CreateEntry(RequirementsHash))
 				{
 					Entry.SetSlidingExpiration(TimeSpan.FromMinutes(10.0));
-					Entry.SetValue(Requirements);
+					Entry.SetValue(Info);
 				}
 			}
-			return Requirements;
+			return Info;
 		}
 	}
 }
diff --git a/Engine/Source/Programs/Horde/HordeServer/Logs/Builder/Impl/RedisLogBuilder.cs b/Engine/Source/Programs/Horde/HordeServer/Logs/Builder/Impl/RedisLogBuilder.cs
--- a/Engine/Source/Programs/Horde/HordeServer/Logs/Builder/Impl/RedisLogBuilder.cs
+++ b/Engine/Source/Programs/Horde/HordeServer/Logs/Builder/Impl/RedisLogBuilder.cs
@@ -18,6 +18,8 @@ using OpenTracing.Util;
 
 namespace HordeServer.Logs.Builder
 {
+	using Condition = StackExchange.Redis.Condition;
+
 	/// <summary>
 	/// Redis-based cache for log file chunks
 	/// </summary>
