Horde: Experimental dark mode
#rnx

diff --git a/Engine/Source/Programs/Horde/HordeDashboard/package.json b/Engine/Source/Programs/Horde/HordeDashboard/package.json
--- a/Engine/Source/Programs/Horde/HordeDashboard/package.json
+++ b/Engine/Source/Programs/Horde/HordeDashboard/package.json
@@ -18,6 +18,7 @@
     "moment-timezone": "^0.5.28",
     "react": "^16.12.0",
     "react-custom-scrollbars": "^4.2.1",
+    "react-darkreader": "^1.5.6",
     "react-dom": "^17.0.1",
     "react-flame-graph": "^1.4.0",
     "react-highlighter": "^0.4.3",
diff --git a/Engine/Source/Programs/Horde/HordeDashboard/src/backend/Api.ts b/Engine/Source/Programs/Horde/HordeDashboard/src/backend/Api.ts
--- a/Engine/Source/Programs/Horde/HordeDashboard/src/backend/Api.ts
+++ b/Engine/Source/Programs/Horde/HordeDashboard/src/backend/Api.ts
@@ -2677,6 +2677,7 @@ export type UserClaim = {
 }
 
 export enum DashboardPreference {
+	Darktheme = "Darktheme",
 	DisplayUTC = "DisplayUTC",
 	DisplayClock = "DisplayClock",
 	ColorSuccess = "ColorSuccess",
diff --git a/Engine/Source/Programs/Horde/HordeDashboard/src/backend/Dashboard.ts b/Engine/Source/Programs/Horde/HordeDashboard/src/backend/Dashboard.ts
--- a/Engine/Source/Programs/Horde/HordeDashboard/src/backend/Dashboard.ts
+++ b/Engine/Source/Programs/Horde/HordeDashboard/src/backend/Dashboard.ts
@@ -166,6 +166,17 @@ export class Dashboard {
 
     }
 
+    get darktheme(): boolean {
+
+        return this.preferences.get(DashboardPreference.Darktheme) === 'true';
+
+    }
+
+    setDarkTheme(value: boolean | undefined) {
+        this.setPreference(DashboardPreference.Darktheme, value ? "true" : "false");
+        this.setUpdated();
+    }
+
     setDisplayUTC(value: boolean | undefined) {
         this.setPreference(DashboardPreference.DisplayUTC, value ? "true" : "false");
     }
