Horde: Allow filtering agents for compute tasks using agent names.

diff --git a/Engine/Source/Programs/Horde/HordeServer/Compute/Impl/ComputeService.cs b/Engine/Source/Programs/Horde/HordeServer/Compute/Impl/ComputeService.cs
--- a/Engine/Source/Programs/Horde/HordeServer/Compute/Impl/ComputeService.cs
+++ b/Engine/Source/Programs/Horde/HordeServer/Compute/Impl/ComputeService.cs
@@ -185,7 +185,7 @@ namespace HordeServer.Compute.Impl
 				byte[] Payload = Any.Pack(ComputeTask).ToByteArray();
 
 				AgentLease Lease = new AgentLease(ObjectId.GenerateNewId(), LeaseName, null, null, null, LeaseState.Pending, Payload, new AgentRequirements(), null);
-				Logger.LogDebug("Created lease {LeaseId} for channel {ChannelId} task {TaskHash}", Lease.Id, ComputeTask.ChannelId, (CbObjectAttachment)ComputeTask.TaskHash);
+				Logger.LogDebug("Created lease {LeaseId} for channel {ChannelId} task {TaskHash} req {RequirementsHash}", Lease.Id, ComputeTask.ChannelId, (CbObjectAttachment)ComputeTask.TaskHash, ComputeTask.RequirementsHash);
 				return Lease;
 			}
 			return null;
@@ -215,7 +215,7 @@ namespace HordeServer.Compute.Impl
 		public Task OnLeaseStartedAsync(IAgent Agent, ObjectId LeaseId, Any Payload)
 		{
 			ComputeTaskMessage ComputeTask = Payload.Unpack<ComputeTaskMessage>();
-			Logger.LogInformation("Compute lease started (lease: {LeaseId}, task: {TaskHash}, agent: {AgentId}, channel: {ChannelId})", LeaseId, ComputeTask.TaskHash, Agent.Id, ComputeTask.ChannelId);
+			Logger.LogInformation("Compute lease started (lease: {LeaseId}, task: {TaskHash}, agent: {AgentId}, channel: {ChannelId})", LeaseId, (CbObjectAttachment)ComputeTask.TaskHash, Agent.Id, ComputeTask.ChannelId);
 
 			ComputeTaskStatus Status = new ComputeTaskStatus(ComputeTask.TaskHash, ComputeTaskState.Executing, Agent.Id, LeaseId);
 			return MessageQueue.PostAsync(ComputeTask.ChannelId, Status);
@@ -233,7 +233,7 @@ namespace HordeServer.Compute.Impl
 				Status.ResultHash = Result.OutputHash;
 			}
 
-			Logger.LogInformation("Compute lease finished (lease: {LeaseId}, task: {TaskHash}, agent: {AgentId}, channel: {ChannelId}, result: {ResultHash})", LeaseId, ComputeTask.TaskHash, Agent.Id, ComputeTask.ChannelId, Status.ResultHash?.Hash ?? IoHash.Zero);
+			Logger.LogInformation("Compute lease finished (lease: {LeaseId}, task: {TaskHash}, agent: {AgentId}, channel: {ChannelId}, result: {ResultHash})", LeaseId, (CbObjectAttachment)ComputeTask.TaskHash, Agent.Id, ComputeTask.ChannelId, Status.ResultHash?.Hash ?? IoHash.Zero);
 			return MessageQueue.PostAsync(ComputeTask.ChannelId, Status);
 		}
 	}
diff --git a/Engine/Source/Programs/Horde/HordeServer/Compute/Impl/RedisTaskScheduler.cs b/Engine/Source/Programs/Horde/HordeServer/Compute/Impl/RedisTaskScheduler.cs
--- a/Engine/Source/Programs/Horde/HordeServer/Compute/Impl/RedisTaskScheduler.cs
+++ b/Engine/Source/Programs/Horde/HordeServer/Compute/Impl/RedisTaskScheduler.cs
@@ -513,6 +513,11 @@ namespace HordeServer.Compute.Impl
 		/// <returns>The property value, or an empty string if it does not eixst</returns>
 		static Scalar GetAgentProperty(IAgent Agent, string Name)
 		{
+			if (Name.Equals("Name", StringComparison.OrdinalIgnoreCase))
+			{
+				return Agent.Id.ToString();
+			}
+
 			HashSet<string>? Properties = Agent.Capabilities.PrimaryDevice.Properties;
 			if (Properties != null)
 			{
