Horde: Support legacy and update name for upgrade tasks.

diff --git a/Engine/Source/Programs/Horde/HordeAgent/Services/WorkerService.cs b/Engine/Source/Programs/Horde/HordeAgent/Services/WorkerService.cs
--- a/Engine/Source/Programs/Horde/HordeAgent/Services/WorkerService.cs
+++ b/Engine/Source/Programs/Horde/HordeAgent/Services/WorkerService.cs
@@ -719,7 +719,7 @@ namespace HordeAgent.Services
 			}
 
 			UpgradeTask UpgradeTask;
-			if (LeaseInfo.Lease.Payload.TryUnpack(out UpgradeTask))
+			if (LeaseInfo.Lease.Payload.TryUnpack(out UpgradeTask) || TryUnpackUpgradeTask(LeaseInfo.Lease.Payload, out UpgradeTask))
 			{
 				GlobalTracer.Instance.ActiveSpan?.SetTag("task", "Upgrade");
 				Func<ILogger, Task<LeaseResult>> Handler = NewLogger => UpgradeAsync(RpcConnection, AgentId, UpgradeTask, NewLogger, LeaseInfo.CancellationTokenSource.Token);
@@ -746,6 +746,26 @@ namespace HordeAgent.Services
 			return LeaseResult.Failed;
 		}
 
+		/// <summary>
+		/// Attempts to parse an upgrade task, given either name
+		/// </summary>
+		/// <param name="Payload"></param>
+		/// <param name="UpgradeTask"></param>
+		/// <returns></returns>
+		static bool TryUnpackUpgradeTask(Any Payload, out UpgradeTask UpgradeTask)
+		{
+			if (Payload.TypeUrl == "type.googleapis.com/UpgradeTask" || Payload.TypeUrl == "type.googleapis.com/Horde.UpgradeTask")
+			{
+				UpgradeTask = UpgradeTask.Parser.ParseFrom(Payload.Value);
+				return true;
+			}
+			else
+			{
+				UpgradeTask = null!;
+				return false;
+			}
+		}
+
 		/// <summary>
 		/// Terminates any processes which are still running under the given directory
 		/// </summary>
@@ -901,6 +921,8 @@ namespace HordeAgent.Services
 		/// <returns></returns>
 		internal async Task<LeaseResult> ComputeAsync(IRpcConnection RpcConnection, string LeaseId, ComputeTaskMessage ComputeTask, ILogger Logger, CancellationToken CancellationToken)
 		{
+			Logger.LogInformation("Starting compute task (lease {LeaseId})", LeaseId);
+
 			DateTimeOffset ActionTaskStartTime = DateTimeOffset.UtcNow;
 			using (IRpcClientRef Client = await RpcConnection.GetClientRef(new RpcContext(), CancellationToken))
 			{
@@ -910,6 +932,7 @@ namespace HordeAgent.Services
 				ComputeTaskExecutor Executor = new ComputeTaskExecutor(Client.Channel, Logger);
 				try
 				{
+					Logger.LogInformation("Running executor");
 					ComputeTaskResultMessage Result = await Executor.ExecuteAsync(LeaseId, ComputeTask, LeaseDir, CancellationToken);
 					return new LeaseResult(Result.ToByteArray());
 				}
