Horde: Special case for upgrading agents on bad version.

diff --git a/Engine/Source/Programs/Horde/HordeServer/Tasks/Impl/UpgradeTaskSource.cs b/Engine/Source/Programs/Horde/HordeServer/Tasks/Impl/UpgradeTaskSource.cs
--- a/Engine/Source/Programs/Horde/HordeServer/Tasks/Impl/UpgradeTaskSource.cs
+++ b/Engine/Source/Programs/Horde/HordeServer/Tasks/Impl/UpgradeTaskSource.cs
@@ -46,7 +46,18 @@ namespace HordeServer.Tasks.Impl
 					Task.SoftwareId = RequiredVersion.ToString();
 					Task.LogId = LogFile.Id.ToString();
 
-					byte[] Payload = Any.Pack(Task).ToByteArray();
+					byte[] Payload;
+					if (Agent.Version == "5.0.0-17448746")
+					{
+						Any Any = new Any();
+						Any.TypeUrl = "type.googleapis.com/Horde.UpgradeTask";
+						Any.Value = Task.ToByteString();
+						Payload = Any.ToByteArray();
+					}
+					else
+					{
+						Payload = Any.Pack(Task).ToByteArray();
+					}
 
 					Lease = new AgentLease(ObjectId.GenerateNewId(), $"Upgrade to {RequiredVersion}", null, null, LogFile.Id, LeaseState.Pending, Payload, new AgentRequirements(), null);
 				}
