Horde: Don't delay job view getting timing info
#rnx

diff --git a/Engine/Source/Programs/Horde/HordeDashboard/src/backend/JobDetails.ts b/Engine/Source/Programs/Horde/HordeDashboard/src/backend/JobDetails.ts
--- a/Engine/Source/Programs/Horde/HordeDashboard/src/backend/JobDetails.ts
+++ b/Engine/Source/Programs/Horde/HordeDashboard/src/backend/JobDetails.ts
@@ -726,9 +726,14 @@ export class JobDetails {
 
             requests.push(backend.getJob(this.id));
 
-            const doTiming = !this.updateTimingInfo;
-            this.updateTimingInfo++;
-            this.updateTimingInfo %= 3;
+            let doTiming = false;
+
+            // job timing is very expensive, don't delay load for it until we have job data
+            if (this.jobdata) {
+                doTiming = !this.updateTimingInfo;
+                this.updateTimingInfo++;
+                this.updateTimingInfo %= 3;
+            }
 
             if (doTiming) {
                 requests.push(backend.getJobTiming(this.id));
