Horde - Adding Linux shutdown and reboot functionality to the Horde Agent. This change imports the libc.so.6 library to invoke the sync() and reboot() syscalls from the Horde agent.

diff --git a/Engine/Source/Programs/Horde/HordeAgent/Utility/Shutdown.cs b/Engine/Source/Programs/Horde/HordeAgent/Utility/Shutdown.cs
--- a/Engine/Source/Programs/Horde/HordeAgent/Utility/Shutdown.cs
+++ b/Engine/Source/Programs/Horde/HordeAgent/Utility/Shutdown.cs
@@ -64,6 +64,17 @@ namespace HordeAgent.Utility
 
 		const int SE_PRIVILEGE_ENABLED = 0x00000002;
 
+		// Linux syscalls and reboot flags
+		[DllImport("libc.so.6", SetLastError = true)]
+		private static extern void sync();
+
+		[DllImport("libc.so.6", SetLastError = true)]
+		private static extern int reboot(int flag);
+
+		const int RB_AUTOBOOT = 0x1234567;
+
+		const int RB_POWER_OFF = 0x4321fedc;
+
 		/// <summary>
 		/// Initiate a shutdown operation
 		/// </summary>
@@ -84,7 +95,7 @@ namespace HordeAgent.Utility
 				}
 
 				IntPtr TokenHandle;
-				if(!OpenProcessToken(Process.GetCurrentProcess().Handle, TOKEN_ADJUST_PRIVILEGES | TOKEN_QUERY, out TokenHandle))
+				if (!OpenProcessToken(Process.GetCurrentProcess().Handle, TOKEN_ADJUST_PRIVILEGES | TOKEN_QUERY, out TokenHandle))
 				{
 					Logger.LogError("OpenProcessToken() failed (code 0x{Code:x8})", Marshal.GetLastWin32Error());
 					return false;
@@ -117,6 +128,29 @@ namespace HordeAgent.Utility
 				}
 				return true;
 			}
+			else if (RuntimeInformation.IsOSPlatform(OSPlatform.Linux))
+			{
+				sync(); // Commit filesystem caches to disk.
+				if (bRestartAfterShutdown)
+				{
+					Logger.LogInformation("Triggering restart");
+					if (reboot(RB_AUTOBOOT) != 0)
+					{
+						Logger.LogError("Restart failed ({0})", Marshal.GetLastWin32Error());
+						return false;
+					}
+				}
+				else
+				{
+					Logger.LogInformation("Triggering shutdown");
+					if (reboot(RB_POWER_OFF) != 0)
+					{
+						Logger.LogError("Shutdown failed ({0})", Marshal.GetLastWin32Error());
+						return false;
+					}
+				}
+				return true;
+			}
 			else
 			{
 				Logger.LogError("Shutdown is not implemented on this platform");
