Horde: Fix exception due to missing logger argument.

diff --git a/Engine/Source/Programs/Horde/HordeServer/Compute/Impl/RedisTaskScheduler.cs b/Engine/Source/Programs/Horde/HordeServer/Compute/Impl/RedisTaskScheduler.cs
--- a/Engine/Source/Programs/Horde/HordeServer/Compute/Impl/RedisTaskScheduler.cs
+++ b/Engine/Source/Programs/Horde/HordeServer/Compute/Impl/RedisTaskScheduler.cs
@@ -553,7 +553,7 @@ namespace HordeServer.Compute.Impl
 					}
 					catch (Exception Ex)
 					{
-						Logger.LogError(Ex, "Unable to parse condition '{Condition}': {Message}", Ex.Message);
+						Logger.LogError(Ex, "Unable to parse condition '{Condition}': {Message}", Info.Requirements.Condition.ToString(), Ex.Message);
 					}
 				}
 
