Horde: Add another agent version that uses the incorrect message name for upgrade tasks.

diff --git a/Engine/Source/Programs/Horde/HordeServer/Tasks/Impl/UpgradeTaskSource.cs b/Engine/Source/Programs/Horde/HordeServer/Tasks/Impl/UpgradeTaskSource.cs
--- a/Engine/Source/Programs/Horde/HordeServer/Tasks/Impl/UpgradeTaskSource.cs
+++ b/Engine/Source/Programs/Horde/HordeServer/Tasks/Impl/UpgradeTaskSource.cs
@@ -47,7 +47,7 @@ namespace HordeServer.Tasks.Impl
 					Task.LogId = LogFile.Id.ToString();
 
 					byte[] Payload;
-					if (Agent.Version == "5.0.0-17448746")
+					if (Agent.Version == "5.0.0-17425336" || Agent.Version == "5.0.0-17448746")
 					{
 						Any Any = new Any();
 						Any.TypeUrl = "type.googleapis.com/Horde.UpgradeTask";
