Horde: Prevent cancellation of compute listeners throwing a TaskCancelledException.

diff --git a/Engine/Source/Programs/Horde/HordeServer/Compute/Impl/ComputeService.cs b/Engine/Source/Programs/Horde/HordeServer/Compute/Impl/ComputeService.cs
--- a/Engine/Source/Programs/Horde/HordeServer/Compute/Impl/ComputeService.cs
+++ b/Engine/Source/Programs/Horde/HordeServer/Compute/Impl/ComputeService.cs
@@ -168,7 +168,8 @@ namespace HordeServer.Compute.Impl
 		{
 			if (!Agent.Enabled || Agent.Id.ToString().StartsWith("REMOTE-EXEC-", StringComparison.OrdinalIgnoreCase))
 			{
-				await Task.Delay(-1, CancellationToken);
+				using CancellationTask Task = new CancellationTask(CancellationToken);
+				await Task.Task;
 				return null;
 			}
 
