Horde: Increase overall site width and adjust some dark mode colors
#rnx

diff --git a/Engine/Source/Programs/Horde/HordeDashboard/src/backend/Dashboard.ts b/Engine/Source/Programs/Horde/HordeDashboard/src/backend/Dashboard.ts
--- a/Engine/Source/Programs/Horde/HordeDashboard/src/backend/Dashboard.ts
+++ b/Engine/Source/Programs/Horde/HordeDashboard/src/backend/Dashboard.ts
@@ -168,12 +168,26 @@ export class Dashboard {
 
     get darktheme(): boolean {
 
+        if (this.preferences.get(DashboardPreference.Darktheme) === undefined) {
+            // avoid intial flash when loading into site before backend is initialized
+            let dark = localStorage?.getItem("horde_darktheme") === "true";
+            return dark;
+        }
+
         return this.preferences.get(DashboardPreference.Darktheme) === 'true';
 
     }
 
     setDarkTheme(value: boolean | undefined) {
+
         this.setPreference(DashboardPreference.Darktheme, value ? "true" : "false");
+
+        if (value) {
+            localStorage?.setItem("horde_darktheme", "true");
+        } else {
+            localStorage?.removeItem("horde_darktheme");
+        }
+
         this.setUpdated();
     }
 
diff --git a/Engine/Source/Programs/Horde/HordeDashboard/src/styles/colors.ts b/Engine/Source/Programs/Horde/HordeDashboard/src/styles/colors.ts
--- a/Engine/Source/Programs/Horde/HordeDashboard/src/styles/colors.ts
+++ b/Engine/Source/Programs/Horde/HordeDashboard/src/styles/colors.ts
@@ -20,7 +20,7 @@ export enum StatusColor {
 
 export const defaultStatusColors = new Map<StatusColor, string>([
     [StatusColor.Success, "#52C705"],
-    [StatusColor.Warnings, "#EDC74A"],
+    [StatusColor.Warnings, dashboard.darktheme ? theme.palette.yellow : "#EDC74A"],
     [StatusColor.Failure, "#DE4522"],
     [StatusColor.Running, theme.palette.blueLight],
     [StatusColor.Waiting, "#A19F9D"],
