Horde: Add null check in blob store RPC.

diff --git a/Engine/Source/Programs/Horde/HordeServer/Storage/Collections/BlobStoreRpc.cs b/Engine/Source/Programs/Horde/HordeServer/Storage/Collections/BlobStoreRpc.cs
--- a/Engine/Source/Programs/Horde/HordeServer/Storage/Collections/BlobStoreRpc.cs
+++ b/Engine/Source/Programs/Horde/HordeServer/Storage/Collections/BlobStoreRpc.cs
@@ -58,7 +58,10 @@ namespace HordeServer.Storage.Collections
 				byte[]? Data = await Tasks[Idx];
 
 				GetBlobResponse Blob = new GetBlobResponse();
-				Blob.Data = ByteString.CopyFrom(Data);
+				if (Data != null)
+				{
+					Blob.Data = ByteString.CopyFrom(Data);
+				}
 				Blob.Exists = Data != null;
 				Response.Blobs.Add(Blob);
 			}
