Horde: Enable dark mode as default
#rnx

diff --git a/Engine/Source/Programs/Horde/HordeDashboard/src/backend/Dashboard.ts b/Engine/Source/Programs/Horde/HordeDashboard/src/backend/Dashboard.ts
--- a/Engine/Source/Programs/Horde/HordeDashboard/src/backend/Dashboard.ts
+++ b/Engine/Source/Programs/Horde/HordeDashboard/src/backend/Dashboard.ts
@@ -6,7 +6,7 @@ import { UserClaim, DashboardPreference, GetUserResponse } from './Api';
 
 export enum WebBrowser {
     Chromium = "Chromium",
-    Safari = "Safari", 
+    Safari = "Safari",
     Other = "Other"
 }
 
@@ -91,21 +91,21 @@ export class Dashboard {
 
     }
 
-    get browser() : WebBrowser {
+    get browser(): WebBrowser {
 
         const agent = window.navigator.userAgent.toLowerCase();
-        
+
         switch (true) {
             case agent.indexOf("edge") > -1: return WebBrowser.Other;
             case agent.indexOf("edg") > -1: return WebBrowser.Chromium;
-            case agent.indexOf("opr") > -1: return WebBrowser.Other ;
+            case agent.indexOf("opr") > -1: return WebBrowser.Other;
             case agent.indexOf("chrome") > -1 && !!(window as any).chrome: return WebBrowser.Chromium;
             case agent.indexOf("trident") > -1: return WebBrowser.Other;
             case agent.indexOf("firefox") > -1: return WebBrowser.Other;
             case agent.indexOf("safari") > -1: return WebBrowser.Safari;
             default: return WebBrowser.Other;
         }
-    
+
     }
 
     get userId(): string {
@@ -168,17 +168,23 @@ export class Dashboard {
 
     get darktheme(): boolean {
 
-        if (this.preferences.get(DashboardPreference.Darktheme) === undefined) {
+        if (!this.available) {
             // avoid intial flash when loading into site before backend is initialized
-            let dark = localStorage?.getItem("horde_darktheme") === "true";
-            return dark;
+            return localStorage?.getItem("horde_darktheme") === "true";
+        }
+
+        const pref = this.preferences.get(DashboardPreference.Darktheme);
+
+        if (pref !== "true" && pref !== "false") {
+            console.log("setting dark theme");
+            this.setDarkTheme(true, false);
         }
 
         return this.preferences.get(DashboardPreference.Darktheme) === 'true';
 
     }
 
-    setDarkTheme(value: boolean | undefined) {
+    setDarkTheme(value: boolean | undefined, update: boolean = true) {
 
         this.setPreference(DashboardPreference.Darktheme, value ? "true" : "false");
 
@@ -188,7 +194,10 @@ export class Dashboard {
             localStorage?.removeItem("horde_darktheme");
         }
 
-        this.setUpdated();
+        if (update) {
+            this.setUpdated();
+        }
+
     }
 
     setDisplayUTC(value: boolean | undefined) {
@@ -198,13 +207,13 @@ export class Dashboard {
     private hasLoggedLocalCache = false;
 
     get localCache(): boolean {
-        
+
         if (this.browser === WebBrowser.Chromium) {
             if (!this.hasLoggedLocalCache) {
                 this.hasLoggedLocalCache = true;
                 console.log("Chromium browser detected, local caching is enabled")
             }
-            
+
             return true;
         }
 
@@ -376,7 +385,7 @@ export class Dashboard {
         return success;
 
     }
-    
+
 
     requestLogout = false;
 
