Horde: Adjust default status color brightness by 20% for dark mode
#rnx

diff --git a/Engine/Source/Programs/Horde/HordeDashboard/src/styles/colors.ts b/Engine/Source/Programs/Horde/HordeDashboard/src/styles/colors.ts
--- a/Engine/Source/Programs/Horde/HordeDashboard/src/styles/colors.ts
+++ b/Engine/Source/Programs/Horde/HordeDashboard/src/styles/colors.ts
@@ -19,10 +19,10 @@ export enum StatusColor {
 }
 
 export const defaultStatusColors = new Map<StatusColor, string>([
-    [StatusColor.Success, "#52C705"],
-    [StatusColor.Warnings, dashboard.darktheme ? theme.palette.yellow : "#EDC74A"],
-    [StatusColor.Failure, "#DE4522"],
-    [StatusColor.Running, theme.palette.blueLight],
+    [StatusColor.Success, dashboard.darktheme ? "#419e03" : "#52C705"],
+    [StatusColor.Warnings, "#EDC74A"],
+    [StatusColor.Failure, dashboard.darktheme ? "#b1361a" : "#DE4522"],
+    [StatusColor.Running, dashboard.darktheme ? "#0d809f" : theme.palette.blueLight],
     [StatusColor.Waiting, "#A19F9D"],
     [StatusColor.Ready, "#A19F9D"],
     [StatusColor.Skipped, "#F3F2F1"],
