Horde: Adjust waiting/skipped/aborted colors
#rnx

diff --git a/Engine/Source/Programs/Horde/HordeDashboard/src/styles/colors.ts b/Engine/Source/Programs/Horde/HordeDashboard/src/styles/colors.ts
--- a/Engine/Source/Programs/Horde/HordeDashboard/src/styles/colors.ts
+++ b/Engine/Source/Programs/Horde/HordeDashboard/src/styles/colors.ts
@@ -23,10 +23,10 @@ export const defaultStatusColors = new Map<StatusColor, string>([
     [StatusColor.Warnings, "#EDC74A"],
     [StatusColor.Failure, dashboard.darktheme ? "#b1361a" : "#DE4522"],
     [StatusColor.Running, dashboard.darktheme ? "#0d809f" : theme.palette.blueLight],
-    [StatusColor.Waiting, "#A19F9D"],
-    [StatusColor.Ready, "#A19F9D"],
-    [StatusColor.Skipped, "#F3F2F1"],
-    [StatusColor.Aborted, "#F3F2F1"],
+    [StatusColor.Waiting, dashboard.darktheme ? "#807e7b" : "#A19F9D"],
+    [StatusColor.Ready, dashboard.darktheme ? "#807e7b" : "#A19F9D"],
+    [StatusColor.Skipped, dashboard.darktheme ? "#c4bfbb" :"#F3F2F1"],
+    [StatusColor.Aborted, dashboard.darktheme ? "#c4bfbb" :"#F3F2F1"],
     [StatusColor.Unspecified, "#637087"]
 ]);
 
