Horde: Fix periodic failures to parse compact binary request objects due to not awaiting stream copies.

diff --git a/Engine/Source/Programs/Horde/HordeServer/Utilities/CompactBinaryFormatter.cs b/Engine/Source/Programs/Horde/HordeServer/Utilities/CompactBinaryFormatter.cs
--- a/Engine/Source/Programs/Horde/HordeServer/Utilities/CompactBinaryFormatter.cs
+++ b/Engine/Source/Programs/Horde/HordeServer/Utilities/CompactBinaryFormatter.cs
@@ -12,6 +12,7 @@ using EpicGames.Serialization.Converters;
 using System.Reflection;
 using System.Net.Mime;
 using EpicGames.Core;
+using Microsoft.Extensions.Primitives;
 
 namespace HordeServer.Utilities
 {
@@ -46,13 +47,13 @@ namespace HordeServer.Utilities
 		}
 
 		/// <inheritdoc/>
-		public override Task<InputFormatterResult> ReadRequestBodyAsync(InputFormatterContext Context)
+		public override async Task<InputFormatterResult> ReadRequestBodyAsync(InputFormatterContext Context)
 		{
 			// Buffer the data into an array
 			byte[] Data;
 			using (MemoryStream Stream = new MemoryStream())
 			{
-				Context.HttpContext.Request.Body.CopyToAsync(Stream);
+				await Context.HttpContext.Request.Body.CopyToAsync(Stream);
 				Data = Stream.ToArray();
 			}
 
@@ -65,9 +66,16 @@ namespace HordeServer.Utilities
 			catch (Exception Ex)
 			{
 				Serilog.Log.Logger.Error("Unable to parse compact binary: {Dump}", FormatHexDump(Data, 256));
+				foreach ((string Name, StringValues Values) in Context.HttpContext.Request.Headers)
+				{
+					foreach (string Value in Values)
+					{
+						Serilog.Log.Logger.Information("Header {Name}: {Value}", Name, Value);
+					}
+				}
 				throw new Exception($"Unable to parse compact binary request: {FormatHexDump(Data, 256)}", Ex);
 			}
-			return InputFormatterResult.SuccessAsync(CbSerializer.Deserialize(new CbField(Data), Context.ModelType));
+			return await InputFormatterResult.SuccessAsync(CbSerializer.Deserialize(new CbField(Data), Context.ModelType));
 		}
 
 		static string FormatHexDump(byte[] Data, int MaxLength)
