Horde: Retry fetching missing requirements objects after 10 seconds.

diff --git a/Engine/Source/Programs/Horde/HordeServer/Compute/Impl/RedisTaskScheduler.cs b/Engine/Source/Programs/Horde/HordeServer/Compute/Impl/RedisTaskScheduler.cs
--- a/Engine/Source/Programs/Horde/HordeServer/Compute/Impl/RedisTaskScheduler.cs
+++ b/Engine/Source/Programs/Horde/HordeServer/Compute/Impl/RedisTaskScheduler.cs
@@ -559,7 +559,14 @@ namespace HordeServer.Compute.Impl
 
 				using (ICacheEntry Entry = RequirementsCache.CreateEntry(RequirementsHash))
 				{
-					Entry.SetSlidingExpiration(TimeSpan.FromMinutes(10.0));
+					if (Info.Requirements == null)
+					{
+						Entry.SetAbsoluteExpiration(TimeSpan.FromSeconds(10.0));
+					}
+					else
+					{
+						Entry.SetSlidingExpiration(TimeSpan.FromMinutes(10.0));
+					}
 					Entry.SetValue(Info);
 				}
 			}
