Horde: Add an error response path that returns additional details about the problem.

diff --git a/Engine/Source/Programs/Horde/HordeServer/Controllers/ErrorController.cs b/Engine/Source/Programs/Horde/HordeServer/Controllers/ErrorController.cs
new file mode 100644
--- /dev/null
+++ b/Engine/Source/Programs/Horde/HordeServer/Controllers/ErrorController.cs
@@ -0,0 +1,29 @@
+using Microsoft.AspNetCore.Diagnostics;
+using Microsoft.AspNetCore.Mvc;
+
+namespace HordeServer.Controllers
+{
+	/// <summary>
+	/// Controller which returns exception metadata to the caller
+	/// </summary>
+	[ApiController]
+	[ApiExplorerSettings(IgnoreApi = true)]
+	public class ErrorController : ControllerBase
+	{
+		/// <summary>
+		/// Handle an exception and generate a problem response
+		/// </summary>
+		/// <returns></returns>
+		[Route("/api/v1/error")]
+		public IActionResult Error()
+		{
+			IExceptionHandlerFeature? Context = HttpContext.Features.Get<IExceptionHandlerFeature>();
+			if (Context == null)
+			{
+				return NoContent();
+			}
+
+			return Problem(detail: Context.Error.StackTrace, title: Context.Error.Message);
+		}
+	}
+}
diff --git a/Engine/Source/Programs/Horde/HordeServer/Startup.cs b/Engine/Source/Programs/Horde/HordeServer/Startup.cs
--- a/Engine/Source/Programs/Horde/HordeServer/Startup.cs
+++ b/Engine/Source/Programs/Horde/HordeServer/Startup.cs
@@ -736,11 +736,7 @@ namespace HordeServer
 		public static void Configure(IApplicationBuilder App, IWebHostEnvironment Env, Microsoft.Extensions.Hosting.IHostApplicationLifetime Lifetime, IOptions<ServerSettings> Settings)
 		{
 			App.UseForwardedHeaders();
-
-			if (Env.IsDevelopment())
-			{
-				App.UseDeveloperExceptionPage();
-			}
+			App.UseExceptionHandler("/api/v1/error");
 
 			if (Settings.Value.CorsEnabled)
 			{
