Horde: Add missing copyright notice.

diff --git a/Engine/Source/Programs/Horde/HordeServer/Controllers/ErrorController.cs b/Engine/Source/Programs/Horde/HordeServer/Controllers/ErrorController.cs
--- a/Engine/Source/Programs/Horde/HordeServer/Controllers/ErrorController.cs
+++ b/Engine/Source/Programs/Horde/HordeServer/Controllers/ErrorController.cs
@@ -1,3 +1,5 @@
+// Copyright Epic Games, Inc. All Rights Reserved.
+
 using Microsoft.AspNetCore.Diagnostics;
 using Microsoft.AspNetCore.Mvc;
 
