Horde: Expire tasks added to queues that are not executed after 10 minutes.

diff --git a/Engine/Source/Programs/Horde/HordeServer/Compute/Impl/ComputeService.cs b/Engine/Source/Programs/Horde/HordeServer/Compute/Impl/ComputeService.cs
--- a/Engine/Source/Programs/Horde/HordeServer/Compute/Impl/ComputeService.cs
+++ b/Engine/Source/Programs/Horde/HordeServer/Compute/Impl/ComputeService.cs
@@ -108,13 +108,13 @@ namespace HordeServer.Compute.Impl
 		{
 		}
 
-		public ComputeTaskStatus(CbObjectAttachment TaskHash, ComputeTaskState State, AgentId AgentId, ObjectId LeaseId)
+		public ComputeTaskStatus(CbObjectAttachment TaskHash, ComputeTaskState State, AgentId? AgentId, ObjectId? LeaseId)
 		{
 			this.TaskHash = TaskHash;
 			this.Time = DateTime.UtcNow;
 			this.State = State;
-			this.AgentId = AgentId.ToString();
-			this.LeaseIdBytes = LeaseId.ToByteArray();
+			this.AgentId = (AgentId == null)? Utf8String.Empty : AgentId.Value.ToString();
+			this.LeaseIdBytes = LeaseId?.ToByteArray();
 		}
 	}
 
@@ -129,6 +129,7 @@ namespace HordeServer.Compute.Impl
 
 		ITaskScheduler<ComputeTaskInfo> TaskScheduler;
 		RedisMessageQueue<ComputeTaskStatus> MessageQueue;
+		BackgroundTick ExpireTasksTicker;
 		ILogger Logger;
 
 		/// <summary>
@@ -141,13 +142,40 @@ namespace HordeServer.Compute.Impl
 		{
 			this.TaskScheduler = new RedisTaskScheduler<ComputeTaskInfo>(Redis, ObjectCollection, DefaultNamespaceId, "compute/tasks/", Logger);
 			this.MessageQueue = new RedisMessageQueue<ComputeTaskStatus>(Redis, "compute/messages/");
+			this.ExpireTasksTicker = new BackgroundTick(ExpireTasksAsync, TimeSpan.FromMinutes(2.0), Logger);
 			this.Logger = Logger;
 		}
 
+		/// <summary>
+		/// Expire tasks that are in inactive queues (ie. no machines can execute them)
+		/// </summary>
+		/// <param name="CancellationToken"></param>
+		/// <returns></returns>
+		async Task ExpireTasksAsync(CancellationToken CancellationToken)
+		{
+			List<IoHash> RequirementsHashes = await TaskScheduler.GetInactiveQueuesAsync();
+			foreach (IoHash RequirementsHash in RequirementsHashes)
+			{
+				for (; ; )
+				{
+					ComputeTaskInfo? ComputeTask = await TaskScheduler.DequeueAsync(RequirementsHash);
+					if (ComputeTask == null)
+					{
+						break;
+					}
+
+					ComputeTaskStatus Status = new ComputeTaskStatus(ComputeTask.TaskHash, ComputeTaskState.Complete, null, null);
+					Logger.LogInformation("Compute task cancelled (task: {TaskHash}, channel: {ChannelId})", ComputeTask.TaskHash, ComputeTask.ChannelId);
+					await MessageQueue.PostAsync(ComputeTask.ChannelId.ToString(), Status);
+				}
+			}
+		}
+
 		/// <inheritdoc/>
 		public void Dispose()
 		{
 			MessageQueue.Dispose();
+			ExpireTasksTicker.Dispose();
 		}
 
 		/// <inheritdoc/>
@@ -186,7 +214,7 @@ namespace HordeServer.Compute.Impl
 				byte[] Payload = Any.Pack(ComputeTask).ToByteArray();
 
 				AgentLease Lease = new AgentLease(ObjectId.GenerateNewId(), LeaseName, null, null, null, LeaseState.Pending, Payload, new AgentRequirements(), null);
-				Logger.LogDebug("Created lease {LeaseId} for channel {ChannelId} task {TaskHash} req {RequirementsHash}", Lease.Id, ComputeTask.ChannelId, (CbObjectAttachment)ComputeTask.TaskHash, ComputeTask.RequirementsHash);
+				Logger.LogDebug("Created lease {LeaseId} for channel {ChannelId} task {TaskHash} req {RequirementsHash}", Lease.Id, ComputeTask.ChannelId, (CbObjectAttachment)ComputeTask.TaskHash, (CbObjectAttachment)ComputeTask.RequirementsHash);
 				return Lease;
 			}
 			return null;
diff --git a/Engine/Source/Programs/Horde/HordeServer/Compute/Impl/RedisTaskScheduler.cs b/Engine/Source/Programs/Horde/HordeServer/Compute/Impl/RedisTaskScheduler.cs
--- a/Engine/Source/Programs/Horde/HordeServer/Compute/Impl/RedisTaskScheduler.cs
+++ b/Engine/Source/Programs/Horde/HordeServer/Compute/Impl/RedisTaskScheduler.cs
@@ -55,7 +55,7 @@ namespace HordeServer.Compute.Impl
 	/// <summary>
 	/// Interface for a queue of compute operations, each of which may have different requirements for the executing machines.
 	/// </summary>
-	interface ITaskScheduler<T>
+	interface ITaskScheduler<T> where T : class
 	{
 		/// <summary>
 		/// Adds a task to the queue
@@ -78,6 +78,19 @@ namespace HordeServer.Compute.Impl
 		/// <param name="Token">Cancellation token for the operation. Will return a null entry rather than throwing an exception.</param>
 		/// <returns>Information about the task to be executed</returns>
 		Task<TaskSchedulerEntry<T>?> DequeueAsync(IAgent Agent, CancellationToken Token = default);
+
+		/// <summary>
+		/// Dequeue any item from the given queue
+		/// </summary>
+		/// <param name="RequirementsHash">Hash of the requirements object</param>
+		/// <returns>Information about the task to be executed</returns>
+		Task<T?> DequeueAsync(IoHash RequirementsHash);
+
+		/// <summary>
+		/// Gets hashes of all the inactive task queues
+		/// </summary>
+		/// <returns></returns>
+		Task<List<IoHash>> GetInactiveQueuesAsync();
 	}
 
 	/// <summary>
