Horde: Do not disable schedules by default; just disable when launching from VS.

diff --git a/Engine/Source/Programs/Horde/HordeServer/Properties/launchSettings.json b/Engine/Source/Programs/Horde/HordeServer/Properties/launchSettings.json
--- a/Engine/Source/Programs/Horde/HordeServer/Properties/launchSettings.json
+++ b/Engine/Source/Programs/Horde/HordeServer/Properties/launchSettings.json
@@ -2,7 +2,9 @@
   "profiles": {
     "HordeServer": {
       "commandName": "Project",
-      "commandLineArgs": ""
+      "environmentVariables": {
+        "Horde__DisableSchedules": "true"
+      }
     }
   }
 }
\ No newline at end of file
diff --git a/Engine/Source/Programs/Horde/HordeServer/ServerSettings.cs b/Engine/Source/Programs/Horde/HordeServer/ServerSettings.cs
--- a/Engine/Source/Programs/Horde/HordeServer/ServerSettings.cs
+++ b/Engine/Source/Programs/Horde/HordeServer/ServerSettings.cs
@@ -349,9 +349,9 @@ namespace HordeServer
 		public FleetManagerType FleetManager { get; set; } = FleetManagerType.None;
 
 		/// <summary>
-		/// Whether to run scheduled jobs. Not wanted for development.
+		/// Whether to run scheduled jobs.
 		/// </summary>
-		public bool DisableSchedules { get; set; } = true;
+		public bool DisableSchedules { get; set; }
 
 		/// <summary>
 		/// Timezone for evaluating schedules
