Horde: Revert timing endpoint mitigation code
#rnx

diff --git a/Engine/Source/Programs/Horde/HordeDashboard/src/backend/JobDetails.ts b/Engine/Source/Programs/Horde/HordeDashboard/src/backend/JobDetails.ts
--- a/Engine/Source/Programs/Horde/HordeDashboard/src/backend/JobDetails.ts
+++ b/Engine/Source/Programs/Horde/HordeDashboard/src/backend/JobDetails.ts
@@ -726,15 +726,10 @@ export class JobDetails {
 
             requests.push(backend.getJob(this.id));
 
-            let doTiming = false;
-
-            // job timing is very expensive, don't delay load for it until we have job data
-            if (this.jobdata) {
-                doTiming = !this.updateTimingInfo;
-                this.updateTimingInfo++;
-                this.updateTimingInfo %= 3;
-            }
-
+            const doTiming = !this.updateTimingInfo;
+            this.updateTimingInfo++;
+            this.updateTimingInfo %= 3;
+            
             if (doTiming) {
                 requests.push(backend.getJobTiming(this.id));
             }
