Horde: Updates and fixes for dark mode
#rnx

diff --git a/Engine/Source/Programs/Horde/HordeDashboard/src/backend/Dashboard.ts b/Engine/Source/Programs/Horde/HordeDashboard/src/backend/Dashboard.ts
--- a/Engine/Source/Programs/Horde/HordeDashboard/src/backend/Dashboard.ts
+++ b/Engine/Source/Programs/Horde/HordeDashboard/src/backend/Dashboard.ts
@@ -194,6 +194,11 @@ export class Dashboard {
             localStorage?.removeItem("horde_darktheme");
         }
 
+        this.setStatusColor(DashboardPreference.ColorWarning, undefined);
+        this.setStatusColor(DashboardPreference.ColorError, undefined);
+        this.setStatusColor(DashboardPreference.ColorRunning, undefined);
+        this.setStatusColor(DashboardPreference.ColorSuccess, undefined);
+
         if (update) {
             this.setDarkThemeUpdated();
         }
diff --git a/Engine/Source/Programs/Horde/HordeDashboard/src/styles/colors.ts b/Engine/Source/Programs/Horde/HordeDashboard/src/styles/colors.ts
--- a/Engine/Source/Programs/Horde/HordeDashboard/src/styles/colors.ts
+++ b/Engine/Source/Programs/Horde/HordeDashboard/src/styles/colors.ts
@@ -18,21 +18,28 @@ export enum StatusColor {
     Unspecified
 }
 
-export const defaultStatusColors = new Map<StatusColor, string>([
-    [StatusColor.Success, dashboard.darktheme ? "#419e03" : "#52C705"],
-    [StatusColor.Warnings, "#EDC74A"],
-    [StatusColor.Failure, dashboard.darktheme ? "#b1361a" : "#DE4522"],
-    [StatusColor.Running, dashboard.darktheme ? "#0d809f" : theme.palette.blueLight],
-    [StatusColor.Waiting, dashboard.darktheme ? "#807e7b" : "#A19F9D"],
-    [StatusColor.Ready, dashboard.darktheme ? "#807e7b" : "#A19F9D"],
-    [StatusColor.Skipped, dashboard.darktheme ? "#c4bfbb" :"#F3F2F1"],
-    [StatusColor.Aborted, dashboard.darktheme ? "#c4bfbb" :"#F3F2F1"],
-    [StatusColor.Unspecified, "#637087"]
-]);
+export const getDefaultStatusColors = (): Map<StatusColor, string> => {
+
+    return new Map<StatusColor, string>([
+        [StatusColor.Success, dashboard.darktheme ? "#30690a" : "#52C705"],
+        [StatusColor.Warnings, dashboard.darktheme ? "#a1862d" : "#EDC74A"],
+        [StatusColor.Failure, dashboard.darktheme ? "#8a3522" : "#DE4522"],
+        [StatusColor.Running, dashboard.darktheme ? "#1a6175" : theme.palette.blueLight],
+        [StatusColor.Waiting, dashboard.darktheme ? "#474542" : "#A19F9D"],
+        [StatusColor.Ready, dashboard.darktheme ? "#474542" : "#A19F9D"],
+        [StatusColor.Skipped, dashboard.darktheme ? "#63625c" : "#F3F2F1"],
+        [StatusColor.Aborted, dashboard.darktheme ? "#63625c" : "#F3F2F1"],
+        [StatusColor.Unspecified, "#637087"]
+    ]);
+
+}
+
 
 
 export const getStatusColors = () => {
 
+    const defaultStatusColors = getDefaultStatusColors();
+
     const success = dashboard.getPreference(DashboardPreference.ColorSuccess);
     const warning = dashboard.getPreference(DashboardPreference.ColorWarning);
     const error = dashboard.getPreference(DashboardPreference.ColorError);
@@ -53,7 +60,7 @@ export const getStatusColors = () => {
 
 export const getJobStateColor = (state: JobState): string => {
 
-    const colors = getStatusColors();    
+    const colors = getStatusColors();
 
     if (state === JobState.Waiting) {
         return colors.get(StatusColor.Waiting)!;
@@ -68,11 +75,11 @@ export const getJobStateColor = (state: JobState): string => {
 
 }
 
-export const getLabelColor = (state: LabelState | undefined, outcome: LabelOutcome | undefined):{primaryColor: string, secondaryColor?: string} => {
+export const getLabelColor = (state: LabelState | undefined, outcome: LabelOutcome | undefined): { primaryColor: string, secondaryColor?: string } => {
 
     if (!state || !outcome) {
         return {
-            primaryColor: "#000000" 
+            primaryColor: "#000000"
         }
     }
 
@@ -81,22 +88,22 @@ export const getLabelColor = (state: LabelState | undefined, outcome: LabelOutco
     if (state === LabelState.Complete) {
         switch (outcome!) {
             case LabelOutcome.Failure:
-                return { primaryColor: colors.get(StatusColor.Failure)!};        
+                return { primaryColor: colors.get(StatusColor.Failure)! };
             case LabelOutcome.Success:
-                return { primaryColor: colors.get(StatusColor.Success)!};        
+                return { primaryColor: colors.get(StatusColor.Success)! };
             case LabelOutcome.Warnings:
-                return { primaryColor: colors.get(StatusColor.Warnings)!};            
+                return { primaryColor: colors.get(StatusColor.Warnings)! };
             default:
-                return { primaryColor: colors.get(StatusColor.Unspecified)!};        
+                return { primaryColor: colors.get(StatusColor.Unspecified)! };
         }
     }
 
     if (outcome === LabelOutcome.Failure) {
-        return { primaryColor: colors.get(StatusColor.Running)!, secondaryColor: colors.get(StatusColor.Failure)!};        
+        return { primaryColor: colors.get(StatusColor.Running)!, secondaryColor: colors.get(StatusColor.Failure)! };
     }
 
     if (outcome === LabelOutcome.Warnings) {
-        return { primaryColor: colors.get(StatusColor.Running)!, secondaryColor: colors.get(StatusColor.Warnings)!};        
+        return { primaryColor: colors.get(StatusColor.Running)!, secondaryColor: colors.get(StatusColor.Warnings)! };
     }
 
     return {
