Horde: Fix default status colors
#rnx

diff --git a/Engine/Source/Programs/Horde/HordeDashboard/src/styles/colors.ts b/Engine/Source/Programs/Horde/HordeDashboard/src/styles/colors.ts
--- a/Engine/Source/Programs/Horde/HordeDashboard/src/styles/colors.ts
+++ b/Engine/Source/Programs/Horde/HordeDashboard/src/styles/colors.ts
@@ -21,10 +21,10 @@ export enum StatusColor {
 export const getDefaultStatusColors = (): Map<StatusColor, string> => {
 
     return new Map<StatusColor, string>([
-        [StatusColor.Success, dashboard.darktheme ? "#30690a" : "#52C705"],
-        [StatusColor.Warnings, dashboard.darktheme ? "#a1862d" : "#EDC74A"],
-        [StatusColor.Failure, dashboard.darktheme ? "#8a3522" : "#DE4522"],
-        [StatusColor.Running, dashboard.darktheme ? "#1a6175" : theme.palette.blueLight],
+        [StatusColor.Success, dashboard.darktheme ? "#3b7b0a" : "#52C705"],
+        [StatusColor.Warnings, dashboard.darktheme ? "#9a7b18" : "#EDC74A"],
+        [StatusColor.Failure, dashboard.darktheme ? "#882f19" : "#DE4522"],
+        [StatusColor.Running, dashboard.darktheme ? "#146579" : theme.palette.blueLight],
         [StatusColor.Waiting, dashboard.darktheme ? "#474542" : "#A19F9D"],
         [StatusColor.Ready, dashboard.darktheme ? "#474542" : "#A19F9D"],
         [StatusColor.Skipped, dashboard.darktheme ? "#63625c" : "#F3F2F1"],
