Horde: Refactoring and cleanups for dark mode
#rnx

diff --git a/Engine/Source/Programs/Horde/HordeDashboard/src/backend/Dashboard.ts b/Engine/Source/Programs/Horde/HordeDashboard/src/backend/Dashboard.ts
--- a/Engine/Source/Programs/Horde/HordeDashboard/src/backend/Dashboard.ts
+++ b/Engine/Source/Programs/Horde/HordeDashboard/src/backend/Dashboard.ts
@@ -3,6 +3,21 @@
 import { action, observable } from 'mobx';
 import backend from '.';
 import { UserClaim, DashboardPreference, GetUserResponse } from './Api';
+import { getTheme } from "@fluentui/react";
+
+const theme = getTheme();
+
+export enum StatusColor {
+    Success,
+    Warnings,
+    Failure,
+    Waiting,
+    Ready,
+    Skipped,
+    Aborted,
+    Running,
+    Unspecified
+}
 
 export enum WebBrowser {
     Chromium = "Chromium",
@@ -176,15 +191,15 @@ export class Dashboard {
         const pref = this.preferences.get(DashboardPreference.Darktheme);
 
         if (pref !== "true" && pref !== "false") {
-            console.log("setting dark theme");
-            this.setDarkTheme(true, false);
+            console.log("setting dark theme");            
+            this.setDarkTheme(true, false, true);            
         }
 
         return this.preferences.get(DashboardPreference.Darktheme) === 'true';
 
     }
 
-    setDarkTheme(value: boolean | undefined, update: boolean = true) {
+    setDarkTheme(value: boolean | undefined, update: boolean = true, resetColors: boolean = false) {
 
         this.setPreference(DashboardPreference.Darktheme, value ? "true" : "false");
 
@@ -194,13 +209,12 @@ export class Dashboard {
             localStorage?.removeItem("horde_darktheme");
         }
 
-        this.setStatusColor(DashboardPreference.ColorWarning, undefined);
-        this.setStatusColor(DashboardPreference.ColorError, undefined);
-        this.setStatusColor(DashboardPreference.ColorRunning, undefined);
-        this.setStatusColor(DashboardPreference.ColorSuccess, undefined);
+        if (resetColors) {
+            this.resetStatusColors();
+        }
 
         if (update) {
-            this.setDarkThemeUpdated();
+            this.setUpdated();
         }
 
     }
@@ -237,8 +251,37 @@ export class Dashboard {
 
     }
 
+    private resetStatusColors() {
+        this.setStatusColor(DashboardPreference.ColorRunning, undefined);
+        this.setStatusColor(DashboardPreference.ColorWarning, undefined);
+        this.setStatusColor(DashboardPreference.ColorError, undefined);
+        this.setStatusColor(DashboardPreference.ColorSuccess, undefined);
+    }
+
     setStatusColor(pref: DashboardPreference, value: string | undefined) {
 
+        const defaultColors = this.getDefaultStatusColors();
+
+        let defaultColor = "";
+        switch (pref) {
+            case DashboardPreference.ColorRunning:
+                defaultColor = defaultColors.get(StatusColor.Running)!;
+                break;
+            case DashboardPreference.ColorWarning:
+                defaultColor = defaultColors.get(StatusColor.Warnings)!;
+                break;
+            case DashboardPreference.ColorError:
+                defaultColor = defaultColors.get(StatusColor.Failure)!;
+                break;
+            case DashboardPreference.ColorSuccess:
+                defaultColor = defaultColors.get(StatusColor.Success)!;
+                break;
+        }
+
+        if (defaultColor.toLowerCase() === value?.toLowerCase()) {
+            value = undefined;
+        }
+
         if (value && !value.startsWith("#")) {
             console.error("Status preference color must be in hex format with preceding #")
         }
@@ -246,6 +289,49 @@ export class Dashboard {
         this.setPreference(pref, value);
     }
 
+    getDefaultStatusColors = (): Map<StatusColor, string> => {
+
+        const dark = this.darktheme;
+
+        return new Map<StatusColor, string>([
+            [StatusColor.Success, dark ? "#3b7b0a" : "#52C705"],
+            [StatusColor.Warnings, dark ? "#9a7b18" : "#EDC74A"],
+            [StatusColor.Failure, dark ? "#882f19" : "#DE4522"],
+            [StatusColor.Running, dark ? "#146579" : theme.palette.blueLight],
+            [StatusColor.Waiting, dark ? "#474542" : "#A19F9D"],
+            [StatusColor.Ready, dark ? "#474542" : "#A19F9D"],
+            [StatusColor.Skipped, dark ? "#63625c" : "#F3F2F1"],
+            [StatusColor.Aborted, dark ? "#63625c" : "#F3F2F1"],
+            [StatusColor.Unspecified, "#637087"]
+        ]);
+
+    }
+
+
+
+    getStatusColors = () => {
+
+        const defaultStatusColors = this.getDefaultStatusColors();
+
+        const success = this.getPreference(DashboardPreference.ColorSuccess);
+        const warning = this.getPreference(DashboardPreference.ColorWarning);
+        const error = this.getPreference(DashboardPreference.ColorError);
+        const running = this.getPreference(DashboardPreference.ColorRunning);
+
+        return new Map<StatusColor, string>([
+            [StatusColor.Success, success ? success : defaultStatusColors.get(StatusColor.Success)!],
+            [StatusColor.Warnings, warning ? warning : defaultStatusColors.get(StatusColor.Warnings)!],
+            [StatusColor.Failure, error ? error : defaultStatusColors.get(StatusColor.Failure)!],
+            [StatusColor.Running, running ? running : defaultStatusColors.get(StatusColor.Running)!],
+            [StatusColor.Waiting, defaultStatusColors.get(StatusColor.Waiting)!],
+            [StatusColor.Ready, defaultStatusColors.get(StatusColor.Ready)!],
+            [StatusColor.Skipped, defaultStatusColors.get(StatusColor.Skipped)!],
+            [StatusColor.Aborted, defaultStatusColors.get(StatusColor.Aborted)!],
+            [StatusColor.Unspecified, defaultStatusColors.get(StatusColor.Unspecified)!]
+        ]);
+    }
+
+
 
     setDisplay24HourClock(value: boolean | undefined) {
         this.setPreference(DashboardPreference.DisplayClock, value ? "24" : "");
@@ -361,15 +447,6 @@ export class Dashboard {
         this.updated++;
     }
 
-    /** Separate observable for dark theme as is a hook at very top component */
-    @observable
-    darkThemeUpdated: number = 0;
-
-    @action
-    private setDarkThemeUpdated() {
-        this.darkThemeUpdated++;
-    }
-
     get available(): boolean {
         return this.data.id !== "";
     }
diff --git a/Engine/Source/Programs/Horde/HordeDashboard/src/styles/colors.ts b/Engine/Source/Programs/Horde/HordeDashboard/src/styles/colors.ts
--- a/Engine/Source/Programs/Horde/HordeDashboard/src/styles/colors.ts
+++ b/Engine/Source/Programs/Horde/HordeDashboard/src/styles/colors.ts
@@ -1,66 +1,11 @@
 // Copyright Epic Games, Inc. All Rights Reserved.
 
-import { getTheme } from "@fluentui/react";
-import { DashboardPreference, JobState, LabelOutcome, LabelState } from "../backend/Api";
-import dashboard from "../backend/Dashboard";
-
-const theme = getTheme();
-
-export enum StatusColor {
-    Success,
-    Warnings,
-    Failure,
-    Waiting,
-    Ready,
-    Skipped,
-    Aborted,
-    Running,
-    Unspecified
-}
-
-export const getDefaultStatusColors = (): Map<StatusColor, string> => {
-
-    return new Map<StatusColor, string>([
-        [StatusColor.Success, dashboard.darktheme ? "#3b7b0a" : "#52C705"],
-        [StatusColor.Warnings, dashboard.darktheme ? "#9a7b18" : "#EDC74A"],
-        [StatusColor.Failure, dashboard.darktheme ? "#882f19" : "#DE4522"],
-        [StatusColor.Running, dashboard.darktheme ? "#146579" : theme.palette.blueLight],
-        [StatusColor.Waiting, dashboard.darktheme ? "#474542" : "#A19F9D"],
-        [StatusColor.Ready, dashboard.darktheme ? "#474542" : "#A19F9D"],
-        [StatusColor.Skipped, dashboard.darktheme ? "#63625c" : "#F3F2F1"],
-        [StatusColor.Aborted, dashboard.darktheme ? "#63625c" : "#F3F2F1"],
-        [StatusColor.Unspecified, "#637087"]
-    ]);
-
-}
-
-
-
-export const getStatusColors = () => {
-
-    const defaultStatusColors = getDefaultStatusColors();
-
-    const success = dashboard.getPreference(DashboardPreference.ColorSuccess);
-    const warning = dashboard.getPreference(DashboardPreference.ColorWarning);
-    const error = dashboard.getPreference(DashboardPreference.ColorError);
-    const running = dashboard.getPreference(DashboardPreference.ColorRunning);
-
-    return new Map<StatusColor, string>([
-        [StatusColor.Success, success ? success : defaultStatusColors.get(StatusColor.Success)!],
-        [StatusColor.Warnings, warning ? warning : defaultStatusColors.get(StatusColor.Warnings)!],
-        [StatusColor.Failure, error ? error : defaultStatusColors.get(StatusColor.Failure)!],
-        [StatusColor.Running, running ? running : defaultStatusColors.get(StatusColor.Running)!],
-        [StatusColor.Waiting, defaultStatusColors.get(StatusColor.Waiting)!],
-        [StatusColor.Ready, defaultStatusColors.get(StatusColor.Ready)!],
-        [StatusColor.Skipped, defaultStatusColors.get(StatusColor.Skipped)!],
-        [StatusColor.Aborted, defaultStatusColors.get(StatusColor.Aborted)!],
-        [StatusColor.Unspecified, defaultStatusColors.get(StatusColor.Unspecified)!]
-    ]);
-}
+import { JobState, LabelOutcome, LabelState } from "../backend/Api";
+import dashboard, {StatusColor} from "../backend/Dashboard";
 
 export const getJobStateColor = (state: JobState): string => {
 
-    const colors = getStatusColors();
+    const colors = dashboard.getStatusColors();
 
     if (state === JobState.Waiting) {
         return colors.get(StatusColor.Waiting)!;
@@ -83,7 +28,7 @@ export const getLabelColor = (state: LabelState | undefined, outcome: LabelOutco
         }
     }
 
-    const colors = getStatusColors();
+    const colors = dashboard.getStatusColors();
 
     if (state === LabelState.Complete) {
         switch (outcome!) {
