Horde: Remove some ugly compute-specific logic from the task scheduler code. Queues are now treated as generic types, with logic for determining which queues to dequeue from delegated to consumers.

diff --git a/Engine/Source/Programs/Horde/HordeServer/Compute/Impl/ComputeService.cs b/Engine/Source/Programs/Horde/HordeServer/Compute/Impl/ComputeService.cs
--- a/Engine/Source/Programs/Horde/HordeServer/Compute/Impl/ComputeService.cs
+++ b/Engine/Source/Programs/Horde/HordeServer/Compute/Impl/ComputeService.cs
@@ -40,6 +40,7 @@ namespace HordeServer.Compute.Impl
 {
 	using ChannelId = StringId<IComputeChannel>;
 	using NamespaceId = StringId<INamespace>;
+	using Condition = HordeServer.Utilities.Condition;
 
 	/// <summary>
 	/// Information about a particular task
@@ -123,13 +124,21 @@ namespace HordeServer.Compute.Impl
 	/// </summary>
 	class ComputeService : IComputeService, IDisposable
 	{
+		class QueueInfo
+		{
+			public Requirements? Requirements;
+			public Condition? Condition;
+		}
+
 		public MessageDescriptor Descriptor => ComputeTaskMessage.Descriptor;
 
 		public static NamespaceId DefaultNamespaceId { get; } = new NamespaceId("default");
 
-		ITaskScheduler<ComputeTaskInfo> TaskScheduler;
+		IObjectCollection ObjectCollection;
+		ITaskScheduler<IoHash, ComputeTaskInfo> TaskScheduler;
 		RedisMessageQueue<ComputeTaskStatus> MessageQueue;
 		BackgroundTick ExpireTasksTicker;
+		IMemoryCache RequirementsCache;
 		ILogger Logger;
 
 		/// <summary>
@@ -140,9 +149,11 @@ namespace HordeServer.Compute.Impl
 		/// <param name="Logger">The logger instance</param>
 		public ComputeService(IDatabase Redis, IObjectCollection ObjectCollection, ILogger<ComputeService> Logger)
 		{
-			this.TaskScheduler = new RedisTaskScheduler<ComputeTaskInfo>(Redis, ObjectCollection, DefaultNamespaceId, "compute/tasks/", Logger);
+			this.ObjectCollection = ObjectCollection;
+			this.TaskScheduler = new RedisTaskScheduler<IoHash, ComputeTaskInfo>(Redis, "compute/tasks/", Logger);
 			this.MessageQueue = new RedisMessageQueue<ComputeTaskStatus>(Redis, "compute/messages/");
 			this.ExpireTasksTicker = new BackgroundTick(ExpireTasksAsync, TimeSpan.FromMinutes(2.0), Logger);
+			this.RequirementsCache = new MemoryCache(new MemoryCacheOptions());
 			this.Logger = Logger;
 		}
 
@@ -176,6 +187,7 @@ namespace HordeServer.Compute.Impl
 		{
 			MessageQueue.Dispose();
 			ExpireTasksTicker.Dispose();
+			RequirementsCache.Dispose();
 		}
 
 		/// <inheritdoc/>
@@ -186,7 +198,7 @@ namespace HordeServer.Compute.Impl
 			{
 				ComputeTaskInfo TaskInfo = new ComputeTaskInfo(TaskHash, ChannelId);
 				Logger.LogDebug("Adding task {TaskHash} to queue {RequirementsHash}", TaskInfo.TaskHash.Hash, RequirementsHash);
-				Tasks.Add(TaskScheduler.EnqueueAsync(TaskInfo, RequirementsHash));
+				Tasks.Add(TaskScheduler.EnqueueAsync(RequirementsHash, TaskInfo));
 			}
 			await Task.WhenAll(Tasks);
 		}
@@ -201,16 +213,18 @@ namespace HordeServer.Compute.Impl
 				return null;
 			}
 
-			TaskSchedulerEntry<ComputeTaskInfo>? Entry = await TaskScheduler.DequeueAsync(Agent, CancellationToken);
-			if(Entry != null)
+			(IoHash, ComputeTaskInfo)? Entry = await TaskScheduler.DequeueAsync(RequirementsHash => CheckRequirements(Agent, RequirementsHash), CancellationToken);
+			if (Entry != null)
 			{
+				(IoHash RequirementsHash, ComputeTaskInfo TaskInfo) = Entry.Value;
+
 				ComputeTaskMessage ComputeTask = new ComputeTaskMessage();
-				ComputeTask.ChannelId = Entry.Item.ChannelId.ToString();
+				ComputeTask.ChannelId = TaskInfo.ChannelId.ToString();
 				ComputeTask.NamespaceId = DefaultNamespaceId.ToString();
-				ComputeTask.RequirementsHash = Entry.RequirementsHash;
-				ComputeTask.TaskHash = Entry.Item.TaskHash;
+				ComputeTask.RequirementsHash = new CbObjectAttachment(RequirementsHash);
+				ComputeTask.TaskHash = TaskInfo.TaskHash;
 
-				string LeaseName = $"Remote action ({Entry.Item.TaskHash})";
+				string LeaseName = $"Remote action ({TaskInfo.TaskHash})";
 				byte[] Payload = Any.Pack(ComputeTask).ToByteArray();
 
 				AgentLease Lease = new AgentLease(ObjectId.GenerateNewId(), LeaseName, null, null, null, LeaseState.Pending, Payload, new AgentRequirements(), null);
@@ -225,7 +239,7 @@ namespace HordeServer.Compute.Impl
 		{
 			ComputeTaskMessage Message = Payload.Unpack<ComputeTaskMessage>();
 			ComputeTaskInfo TaskInfo = new ComputeTaskInfo(Message.TaskHash, new ChannelId(Message.ChannelId));
-			return TaskScheduler.EnqueueAsync(TaskInfo, Message.RequirementsHash);
+			return TaskScheduler.EnqueueAsync((CbObjectAttachment)Message.RequirementsHash, TaskInfo);
 		}
 
 		/// <inheritdoc/>
@@ -265,6 +279,103 @@ namespace HordeServer.Compute.Impl
 			Logger.LogInformation("Compute lease finished (lease: {LeaseId}, task: {TaskHash}, agent: {AgentId}, channel: {ChannelId}, result: {ResultHash})", LeaseId, (CbObjectAttachment)ComputeTask.TaskHash, Agent.Id, ComputeTask.ChannelId, Status.ResultHash?.Hash ?? IoHash.Zero);
 			return MessageQueue.PostAsync(ComputeTask.ChannelId, Status);
 		}
+
+		/// <summary>
+		/// Checks that an agent matches the necessary criteria to execute a task
+		/// </summary>
+		/// <param name="Agent"></param>
+		/// <param name="RequirementsHash"></param>
+		/// <returns></returns>
+		async ValueTask<bool> CheckRequirements(IAgent Agent, IoHash RequirementsHash)
+		{
+			QueueInfo QueueInfo = await GetQueueInfoAsync(RequirementsHash);
+			if (QueueInfo.Requirements == null)
+			{
+				return false;
+			}
+
+			if (!QueueInfo.Requirements.Condition.IsEmpty)
+			{
+				if (QueueInfo.Condition == null)
+				{
+					return false;
+				}
+				else if (!QueueInfo.Condition.Evaluate(x => GetAgentProperty(Agent, x)))
+				{
+					return false;
+				}
+			}
+
+			return true;
+		}
+
+		/// <summary>
+		/// Gets a named property for an agent
+		/// </summary>
+		/// <param name="Agent">The agent instance</param>
+		/// <param name="Name">Name of the property to read</param>
+		/// <returns>The property value, or an empty string if it does not eixst</returns>
+		static Scalar GetAgentProperty(IAgent Agent, string Name)
+		{
+			if (Name.Equals("Name", StringComparison.OrdinalIgnoreCase))
+			{
+				return Agent.Id.ToString();
+			}
+
+			HashSet<string>? Properties = Agent.Capabilities.PrimaryDevice.Properties;
+			if (Properties != null)
+			{
+				foreach (string Property in Properties)
+				{
+					if (Property.Length > Name.Length && Property[Name.Length] == '=' && Property.StartsWith(Name, StringComparison.OrdinalIgnoreCase))
+					{
+						return Property.Substring(Name.Length + 1);
+					}
+				}
+			}
+			return new Scalar("");
+		}
+
+		/// <summary>
+		/// Gets the requirements object from the CAS
+		/// </summary>
+		/// <param name="RequirementsHash"></param>
+		/// <returns></returns>
+		async ValueTask<QueueInfo> GetQueueInfoAsync(IoHash RequirementsHash)
+		{
+			QueueInfo? Info;
+			if (!RequirementsCache.TryGetValue(RequirementsHash, out Info))
+			{
+				Info = new QueueInfo();
+				Info.Requirements = await ObjectCollection.GetAsync<Requirements>(DefaultNamespaceId, RequirementsHash);
+
+				if (Info.Requirements != null && !Info.Requirements.Condition.IsEmpty)
+				{
+					try
+					{
+						Info.Condition = Condition.Parse(Info.Requirements.Condition.ToString());
+					}
+					catch (Exception Ex)
+					{
+						Logger.LogError(Ex, "Unable to parse condition '{Condition}': {Message}", Info.Requirements.Condition.ToString(), Ex.Message);
+					}
+				}
+
+				using (ICacheEntry Entry = RequirementsCache.CreateEntry(RequirementsHash))
+				{
+					if (Info.Requirements == null)
+					{
+						Entry.SetAbsoluteExpiration(TimeSpan.FromSeconds(10.0));
+					}
+					else
+					{
+						Entry.SetSlidingExpiration(TimeSpan.FromMinutes(10.0));
+					}
+					Entry.SetValue(Info);
+				}
+			}
+			return Info!;
+		}
 	}
 
 	/// <summary>
diff --git a/Engine/Source/Programs/Horde/HordeServer/Compute/Impl/RedisTaskScheduler.cs b/Engine/Source/Programs/Horde/HordeServer/Compute/Impl/RedisTaskScheduler.cs
--- a/Engine/Source/Programs/Horde/HordeServer/Compute/Impl/RedisTaskScheduler.cs
+++ b/Engine/Source/Programs/Horde/HordeServer/Compute/Impl/RedisTaskScheduler.cs
@@ -23,134 +23,100 @@ using System.Threading.Tasks;
 namespace HordeServer.Compute.Impl
 {
 	using NamespaceId = StringId<INamespace>;
-
-	/// <summary>
-	/// Describes an entry for a compute queue.
-	/// </summary>
-	/// <typeparam name="T">The type describing the work to be performed. Must be serializable to a RedisValue via RedisSerializer.</typeparam>
-	class TaskSchedulerEntry<T>
-	{
-		/// <summary>
-		/// The task description, encoded as a RedisValue
-		/// </summary>
-		public T Item { get; }
-
-		/// <summary>
-		/// Hash of the <see cref="Requirements"/> object in CAS for the agent to execute the task
-		/// </summary>
-		public CbObjectAttachment RequirementsHash { get; }
-
-		/// <summary>
-		/// Constructor
-		/// </summary>
-		/// <param name="Item">The task item</param>
-		/// <param name="RequirementsHash">Requirements for executing the task</param>
-		public TaskSchedulerEntry(T Item, CbObjectAttachment RequirementsHash)
-		{
-			this.Item = Item;
-			this.RequirementsHash = RequirementsHash;
-		}
-	}
+	using Condition = StackExchange.Redis.Condition;
 
 	/// <summary>
 	/// Interface for a queue of compute operations, each of which may have different requirements for the executing machines.
 	/// </summary>
-	interface ITaskScheduler<T> where T : class
+	/// <typeparam name="TQueueId">Type used to identify a particular queue</typeparam>
+	/// <typeparam name="TTask">Type used to describe a task to be performed</typeparam>
+	interface ITaskScheduler<TQueueId, TTask> 
+		where TQueueId : struct 
+		where TTask : class
 	{
 		/// <summary>
-		/// Adds a task to the queue
+		/// Adds a task to a queue
 		/// </summary>
-		/// <param name="Item">The item to add</param>
-		/// <param name="RequirementsHash">Hash of a <see cref="Requirements"/> object stored in the CAS which describes the agent to execute the task</param>
-		Task EnqueueAsync(T Item, CbObjectAttachment RequirementsHash);
+		/// <param name="QueueId">The queue identifier</param>
+		/// <param name="TaskId">The task to add</param>
+		Task EnqueueAsync(TQueueId QueueId, TTask TaskId);
 
 		/// <summary>
 		/// Inserts a previously dequeued task back at the front of the queue
 		/// </summary>
-		/// <param name="Item">The item to add</param>
-		/// <param name="RequirementsHash">Hash of a <see cref="Requirements"/> object stored in the CAS which describes the agent to execute the task</param>
-		Task RequeueAsync(T Item, CbObjectAttachment RequirementsHash);
+		/// <param name="QueueId">The queue identifier</param>
+		/// <param name="Task">The task to add</param>
+		Task RequeueAsync(TQueueId QueueId, TTask Task);
 
 		/// <summary>
-		/// Dequeues a task that the given agent can execute
+		/// Dequeue any task from a particular queue
 		/// </summary>
-		/// <param name="Agent">The agent to execute the task</param>
-		/// <param name="Token">Cancellation token for the operation. Will return a null entry rather than throwing an exception.</param>
+		/// <param name="QueueId">The queue to remove a task from</param>
 		/// <returns>Information about the task to be executed</returns>
-		Task<TaskSchedulerEntry<T>?> DequeueAsync(IAgent Agent, CancellationToken Token = default);
+		Task<TTask?> DequeueAsync(TQueueId QueueId);
 
 		/// <summary>
-		/// Dequeue any item from the given queue
+		/// Dequeues a task that the given agent can execute
 		/// </summary>
-		/// <param name="RequirementsHash">Hash of the requirements object</param>
+		/// <param name="Predicate">Predicate for determining which queues can be removed from</param>
+		/// <param name="Token">Cancellation token for the operation. Will return a null entry rather than throwing an exception.</param>
 		/// <returns>Information about the task to be executed</returns>
-		Task<T?> DequeueAsync(IoHash RequirementsHash);
+		Task<(TQueueId, TTask)?> DequeueAsync(Func<TQueueId, ValueTask<bool>> Predicate, CancellationToken Token = default);
 
 		/// <summary>
 		/// Gets hashes of all the inactive task queues
 		/// </summary>
 		/// <returns></returns>
-		Task<List<IoHash>> GetInactiveQueuesAsync();
+		Task<List<TQueueId>> GetInactiveQueuesAsync();
 	}
 
 	/// <summary>
-	/// Implementation of <see cref="ITaskScheduler{T}"/> using Redis for storage
+	/// Implementation of <see cref="ITaskScheduler{TQueue, TTask}"/> using Redis for storage
 	/// </summary>
-	/// <typeparam name="T">The task definition type</typeparam>
-	class RedisTaskScheduler<T> : ITaskScheduler<T>, IDisposable where T : class
+	/// <typeparam name="TQueueId">Type used to identify a particular queue</typeparam>
+	/// <typeparam name="TTask">Type used to describe a task to be performed</typeparam>
+	class RedisTaskScheduler<TQueueId, TTask> : ITaskScheduler<TQueueId, TTask>, IDisposable 
+		where TQueueId : struct 
+		where TTask : class
 	{
 		class Listener
 		{
-			public IAgent Agent { get; }
-			public TaskCompletionSource<TaskSchedulerEntry<T>?> CompletionSource { get; }
+			public Func<TQueueId, ValueTask<bool>> Predicate;
+			public TaskCompletionSource<(TQueueId, TTask)?> CompletionSource { get; }
 
-			public Listener(IAgent Agent)
+			public Listener(Func<TQueueId, ValueTask<bool>> Predicate)
 			{
-				this.Agent = Agent;
-				this.CompletionSource = new TaskCompletionSource<TaskSchedulerEntry<T>?>();
+				this.Predicate = Predicate;
+				this.CompletionSource = new TaskCompletionSource<(TQueueId, TTask)?>();
 			}
 		}
 
-		class QueueInfo
-		{
-			public Requirements? Requirements;
-			public Utilities.Condition? Condition;
-		}
-
 		IDatabase Redis;
-		IObjectCollection ObjectCollection;
 		RedisKey BaseKey;
-		NamespaceId NamespaceId;
-		RedisSet<IoHash> QueueIndex;
-		RedisHash<IoHash, DateTime> ActiveQueues; // Queues which are actively being dequeued from
-		ReadOnlyHashSet<IoHash> LocalActiveQueues = new HashSet<IoHash>();
+		RedisSet<TQueueId> QueueIndex;
+		RedisHash<TQueueId, DateTime> ActiveQueues; // Queues which are actively being dequeued from
+		ReadOnlyHashSet<TQueueId> LocalActiveQueues = new HashSet<TQueueId>();
 		Stopwatch ResetActiveQueuesTimer = Stopwatch.StartNew();
 
 		List<Listener> Listeners = new List<Listener>();
-		RedisChannel<IoHash> NewQueueChannel;
+		RedisChannel<TQueueId> NewQueueChannel;
 		Task QueueUpdateTask;
 		CancellationTokenSource CancellationSource = new CancellationTokenSource();
-		IMemoryCache RequirementsCache;
 		ILogger Logger;
 
 		/// <summary>
 		/// Constructor
 		/// </summary>
 		/// <param name="Redis">The Redis database instance</param>
-		/// <param name="ObjectCollection">Interface to the CAS object store</param>
 		/// <param name="BaseKey">Base key for all keys used by this scheduler</param>
-		/// <param name="NamespaceId">The namespace identifier</param>
 		/// <param name="Logger"></param>
-		public RedisTaskScheduler(IDatabase Redis, IObjectCollection ObjectCollection, NamespaceId NamespaceId, RedisKey BaseKey, ILogger Logger)
+		public RedisTaskScheduler(IDatabase Redis, RedisKey BaseKey, ILogger Logger)
 		{
 			this.Redis = Redis;
-			this.ObjectCollection = ObjectCollection;
-			this.NamespaceId = NamespaceId;
 			this.BaseKey = BaseKey;
-			this.QueueIndex = new RedisSet<IoHash>(BaseKey.Append("index"));
-			this.ActiveQueues = new RedisHash<IoHash, DateTime>(BaseKey.Append("active"));
-			this.NewQueueChannel = new RedisChannel<IoHash>(BaseKey.Append("new_queues").ToString());
-			this.RequirementsCache = new MemoryCache(new MemoryCacheOptions());
+			this.QueueIndex = new RedisSet<TQueueId>(BaseKey.Append("index"));
+			this.ActiveQueues = new RedisHash<TQueueId, DateTime>(BaseKey.Append("active"));
+			this.NewQueueChannel = new RedisChannel<TQueueId>(BaseKey.Append("new_queues").ToString());
 			this.Logger = Logger;
 
 			QueueUpdateTask = Task.Run(() => UpdateQueuesAsync(CancellationSource.Token));
@@ -164,7 +130,6 @@ namespace HordeServer.Compute.Impl
 				CancellationSource.Cancel();
 				await QueueUpdateTask;
 			}
-			RequirementsCache.Dispose();
 			CancellationSource.Dispose();
 		}
 
@@ -177,98 +142,98 @@ namespace HordeServer.Compute.Impl
 		/// <summary>
 		/// Gets the key for a list of tasks for a particular queue
 		/// </summary>
-		/// <param name="RequirementsHash"></param>
+		/// <param name="QueueId">The queue identifier</param>
 		/// <returns></returns>
-		RedisKey GetQueueKey(CbObjectAttachment RequirementsHash)
+		RedisKey GetQueueKey(TQueueId QueueId)
 		{
-			return BaseKey.Append(RedisSerializer.Serialize(RequirementsHash).AsKey());
+			return BaseKey.Append(RedisSerializer.Serialize(QueueId).AsKey());
 		}
 
 		/// <summary>
 		/// Gets the key for a list of tasks for a particular queue
 		/// </summary>
-		/// <param name="RequirementsHash"></param>
+		/// <param name="QueueId">The queue identifier</param>
 		/// <returns></returns>
-		RedisList<T> GetQueue(CbObjectAttachment RequirementsHash)
+		RedisList<TTask> GetQueue(TQueueId QueueId)
 		{
-			return new RedisList<T>(GetQueueKey(RequirementsHash));
+			return new RedisList<TTask>(GetQueueKey(QueueId));
 		}
 
 		/// <summary>
 		/// Adds a queue to the index, atomically checking that the queue exists
 		/// </summary>
-		/// <param name="RequirementsHash">Hash of the queue requirements objectadd</param>
-		async ValueTask AddQueueToIndexAsync(CbObjectAttachment RequirementsHash)
+		/// <param name="QueueId">The queue to add to the index</param>
+		async ValueTask AddQueueToIndexAsync(TQueueId QueueId)
 		{
 			ITransaction Transaction = Redis.CreateTransaction();
-			Transaction.AddCondition(StackExchange.Redis.Condition.KeyExists(GetQueueKey(RequirementsHash)));
-			_ = Transaction.SetAddAsync(QueueIndex, RequirementsHash);
+			Transaction.AddCondition(Condition.KeyExists(GetQueueKey(QueueId)));
+			_ = Transaction.SetAddAsync(QueueIndex, QueueId);
 			await Transaction.ExecuteAsync();
 
-			await Redis.PublishAsync(NewQueueChannel, RequirementsHash);
+			await Redis.PublishAsync(NewQueueChannel, QueueId);
 		}
 
 		/// <summary>
 		/// Removes a queue from the index, atomically checking that it is empty
 		/// </summary>
-		/// <param name="RequirementsHash">Hash of the queue requirements objectadd</param>
-		async ValueTask RemoveQueueFromIndexAsync(CbObjectAttachment RequirementsHash)
+		/// <param name="QueueId">The queue to remove</param>
+		async ValueTask RemoveQueueFromIndexAsync(TQueueId QueueId)
 		{
 			ITransaction Transaction = Redis.CreateTransaction();
-			Transaction.AddCondition(StackExchange.Redis.Condition.KeyNotExists(GetQueueKey(RequirementsHash)));
-			_ = Transaction.SetRemoveAsync(QueueIndex, RequirementsHash);
+			Transaction.AddCondition(Condition.KeyNotExists(GetQueueKey(QueueId)));
+			_ = Transaction.SetRemoveAsync(QueueIndex, QueueId);
 			await Transaction.ExecuteAsync(CommandFlags.FireAndForget);
 		}
 
 		/// <summary>
-		/// Adds an item to a particular queue, creating and adding that queue to the index if necessary
+		/// Adds a task to a particular queue, creating and adding that queue to the index if necessary
 		/// </summary>
-		/// <param name="Item">The item to be added</param>
-		/// <param name="RequirementsHash">Requirements for the task</param>
-		public async Task EnqueueAsync(T Item, CbObjectAttachment RequirementsHash)
+		/// <param name="QueueId">The queue to add the task to</param>
+		/// <param name="Task">Task to be scheduled</param>
+		public async Task EnqueueAsync(TQueueId QueueId, TTask Task)
 		{
-			RedisList<T> List = GetQueue(RequirementsHash);
+			RedisList<TTask> List = GetQueue(QueueId);
 
-			long NewLength = await Redis.ListRightPushAsync(List, Item);
+			long NewLength = await Redis.ListRightPushAsync(List, Task);
 			if (NewLength == 1)
 			{
-				await AddQueueToIndexAsync(RequirementsHash);
+				await AddQueueToIndexAsync(QueueId);
 			}
 		}
 
 		/// <summary>
 		/// Adds an item to the front of a particular queue, creating and adding that queue to the index if necessary
 		/// </summary>
-		/// <param name="Item">The item to be added</param>
-		/// <param name="RequirementsHash">Requirements for the task</param>
-		public async Task RequeueAsync(T Item, CbObjectAttachment RequirementsHash)
+		/// <param name="QueueId">The queue to add the task to</param>
+		/// <param name="Task">Task to be scheduled</param>
+		public async Task RequeueAsync(TQueueId QueueId, TTask Task)
 		{
-			RedisList<T> Queue = GetQueue(RequirementsHash);
+			RedisList<TTask> List = GetQueue(QueueId);
 
-			long NewLength = await Redis.ListLeftPushAsync(Queue, Item);
+			long NewLength = await Redis.ListLeftPushAsync(List, Task);
 			if (NewLength == 1)
 			{
-				await AddQueueToIndexAsync(RequirementsHash);
+				await AddQueueToIndexAsync(QueueId);
 			}
 		}
 
 		/// <summary>
 		/// Dequeues an item for execution by the given agent
 		/// </summary>
-		/// <param name="Agent">The agent to dequeue an item for</param>
+		/// <param name="Predicate">Predicate for queues that tasks can be removed from</param>
 		/// <param name="Token">Cancellation token for waiting for an item</param>
 		/// <returns>The dequeued item, or null if no item is available</returns>
-		public async Task<TaskSchedulerEntry<T>?> DequeueAsync(IAgent Agent, CancellationToken Token = default)
+		public async Task<(TQueueId, TTask)?> DequeueAsync(Func<TQueueId, ValueTask<bool>> Predicate, CancellationToken Token = default)
 		{
 			// Compare against all the list of cached queues to see if we can dequeue something from any of them
 			Listener? Listener = null;
 			try
 			{
-				IoHash[] Queues = await Redis.SetMembersAsync(QueueIndex);
+				TQueueId[] Queues = await Redis.SetMembersAsync(QueueIndex);
 				while (!Token.IsCancellationRequested)
 				{
 					// Try to dequeue an item from the list
-					TaskSchedulerEntry<T>? Entry = await TryAssignToLocalAgentAsync(Queues, Agent);
+					(TQueueId, TTask)? Entry = await TryAssignToLocalAgentAsync(Queues, Predicate);
 					if (Entry != null)
 					{
 						return Entry;
@@ -277,7 +242,7 @@ namespace HordeServer.Compute.Impl
 					// Create and register a listener for this waiter 
 					if (Listener == null)
 					{
-						Listener = new Listener(Agent);
+						Listener = new Listener(Predicate);
 						lock (Listeners)
 						{
 							Listeners.Add(Listener);
@@ -305,17 +270,22 @@ namespace HordeServer.Compute.Impl
 			return null;
 		}
 
-		async Task<TaskSchedulerEntry<T>?> TryAssignToLocalAgentAsync(IoHash[] Queues, IAgent Agent)
+		/// <summary>
+		/// Attempts to dequeue a task from a set of queue
+		/// </summary>
+		/// <param name="QueueIds">The current array of queues</param>
+		/// <param name="Predicate">Predicate for queues that tasks can be removed from</param>
+		/// <returns>The dequeued item, or null if no item is available</returns>
+		async Task<(TQueueId, TTask)?> TryAssignToLocalAgentAsync(TQueueId[] QueueIds, Func<TQueueId, ValueTask<bool>> Predicate)
 		{
-			foreach (IoHash RequirementsHash in Queues)
+			foreach (TQueueId QueueId in QueueIds)
 			{
-				QueueInfo? Requirements = await GetQueueInfoAsync(RequirementsHash);
-				if (Requirements != null && CheckRequirements(Agent, Requirements))
+				if (await Predicate(QueueId))
 				{
-					T? Item = await DequeueAsync(RequirementsHash);
-					if (Item != null)
+					TTask? Task = await DequeueAsync(QueueId);
+					if (Task != null)
 					{
-						return new TaskSchedulerEntry<T>(Item, RequirementsHash);
+						return (QueueId, Task);
 					}
 				}
 			}
@@ -325,16 +295,16 @@ namespace HordeServer.Compute.Impl
 		/// <summary>
 		/// Dequeues the item at the front of a queue
 		/// </summary>
-		/// <param name="RequirementsHash">Key of the queue to remove from</param>
+		/// <param name="QueueId">The queue to dequeue from</param>
 		/// <returns>The dequeued item, or null if the queue is empty</returns>
-		public async Task<T?> DequeueAsync(IoHash RequirementsHash)
+		public async Task<TTask?> DequeueAsync(TQueueId QueueId)
 		{
-			await AddActiveQueue(RequirementsHash);
+			await AddActiveQueue(QueueId);
 
-			T? Item = await Redis.ListLeftPopAsync(GetQueue(RequirementsHash));
+			TTask? Item = await Redis.ListLeftPopAsync(GetQueue(QueueId));
 			if (Item == null)
 			{
-				await RemoveQueueFromIndexAsync(RequirementsHash);
+				await RemoveQueueFromIndexAsync(QueueId);
 			}
 			return Item;
 		}
@@ -342,9 +312,9 @@ namespace HordeServer.Compute.Impl
 		/// <summary>
 		/// Marks a queue key as being actively monitored, preventing it being returned by <see cref="GetInactiveQueuesAsync"/>
 		/// </summary>
-		/// <param name="Key">The queue key</param>
+		/// <param name="QueueId">The queue key</param>
 		/// <returns></returns>
-		async ValueTask AddActiveQueue(IoHash Key)
+		async ValueTask AddActiveQueue(TQueueId QueueId)
 		{
 			// Periodically clear out the set of active keys
 			TimeSpan ResetTime = TimeSpan.FromSeconds(10.0);
@@ -354,7 +324,7 @@ namespace HordeServer.Compute.Impl
 				{
 					if (ResetActiveQueuesTimer.Elapsed > ResetTime)
 					{
-						LocalActiveQueues = new HashSet<IoHash>();
+						LocalActiveQueues = new HashSet<TQueueId>();
 						ResetActiveQueuesTimer.Reset();
 					}
 				}
@@ -363,19 +333,19 @@ namespace HordeServer.Compute.Impl
 			// Check if the set of active keys already contains the key we're adding. In order to optimize the 
 			// common case under heavy load where the key is in the set, updating it creates a full copy of it. Any
 			// readers can thus access it without the need for any locking.
-			ReadOnlyHashSet<IoHash> LocalActiveKeysCopy = LocalActiveQueues;
-			if (!LocalActiveKeysCopy.Contains(Key))
+			ReadOnlyHashSet<TQueueId> LocalActiveKeysCopy = LocalActiveQueues;
+			if (!LocalActiveKeysCopy.Contains(QueueId))
 			{
 				for (; ; )
 				{
-					HashSet<IoHash> NewLocalActiveKeys = new HashSet<IoHash>(LocalActiveKeysCopy);
-					if (!NewLocalActiveKeys.Add(Key))
+					HashSet<TQueueId> NewLocalActiveKeys = new HashSet<TQueueId>(LocalActiveKeysCopy);
+					if (!NewLocalActiveKeys.Add(QueueId))
 					{
 						break;
 					}
 					if (Interlocked.CompareExchange(ref LocalActiveQueues, NewLocalActiveKeys, LocalActiveKeysCopy) == LocalActiveKeysCopy)
 					{
-						await Redis.HashSetAsync(ActiveQueues, Key, DateTime.UtcNow);
+						await Redis.HashSetAsync(ActiveQueues, QueueId, DateTime.UtcNow);
 						break;
 					}
 				}
@@ -386,15 +356,15 @@ namespace HordeServer.Compute.Impl
 		/// Find any inactive keys
 		/// </summary>
 		/// <returns></returns>
-		public async Task<List<IoHash>> GetInactiveQueuesAsync()
+		public async Task<List<TQueueId>> GetInactiveQueuesAsync()
 		{
-			HashSet<IoHash> Keys = new HashSet<IoHash>(await Redis.SetMembersAsync(QueueIndex));
-			HashSet<IoHash> InvalidKeys = new HashSet<IoHash>();
+			HashSet<TQueueId> Keys = new HashSet<TQueueId>(await Redis.SetMembersAsync(QueueIndex));
+			HashSet<TQueueId> InvalidKeys = new HashSet<TQueueId>();
 
 			DateTime MinTime = DateTime.UtcNow - TimeSpan.FromMinutes(1.0);
 
-			HashEntry<IoHash, DateTime>[] Entries = await Redis.HashGetAllAsync(ActiveQueues);
-			foreach (HashEntry<IoHash, DateTime> Entry in Entries)
+			HashEntry<TQueueId, DateTime>[] Entries = await Redis.HashGetAllAsync(ActiveQueues);
+			foreach (HashEntry<TQueueId, DateTime> Entry in Entries)
 			{
 				if (Entry.Value < MinTime)
 				{
@@ -416,62 +386,72 @@ namespace HordeServer.Compute.Impl
 
 		async Task UpdateQueuesAsync(CancellationToken CancellationToken)
 		{
-			Channel<IoHash> NewQueues = Channel.CreateUnbounded<IoHash>();
+			Channel<TQueueId> NewQueues = Channel.CreateUnbounded<TQueueId>();
 
 			ISubscriber Subscriber = Redis.Multiplexer.GetSubscriber();
 			await using var _ = await Subscriber.SubscribeAsync(NewQueueChannel, (_, v) => NewQueues.Writer.TryWrite(v));
 
 			while (await NewQueues.Reader.WaitToReadAsync(CancellationToken))
 			{
-				HashSet<IoHash> NewQueueKeys = new HashSet<IoHash>();
-				while (NewQueues.Reader.TryRead(out IoHash Key))
+				HashSet<TQueueId> NewQueueIds = new HashSet<TQueueId>();
+				while (NewQueues.Reader.TryRead(out TQueueId QueueId))
 				{
-					NewQueueKeys.Add(Key);
+					NewQueueIds.Add(QueueId);
 				}
-				foreach (IoHash NewQueueKey in NewQueueKeys)
+				foreach (TQueueId NewQueueId in NewQueueIds)
 				{
-					await TryDispatchToNewQueueAsync(NewQueueKey);
+					await TryDispatchToNewQueueAsync(NewQueueId);
 				}
 			}
 		}
 
-		async Task<bool> TryDispatchToNewQueueAsync(IoHash RequirementsHash)
+		async Task<bool> TryDispatchToNewQueueAsync(TQueueId QueueId)
 		{
-			RedisList<T> Queue = GetQueue(RequirementsHash);
-
-			// Get the requirements for this queue
-			QueueInfo? Requirements = await GetQueueInfoAsync(RequirementsHash);
-			if (Requirements == null)
-			{
-				throw new Exception($"Unable to find requirements with hash {RequirementsHash}");
-			}
+			RedisList<TTask> Queue = GetQueue(QueueId);
 
 			// Find a local listener that can execute the work
-			TaskSchedulerEntry<T>? Entry = null;
+			(TQueueId QueueId, TTask TaskId)? Entry = null;
 			try
 			{
 				for (; ; )
 				{
+					Listener? Listener = null;
+
 					// Look for a listener that can execute the task
-					Listener? Listener;
-					lock (Listeners)
-					{
-						Listener = Listeners.FirstOrDefault(x => !x.CompletionSource.Task.IsCompleted && CheckRequirements(x.Agent, Requirements));
-					}
-					if (Listener == null)
+					HashSet<Listener> CheckedListeners = new HashSet<Listener>();
+					while(Listener == null)
 					{
-						return false;
+						// Find up to 10 listeners we haven't seen before
+						List<Listener> NewListeners = new List<Listener>();
+						lock (Listeners)
+						{
+							NewListeners.AddRange(Listeners.Where(x => !x.CompletionSource.Task.IsCompleted && CheckedListeners.Add(x)).Take(10));
+						}
+						if (NewListeners.Count == 0)
+						{
+							return false;
+						}
+
+						// Check predicates for each one against the new queue
+						foreach (Listener NewListener in NewListeners)
+						{
+							if (await NewListener.Predicate(QueueId))
+							{
+								Listener = NewListener;
+								break;
+							}
+						}
 					}
 
 					// Pop an entry from the queue
 					if (Entry == null)
 					{
-						T? TaskData = await DequeueAsync(RequirementsHash);
-						if (TaskData == null)
+						TTask? Task = await DequeueAsync(QueueId);
+						if (Task == null)
 						{
 							return false;
 						}
-						Entry = new TaskSchedulerEntry<T>(TaskData, RequirementsHash);
+						Entry = (QueueId, Task);
 					}
 
 					// Assign it to the listener
@@ -485,105 +465,9 @@ namespace HordeServer.Compute.Impl
 			{
 				if (Entry != null)
 				{
-					await RequeueAsync(Entry.Item, Entry.RequirementsHash);
-				}
-			}
-		}
-
-		/// <summary>
-		/// Checks that an agent matches the necessary criteria to execute a task
-		/// </summary>
-		/// <param name="Agent"></param>
-		/// <param name="QueueInfo"></param>
-		/// <returns></returns>
-		static bool CheckRequirements(IAgent Agent, QueueInfo QueueInfo)
-		{
-			if (QueueInfo.Requirements == null)
-			{
-				return false;
-			}
-
-			if (!QueueInfo.Requirements.Condition.IsEmpty)
-			{
-				if (QueueInfo.Condition == null)
-				{
-					return false;
-				}
-				else if (!QueueInfo.Condition.Evaluate(x => GetAgentProperty(Agent, x)))
-				{
-					return false;
-				}
-			}
-
-			return true;
-		}
-
-		/// <summary>
-		/// Gets a named property for an agent
-		/// </summary>
-		/// <param name="Agent">The agent instance</param>
-		/// <param name="Name">Name of the property to read</param>
-		/// <returns>The property value, or an empty string if it does not eixst</returns>
-		static Scalar GetAgentProperty(IAgent Agent, string Name)
-		{
-			if (Name.Equals("Name", StringComparison.OrdinalIgnoreCase))
-			{
-				return Agent.Id.ToString();
-			}
-
-			HashSet<string>? Properties = Agent.Capabilities.PrimaryDevice.Properties;
-			if (Properties != null)
-			{
-				foreach (string Property in Properties)
-				{
-					if (Property.Length > Name.Length && Property[Name.Length] == '=' && Property.StartsWith(Name, StringComparison.OrdinalIgnoreCase))
-					{
-						return Property.Substring(Name.Length + 1);
-					}
-				}
-			}
-			return new Scalar("");
-		}
-
-		/// <summary>
-		/// Gets the requirements object from the CAS
-		/// </summary>
-		/// <param name="RequirementsHash"></param>
-		/// <returns></returns>
-		async Task<QueueInfo?> GetQueueInfoAsync(IoHash RequirementsHash)
-		{
-			QueueInfo? Info;
-			if (!RequirementsCache.TryGetValue(RequirementsHash, out Info))
-			{
-				Info = new QueueInfo();
-				Info.Requirements = await ObjectCollection.GetAsync<Requirements>(NamespaceId, RequirementsHash);
-
-				if (Info.Requirements != null && !Info.Requirements.Condition.IsEmpty)
-				{
-					try
-					{
-						Info.Condition = Utilities.Condition.Parse(Info.Requirements.Condition.ToString());
-					}
-					catch (Exception Ex)
-					{
-						Logger.LogError(Ex, "Unable to parse condition '{Condition}': {Message}", Info.Requirements.Condition.ToString(), Ex.Message);
-					}
-				}
-
-				using (ICacheEntry Entry = RequirementsCache.CreateEntry(RequirementsHash))
-				{
-					if (Info.Requirements == null)
-					{
-						Entry.SetAbsoluteExpiration(TimeSpan.FromSeconds(10.0));
-					}
-					else
-					{
-						Entry.SetSlidingExpiration(TimeSpan.FromMinutes(10.0));
-					}
-					Entry.SetValue(Info);
+					await RequeueAsync(Entry.Value.QueueId, Entry.Value.TaskId);
 				}
 			}
-			return Info;
 		}
 	}
 }
