Horde: Add support for multivalue properties in conditions (such as the 'pool' property on an agent).

When expressions are evaluated, any binary operator is taken to match if any of its values match. For example, a condition matching Pool='Win-UE5' will match any agent in the Win-UE5 pool, even if it also happens to be in another pool. Likewise, Pool != 'Win-UE5' will match any agent in the Win-UE5 pool, even if it also belongs to other pools.

As part of this change, the query syntax has been adjusted to distinguish conditional expressions from comparison expressions. It is no longer possible to write expressions such as "(a == 1) == (b == 2)"; the more expressive form is required ("(a == 1 && b == 2) || (a != 1 && b != 2)".

diff --git a/Engine/Source/Programs/Horde/HordeServer/Compute/Impl/ComputeService.cs b/Engine/Source/Programs/Horde/HordeServer/Compute/Impl/ComputeService.cs
--- a/Engine/Source/Programs/Horde/HordeServer/Compute/Impl/ComputeService.cs
+++ b/Engine/Source/Programs/Horde/HordeServer/Compute/Impl/ComputeService.cs
@@ -35,6 +35,7 @@ using EpicGames.Horde.Compute;
 using Microsoft.Extensions.Hosting;
 using System.Threading.Channels;
 using EpicGames.Redis;
+using HordeServer.Collections;
 
 namespace HordeServer.Compute.Impl
 {
@@ -135,10 +136,12 @@ namespace HordeServer.Compute.Impl
 		public static NamespaceId DefaultNamespaceId { get; } = new NamespaceId("default");
 
 		IObjectCollection ObjectCollection;
+		IPoolCollection PoolCollection;
 		ITaskScheduler<IoHash, ComputeTaskInfo> TaskScheduler;
 		RedisMessageQueue<ComputeTaskStatus> MessageQueue;
 		BackgroundTick ExpireTasksTicker;
 		IMemoryCache RequirementsCache;
+		LazyCachedValue<Task<List<IPool>>> CachedPools;
 		ILogger Logger;
 
 		/// <summary>
@@ -146,14 +149,17 @@ namespace HordeServer.Compute.Impl
 		/// </summary>
 		/// <param name="Redis">Redis instance</param>
 		/// <param name="ObjectCollection"></param>
+		/// <param name="PoolCollection">Collection of pool documents</param>
 		/// <param name="Logger">The logger instance</param>
-		public ComputeService(IDatabase Redis, IObjectCollection ObjectCollection, ILogger<ComputeService> Logger)
+		public ComputeService(IDatabase Redis, IObjectCollection ObjectCollection, IPoolCollection PoolCollection, ILogger<ComputeService> Logger)
 		{
 			this.ObjectCollection = ObjectCollection;
+			this.PoolCollection = PoolCollection;
 			this.TaskScheduler = new RedisTaskScheduler<IoHash, ComputeTaskInfo>(Redis, "compute/tasks/", Logger);
 			this.MessageQueue = new RedisMessageQueue<ComputeTaskStatus>(Redis, "compute/messages/");
 			this.ExpireTasksTicker = new BackgroundTick(ExpireTasksAsync, TimeSpan.FromMinutes(2.0), Logger);
 			this.RequirementsCache = new MemoryCache(new MemoryCacheOptions());
+			this.CachedPools = new LazyCachedValue<Task<List<IPool>>>(() => PoolCollection.GetAsync(), TimeSpan.FromSeconds(30.0));
 			this.Logger = Logger;
 		}
 
@@ -206,14 +212,19 @@ namespace HordeServer.Compute.Impl
 		/// <inheritdoc/>
 		public async Task<AgentLease?> TryAssignLeaseAsync(IAgent Agent, CancellationToken CancellationToken)
 		{
-			if (!Agent.Enabled || Agent.Id.ToString().StartsWith("REMOTE-EXEC-", StringComparison.OrdinalIgnoreCase))
+			// If the agent is disabled, just block until we cancel
+			if (!Agent.Enabled)
 			{
 				using CancellationTask Task = new CancellationTask(CancellationToken);
 				await Task.Task;
 				return null;
 			}
 
-			(IoHash, ComputeTaskInfo)? Entry = await TaskScheduler.DequeueAsync(RequirementsHash => CheckRequirements(Agent, RequirementsHash), CancellationToken);
+			// Get all the current agent properties
+			Lazy<Task<List<string>>> Properties = new Lazy<Task<List<string>>>(() => GetAgentProperties(Agent));
+
+			// Find a task to execute
+			(IoHash, ComputeTaskInfo)? Entry = await TaskScheduler.DequeueAsync(RequirementsHash => CheckRequirements(RequirementsHash, Properties), CancellationToken);
 			if (Entry != null)
 			{
 				(IoHash RequirementsHash, ComputeTaskInfo TaskInfo) = Entry.Value;
@@ -283,10 +294,10 @@ namespace HordeServer.Compute.Impl
 		/// <summary>
 		/// Checks that an agent matches the necessary criteria to execute a task
 		/// </summary>
-		/// <param name="Agent"></param>
 		/// <param name="RequirementsHash"></param>
+		/// <param name="LazyProperties">Properties for this agentthat the agent is in</param>
 		/// <returns></returns>
-		async ValueTask<bool> CheckRequirements(IAgent Agent, IoHash RequirementsHash)
+		async ValueTask<bool> CheckRequirements(IoHash RequirementsHash, Lazy<Task<List<string>>> LazyProperties)
 		{
 			QueueInfo QueueInfo = await GetQueueInfoAsync(RequirementsHash);
 			if (QueueInfo.Requirements == null)
@@ -300,7 +311,9 @@ namespace HordeServer.Compute.Impl
 				{
 					return false;
 				}
-				else if (!QueueInfo.Condition.Evaluate(x => GetAgentProperty(Agent, x)))
+
+				List<string> Properties = await LazyProperties.Value;
+				if (!QueueInfo.Condition.Evaluate(x => FindAgentProperty(x, Properties)))
 				{
 					return false;
 				}
@@ -309,31 +322,53 @@ namespace HordeServer.Compute.Impl
 			return true;
 		}
 
+		/// <summary>
+		/// Finds property values from a sorted list of Name=Value pairs
+		/// </summary>
+		/// <param name="Name">Name of the property to find</param>
+		/// <param name="Properties">The full list of properties</param>
+		/// <returns>Property values</returns>
+		static IEnumerable<string> FindAgentProperty(string Name, List<string> Properties)
+		{
+			int Index = Properties.BinarySearch(Name, StringComparer.OrdinalIgnoreCase);
+			if (Index < 0)
+			{
+				Index = ~Index;
+				while (Index < Properties.Count)
+				{
+					string Property = Properties[Index];
+					if (Property.Length <= Name.Length || !Property.StartsWith(Name, StringComparison.OrdinalIgnoreCase) || Property[Name.Length] != '=')
+					{
+						break;
+					}
+					yield return Property.Substring(Name.Length + 1);
+				}
+			}
+		}
+
 		/// <summary>
 		/// Gets a named property for an agent
 		/// </summary>
 		/// <param name="Agent">The agent instance</param>
-		/// <param name="Name">Name of the property to read</param>
 		/// <returns>The property value, or an empty string if it does not eixst</returns>
-		static Scalar GetAgentProperty(IAgent Agent, string Name)
+		async Task<List<string>> GetAgentProperties(IAgent Agent)
 		{
-			if (Name.Equals("Name", StringComparison.OrdinalIgnoreCase))
+			List<string> Properties = new List<string>();
+			Properties.Add($"Name={Agent.Id}");
+
+			foreach (IPool Pool in Agent.GetPools(await CachedPools.GetCached()))
 			{
-				return Agent.Id.ToString();
+				Properties.Add($"Pool={Pool.Id}");
 			}
 
-			HashSet<string>? Properties = Agent.Capabilities.PrimaryDevice.Properties;
-			if (Properties != null)
+			HashSet<string>? DeviceProperties = Agent.Capabilities.PrimaryDevice.Properties;
+			if (DeviceProperties != null)
 			{
-				foreach (string Property in Properties)
-				{
-					if (Property.Length > Name.Length && Property[Name.Length] == '=' && Property.StartsWith(Name, StringComparison.OrdinalIgnoreCase))
-					{
-						return Property.Substring(Name.Length + 1);
-					}
-				}
+				Properties.AddRange(DeviceProperties);
 			}
-			return new Scalar("");
+
+			Properties.Sort(StringComparer.OrdinalIgnoreCase);
+			return Properties;
 		}
 
 		/// <summary>
diff --git a/Engine/Source/Programs/Horde/HordeServer/Utilities/Condition.cs b/Engine/Source/Programs/Horde/HordeServer/Utilities/Condition.cs
--- a/Engine/Source/Programs/Horde/HordeServer/Utilities/Condition.cs
+++ b/Engine/Source/Programs/Horde/HordeServer/Utilities/Condition.cs
@@ -12,69 +12,11 @@ using System.Threading.Tasks;
 
 namespace HordeServer.Utilities
 {
-	/// <summary>
-	/// Represents a scalar value for use by a condition expression
-	/// </summary>
-	readonly struct Scalar
-	{
-		public bool IsInteger => !IsString;
-		public bool IsString => String != null;
-
-		public readonly ulong Integer;
-		public readonly string? String;
-
-		public Scalar(ulong Integer)
-		{
-			this.Integer = Integer;
-			this.String = null;
-		}
-
-		public Scalar(string String)
-		{
-			this.Integer = 0;
-			this.String = String;
-		}
-
-		public ulong AsInteger()
-		{
-			if (!IsInteger)
-			{
-				throw new ConditionException("Value is not an integer");
-			}
-			return Integer;
-		}
-		
-		public string AsString()
-		{
-			if (!IsString)
-			{
-				throw new ConditionException("Value is not a string");
-			}
-			return String!;
-		}
-
-		public static implicit operator Scalar(bool Value)
-		{
-			return new Scalar(Value ? 1UL : 0UL);
-		}
-
-		public static implicit operator Scalar(ulong Integer)
-		{
-			return new Scalar(Integer);
-		}
-
-		public static implicit operator Scalar(string String)
-		{
-			return new Scalar(String);
-		}
-	}
-
 	enum TokenType : byte
 	{
 		End,
 		Identifier,
-		Number,
-		String,
+		Scalar,
 		Not,
 		Eq,
 		Neq,
@@ -97,8 +39,7 @@ namespace HordeServer.Utilities
 
 		public StringView Token => new StringView(Input, Offset, Length);
 		public TokenType Type { get; private set; }
-		public ulong IntegerValue { get; private set; }
-		public string? StringValue { get; private set; }
+		public string? Scalar { get; private set; }
 
 		public TokenReader(string Input)
 		{
@@ -106,12 +47,11 @@ namespace HordeServer.Utilities
 			MoveNext();
 		}
 
-		private void SetCurrent(int Length, TokenType Type, ulong IntegerValue = 0, string? StringValue = null)
+		private void SetCurrent(int Length, TokenType Type, string? Scalar = null)
 		{
 			this.Length = Length;
 			this.Type = Type;
-			this.IntegerValue = IntegerValue;
-			this.StringValue = StringValue;
+			this.Scalar = Scalar;
 		}
 
 		public void MoveNext()
@@ -181,11 +121,11 @@ namespace HordeServer.Utilities
 					case '\"':
 					case '\'':
 						Utf8String String = ParseString(Input, Offset, out int StringLength);
-						SetCurrent(StringLength, TokenType.String, StringValue: String.ToString());
+						SetCurrent(StringLength, TokenType.Scalar, Scalar: String.ToString());
 						return;
 					case char Char when IsNumber(Char):
 						ulong Integer = ParseNumber(Input, Offset, out int IntegerLength);
-						SetCurrent(IntegerLength, TokenType.Number, IntegerValue: Integer);
+						SetCurrent(IntegerLength, TokenType.Scalar, Scalar: Integer.ToString(CultureInfo.InvariantCulture));
 						return;
 					case char Char when IsIdentifier(Char):
 						int EndIdx = Offset + 1;
@@ -231,7 +171,6 @@ namespace HordeServer.Utilities
 				}
 			}
 
-
 			char[] Copy = new char[EndIdx - (Idx + 1) - NumEscapeChars];
 
 			int InputIdx = Idx + 1;
@@ -344,7 +283,6 @@ namespace HordeServer.Utilities
 
 		List<Token> Tokens = new List<Token>();
 		List<string> Strings = new List<string>();
-		List<ulong> Integers = new List<ulong>();
 
 		private Condition(string Text)
 		{
@@ -361,71 +299,91 @@ namespace HordeServer.Utilities
 			return new Condition(Text);
 		}
 
-		public bool Evaluate(Func<string, Scalar> ResolveIdentifier)
+		public bool Evaluate(Func<string, IEnumerable<string>> GetPropertyValues)
 		{
 			int Idx = 0;
-			Scalar Scalar = Evaluate(ref Idx, ResolveIdentifier);
-			return Scalar.AsInteger() != 0;
+			return EvaluateCondition(ref Idx, GetPropertyValues);
 		}
 
-		Scalar Evaluate(ref int Idx, Func<string, Scalar> ResolveIdentifier)
+		bool EvaluateCondition(ref int Idx, Func<string, IEnumerable<string>> GetPropertyValues)
 		{
-			Scalar Lhs;
-			Scalar Rhs;
+			bool LhsBool;
+			bool RhsBool;
+			IEnumerable<string> LhsScalar;
+			IEnumerable<string> RhsScalar;
 
 			Token Token = Tokens[Idx++];
 			switch (Token.Type)
 			{
-				case TokenType.Identifier:
-					return ResolveIdentifier(Strings[Token.Index]);
-				case TokenType.Number:
-					return Integers[Token.Index];
-				case TokenType.String:
-					return Strings[Token.Index].ToString();
 				case TokenType.Not:
-					Lhs = Evaluate(ref Idx, ResolveIdentifier);
-					return Lhs.AsInteger() == 0UL;
+					return !EvaluateCondition(ref Idx, GetPropertyValues);
 				case TokenType.Eq:
-					Lhs = Evaluate(ref Idx, ResolveIdentifier);
-					Rhs = Evaluate(ref Idx, ResolveIdentifier);
-					return (Lhs.IsString || Rhs.IsString) ? String.Equals(Lhs.AsString(), Rhs.AsString(), StringComparison.OrdinalIgnoreCase) : (Lhs.AsInteger() == Rhs.AsInteger());
+					LhsScalar = EvaluateScalar(ref Idx, GetPropertyValues);
+					RhsScalar = EvaluateScalar(ref Idx, GetPropertyValues);
+					return LhsScalar.Any(x => RhsScalar.Contains(x, StringComparer.OrdinalIgnoreCase));
 				case TokenType.Neq:
-					Lhs = Evaluate(ref Idx, ResolveIdentifier);
-					Rhs = Evaluate(ref Idx, ResolveIdentifier);
-					return (Lhs.IsString || Rhs.IsString) ? !String.Equals(Lhs.AsString(), Rhs.AsString(), StringComparison.OrdinalIgnoreCase) : (Lhs.AsInteger() != Rhs.AsInteger());
+					LhsScalar = EvaluateScalar(ref Idx, GetPropertyValues);
+					RhsScalar = EvaluateScalar(ref Idx, GetPropertyValues);
+					return !LhsScalar.Any(x => RhsScalar.Contains(x, StringComparer.OrdinalIgnoreCase));
 				case TokenType.LogicalOr:
-					Lhs = Evaluate(ref Idx, ResolveIdentifier);
-					Rhs = Evaluate(ref Idx, ResolveIdentifier);
-					return (Lhs.AsInteger() != 0) || (Rhs.AsInteger() != 0);
+					LhsBool = EvaluateCondition(ref Idx, GetPropertyValues);
+					RhsBool = EvaluateCondition(ref Idx, GetPropertyValues);
+					return LhsBool || RhsBool;
 				case TokenType.LogicalAnd:
-					Lhs = Evaluate(ref Idx, ResolveIdentifier);
-					Rhs = Evaluate(ref Idx, ResolveIdentifier);
-					return (Lhs.AsInteger() != 0) && (Rhs.AsInteger() != 0);
+					LhsBool = EvaluateCondition(ref Idx, GetPropertyValues);
+					RhsBool = EvaluateCondition(ref Idx, GetPropertyValues);
+					return LhsBool && RhsBool;
 				case TokenType.Lt:
-					Lhs = Evaluate(ref Idx, ResolveIdentifier);
-					Rhs = Evaluate(ref Idx, ResolveIdentifier);
-					return Lhs.AsInteger() < Rhs.AsInteger();
+					LhsScalar = EvaluateScalar(ref Idx, GetPropertyValues);
+					RhsScalar = EvaluateScalar(ref Idx, GetPropertyValues);
+					return AsIntegers(LhsScalar).Any(x => AsIntegers(RhsScalar).Any(y => x < y));
 				case TokenType.Lte:
-					Lhs = Evaluate(ref Idx, ResolveIdentifier);
-					Rhs = Evaluate(ref Idx, ResolveIdentifier);
-					return Lhs.AsInteger() <= Rhs.AsInteger();
+					LhsScalar = EvaluateScalar(ref Idx, GetPropertyValues);
+					RhsScalar = EvaluateScalar(ref Idx, GetPropertyValues);
+					return AsIntegers(LhsScalar).Any(x => AsIntegers(RhsScalar).Any(y => x <= y));
 				case TokenType.Gt:
-					Lhs = Evaluate(ref Idx, ResolveIdentifier);
-					Rhs = Evaluate(ref Idx, ResolveIdentifier);
-					return Lhs.AsInteger() > Rhs.AsInteger();
+					LhsScalar = EvaluateScalar(ref Idx, GetPropertyValues);
+					RhsScalar = EvaluateScalar(ref Idx, GetPropertyValues);
+					return AsIntegers(LhsScalar).Any(x => AsIntegers(RhsScalar).Any(y => x > y));
 				case TokenType.Gte:
-					Lhs = Evaluate(ref Idx, ResolveIdentifier);
-					Rhs = Evaluate(ref Idx, ResolveIdentifier);
-					return Lhs.AsInteger() >= Rhs.AsInteger();
+					LhsScalar = EvaluateScalar(ref Idx, GetPropertyValues);
+					RhsScalar = EvaluateScalar(ref Idx, GetPropertyValues);
+					return AsIntegers(LhsScalar).Any(x => AsIntegers(RhsScalar).Any(y => x >= y));
 				case TokenType.Regex:
-					Lhs = Evaluate(ref Idx, ResolveIdentifier);
-					Rhs = Evaluate(ref Idx, ResolveIdentifier);
-					return Regex.IsMatch(Lhs.AsString(), Rhs.AsString());
+					LhsScalar = EvaluateScalar(ref Idx, GetPropertyValues);
+					RhsScalar = EvaluateScalar(ref Idx, GetPropertyValues);
+					return LhsScalar.Any(x => RhsScalar.Any(y => Regex.IsMatch(x, y, RegexOptions.IgnoreCase)));
 				default:
 					throw new ConditionException("Invalid token type");
 			}
 		}
 
+		IEnumerable<string> EvaluateScalar(ref int Idx, Func<string, IEnumerable<string>> GetPropertyValues)
+		{
+			Token Token = Tokens[Idx++];
+			switch (Token.Type)
+			{
+				case TokenType.Identifier:
+					return GetPropertyValues(Strings[Token.Index]);
+				case TokenType.Scalar:
+					return new string[] { Strings[Token.Index] };
+				default:
+					throw new ConditionException("Invalid token type");
+			}
+		}
+
+		static IEnumerable<long> AsIntegers(IEnumerable<string> Scalars)
+		{
+			foreach (string Scalar in Scalars)
+			{
+				long Value;
+				if (long.TryParse(Scalar, out Value))
+				{
+					yield return Value;
+				}
+			}
+		}
+
 		void ParseOr(TokenReader Reader)
 		{
 			int StartCount = Tokens.Count;
@@ -442,7 +400,7 @@ namespace HordeServer.Utilities
 		void ParseAnd(TokenReader Reader)
 		{
 			int StartCount = Tokens.Count;
-			ParseComparison(Reader);
+			ParseBoolean(Reader);
 
 			if (Reader.Type == TokenType.LogicalAnd)
 			{
@@ -452,10 +410,34 @@ namespace HordeServer.Utilities
 			}
 		}
 
+		void ParseBoolean(TokenReader Reader)
+		{
+			while (Reader.Type == TokenType.Not)
+			{
+				Tokens.Add(new Token(Reader.Type, 0));
+				Reader.MoveNext();
+			}
+
+			if (Reader.Type == TokenType.Lparen)
+			{
+				Reader.MoveNext();
+				ParseOr(Reader);
+				if (Reader.Type != TokenType.Rparen)
+				{
+					throw new ConditionException($"Missing ')' at offset {Reader.Offset}");
+				}
+				Reader.MoveNext();
+			}
+			else
+			{
+				ParseComparison(Reader);
+			}
+		}
+
 		void ParseComparison(TokenReader Reader)
 		{
 			int StartCount = Tokens.Count;
-			ParseUnary(Reader);
+			ParseScalar(Reader);
 
 			if (Reader.Type == TokenType.Lt 
 				|| Reader.Type == TokenType.Lte 
@@ -467,18 +449,8 @@ namespace HordeServer.Utilities
 			{
 				Tokens.Insert(StartCount, new Token(Reader.Type, 0));
 				Reader.MoveNext();
-				ParseUnary(Reader);
-			}
-		}
-
-		void ParseUnary(TokenReader Reader)
-		{
-			while (Reader.Type == TokenType.Not)
-			{
-				Tokens.Add(new Token(Reader.Type, 0));
-				Reader.MoveNext();
+				ParseScalar(Reader);
 			}
-			ParseScalar(Reader);
 		}
 
 		void ParseScalar(TokenReader Reader)
@@ -489,26 +461,10 @@ namespace HordeServer.Utilities
 				Tokens.Add(new Token(TokenType.Identifier, Strings.Count - 1));
 				Reader.MoveNext();
 			}
-			else if (Reader.Type == TokenType.Number)
+			else if (Reader.Type == TokenType.Scalar)
 			{
-				Integers.Add(Reader.IntegerValue);
-				Tokens.Add(new Token(TokenType.Number, Integers.Count - 1));
-				Reader.MoveNext();
-			}
-			else if (Reader.Type == TokenType.String)
-			{
-				Strings.Add(Reader.StringValue!);
-				Tokens.Add(new Token(TokenType.String, Strings.Count - 1));
-				Reader.MoveNext();
-			}
-			else if (Reader.Type == TokenType.Lparen)
-			{
-				Reader.MoveNext();
-				ParseOr(Reader);
-				if (Reader.Type != TokenType.Rparen)
-				{
-					throw new ConditionException($"Missing ')' at offset {Reader.Offset}");
-				}
+				Strings.Add(Reader.Scalar!);
+				Tokens.Add(new Token(TokenType.Scalar, Strings.Count - 1));
 				Reader.MoveNext();
 			}
 			else
