Horde: Remove the 'ephemeral' flag from agents. We never used this, and we'll solve the deletion of agents through autoscaling more explicitly.

diff --git a/Engine/Source/Programs/Horde/HordeServer/Api/Agents.cs b/Engine/Source/Programs/Horde/HordeServer/Api/Agents.cs
--- a/Engine/Source/Programs/Horde/HordeServer/Api/Agents.cs
+++ b/Engine/Source/Programs/Horde/HordeServer/Api/Agents.cs
@@ -473,7 +473,6 @@ namespace HordeServer.Api
 			this.Name = Agent.Id.ToString();
 			this.Enabled = Agent.Enabled;
 			this.SessionId = Agent.SessionId?.ToString();
-			this.Ephemeral = Agent.Ephemeral;
 			this.Online = Agent.IsSessionValid(DateTime.UtcNow);
 			this.Deleted = Agent.Deleted;
 			this.PendingConform = Agent.RequestConform;
diff --git a/Engine/Source/Programs/Horde/HordeServer/Collections/IAgentCollection.cs b/Engine/Source/Programs/Horde/HordeServer/Collections/IAgentCollection.cs
--- a/Engine/Source/Programs/Horde/HordeServer/Collections/IAgentCollection.cs
+++ b/Engine/Source/Programs/Horde/HordeServer/Collections/IAgentCollection.cs
@@ -31,10 +31,9 @@ namespace HordeServer.Collections
 		/// </summary>
 		/// <param name="Id">Id for the new agent</param>
 		/// <param name="bEnabled">Whether the agent is enabled or not</param>
-		/// <param name="bEphemeral">Whether the agent should be deleted automatically once offline</param>
 		/// <param name="Channel">Channel to use for software run by this agent</param>
 		/// <param name="Pools">Pools for the agent</param>
-		Task<IAgent> AddAsync(AgentId Id, bool bEnabled, bool bEphemeral, AgentSoftwareChannelName? Channel = null, List<PoolId>? Pools = null);
+		Task<IAgent> AddAsync(AgentId Id, bool bEnabled, AgentSoftwareChannelName? Channel = null, List<PoolId>? Pools = null);
 
 		/// <summary>
 		/// Deletes an agent
diff --git a/Engine/Source/Programs/Horde/HordeServer/Collections/Impl/AgentCollection.cs b/Engine/Source/Programs/Horde/HordeServer/Collections/Impl/AgentCollection.cs
--- a/Engine/Source/Programs/Horde/HordeServer/Collections/Impl/AgentCollection.cs
+++ b/Engine/Source/Programs/Horde/HordeServer/Collections/Impl/AgentCollection.cs
@@ -91,13 +91,11 @@ namespace HordeServer.Collections.Impl
 			{
 			}
 
-			public AgentDocument(AgentId Id, bool bEnabled, bool bEphemeral, AgentSoftwareChannelName? Channel, List<PoolId> Pools)
+			public AgentDocument(AgentId Id, bool bEnabled, AgentSoftwareChannelName? Channel, List<PoolId> Pools)
 			{
 				this.Id = Id;
 				this.Acl = new Acl();
 				this.Enabled = bEnabled;
-				this.Ephemeral = bEphemeral;
-				this.Deleted = bEphemeral;
 				this.Channel = Channel;
 				this.Pools = Pools;
 			}
@@ -122,9 +120,9 @@ namespace HordeServer.Collections.Impl
 		}
 
 		/// <inheritdoc/>
-		public async Task<IAgent> AddAsync(AgentId Id, bool bEnabled, bool bEphemeral, AgentSoftwareChannelName? Channel, List<PoolId>? Pools)
+		public async Task<IAgent> AddAsync(AgentId Id, bool bEnabled, AgentSoftwareChannelName? Channel, List<PoolId>? Pools)
 		{
-			AgentDocument Agent = new AgentDocument(Id, bEnabled, bEphemeral, Channel, Pools ?? new List<PoolId>());
+			AgentDocument Agent = new AgentDocument(Id, bEnabled, Channel, Pools ?? new List<PoolId>());
 			await Agents.InsertOneAsync(Agent);
 			return Agent;
 		}
diff --git a/Engine/Source/Programs/Horde/HordeServer/Controllers/AgentsController.cs b/Engine/Source/Programs/Horde/HordeServer/Controllers/AgentsController.cs
--- a/Engine/Source/Programs/Horde/HordeServer/Controllers/AgentsController.cs
+++ b/Engine/Source/Programs/Horde/HordeServer/Controllers/AgentsController.cs
@@ -77,7 +77,7 @@ namespace HordeServer.Controllers
 			}
 
 			AgentSoftwareChannelName? Channel = String.IsNullOrEmpty(Request.Channel) ? (AgentSoftwareChannelName?)null : new AgentSoftwareChannelName(Request.Channel);
-			IAgent Agent = await AgentService.CreateAgentAsync(Request.Name, Request.Enabled, Request.Ephemeral, Channel, Request.Pools?.ConvertAll(x => new PoolId(x)));
+			IAgent Agent = await AgentService.CreateAgentAsync(Request.Name, Request.Enabled, Channel, Request.Pools?.ConvertAll(x => new PoolId(x)));
 
 			return new CreateAgentResponse(Agent.Id.ToString());
 		}
diff --git a/Engine/Source/Programs/Horde/HordeServer/Models/Agent.cs b/Engine/Source/Programs/Horde/HordeServer/Models/Agent.cs
--- a/Engine/Source/Programs/Horde/HordeServer/Models/Agent.cs
+++ b/Engine/Source/Programs/Horde/HordeServer/Models/Agent.cs
@@ -451,16 +451,6 @@ namespace HordeServer.Models
 		/// Whether the agent is enabled
 		/// </summary>
 		public bool Enabled { get; }
-		
-		/// <summary>
-		/// Arbitrary comment for the agent (useful for disable reasons etc)
-		/// </summary>
-		public string? Comment { get; }
-
-		/// <summary>
-		/// Whether the agent is ephemeral
-		/// </summary>
-		public bool Ephemeral { get; }
 
 		/// <summary>
 		/// Whether the agent should be included on the dashboard. This is set to true for ephemeral agents once they are no longer online, or agents that are explicitly deleted.
@@ -472,6 +462,11 @@ namespace HordeServer.Models
 		/// </summary>
 		public string? Version { get; }
 
+		/// <summary>
+		/// Arbitrary comment for the agent (useful for disable reasons etc)
+		/// </summary>
+		public string? Comment { get; }
+
 		/// <summary>
 		/// Channel for the software running on this agent. Uses <see cref="AgentSoftwareService.DefaultChannelName"/> if not specified
 		/// </summary>
diff --git a/Engine/Source/Programs/Horde/HordeServer/Services/AgentService.cs b/Engine/Source/Programs/Horde/HordeServer/Services/AgentService.cs
--- a/Engine/Source/Programs/Horde/HordeServer/Services/AgentService.cs
+++ b/Engine/Source/Programs/Horde/HordeServer/Services/AgentService.cs
@@ -159,13 +159,12 @@ namespace HordeServer.Services
 		/// </summary>
 		/// <param name="Name">Name of the agent</param>
 		/// <param name="bEnabled">Whether the agent is currently enabled</param>
-		/// <param name="bEphemeral">Whether this agent is ephemeral</param>
 		/// <param name="Channel">Override for the desired software version</param>
 		/// <param name="Pools">Pools for this agent</param>
 		/// <returns>Unique id for the agent</returns>
-		public Task<IAgent> CreateAgentAsync(string Name, bool bEnabled, bool bEphemeral, AgentSoftwareChannelName? Channel, List<PoolId>? Pools)
+		public Task<IAgent> CreateAgentAsync(string Name, bool bEnabled, AgentSoftwareChannelName? Channel, List<PoolId>? Pools)
 		{
-			return Agents.AddAsync(new AgentId(Name), bEnabled, bEphemeral, Channel, Pools);
+			return Agents.AddAsync(new AgentId(Name), bEnabled, Channel, Pools);
 		}
 
 		/// <summary>
diff --git a/Engine/Source/Programs/Horde/HordeServer/Services/RpcService.cs b/Engine/Source/Programs/Horde/HordeServer/Services/RpcService.cs
--- a/Engine/Source/Programs/Horde/HordeServer/Services/RpcService.cs
+++ b/Engine/Source/Programs/Horde/HordeServer/Services/RpcService.cs
@@ -314,8 +314,7 @@ namespace HordeServer.Services
 				}
 
 				const bool bEnabled = true;
-				const bool bEphemeral = false;
-				Agent = await AgentService.CreateAgentAsync(Request.Name, bEnabled, bEphemeral, null, null);
+				Agent = await AgentService.CreateAgentAsync(Request.Name, bEnabled, null, null);
 			}
 
 			// Make sure we're allowed to create sessions on this agent
