Horde: Fix tests.

diff --git a/Engine/Source/Programs/Horde/HordeServerTests/Fixture.cs b/Engine/Source/Programs/Horde/HordeServerTests/Fixture.cs
--- a/Engine/Source/Programs/Horde/HordeServerTests/Fixture.cs
+++ b/Engine/Source/Programs/Horde/HordeServerTests/Fixture.cs
@@ -135,7 +135,7 @@ namespace HordeServerTests
 				"text/plain", new MemoryStream(Encoding.UTF8.GetBytes(Job1ArtifactData)));
 
 			Agent1Name = "testAgent1";
-			Agent1 = await AgentService.CreateAgentAsync(Agent1Name, true, false, null, null);
+			Agent1 = await AgentService.CreateAgentAsync(Agent1Name, true, null, null);
 		}
 
 		private class FixtureGraph : IGraph
