Horde: Unify capabilities/requirements for agents.

gRPC-style devices are no longer supported (as much as they ever were); instead, each agent supplies a flat list of properties and resources which jobs can consume. Pools can specify conditions to collect related agents together which are evaluated on the list of properties.

Dynamic pools are distinguished from explicit (user-defined) pools, and are updated on the agent along with session updates.

diff --git a/Engine/Source/Programs/Horde/HordeServer/Api/Agents.cs b/Engine/Source/Programs/Horde/HordeServer/Api/Agents.cs
--- a/Engine/Source/Programs/Horde/HordeServer/Api/Agents.cs
+++ b/Engine/Source/Programs/Horde/HordeServer/Api/Agents.cs
@@ -246,9 +246,9 @@ namespace HordeServer.Api
 		public DateTime? FinishTime { get; set; }
 
 		/// <summary>
-		/// Capabilities of this agent
+		/// Properties of this agent
 		/// </summary>
-		public GetAgentCapabilitiesResponse? Capabilities { get; set; }
+		public List<string>? Properties { get; set; }
 
 		/// <summary>
 		/// Version of the software running during this session
@@ -272,7 +272,7 @@ namespace HordeServer.Api
 			this.Id = Session.Id.ToString();
 			this.StartTime = Session.StartTime;
 			this.FinishTime = Session.FinishTime;
-			this.Capabilities = (Session.Capabilities != null)? new GetAgentCapabilitiesResponse(Session.Capabilities) : null;
+			this.Properties = (Session.Properties != null) ? new List<string>(Session.Properties) : null;
 			this.Version = Session.Version?.ToString();
 		}
 	}
@@ -416,6 +416,16 @@ namespace HordeServer.Api
 		/// </summary>
 		public string? Channel { get; set; }
 
+		/// <summary>
+		/// Properties for the agent
+		/// </summary>
+		public List<string> Properties { get; set; }
+
+		/// <summary>
+		/// Resources for the agent
+		/// </summary>
+		public Dictionary<string, int> Resources { get; set; }
+
 		/// <summary>
 		/// Last update time of this agent
 		/// </summary>
@@ -434,7 +444,7 @@ namespace HordeServer.Api
 		/// <summary>
 		/// Capabilities of this agent
 		/// </summary>
-		public GetAgentCapabilitiesResponse? Capabilities { get; }
+		public object? Capabilities { get; }
 
 		/// <summary>
 		/// Array of active leases.
@@ -458,6 +468,8 @@ namespace HordeServer.Api
 		{
 			this.Id = null!;
 			this.Name = null!;
+			Properties = new List<string>();
+			Resources = new Dictionary<string, int>();
 			this.Leases = new List<GetAgentLeaseResponse>();
 		}
 
@@ -465,13 +477,14 @@ namespace HordeServer.Api
 		/// Constructor
 		/// </summary>
 		/// <param name="Agent">The agent to construct from</param>
-		/// <param name="Pools">List of pools for this agent</param>
 		/// <param name="bIncludeAcl">Whether to include the ACL in the response</param>
-		public GetAgentResponse(IAgent Agent, IEnumerable<PoolId> Pools, bool bIncludeAcl)
+		public GetAgentResponse(IAgent Agent, bool bIncludeAcl)
 		{
 			this.Id = Agent.Id.ToString();
 			this.Name = Agent.Id.ToString();
 			this.Enabled = Agent.Enabled;
+			this.Properties = new List<string>(Agent.Properties);
+			this.Resources = new Dictionary<string, int>(Agent.Resources);
 			this.SessionId = Agent.SessionId?.ToString();
 			this.Online = Agent.IsSessionValid(DateTime.UtcNow);
 			this.Deleted = Agent.Deleted;
@@ -489,9 +502,9 @@ namespace HordeServer.Api
 			this.Version = Agent.Version?.ToString();
 			this.Channel = Agent.Channel?.ToString();
 			this.UpdateTime = Agent.UpdateTime;
-			this.Pools = Pools.Select(x => x.ToString()).ToList();
+			this.Pools = Agent.GetPools().Select(x => x.ToString()).ToList();
 			this.Workspaces = Agent.Workspaces.ConvertAll(x => new GetAgentWorkspaceResponse(x));
-			this.Capabilities = new GetAgentCapabilitiesResponse(Agent.Capabilities);
+			this.Capabilities = new { Devices = new[] { new { Properties = Agent.Properties, Resources = Agent.Resources } } };
 			this.Leases = Agent.Leases.ConvertAll(x => new GetAgentLeaseResponse(x));
 			this.Acl = (bIncludeAcl && Agent.Acl != null) ? new GetAclResponse(Agent.Acl) : null;
 			this.Comment = Agent.Comment;
diff --git a/Engine/Source/Programs/Horde/HordeServer/Api/Pools.cs b/Engine/Source/Programs/Horde/HordeServer/Api/Pools.cs
--- a/Engine/Source/Programs/Horde/HordeServer/Api/Pools.cs
+++ b/Engine/Source/Programs/Horde/HordeServer/Api/Pools.cs
@@ -1,5 +1,6 @@
 // Copyright Epic Games, Inc. All Rights Reserved.
 
+using EpicGames.Horde.Common;
 using HordeServer.Models;
 using System;
 using System.Collections.Generic;
@@ -21,9 +22,9 @@ namespace HordeServer.Api
 		public string Name { get; set; } = null!;
 
 		/// <summary>
-		/// Agent requirements to automatically include in this pool
+		/// Condition to satisfy for agents to be included in this pool
 		/// </summary>
-		public AgentRequirements? Requirements { get; set; }
+		public Condition? Condition { get; set; }
 
 		/// <summary>
 		/// Whether to enable autoscaling for this pool
@@ -87,7 +88,7 @@ namespace HordeServer.Api
 		/// <summary>
 		/// Requirements for this pool
 		/// </summary>
-		public AgentRequirements? Requirements { get; set; }
+		public Condition? Condition { get; set; }
 
 		/// <summary>
 		/// Whether to enable autoscaling for this pool
@@ -148,9 +149,9 @@ namespace HordeServer.Api
 		public string Name { get; set; }
 
 		/// <summary>
-		/// Requirements for agents to be auto-added to the pool
+		/// Condition for agents to be auto-added to the pool
 		/// </summary>
-		public AgentRequirements? Requirements { get; set; }
+		public Condition? Condition { get; set; }
 
 		/// <summary>
 		/// Whether to enable autoscaling for this pool
@@ -196,7 +197,7 @@ namespace HordeServer.Api
 		{
 			this.Id = Pool.Id.ToString();
 			this.Name = Pool.Name;
-			this.Requirements = Pool.Requirements;
+			this.Condition = Pool.Condition;
 			this.EnableAutoscaling = Pool.EnableAutoscaling;
 			this.MinAgents = Pool.MinAgents;
 			this.NumReserveAgents = Pool.NumReserveAgents;
diff --git a/Engine/Source/Programs/Horde/HordeServer/Api/Streams.cs b/Engine/Source/Programs/Horde/HordeServer/Api/Streams.cs
--- a/Engine/Source/Programs/Horde/HordeServer/Api/Streams.cs
+++ b/Engine/Source/Programs/Horde/HordeServer/Api/Streams.cs
@@ -347,11 +347,6 @@ namespace HordeServer.Api
 		/// </summary>
 		public string? Workspace { get; set; }
 
-		/// <summary>
-		/// Pool of agents to use to execute this work
-		/// </summary>
-		public GetAgentRequirementsResponse? Requirements { get; set; }
-
 		/// <summary>
 		/// Path to the temporary storage dir
 		/// </summary>
diff --git a/Engine/Source/Programs/Horde/HordeServer/Api/Workers.cs b/Engine/Source/Programs/Horde/HordeServer/Api/Workers.cs
deleted file mode 100644
diff --git a/Engine/Source/Programs/Horde/HordeServer/Collections/IAgentCollection.cs b/Engine/Source/Programs/Horde/HordeServer/Collections/IAgentCollection.cs
--- a/Engine/Source/Programs/Horde/HordeServer/Collections/IAgentCollection.cs
+++ b/Engine/Source/Programs/Horde/HordeServer/Collections/IAgentCollection.cs
@@ -107,10 +107,12 @@ namespace HordeServer.Collections
 		/// <param name="SessionId">New session id</param>
 		/// <param name="SessionExpiresAt">Expiry time for the new session</param>
 		/// <param name="Status">Status of the agent</param>
-		/// <param name="Capabilities">Capabilities for the agent</param>
+		/// <param name="Properties">Properties for the current session</param>
+		/// <param name="Resources">Resources for the agent</param>
+		/// <param name="DynamicPools">New list of dynamic pools for the agent</param>
 		/// <param name="Version">Current version of the agent software</param>
 		/// <returns>New agent state</returns>
-		Task<IAgent?> TryStartSessionAsync(IAgent Agent, ObjectId SessionId, DateTime SessionExpiresAt, AgentStatus Status, AgentCapabilities Capabilities, string? Version);
+		Task<IAgent?> TryStartSessionAsync(IAgent Agent, ObjectId SessionId, DateTime SessionExpiresAt, AgentStatus Status, IReadOnlyList<string> Properties, IReadOnlyDictionary<string, int> Resources, IReadOnlyList<PoolId> DynamicPools, string? Version);
 
 		/// <summary>
 		/// Attempt to update the agent state
@@ -118,10 +120,12 @@ namespace HordeServer.Collections
 		/// <param name="Agent">Agent instance</param>
 		/// <param name="Status">New status of the agent</param>
 		/// <param name="SessionExpiresAt">New expiry time for the current session</param>
-		/// <param name="Capabilities">Capabilities for the agent</param>
+		/// <param name="Properties">Properties for the current session</param>
+		/// <param name="Resources">Resources for the agent</param>
+		/// <param name="DynamicPools">New list of dynamic pools for the agent</param>
 		/// <param name="Leases">New set of leases</param>
 		/// <returns>True if the document was updated, false if another writer updated the document first</returns>
-		Task<IAgent?> TryUpdateSessionAsync(IAgent Agent, AgentStatus? Status, DateTime? SessionExpiresAt, AgentCapabilities? Capabilities, List<AgentLease>? Leases);
+		Task<IAgent?> TryUpdateSessionAsync(IAgent Agent, AgentStatus? Status, DateTime? SessionExpiresAt, IReadOnlyList<string>? Properties, IReadOnlyDictionary<string, int>? Resources, IReadOnlyList<PoolId>? DynamicPools, List<AgentLease>? Leases);
 
 		/// <summary>
 		/// Terminates the current session
diff --git a/Engine/Source/Programs/Horde/HordeServer/Collections/IPoolCollection.cs b/Engine/Source/Programs/Horde/HordeServer/Collections/IPoolCollection.cs
--- a/Engine/Source/Programs/Horde/HordeServer/Collections/IPoolCollection.cs
+++ b/Engine/Source/Programs/Horde/HordeServer/Collections/IPoolCollection.cs
@@ -1,5 +1,6 @@
 // Copyright Epic Games, Inc. All Rights Reserved.
 
+using EpicGames.Horde.Common;
 using HordeServer.Models;
 using HordeServer.Services;
 using HordeServer.Utilities;
@@ -27,13 +28,13 @@ namespace HordeServer.Collections
 		/// </summary>
 		/// <param name="Id">Unique id for the new pool</param>
 		/// <param name="Name">Name of the new pool</param>
-		/// <param name="Requirements">Requirements for agents to include in this pool</param>
+		/// <param name="Condition">Condition for agents to be included in this pool</param>
 		/// <param name="EnableAutoscaling">Whether to enable autoscaling for this pool</param>
 		/// <param name="MinAgents">Minimum number of agents in the pool</param>
 		/// <param name="NumReserveAgents">Minimum number of idle agents to maintain</param>
 		/// <param name="Properties">Properties for the pool</param>
 		/// <returns>The new pool document</returns>
-		Task<IPool> AddAsync(PoolId Id, string Name, AgentRequirements? Requirements = null, bool? EnableAutoscaling = null, int? MinAgents = null, int? NumReserveAgents = null, IEnumerable<KeyValuePair<string, string>>? Properties = null);
+		Task<IPool> AddAsync(PoolId Id, string Name, Condition? Condition = null, bool? EnableAutoscaling = null, int? MinAgents = null, int? NumReserveAgents = null, IEnumerable<KeyValuePair<string, string>>? Properties = null);
 
 		/// <summary>
 		/// Enumerates all the pools
@@ -66,7 +67,7 @@ namespace HordeServer.Collections
 		/// </summary>
 		/// <param name="Pool">The pool to update</param>
 		/// <param name="NewName">New name for the pool</param>
-		/// <param name="NewRequirements">New requirements for the pool</param>
+		/// <param name="NewCondition">New condition for the pool</param>
 		/// <param name="NewEnableAutoscaling">New setting for whether to enable autoscaling</param>
 		/// <param name="NewMinAgents">Minimum number of agents in the pool</param>
 		/// <param name="NewNumReserveAgents">Minimum number of idle agents to maintain</param>
@@ -75,6 +76,6 @@ namespace HordeServer.Collections
 		/// <param name="LastScaleUpTime">New time for last (auto) scale up</param>
 		/// <param name="LastScaleDownTime">New time for last (auto) scale down</param>
 		/// <returns>Async task</returns>
-		Task<IPool?> TryUpdateAsync(IPool Pool, string? NewName = null, AgentRequirements? NewRequirements = null, bool? NewEnableAutoscaling = null, int? NewMinAgents = null, int? NewNumReserveAgents = null, List<AgentWorkspace>? NewWorkspaces = null, Dictionary<string, string?>? NewProperties = null, DateTime? LastScaleUpTime = null, DateTime? LastScaleDownTime = null);
+		Task<IPool?> TryUpdateAsync(IPool Pool, string? NewName = null, Condition? NewCondition = null, bool? NewEnableAutoscaling = null, int? NewMinAgents = null, int? NewNumReserveAgents = null, List<AgentWorkspace>? NewWorkspaces = null, Dictionary<string, string?>? NewProperties = null, DateTime? LastScaleUpTime = null, DateTime? LastScaleDownTime = null);
 	}
 }
diff --git a/Engine/Source/Programs/Horde/HordeServer/Collections/ISessionCollection.cs b/Engine/Source/Programs/Horde/HordeServer/Collections/ISessionCollection.cs
--- a/Engine/Source/Programs/Horde/HordeServer/Collections/ISessionCollection.cs
+++ b/Engine/Source/Programs/Horde/HordeServer/Collections/ISessionCollection.cs
@@ -25,9 +25,10 @@ namespace HordeServer.Collections
 		/// <param name="Id">The session id</param>
 		/// <param name="AgentId">The agent this session is for</param>
 		/// <param name="StartTime">Start time of this session</param>
-		/// <param name="Capabilities">Properties of this agent at the time the session started</param>
+		/// <param name="Properties">Properties of this agent at the time the session started</param>
+		/// <param name="Resources">Resources which the agent has</param>
 		/// <param name="Version">Version of the agent software</param>
-		Task<ISession> AddAsync(ObjectId Id, AgentId AgentId, DateTime StartTime, AgentCapabilities? Capabilities, string? Version);
+		Task<ISession> AddAsync(ObjectId Id, AgentId AgentId, DateTime StartTime, IReadOnlyList<string>? Properties, IReadOnlyDictionary<string, int>? Resources, string? Version);
 
 		/// <summary>
 		/// Gets information about a particular session
@@ -60,9 +61,10 @@ namespace HordeServer.Collections
 		/// </summary>
 		/// <param name="SessionId">The session to update</param>
 		/// <param name="FinishTime">Time at which the session finished</param>
-		/// <param name="Capabilities">Capabilities of the agent at the end of the session</param>
+		/// <param name="Properties">The agent properties</param>
+		/// <param name="Resources">Resources which the agent has</param>
 		/// <returns>Async task</returns>
-		Task UpdateAsync(ObjectId SessionId, DateTime FinishTime, AgentCapabilities Capabilities);
+		Task UpdateAsync(ObjectId SessionId, DateTime FinishTime, IReadOnlyList<string> Properties, IReadOnlyDictionary<string, int> Resources);
 
 		/// <summary>
 		/// Delete a session from the collection
diff --git a/Engine/Source/Programs/Horde/HordeServer/Collections/Impl/AgentCollection.cs b/Engine/Source/Programs/Horde/HordeServer/Collections/Impl/AgentCollection.cs
--- a/Engine/Source/Programs/Horde/HordeServer/Collections/Impl/AgentCollection.cs
+++ b/Engine/Source/Programs/Horde/HordeServer/Collections/Impl/AgentCollection.cs
@@ -25,6 +25,29 @@ namespace HordeServer.Collections.Impl
 	/// </summary>
 	public class AgentCollection : IAgentCollection
 	{
+		/// <summary>
+		/// Legacy information about a device attached to an agent
+		/// </summary>
+		class DeviceCapabilities
+		{
+			[BsonIgnoreIfNull]
+			public HashSet<string>? Properties { get; set; }
+
+			[BsonIgnoreIfNull]
+			public Dictionary<string, int>? Resources { get; set; }
+		}
+
+		/// <summary>
+		/// Legacy capabilities of an agent
+		/// </summary>
+		class AgentCapabilities
+		{
+			public List<DeviceCapabilities> Devices { get; set; } = new List<DeviceCapabilities>();
+
+			[BsonIgnoreIfNull]
+			public HashSet<string>? Properties { get; set; }
+		}
+
 		/// <summary>
 		/// Concrete implementation of an agent document
 		/// </summary>
@@ -51,6 +74,9 @@ namespace HordeServer.Collections.Impl
 			[BsonElement("Version2")]
 			public string? Version { get; set; }
 
+			public List<string>? Properties { get; set; }
+			public Dictionary<string, int>? Resources { get; set; }
+
 			[BsonIgnoreIfNull]
 			public AgentSoftwareChannelName? Channel { get; set; }
 
@@ -60,6 +86,7 @@ namespace HordeServer.Collections.Impl
 			[BsonIgnoreIfNull]
 			public DateTime? LastUpgradeTime { get; set; }
 
+			public List<PoolId> DynamicPools { get; set; } = new List<PoolId>();
 			public List<PoolId> Pools { get; set; } = new List<PoolId>();
 			public bool RequestConform { get; set; }
 
@@ -82,9 +109,12 @@ namespace HordeServer.Collections.Impl
 			public uint UpdateIndex { get; set; }
 			public string? Comment { get; set; }
 
+			IReadOnlyList<PoolId> IAgent.DynamicPools => DynamicPools;
 			IReadOnlyList<PoolId> IAgent.ExplicitPools => Pools;
 			IReadOnlyList<AgentWorkspace> IAgent.Workspaces => Workspaces;
 			IReadOnlyList<AgentLease> IAgent.Leases => Leases ?? EmptyLeases;
+			IReadOnlyList<string> IAgent.Properties => Properties ?? Capabilities.Devices.FirstOrDefault()?.Properties.ToList() ?? new List<string>();
+			IReadOnlyDictionary<string, int> IAgent.Resources => Resources ?? Capabilities.Devices.FirstOrDefault()?.Resources ?? new Dictionary<string, int>();
 
 			[BsonConstructor]
 			private AgentDocument()
@@ -321,7 +351,7 @@ namespace HordeServer.Collections.Impl
 		}
 
 		/// <inheritdoc/>
-		public async Task<IAgent?> TryUpdateSessionAsync(IAgent AgentInterface, AgentStatus? Status, DateTime? SessionExpiresAt, AgentCapabilities? Capabilities, List<AgentLease>? Leases)
+		public async Task<IAgent?> TryUpdateSessionAsync(IAgent AgentInterface, AgentStatus? Status, DateTime? SessionExpiresAt, IReadOnlyList<string>? Properties, IReadOnlyDictionary<string, int>? Resources, IReadOnlyList<PoolId>? DynamicPools, List<AgentLease>? Leases)
 		{
 			AgentDocument Agent = (AgentDocument)AgentInterface;
 
@@ -339,10 +369,20 @@ namespace HordeServer.Collections.Impl
 				Agent.SessionExpiresAt = SessionExpiresAt.Value;
 				Updates.Add(UpdateBuilder.Set(x => x.SessionExpiresAt, Agent.SessionExpiresAt));
 			}
-			if (Capabilities != null)
+			if (Properties != null)
+			{
+				Agent.Properties = Properties.OrderBy(x => x, StringComparer.OrdinalIgnoreCase).ToList();
+				Updates.Add(UpdateBuilder.Set(x => x.Properties, Agent.Properties));
+			}
+			if (Resources != null)
+			{
+				Agent.Resources = new Dictionary<string, int>(Resources);
+				Updates.Add(UpdateBuilder.Set(x => x.Resources, Agent.Resources));
+			}
+			if (DynamicPools != null && !DynamicPools.SequenceEqual(Agent.DynamicPools))
 			{
-				Agent.Capabilities = Capabilities;
-				Updates.Add(UpdateBuilder.Set(x => x.Capabilities, Agent.Capabilities));
+				Agent.DynamicPools = new List<PoolId>(DynamicPools);
+				Updates.Add(UpdateBuilder.Set(x => x.DynamicPools, Agent.DynamicPools));
 			}
 			if (Leases != null)
 			{
@@ -411,9 +451,12 @@ namespace HordeServer.Collections.Impl
 		}
 
 		/// <inheritdoc/>
-		public async Task<IAgent?> TryStartSessionAsync(IAgent AgentInterface, ObjectId SessionId, DateTime SessionExpiresAt, AgentStatus Status, AgentCapabilities Capabilities, string? Version)
+		public async Task<IAgent?> TryStartSessionAsync(IAgent AgentInterface, ObjectId SessionId, DateTime SessionExpiresAt, AgentStatus Status, IReadOnlyList<string> Properties, IReadOnlyDictionary<string, int> Resources, IReadOnlyList<PoolId> DynamicPools, string? Version)
 		{
 			AgentDocument Agent = (AgentDocument)AgentInterface;
+			List<string> NewProperties = Properties.OrderBy(x => x, StringComparer.OrdinalIgnoreCase).ToList();
+			Dictionary<string, int> NewResources = new Dictionary<string, int>(Resources);
+			List<PoolId> NewDynamicPools = new List<PoolId>(DynamicPools);
 
 			// Reset the agent to use the new session
 			UpdateDefinitionBuilder<AgentDocument> UpdateBuilder = Builders<AgentDocument>.Update;
@@ -424,7 +467,9 @@ namespace HordeServer.Collections.Impl
 			Updates.Add(UpdateBuilder.Set(x => x.Status, Status));
 			Updates.Add(UpdateBuilder.Unset(x => x.Leases));
 			Updates.Add(UpdateBuilder.Unset(x => x.Deleted));
-			Updates.Add(UpdateBuilder.Set(x => x.Capabilities, Capabilities));
+			Updates.Add(UpdateBuilder.Set(x => x.Properties, NewProperties));
+			Updates.Add(UpdateBuilder.Set(x => x.Resources, NewResources));
+			Updates.Add(UpdateBuilder.Set(x => x.DynamicPools, NewDynamicPools));
 			Updates.Add(UpdateBuilder.Set(x => x.Version, Version));
 			Updates.Add(UpdateBuilder.Unset(x => x.RequestRestart));
 			Updates.Add(UpdateBuilder.Unset(x => x.RequestShutdown));
@@ -441,7 +486,9 @@ namespace HordeServer.Collections.Impl
 			Agent.Status = Status;
 			Agent.Leases = null;
 			Agent.Deleted = false;
-			Agent.Capabilities = Capabilities;
+			Agent.Properties = NewProperties;
+			Agent.Resources = NewResources;
+			Agent.DynamicPools = NewDynamicPools;
 			Agent.Version = Version;
 
 			return Agent;
diff --git a/Engine/Source/Programs/Horde/HordeServer/Collections/Impl/PoolCollection.cs b/Engine/Source/Programs/Horde/HordeServer/Collections/Impl/PoolCollection.cs
--- a/Engine/Source/Programs/Horde/HordeServer/Collections/Impl/PoolCollection.cs
+++ b/Engine/Source/Programs/Horde/HordeServer/Collections/Impl/PoolCollection.cs
@@ -1,5 +1,6 @@
 // Copyright Epic Games, Inc. All Rights Reserved.
 
+using EpicGames.Horde.Common;
 using HordeServer.Models;
 using HordeServer.Services;
 using HordeServer.Utilities;
@@ -29,8 +30,9 @@ namespace HordeServer.Collections.Impl
 			public PoolId Id { get; set; }
 			[BsonRequired]
 			public string Name { get; set; } = null!;
+			[BsonIgnoreIfNull]
+			public Condition? Condition { get; set; }
 			public List<AgentWorkspace> Workspaces { get; set; } = new List<AgentWorkspace>();
-			public AgentRequirements? Requirements { get; set; }
 			public Dictionary<string, string> Properties { get; set; } = new Dictionary<string, string>();
 			[BsonIgnoreIfDefault(true)]
 			public bool EnableAutoscaling { get; set; } = true;
@@ -52,11 +54,8 @@ namespace HordeServer.Collections.Impl
 			{
 				Id = Other.Id;
 				Name = Other.Name;
+				Condition = Other.Condition;
 				Workspaces.AddRange(Other.Workspaces);
-				if(Other.Requirements != null)
-				{
-					Requirements = new AgentRequirements(Other.Requirements);
-				}
 				Properties = new Dictionary<string, string>(Other.Properties);
 				EnableAutoscaling = Other.EnableAutoscaling;
 				MinAgents = Other.MinAgents;
@@ -82,15 +81,12 @@ namespace HordeServer.Collections.Impl
 		}
 
 		/// <inheritdoc/>
-		public async Task<IPool> AddAsync(PoolId Id, string Name, AgentRequirements? Requirements, bool? EnableAutoscaling, int? MinAgents, int? NumReserveAgents, IEnumerable<KeyValuePair<string, string>>? Properties)
+		public async Task<IPool> AddAsync(PoolId Id, string Name, Condition? Condition, bool? EnableAutoscaling, int? MinAgents, int? NumReserveAgents, IEnumerable<KeyValuePair<string, string>>? Properties)
 		{
 			PoolDocument Pool = new PoolDocument();
 			Pool.Id = Id;
 			Pool.Name = Name;
-			if(Requirements != null)
-			{
-				Pool.Requirements = new AgentRequirements(Requirements);
-			}
+			Pool.Condition = Condition;
 			if (EnableAutoscaling != null)
 			{
 				Pool.EnableAutoscaling = EnableAutoscaling.Value;
@@ -175,22 +171,22 @@ namespace HordeServer.Collections.Impl
 		}
 
 		/// <inheritdoc/>
-		public Task<IPool?> TryUpdateAsync(IPool Pool, string? NewName, AgentRequirements? NewRequirements, bool? NewEnableAutoscaling, int? NewMinAgents, int? NewNumReserveAgents, List<AgentWorkspace>? NewWorkspaces, Dictionary<string, string?>? NewProperties, DateTime? LastScaleUpTime, DateTime? LastScaleDownTime)
+		public Task<IPool?> TryUpdateAsync(IPool Pool, string? NewName, Condition? NewCondition, bool? NewEnableAutoscaling, int? NewMinAgents, int? NewNumReserveAgents, List<AgentWorkspace>? NewWorkspaces, Dictionary<string, string?>? NewProperties, DateTime? LastScaleUpTime, DateTime? LastScaleDownTime)
 		{
 			TransactionBuilder<PoolDocument> Transaction = new TransactionBuilder<PoolDocument>();
 			if (NewName != null)
 			{
 				Transaction.Set(x => x.Name, NewName);
 			}
-			if (NewRequirements != null)
+			if (NewCondition != null)
 			{
-				if (NewRequirements.IsEmpty())
+				if (NewCondition.IsEmpty())
 				{
-					Transaction.Unset(x => x.Requirements!);
+					Transaction.Unset(x => x.Condition!);
 				}
 				else
 				{
-					Transaction.Set(x => x.Requirements, NewRequirements);
+					Transaction.Set(x => x.Condition, NewCondition);
 				}
 			}
 			if (NewEnableAutoscaling != null)
diff --git a/Engine/Source/Programs/Horde/HordeServer/Collections/Impl/SessionCollection.cs b/Engine/Source/Programs/Horde/HordeServer/Collections/Impl/SessionCollection.cs
--- a/Engine/Source/Programs/Horde/HordeServer/Collections/Impl/SessionCollection.cs
+++ b/Engine/Source/Programs/Horde/HordeServer/Collections/Impl/SessionCollection.cs
@@ -32,20 +32,30 @@ namespace HordeServer.Collections.Impl
 
 			public DateTime StartTime { get; set; }
 			public DateTime? FinishTime { get; set; }
-			public AgentCapabilities? Capabilities { get; set; }
+			public List<string>? Properties { get; set; }
+			public Dictionary<string, int>? Resources { get; set; }
 			public string? Version { get; set; }
 
+			IReadOnlyList<string>? ISession.Properties => Properties;
+
 			[BsonConstructor]
 			private SessionDocument()
 			{
 			}
 
-			public SessionDocument(ObjectId Id, AgentId AgentId, DateTime StartTime, AgentCapabilities? Capabilities, string? Version)
+			public SessionDocument(ObjectId Id, AgentId AgentId, DateTime StartTime, IReadOnlyList<string>? Properties, IReadOnlyDictionary<string, int>? Resources, string? Version)
 			{
 				this.Id = Id;
 				this.AgentId = AgentId;
 				this.StartTime = StartTime;
-				this.Capabilities = Capabilities;
+				if (Properties != null)
+				{
+					this.Properties = new List<string>(Properties);
+				}
+				if (Resources != null)
+				{
+					this.Resources = new Dictionary<string, int>(Resources);
+				}
 				this.Version = Version;
 			}
 		}
@@ -72,9 +82,9 @@ namespace HordeServer.Collections.Impl
 		}
 
 		/// <inheritdoc/>
-		public async Task<ISession> AddAsync(ObjectId Id, AgentId AgentId, DateTime StartTime, AgentCapabilities? Capabilities, string? Version)
+		public async Task<ISession> AddAsync(ObjectId Id, AgentId AgentId, DateTime StartTime, IReadOnlyList<string>? Properties, IReadOnlyDictionary<string, int>? Resources, string? Version)
 		{
-			SessionDocument NewSession = new SessionDocument(Id, AgentId, StartTime, Capabilities, Version);
+			SessionDocument NewSession = new SessionDocument(Id, AgentId, StartTime, Properties, Resources, Version);
 			await Sessions.InsertOneAsync(NewSession);
 			return NewSession;
 		}
@@ -112,11 +122,12 @@ namespace HordeServer.Collections.Impl
 		}
 
 		/// <inheritdoc/>
-		public Task UpdateAsync(ObjectId SessionId, DateTime FinishTime, AgentCapabilities Capabilities)
+		public Task UpdateAsync(ObjectId SessionId, DateTime FinishTime, IReadOnlyList<string> Properties, IReadOnlyDictionary<string, int> Resources)
 		{
 			UpdateDefinition<SessionDocument> Update = Builders<SessionDocument>.Update
 				.Set(x => x.FinishTime, FinishTime)
-				.Set(x => x.Capabilities, Capabilities);
+				.Set(x => (IEnumerable<string>?)x.Properties, Properties)
+				.Set(x => (IEnumerable<KeyValuePair<string, int>>?)x.Resources, Resources);
 
 			return Sessions.FindOneAndUpdateAsync(x => x.Id == SessionId, Update);
 		}
diff --git a/Engine/Source/Programs/Horde/HordeServer/Compute/Impl/ComputeService.cs b/Engine/Source/Programs/Horde/HordeServer/Compute/Impl/ComputeService.cs
--- a/Engine/Source/Programs/Horde/HordeServer/Compute/Impl/ComputeService.cs
+++ b/Engine/Source/Programs/Horde/HordeServer/Compute/Impl/ComputeService.cs
@@ -213,11 +213,8 @@ namespace HordeServer.Compute.Impl
 				return null;
 			}
 
-			// Get all the current agent properties
-			Lazy<Task<List<string>>> Properties = new Lazy<Task<List<string>>>(() => GetAgentProperties(Agent));
-
 			// Find a task to execute
-			(IoHash, ComputeTaskInfo)? Entry = await TaskScheduler.DequeueAsync(RequirementsHash => CheckRequirements(RequirementsHash, Properties), CancellationToken);
+			(IoHash, ComputeTaskInfo)? Entry = await TaskScheduler.DequeueAsync(RequirementsHash => CheckRequirements(Agent, RequirementsHash), CancellationToken);
 			if (Entry != null)
 			{
 				(IoHash RequirementsHash, ComputeTaskInfo TaskInfo) = Entry.Value;
@@ -231,7 +228,7 @@ namespace HordeServer.Compute.Impl
 				string LeaseName = $"Remote action ({TaskInfo.TaskHash})";
 				byte[] Payload = Any.Pack(ComputeTask).ToByteArray();
 
-				AgentLease Lease = new AgentLease(ObjectId.GenerateNewId(), LeaseName, null, null, null, LeaseState.Pending, Payload, new AgentRequirements(), null);
+				AgentLease Lease = new AgentLease(ObjectId.GenerateNewId(), LeaseName, null, null, null, LeaseState.Pending, Payload, null);
 				Logger.LogDebug("Created lease {LeaseId} for channel {ChannelId} task {TaskHash} req {RequirementsHash}", Lease.Id, ComputeTask.ChannelId, (CbObjectAttachment)ComputeTask.TaskHash, (CbObjectAttachment)ComputeTask.RequirementsHash);
 				return Lease;
 			}
@@ -287,76 +284,17 @@ namespace HordeServer.Compute.Impl
 		/// <summary>
 		/// Checks that an agent matches the necessary criteria to execute a task
 		/// </summary>
+		/// <param name="Agent"></param>
 		/// <param name="RequirementsHash"></param>
-		/// <param name="LazyProperties">Properties for this agentthat the agent is in</param>
 		/// <returns></returns>
-		async ValueTask<bool> CheckRequirements(IoHash RequirementsHash, Lazy<Task<List<string>>> LazyProperties)
+		async ValueTask<bool> CheckRequirements(IAgent Agent, IoHash RequirementsHash)
 		{
 			Requirements? Requirements = await GetCachedRequirementsAsync(RequirementsHash);
 			if (Requirements == null)
 			{
 				return false;
 			}
-
-			if (Requirements.Condition != null)
-			{
-				List<string> Properties = await LazyProperties.Value;
-				if (!Requirements.Condition.Evaluate(x => FindAgentProperty(x, Properties)))
-				{
-					return false;
-				}
-			}
-
-			return true;
-		}
-
-		/// <summary>
-		/// Finds property values from a sorted list of Name=Value pairs
-		/// </summary>
-		/// <param name="Name">Name of the property to find</param>
-		/// <param name="Properties">The full list of properties</param>
-		/// <returns>Property values</returns>
-		static IEnumerable<string> FindAgentProperty(string Name, List<string> Properties)
-		{
-			int Index = Properties.BinarySearch(Name, StringComparer.OrdinalIgnoreCase);
-			if (Index < 0)
-			{
-				Index = ~Index;
-				while (Index < Properties.Count)
-				{
-					string Property = Properties[Index];
-					if (Property.Length <= Name.Length || !Property.StartsWith(Name, StringComparison.OrdinalIgnoreCase) || Property[Name.Length] != '=')
-					{
-						break;
-					}
-					yield return Property.Substring(Name.Length + 1);
-				}
-			}
-		}
-
-		/// <summary>
-		/// Gets a named property for an agent
-		/// </summary>
-		/// <param name="Agent">The agent instance</param>
-		/// <returns>The property value, or an empty string if it does not eixst</returns>
-		async Task<List<string>> GetAgentProperties(IAgent Agent)
-		{
-			List<string> Properties = new List<string>();
-			Properties.Add($"Name={Agent.Id}");
-
-			foreach (IPool Pool in Agent.GetPools(await CachedPools.GetCached()))
-			{
-				Properties.Add($"Pool={Pool.Id}");
-			}
-
-			HashSet<string>? DeviceProperties = Agent.Capabilities.PrimaryDevice.Properties;
-			if (DeviceProperties != null)
-			{
-				Properties.AddRange(DeviceProperties);
-			}
-
-			Properties.Sort(StringComparer.OrdinalIgnoreCase);
-			return Properties;
+			return Agent.MeetsRequirements(Requirements);
 		}
 
 		/// <summary>
diff --git a/Engine/Source/Programs/Horde/HordeServer/Controllers/AgentsController.cs b/Engine/Source/Programs/Horde/HordeServer/Controllers/AgentsController.cs
--- a/Engine/Source/Programs/Horde/HordeServer/Controllers/AgentsController.cs
+++ b/Engine/Source/Programs/Horde/HordeServer/Controllers/AgentsController.cs
@@ -98,14 +98,12 @@ namespace HordeServer.Controllers
 		{
 			GlobalPermissionsCache PermissionsCache = new GlobalPermissionsCache();
 			List<IAgent> Agents = await AgentService.FindAgentsAsync(Pool?.ToObjectId(), ModifiedAfter?.UtcDateTime, Index, Count);
-			List<IPool> Pools = await PoolService.GetPoolsAsync();
 
 			List<object> Responses = new List<object>();
 			foreach (IAgent Agent in Agents)
 			{
 				bool bIncludeAcl = await AgentService.AuthorizeAsync(Agent, AclAction.ViewPermissions, User, PermissionsCache);
-				List<PoolId> PoolIds = Agent.GetPools(Pools).Select(x => x.Id).ToList();
-				Responses.Add(new GetAgentResponse(Agent, PoolIds, bIncludeAcl).ApplyFilter(Filter));
+				Responses.Add(new GetAgentResponse(Agent, bIncludeAcl).ApplyFilter(Filter));
 			}
 
 			return Responses;
@@ -132,8 +130,7 @@ namespace HordeServer.Controllers
 
 			GlobalPermissionsCache Cache = new GlobalPermissionsCache();
 			bool bIncludeAcl = await AgentService.AuthorizeAsync(Agent, AclAction.ViewPermissions, User, Cache);
-			List<PoolId> PoolIds = Agent.GetPools(Pools).Select(x => x.Id).ToList();
-			return new GetAgentResponse(Agent, PoolIds, bIncludeAcl).ApplyFilter(Filter);
+			return new GetAgentResponse(Agent, bIncludeAcl).ApplyFilter(Filter);
 		}
 
 		/// <summary>
diff --git a/Engine/Source/Programs/Horde/HordeServer/Controllers/DebugController.cs b/Engine/Source/Programs/Horde/HordeServer/Controllers/DebugController.cs
--- a/Engine/Source/Programs/Horde/HordeServer/Controllers/DebugController.cs
+++ b/Engine/Source/Programs/Horde/HordeServer/Controllers/DebugController.cs
@@ -433,7 +433,7 @@ namespace HordeServer.Controllers
 				return Forbid();
 			}
 
-			return await JobTaskSource.GetStatus();
+			return JobTaskSource.GetStatus();
 		}
 
 		/// <summary>
@@ -880,7 +880,7 @@ namespace HordeServer.Controllers
 			else if (ChangeType == "shrink")
 			{
 				Agents = await AgentCollection.FindAsync();
-				Agents = Agents.FindAll(a => a.InPool(Pool)).ToList();
+				Agents = Agents.FindAll(a => a.IsInPool(Pool.Id)).ToList();
 				await FleetManager.ShrinkPool(Pool, Agents, Count.Value);
 				Message = $"Shrank pool {Pool.Name} by {Count}";
 			}
diff --git a/Engine/Source/Programs/Horde/HordeServer/Controllers/PoolsController.cs b/Engine/Source/Programs/Horde/HordeServer/Controllers/PoolsController.cs
--- a/Engine/Source/Programs/Horde/HordeServer/Controllers/PoolsController.cs
+++ b/Engine/Source/Programs/Horde/HordeServer/Controllers/PoolsController.cs
@@ -60,7 +60,7 @@ namespace HordeServer.Controllers
 				return Forbid();
 			}
 
-			IPool NewPool = await PoolService.CreatePoolAsync(Create.Name, Create.Requirements, Create.EnableAutoscaling, Create.MinAgents, Create.NumReserveAgents, Create.Properties);
+			IPool NewPool = await PoolService.CreatePoolAsync(Create.Name, Create.Condition, Create.EnableAutoscaling, Create.MinAgents, Create.NumReserveAgents, Create.Properties);
 			return new CreatePoolResponse(NewPool.Id.ToString());
 		}
 
@@ -139,7 +139,7 @@ namespace HordeServer.Controllers
 				return NotFound();
 			}
 
-			await PoolService.UpdatePoolAsync(Pool, Update.Name, Update.Requirements, Update.EnableAutoscaling, Update.MinAgents, Update.NumReserveAgents, Update.Properties);
+			await PoolService.UpdatePoolAsync(Pool, Update.Name, Update.Condition, Update.EnableAutoscaling, Update.MinAgents, Update.NumReserveAgents, Update.Properties);
 			return new OkResult();
 		}
 
diff --git a/Engine/Source/Programs/Horde/HordeServer/Models/Agent.cs b/Engine/Source/Programs/Horde/HordeServer/Models/Agent.cs
--- a/Engine/Source/Programs/Horde/HordeServer/Models/Agent.cs
+++ b/Engine/Source/Programs/Horde/HordeServer/Models/Agent.cs
@@ -1,6 +1,7 @@
 // Copyright Epic Games, Inc. All Rights Reserved.
 
 using Google.Protobuf.WellKnownTypes;
+using EpicGames.Core;
 using HordeServer.Api;
 using HordeCommon;
 using HordeCommon.Rpc.Tasks;
@@ -19,6 +20,8 @@ using Serilog;
 using System.Diagnostics;
 using System.Threading.Tasks;
 using HordeServer.Collections;
+using EpicGames.Horde.Compute;
+using EpicGames.Horde.Common;
 
 namespace HordeServer.Models
 {
@@ -174,52 +177,6 @@ namespace HordeServer.Models
 		}
 	}
 
-	/// <summary>
-	/// Information about a device allocated for a lease
-	/// </summary>
-	public class AgentLeaseDevice
-	{
-		/// <summary>
-		/// Index of the device in the agent's device list
-		/// </summary>
-		[BsonRequired]
-		public int Index { get; set; }
-
-		/// <summary>
-		/// Handle of the device in the lease, ie. the logical name in the context of the work being done, not the physical device name.
-		/// </summary>
-		[BsonRequired]
-		public string? Handle { get; set; }
-
-		/// <summary>
-		/// Resources claimed by the lease. If null, the whole device is claimed, and nothing else can be allocated to it.
-		/// </summary>
-		[BsonIgnoreIfNull]
-		public Dictionary<string, int>? Resources { get; set; }
-
-		/// <summary>
-		/// Private constructor for serialization
-		/// </summary>
-		[BsonConstructor]
-		private AgentLeaseDevice()
-		{
-			Handle = null!;
-		}
-
-		/// <summary>
-		/// Constructor
-		/// </summary>
-		/// <param name="Index"></param>
-		/// <param name="Handle"></param>
-		/// <param name="Resources"></param>
-		public AgentLeaseDevice(int Index, string? Handle, Dictionary<string, int>? Resources)
-		{
-			this.Index = Index;
-			this.Handle = Handle;
-			this.Resources = Resources;
-		}
-	}
-
 	/// <summary>
 	/// Document describing an active lease
 	/// </summary>
@@ -273,21 +230,15 @@ namespace HordeServer.Models
 		public bool Active { get; set; }
 
 		/// <summary>
-		/// Requirements for this lease
+		/// Resources used by this lease
 		/// </summary>
-		public AgentRequirements Requirements { get; set; }
+		public IReadOnlyDictionary<string, int>? Resources { get; set; }
 
 		/// <summary>
 		/// For leases in the pending state, encodes an "any" protobuf containing the payload for the agent to execute the lease.
 		/// </summary>
 		public byte[]? Payload { get; set; }
 
-		/// <summary>
-		/// List of devices allocated to this lease
-		/// </summary>
-		[BsonIgnoreIfNull]
-		public List<AgentLeaseDevice>? Devices { get; set; }
-
 		/// <summary>
 		/// Private constructor
 		/// </summary>
@@ -295,8 +246,6 @@ namespace HordeServer.Models
 		private AgentLease()
 		{
 			Name = String.Empty;
-			Requirements = null!;
-			Devices = null!;
 		}
 
 		/// <summary>
@@ -309,9 +258,8 @@ namespace HordeServer.Models
 		/// <param name="LogId">Unique id for the log</param>
 		/// <param name="State">State for the lease</param>
 		/// <param name="Payload">Encoded "any" protobuf describing the contents of the payload</param>
-		/// <param name="Requirements">Requirements for this lease</param>
-		/// <param name="Devices">Device mapping for this lease</param>
-		public AgentLease(ObjectId Id, string Name, StreamId? StreamId, PoolId? PoolId, ObjectId? LogId, LeaseState State, byte[]? Payload, AgentRequirements Requirements, List<AgentLeaseDevice>? Devices)
+		/// <param name="Resources">Resources required for this lease</param>
+		public AgentLease(ObjectId Id, string Name, StreamId? StreamId, PoolId? PoolId, ObjectId? LogId, LeaseState State, byte[]? Payload, IReadOnlyDictionary<string, int>? Resources)
 		{
 			this.Id = Id;
 			this.Name = Name;
@@ -320,8 +268,7 @@ namespace HordeServer.Models
 			this.LogId = LogId;
 			this.State = State;
 			this.Payload = Payload;
-			this.Requirements = Requirements;
-			this.Devices = Devices;
+			this.Resources = Resources;
 			StartTime = DateTime.UtcNow;
 		}
 
@@ -422,6 +369,17 @@ namespace HordeServer.Models
 		}
 	}
 
+	/// <summary>
+	/// Well-known property names for agents
+	/// </summary>
+	static class KnownPropertyNames
+	{
+		/// <summary>
+		/// Pools that this agent belongs to
+		/// </summary>
+		public const string Pool = "Pool";
+	}
+
 	/// <summary>
 	/// Mirrors an Agent document in the database
 	/// </summary>
@@ -467,6 +425,16 @@ namespace HordeServer.Models
 		/// </summary>
 		public string? Comment { get; }
 
+		/// <summary>
+		/// List of properties for this agent
+		/// </summary>
+		public IReadOnlyList<string> Properties { get; }
+
+		/// <summary>
+		/// List of resources available to the agent
+		/// </summary>
+		public IReadOnlyDictionary<string, int> Resources { get; }
+
 		/// <summary>
 		/// Channel for the software running on this agent. Uses <see cref="AgentSoftwareService.DefaultChannelName"/> if not specified
 		/// </summary>
@@ -482,6 +450,11 @@ namespace HordeServer.Models
 		/// </summary>
 		public DateTime? LastUpgradeTime { get; }
 
+		/// <summary>
+		/// Dynamically applied pools
+		/// </summary>
+		public IReadOnlyList<PoolId> DynamicPools { get; }
+
 		/// <summary>
 		/// List of manually assigned pools for agent
 		/// </summary>
@@ -517,11 +490,6 @@ namespace HordeServer.Models
 		/// </summary>
 		public int? ConformAttemptCount { get; }
 
-		/// <summary>
-		/// Capabilities of this agent
-		/// </summary>
-		public AgentCapabilities Capabilities { get; }
-
 		/// <summary>
 		/// Array of active leases.
 		/// </summary>
@@ -558,216 +526,146 @@ namespace HordeServer.Models
 		}
 
 		/// <summary>
-		/// Gets a list of pools for the given agent. Includes all automatically assigned pools based on agent capabilities.
+		/// Tests whether an agent is in the given pool
 		/// </summary>
-		/// <param name="Agent">The agent instance</param>
-		/// <param name="Pools">List of all available pools</param>
-		/// <returns>List of pools</returns>
-		public static IEnumerable<IPool> GetPools(this IAgent Agent, IEnumerable<IPool> Pools)
+		/// <param name="Agent"></param>
+		/// <param name="PoolId"></param>
+		/// <returns></returns>
+		public static bool IsInPool(this IAgent Agent, PoolId PoolId)
 		{
-			return Pools.Where(x => Agent.InPool(x));
+			return Agent.DynamicPools.Contains(PoolId) || Agent.ExplicitPools.Contains(PoolId);
 		}
 
 		/// <summary>
-		/// Checks whether an agent is in the given pool
+		/// Get all the pools for each agent
 		/// </summary>
-		/// <param name="Agent">The agent to check</param>
-		/// <param name="Pool">The pool to test against</param>
-		/// <returns>True if the agent is in the pool</returns>
-		public static bool InPool(this IAgent Agent, IPool Pool)
+		/// <param name="Agent">The agent to query</param>
+		/// <returns></returns>
+		public static IEnumerable<PoolId> GetPools(this IAgent Agent)
 		{
-			if (Agent.ExplicitPools.Contains(Pool.Id))
+			foreach (PoolId PoolId in Agent.DynamicPools)
 			{
-				return true;
+				yield return PoolId;
 			}
-			else if (Pool.Requirements != null && AgentExtensions.TryCreateLease(Agent.Capabilities, Pool.Requirements, Enumerable.Empty<AgentLease>(), out _))
+			foreach (PoolId PoolId in Agent.ExplicitPools)
 			{
-				return true;
+				yield return PoolId;
 			}
-			return false;
 		}
 
 		/// <summary>
-		/// Determine whether it's possible to add a lease for the given resources
+		/// Tests whether an agent has a particular property
 		/// </summary>
-		/// <param name="Agent">The agent to create a lease for</param>
-		/// <param name="Pool">The pool required</param>
-		/// <param name="Requirements">Requirements for this lease</param>
-		/// <param name="LeasedDevices">On success, recieves the list of leased devices</param>
-		/// <returns>True if the new lease can be granted</returns>
-		public static bool TryCreateLease(this IAgent Agent, IPool Pool, AgentRequirements? Requirements, out List<AgentLeaseDevice>? LeasedDevices)
+		/// <param name="Agent"></param>
+		/// <param name="Property"></param>
+		/// <returns></returns>
+		public static bool HasProperty(this IAgent Agent, string Property)
 		{
-			if (!Agent.Enabled || Agent.Status != AgentStatus.Ok || !Agent.InPool(Pool))
-			{
-				LeasedDevices = null!;
-				return false;
-			}
-			return TryCreateLease(Agent.Capabilities, Requirements, Agent.Leases, out LeasedDevices);
+			return Agent.Properties.BinarySearch(Property, StringComparer.OrdinalIgnoreCase) >= 0;
 		}
 
 		/// <summary>
-		/// Attempts to match the requirements for a particular agent, and return the list of leased devices that fulfils the requirements
+		/// Finds property values from a sorted list of Name=Value pairs
 		/// </summary>
-		/// <param name="Capabilities">Capabilities of the gent</param>
-		/// <param name="Requirements">Requirements for the lease</param>
-		/// <param name="Leases">List of current leases</param>
-		/// <param name="LeasedDevices">Receives the list of device leases</param>
-		/// <returns>True if successful</returns>
-		public static bool TryCreateLease(AgentCapabilities Capabilities, AgentRequirements? Requirements, IEnumerable<AgentLease> Leases, out List<AgentLeaseDevice>? LeasedDevices)
+		/// <param name="Agent">The agent to query</param>
+		/// <param name="Name">Name of the property to find</param>
+		/// <returns>Property values</returns>
+		public static IEnumerable<string> GetPropertyValues(this IAgent Agent, string Name)
 		{
-			List<AgentLeaseDevice>? Result = null;
-			if (Requirements != null)
+			if (Name.Equals(KnownPropertyNames.Pool, StringComparison.OrdinalIgnoreCase))
 			{
-				if (!Requirements.Shared && Leases.Any())
+				foreach (PoolId PoolId in Agent.GetPools())
 				{
-					LeasedDevices = null;
-					return false;
+					yield return PoolId.ToString();
 				}
-
-				if (Requirements.Properties != null)
+			}
+			else
+			{
+				int Index = Agent.Properties.BinarySearch(Name, StringComparer.OrdinalIgnoreCase);
+				if (Index < 0)
 				{
-					foreach (string Property in Requirements.Properties)
+					Index = ~Index;
+					while (Index < Agent.Properties.Count)
 					{
-						if (Capabilities.Properties == null || !Capabilities.Properties.Contains(Property))
+						string Property = Agent.Properties[Index];
+						if (Property.Length <= Name.Length || !Property.StartsWith(Name, StringComparison.OrdinalIgnoreCase) || Property[Name.Length] != '=')
 						{
-							LeasedDevices = null!;
-							return false;
+							break;
 						}
-					}
-				}
-
-				if (Requirements.Devices != null)
-				{
-					if (!TryCreateDeviceLeases(Capabilities.Devices, Requirements.Devices, 0, 0, Leases.SelectMany(x => x.Devices), out Result))
-					{
-						LeasedDevices = null!;
-						return false;
+						yield return Property.Substring(Name.Length + 1);
 					}
 				}
 			}
-
-			LeasedDevices = Result;
-			return true;
 		}
 
 		/// <summary>
-		/// Try to match up the required devices for a lease to the available devices on an agent. This is done through a recursive search with backtracking for completeness,
-		/// due to not knowing which devices in the lease will correspond to each device in the agent, but it should be quick in practice due to agents having a small number 
-		/// of homogenous devices.
+		/// Evaluates a condition against an agent
 		/// </summary>
-		/// <param name="CapableDevices">Device capabilities for the agent</param>
-		/// <param name="RequiredDevices">Device requirements for the lease</param>
-		/// <param name="RequiredDeviceIdx">Current index of the device in the lease requirements</param>
-		/// <param name="LeasedDeviceMask">Bitmask for devices that have currently been allocated. This is limited to 32, which is unlikely to be a problem in practice.</param>
-		/// <param name="CurrentLeasedDevices">The current set of leased devices</param>
-		/// <param name="LeasedDevices">On success, receives a list of the leased devices, mapping each requested device to a device in the agent device array.</param>
-		/// <returns>True if the leases were allocated</returns>
-		private static bool TryCreateDeviceLeases(List<DeviceCapabilities> CapableDevices, List<DeviceRequirements> RequiredDevices, int RequiredDeviceIdx, uint LeasedDeviceMask, IEnumerable<AgentLeaseDevice> CurrentLeasedDevices, [MaybeNullWhen(false)] out List<AgentLeaseDevice> LeasedDevices)
+		/// <param name="Agent">The agent to evaluate</param>
+		/// <param name="Condition">The condition to evaluate</param>
+		/// <returns>True if the agent satisfies the condition</returns>
+		public static bool SatisfiesCondition(this IAgent Agent, Condition Condition)
 		{
-			// If we've assigned all the required devices now, we can create the list of assigned devices
-			if (RequiredDeviceIdx == RequiredDevices.Count)
-			{
-				LeasedDevices = new List<AgentLeaseDevice>(RequiredDevices.Count);
-				foreach (DeviceRequirements Device in RequiredDevices)
-				{
-					LeasedDevices.Add(new AgentLeaseDevice(-1, Device.Handle, Device.Resources));
-				}
-				return true;
-			}
-
-			// Otherwise try to assign the next device
-			DeviceRequirements RequiredDevice = RequiredDevices[RequiredDeviceIdx];
-			for (int DeviceIdx = 0; DeviceIdx < CapableDevices.Count; DeviceIdx++)
-			{
-				uint DeviceFlag = 1U << DeviceIdx;
-				if ((LeasedDeviceMask & DeviceFlag) == 0 && MatchDevice(CapableDevices[DeviceIdx], RequiredDevice, CurrentLeasedDevices.Where(x => x.Index == DeviceIdx)))
-				{
-					List<AgentLeaseDevice>? Result;
-					if (TryCreateDeviceLeases(CapableDevices, RequiredDevices, RequiredDeviceIdx + 1, LeasedDeviceMask | DeviceFlag, CurrentLeasedDevices, out Result))
-					{
-						LeasedDevices = Result;
-						LeasedDevices[RequiredDeviceIdx].Index = DeviceIdx;
-						return true;
-					}
-				}
-			}
-
-			// Otherwise failed
-			LeasedDevices = null!;
-			return false;
+			return Condition.Evaluate(x => Agent.GetPropertyValues(x));
 		}
 
 		/// <summary>
-		/// Determines if an agent device has the given requirements
+		/// Determine whether it's possible to add a lease for the given resources
 		/// </summary>
-		/// <param name="Capabilities">The device capabilities</param>
+		/// <param name="Agent">The agent to create a lease for</param>
 		/// <param name="Requirements">Requirements for the lease</param>
-		/// <param name="Leases">Current leases for the device</param>
-		/// <returns>True if the device can satisfy the given requirements</returns>
-		private static bool MatchDevice(DeviceCapabilities Capabilities, DeviceRequirements Requirements, IEnumerable<AgentLeaseDevice> Leases)
+		/// <returns>True if the new lease can be granted</returns>
+		public static bool MeetsRequirements(this IAgent Agent, Requirements Requirements)
+		{
+			return MeetsRequirements(Agent, Requirements.Condition, Requirements.Resources, Requirements.Exclusive);
+		}
+
+		/// <summary>
+		/// Determine whether it's possible to add a lease for the given resources
+		/// </summary>
+		/// <param name="Agent">The agent to create a lease for</param>
+		/// <param name="Exclusive">Whether t</param>
+		/// <param name="Condition">Condition to satisfy</param>
+		/// <param name="Resources">Resources required to execute</param>
+		/// <returns>True if the new lease can be granted</returns>
+		public static bool MeetsRequirements(this IAgent Agent, Condition? Condition, Dictionary<string, int>? Resources, bool Exclusive)
 		{
-			// Check the device has all the required properties
-			if (Requirements.Properties != null)
+			if (!Agent.Enabled || Agent.Status != AgentStatus.Ok)
 			{
-				foreach (string Property in Requirements.Properties)
-				{
-					if (Capabilities.Properties == null || !Capabilities.Properties.Contains(Property))
-					{
-						return false;
-					}
-				}
+				return false;
 			}
-
-			// Check the device has all the required resources
-			if (Requirements.Resources == null)
+			if (Exclusive && Agent.Leases.Any())
 			{
-				// Requires exclusive access to the device
-				if (Leases.Any())
-				{
-					return false;
-				}
+				return false;
 			}
-			else
+			if (Condition != null && !Agent.SatisfiesCondition(Condition))
 			{
-				// Requires shared access to the device. Check the device can be shared.
-				if (Capabilities.Resources == null)
-				{
-					return false;
-				}
-
-				// Check there are enough available of each resource type
-				foreach (KeyValuePair<string, int> Resource in Requirements.Resources)
+				return false;
+			}
+			if (Resources != null)
+			{
+				foreach ((string Name, int Count) in Resources)
 				{
-					// Make sure the device has enough of the named resource to start with
 					int RemainingCount;
-					if (!Capabilities.Resources.TryGetValue(Resource.Key, out RemainingCount) || RemainingCount < Resource.Value)
+					if (!Agent.Resources.TryGetValue(Name, out RemainingCount))
 					{
 						return false;
 					}
-
-					// Check each existing lease of this device
-					foreach (AgentLeaseDevice Lease in Leases)
+					foreach (AgentLease Lease in Agent.Leases)
 					{
-						// If the lease has an exclusive reservation, we can't use it
-						if (Lease.Resources == null)
-						{
-							return false;
-						}
-
-						// Update the remaining count for this resource
-						int UsedCount;
-						if (Lease.Resources.TryGetValue(Resource.Key, out UsedCount))
+						if (Lease.Resources != null)
 						{
-							RemainingCount -= UsedCount;
-							if (RemainingCount < Resource.Value)
-							{
-								return false;
-							}
+							int LeaseCount;
+							Lease.Resources.TryGetValue(Name, out LeaseCount);
+							RemainingCount -= LeaseCount;
 						}
 					}
+					if (RemainingCount < Count)
+					{
+						return false;
+					}
 				}
 			}
-
 			return true;
 		}
 
@@ -804,18 +702,11 @@ namespace HordeServer.Models
 		/// <returns></returns>
 		public static AgentWorkspace? GetAutoSdkWorkspace(this IAgent Agent, PerforceCluster Cluster)
 		{
-			if (Agent.Capabilities.Devices.Count > 0)
+			foreach (AutoSdkWorkspace AutoSdk in Cluster.AutoSdk)
 			{
-				DeviceCapabilities PrimaryDevice = Agent.Capabilities.Devices[0];
-				if (PrimaryDevice.Properties != null)
+				if (AutoSdk.Stream != null && AutoSdk.Properties.All(x => Agent.Properties.Contains(x)))
 				{
-					foreach (AutoSdkWorkspace AutoSdk in Cluster.AutoSdk)
-					{
-						if (AutoSdk.Stream != null && AutoSdk.Properties.All(x => PrimaryDevice.Properties.Contains(x)))
-						{
-							return new AgentWorkspace(Cluster.Name, AutoSdk.UserName, AutoSdk.Name ?? "AutoSDK", AutoSdk.Stream!, null, true);
-						}
-					}
+					return new AgentWorkspace(Cluster.Name, AutoSdk.UserName, AutoSdk.Name ?? "AutoSDK", AutoSdk.Stream!, null, true);
 				}
 			}
 			return null;
diff --git a/Engine/Source/Programs/Horde/HordeServer/Models/Pool.cs b/Engine/Source/Programs/Horde/HordeServer/Models/Pool.cs
--- a/Engine/Source/Programs/Horde/HordeServer/Models/Pool.cs
+++ b/Engine/Source/Programs/Horde/HordeServer/Models/Pool.cs
@@ -1,5 +1,6 @@
 // Copyright Epic Games, Inc. All Rights Reserved.
 
+using EpicGames.Horde.Common;
 using HordeServer.Api;
 using HordeServer.Utilities;
 using MongoDB.Bson;
@@ -30,14 +31,14 @@ namespace HordeServer.Models
 		public string Name { get; }
 
 		/// <summary>
-		/// List of workspaces currently assigned to this pool
+		/// Condition for agents to automatically be included in this pool
 		/// </summary>
-		public IReadOnlyList<AgentWorkspace> Workspaces { get; }
+		public Condition? Condition { get; }
 
 		/// <summary>
-		/// Requirements for agents to be included in this pool
+		/// List of workspaces currently assigned to this pool
 		/// </summary>
-		public AgentRequirements? Requirements { get; }
+		public IReadOnlyList<AgentWorkspace> Workspaces { get; }
 
 		/// <summary>
 		/// Arbitrary properties related to this pool
diff --git a/Engine/Source/Programs/Horde/HordeServer/Models/Session.cs b/Engine/Source/Programs/Horde/HordeServer/Models/Session.cs
--- a/Engine/Source/Programs/Horde/HordeServer/Models/Session.cs
+++ b/Engine/Source/Programs/Horde/HordeServer/Models/Session.cs
@@ -37,9 +37,9 @@ namespace HordeServer.Models
 		public DateTime? FinishTime { get; }
 
 		/// <summary>
-		/// Capabilities of this agent at the time the session started
+		/// Properties of this agent at the time the session started
 		/// </summary>
-		public AgentCapabilities? Capabilities { get; }
+		public IReadOnlyList<string>? Properties { get; }
 
 		/// <summary>
 		/// Version of the agent software
diff --git a/Engine/Source/Programs/Horde/HordeServer/Models/Worker.cs b/Engine/Source/Programs/Horde/HordeServer/Models/Worker.cs
deleted file mode 100644
diff --git a/Engine/Source/Programs/Horde/HordeServer/Services/AgentService.cs b/Engine/Source/Programs/Horde/HordeServer/Services/AgentService.cs
--- a/Engine/Source/Programs/Horde/HordeServer/Services/AgentService.cs
+++ b/Engine/Source/Programs/Horde/HordeServer/Services/AgentService.cs
@@ -113,6 +113,11 @@ namespace HordeServer.Services
 		/// </summary>
 		Dictionary<AgentId, AgentEventListener> Listeners = new Dictionary<AgentId, AgentEventListener>();
 
+		/// <summary>
+		/// Lazily updated list of current pools
+		/// </summary>
+		AsyncCachedValue<List<IPool>> PoolsList;
+
 		/// <summary>
 		/// Constructor
 		/// </summary>
@@ -121,12 +126,13 @@ namespace HordeServer.Services
 		/// <param name="Sessions">Collection of session documents</param>
 		/// <param name="AclService">The ACL service</param>
 		/// <param name="DowntimeService">The downtime service</param>
+		/// <param name="PoolCollection">Pool collection</param>
 		/// <param name="DogStatsd">The DogStatsd metric client</param>
 		/// <param name="TaskSources">The list of task sources</param>
 		/// <param name="ApplicationLifetime"></param>
 		/// <param name="Logger">Log output writer</param>
 		/// <param name="Clock">Clock</param>
-		public AgentService(IAgentCollection Agents, ILeaseCollection Leases, ISessionCollection Sessions, AclService AclService, IDowntimeService DowntimeService, IDogStatsd DogStatsd, IEnumerable<ITaskSource> TaskSources, IHostApplicationLifetime ApplicationLifetime, ILogger<AgentService> Logger, IClock Clock)
+		public AgentService(IAgentCollection Agents, ILeaseCollection Leases, ISessionCollection Sessions, AclService AclService, IDowntimeService DowntimeService, IPoolCollection PoolCollection, IDogStatsd DogStatsd, IEnumerable<ITaskSource> TaskSources, IHostApplicationLifetime ApplicationLifetime, ILogger<AgentService> Logger, IClock Clock)
 			: base(TimeSpan.FromSeconds(30.0), Logger)
 		{
 			this.Agents = Agents;
@@ -134,6 +140,7 @@ namespace HordeServer.Services
 			this.Sessions = Sessions;
 			this.AclService = AclService;
 			this.DowntimeService = DowntimeService;
+			this.PoolsList = new AsyncCachedValue<List<IPool>>(() => PoolCollection.GetAsync(), TimeSpan.FromSeconds(30.0));
 			this.DogStatsd = DogStatsd;
 			this.TaskSources = TaskSources.ToArray();
 			this.ApplicationLifetime = ApplicationLifetime;
@@ -259,15 +266,32 @@ namespace HordeServer.Services
 			}
 		}
 
+		async ValueTask<List<PoolId>> GetDynamicPoolsAsync(IAgent Agent)
+		{
+			List<PoolId> NewDynamicPools = new List<PoolId>();
+
+			List<IPool> Pools = await PoolsList.GetAsync();
+			foreach (IPool Pool in Pools)
+			{
+				if (Pool.Condition != null && Agent.SatisfiesCondition(Pool.Condition))
+				{
+					NewDynamicPools.Add(Pool.Id);
+				}
+			}
+
+			return NewDynamicPools;
+		}
+
 		/// <summary>
 		/// Creates a new agent session
 		/// </summary>
 		/// <param name="Agent">The agent to create a session for</param>
 		/// <param name="Status">Current status of the agent</param>
-		/// <param name="Capabilities">Capabilities of the agent software</param>
+		/// <param name="Properties">Properties for the agent</param>
+		/// <param name="Resources">Resources which the agent has</param>
 		/// <param name="Version">Version of the software that's running</param>
 		/// <returns>New agent state</returns>
-		public async Task<IAgent> CreateSessionAsync(IAgent Agent, AgentStatus Status, AgentCapabilities Capabilities, string? Version)
+		public async Task<IAgent> CreateSessionAsync(IAgent Agent, AgentStatus Status, IReadOnlyList<string> Properties, IReadOnlyDictionary<string, int> Resources, string? Version)
 		{
 			for (; ; )
 			{
@@ -290,11 +314,14 @@ namespace HordeServer.Services
 					}
 
 					// Create a new session document
-					ISession NewSession = await Sessions.AddAsync(ObjectId.GenerateNewId(), Agent.Id, Clock.UtcNow, Capabilities, Version);
+					ISession NewSession = await Sessions.AddAsync(ObjectId.GenerateNewId(), Agent.Id, Clock.UtcNow, Properties, Resources, Version);
 					DateTime SessionExpiresAt = UtcNow + SessionExpiryTime;
 
+					// Get the new dynamic pools for the agent
+					List<PoolId> DynamicPools = await GetDynamicPoolsAsync(Agent);
+
 					// Reset the agent to use the new session
-					NewAgent = await Agents.TryStartSessionAsync(Agent, NewSession.Id, SessionExpiresAt, Status, Capabilities, Version);
+					NewAgent = await Agents.TryStartSessionAsync(Agent, NewSession.Id, SessionExpiresAt, Status, Properties, Resources, DynamicPools, Version);
 					if(NewAgent != null)
 					{
 						Agent = NewAgent;
@@ -477,10 +504,11 @@ namespace HordeServer.Services
 		/// <param name="InAgent">The current agent state</param>
 		/// <param name="SessionId">Id of the session</param>
 		/// <param name="Status">New status for the agent</param>
-		/// <param name="Capabilities">New agent capabilities</param>
+		/// <param name="Properties">New agent properties</param>
+		/// <param name="Resources">New agent resources</param>
 		/// <param name="NewLeases">New list of leases for this session</param>
 		/// <returns>Updated agent state</returns>
-		public async Task<IAgent?> UpdateSessionAsync(IAgent InAgent, ObjectId SessionId, AgentStatus Status, AgentCapabilities? Capabilities, IList<HordeCommon.Rpc.Messages.Lease> NewLeases)
+		public async Task<IAgent?> UpdateSessionAsync(IAgent InAgent, ObjectId SessionId, AgentStatus Status, IReadOnlyList<string>? Properties, IReadOnlyDictionary<string, int>? Resources, IList<HordeCommon.Rpc.Messages.Lease> NewLeases)
 		{
 			AgentEventListener? Listener = null;
 			try
@@ -593,8 +621,11 @@ namespace HordeServer.Services
 						}
 					}
 
+					// Get the new dynamic pools for the agent
+					List<PoolId> DynamicPools = await GetDynamicPoolsAsync(Agent);
+
 					// Update the agent, and try to create new lease documents if we succeed
-					IAgent? NewAgent = await Agents.TryUpdateSessionAsync(Agent, Status, SessionExpiresAt, Capabilities, bUpdateLeases ? Leases : null);
+					IAgent? NewAgent = await Agents.TryUpdateSessionAsync(Agent, Status, SessionExpiresAt, Properties, Resources, DynamicPools, bUpdateLeases ? Leases : null);
 					if (NewAgent != null)
 					{
 						break;
@@ -768,7 +799,7 @@ namespace HordeServer.Services
 				}
 
 				// Update the session document
-				await Sessions.UpdateAsync(SessionId, FinishTime, Agent.Capabilities);
+				await Sessions.UpdateAsync(SessionId, FinishTime, Agent.Properties, Agent.Resources);
 				return true;
 			}
 			return false;
diff --git a/Engine/Source/Programs/Horde/HordeServer/Services/AutoscaleService.cs b/Engine/Source/Programs/Horde/HordeServer/Services/AutoscaleService.cs
--- a/Engine/Source/Programs/Horde/HordeServer/Services/AutoscaleService.cs
+++ b/Engine/Source/Programs/Horde/HordeServer/Services/AutoscaleService.cs
@@ -166,15 +166,15 @@ namespace HordeServer.Services
 
 			// Get all the pools
 			List<IPool> Pools = await PoolCollection.GetAsync();
-			Dictionary<IPool, PoolData> PoolToData = Pools.ToDictionary(x => x, x => new PoolData(x));
+			Dictionary<PoolId, PoolData> PoolToData = Pools.ToDictionary(x => x.Id, x => new PoolData(x));
 
 			// Find pool utilization over the query period
 			foreach (AgentData AgentData in AgentIdToData.Values)
 			{
-				foreach (IPool Pool in AgentData.Agent.GetPools(Pools))
+				foreach (PoolId PoolId in AgentData.Agent.GetPools())
 				{
 					PoolData? PoolData;
-					if (PoolToData.TryGetValue(Pool, out PoolData))
+					if (PoolToData.TryGetValue(PoolId, out PoolData))
 					{
 						PoolData.Add(AgentData);
 					}
diff --git a/Engine/Source/Programs/Horde/HordeServer/Services/ConfigService.cs b/Engine/Source/Programs/Horde/HordeServer/Services/ConfigService.cs
--- a/Engine/Source/Programs/Horde/HordeServer/Services/ConfigService.cs
+++ b/Engine/Source/Programs/Horde/HordeServer/Services/ConfigService.cs
@@ -603,8 +603,7 @@ namespace HordeServer.Services
 				IPool? Pool = await PoolService.GetPoolAsync(PoolIdValue);
 				if (Pool == null)
 				{
-					AgentRequirements Requirements = new AgentRequirements(new List<DeviceRequirements>() { new DeviceRequirements(null, new HashSet<string>(){ "OSFamily=Windows"}, null) }, null, null, false);
-					await PoolService.CreatePoolAsync(Request.DefaultPoolName, Requirements: Requirements, Properties: new Dictionary<string, string>() { ["Color"] = "0" });
+					await PoolService.CreatePoolAsync(Request.DefaultPoolName, Condition: "OSFamily == 'Windows'", Properties: new Dictionary<string, string>() { ["Color"] = "0" });
 				}
 			}
 
diff --git a/Engine/Source/Programs/Horde/HordeServer/Services/Impl/AwsFleetManager.cs b/Engine/Source/Programs/Horde/HordeServer/Services/Impl/AwsFleetManager.cs
--- a/Engine/Source/Programs/Horde/HordeServer/Services/Impl/AwsFleetManager.cs
+++ b/Engine/Source/Programs/Horde/HordeServer/Services/Impl/AwsFleetManager.cs
@@ -86,7 +86,7 @@ namespace HordeServer.Services.Impl
 			for (int Idx = 0; Idx < SortedAgents.Count && Count > 0; Idx++)
 			{
 				IAgent Agent = SortedAgents[Idx];
-				if (Agent.Capabilities.PrimaryDevice.HasProperty(AwsTagProperty))
+				if (Agent.HasProperty(AwsTagProperty))
 				{
 					for (IAgent? NewAgent = Agent; NewAgent != null; NewAgent = await AgentCollection.GetAsync(Agent.Id))
 					{
diff --git a/Engine/Source/Programs/Horde/HordeServer/Services/PerforceLoadBalancer.cs b/Engine/Source/Programs/Horde/HordeServer/Services/PerforceLoadBalancer.cs
--- a/Engine/Source/Programs/Horde/HordeServer/Services/PerforceLoadBalancer.cs
+++ b/Engine/Source/Programs/Horde/HordeServer/Services/PerforceLoadBalancer.cs
@@ -233,7 +233,7 @@ namespace HordeServer.Services
 		/// <returns></returns>
 		public Task<IPerforceServer?> SelectServerAsync(PerforceCluster Cluster)
 		{
-			HashSet<string> Properties = new HashSet<string>{ "HordeServer=1" };
+			List<string> Properties = new List<string>{ "HordeServer=1" };
 			return SelectServerAsync(Cluster, Properties);
 		}
 
@@ -245,9 +245,7 @@ namespace HordeServer.Services
 		/// <returns></returns>
 		public Task<IPerforceServer?> SelectServerAsync(PerforceCluster Cluster, IAgent Agent)
 		{
-			// Get the set of properties to use for checking server validity
-			HashSet<string>? Properties = Agent.Capabilities.Devices[0].Properties;
-			return SelectServerAsync(Cluster, Properties);
+			return SelectServerAsync(Cluster, Agent.Properties);
 		}
 
 		/// <summary>
@@ -256,7 +254,7 @@ namespace HordeServer.Services
 		/// <param name="Cluster"></param>
 		/// <param name="Properties"></param>
 		/// <returns></returns>
-		public async Task<IPerforceServer?> SelectServerAsync(PerforceCluster Cluster, HashSet<string>? Properties)
+		public async Task<IPerforceServer?> SelectServerAsync(PerforceCluster Cluster, IReadOnlyList<string> Properties)
 		{
 			// Find all the valid servers for this agent
 			List<PerforceServer> ValidServers = new List<PerforceServer>();
diff --git a/Engine/Source/Programs/Horde/HordeServer/Services/PoolService.cs b/Engine/Source/Programs/Horde/HordeServer/Services/PoolService.cs
--- a/Engine/Source/Programs/Horde/HordeServer/Services/PoolService.cs
+++ b/Engine/Source/Programs/Horde/HordeServer/Services/PoolService.cs
@@ -1,5 +1,6 @@
 // Copyright Epic Games, Inc. All Rights Reserved.
 
+using EpicGames.Horde.Common;
 using HordeCommon;
 using HordeServer.Api;
 using HordeServer.Collections;
@@ -61,15 +62,15 @@ namespace HordeServer.Services
 		/// Creates a new pool
 		/// </summary>
 		/// <param name="Name">Name of the new pool</param>
-		/// <param name="Requirements">Requirements for agents to add to this pool</param>
+		/// <param name="Condition">Condition for agents to be automatically included in this pool</param>
 		/// <param name="EnableAutoscaling">Whether to enable autoscaling for this pool</param>
 		/// <param name="MinAgents">Minimum number of agents in the pool</param>
 		/// <param name="NumReserveAgents">Minimum number of idle agents to maintain</param>
 		/// <param name="Properties">Properties for the new pool</param>
 		/// <returns>The new pool document</returns>
-		public Task<IPool> CreatePoolAsync(string Name, AgentRequirements? Requirements = null, bool? EnableAutoscaling = null, int? MinAgents = null, int? NumReserveAgents = null, Dictionary<string, string>? Properties = null)
+		public Task<IPool> CreatePoolAsync(string Name, Condition? Condition = null, bool? EnableAutoscaling = null, int? MinAgents = null, int? NumReserveAgents = null, Dictionary<string, string>? Properties = null)
 		{
-			return Pools.AddAsync(PoolId.Sanitize(Name), Name, Requirements, EnableAutoscaling, MinAgents, NumReserveAgents, Properties);
+			return Pools.AddAsync(PoolId.Sanitize(Name), Name, Condition, EnableAutoscaling, MinAgents, NumReserveAgents, Properties);
 		}
 
 		/// <summary>
@@ -87,17 +88,17 @@ namespace HordeServer.Services
 		/// </summary>
 		/// <param name="Pool">The pool to update</param>
 		/// <param name="NewName">The new name for the pool</param>
-		/// <param name="NewRequirements">New requirements for the pool</param>
+		/// <param name="NewCondition">New requirements for the pool</param>
 		/// <param name="NewEnableAutoscaling">Whether to enable autoscaling</param>
 		/// <param name="NewMinAgents">Minimum number of agents in the pool</param>
 		/// <param name="NewNumReserveAgents">Minimum number of idle agents to maintain</param>
 		/// <param name="NewProperties">Properties on the pool to update. Any properties with a value of null will be removed.</param>
 		/// <returns>Async task object</returns>
-		public async Task<IPool?> UpdatePoolAsync(IPool? Pool, string? NewName = null, AgentRequirements? NewRequirements = null, bool? NewEnableAutoscaling = null, int? NewMinAgents = null, int? NewNumReserveAgents = null, Dictionary<string, string?>? NewProperties = null)
+		public async Task<IPool?> UpdatePoolAsync(IPool? Pool, string? NewName = null, Condition? NewCondition = null, bool? NewEnableAutoscaling = null, int? NewMinAgents = null, int? NewNumReserveAgents = null, Dictionary<string, string?>? NewProperties = null)
 		{
 			for (; Pool != null; Pool = await Pools.GetAsync(Pool.Id))
 			{
-				IPool? NewPool = await Pools.TryUpdateAsync(Pool, NewName, NewRequirements, NewEnableAutoscaling, NewMinAgents, NewNumReserveAgents, null, NewProperties, null, null);
+				IPool? NewPool = await Pools.TryUpdateAsync(Pool, NewName, NewCondition, NewEnableAutoscaling, NewMinAgents, NewNumReserveAgents, null, NewProperties, null, null);
 				if (NewPool != null)
 				{
 					return NewPool;
@@ -136,7 +137,14 @@ namespace HordeServer.Services
 			HashSet<AgentWorkspace> Workspaces = new HashSet<AgentWorkspace>();
 
 			Dictionary<PoolId, IPool> PoolMapping = await GetPoolLookupAsync(ValidAtTime);
-			Workspaces.UnionWith(Agent.GetPools(PoolMapping.Values).SelectMany(x => x.Workspaces));
+			foreach (PoolId PoolId in Agent.GetPools())
+			{
+				IPool? Pool;
+				if (PoolMapping.TryGetValue(PoolId, out Pool))
+				{
+					Workspaces.UnionWith(Pool.Workspaces);
+				}
+			}
 
 			Globals Globals = await DatabaseService.GetGlobalsAsync();
 			Workspaces.UnionWith(Agent.GetAutoSdkWorkspaces(Globals, Workspaces.ToList()));
diff --git a/Engine/Source/Programs/Horde/HordeServer/Services/RpcService.cs b/Engine/Source/Programs/Horde/HordeServer/Services/RpcService.cs
--- a/Engine/Source/Programs/Horde/HordeServer/Services/RpcService.cs
+++ b/Engine/Source/Programs/Horde/HordeServer/Services/RpcService.cs
@@ -31,6 +31,8 @@ namespace HordeServer.Services
 	using AgentSoftwareVersion = StringId<IAgentSoftwareCollection>;
 	using IStream = HordeServer.Models.IStream;
 	using StreamId = StringId<IStream>;
+	using RpcAgentCapabilities = HordeCommon.Rpc.Messages.AgentCapabilities;
+	using RpcDeviceCapabilities = HordeCommon.Rpc.Messages.DeviceCapabilities;
 
 	/// <summary>
 	/// Implements the Horde gRPC service for bots updating their status and dequeing work
@@ -288,6 +290,25 @@ namespace HordeServer.Services
 			}
 		}
 
+		static void GetCapabilities(RpcAgentCapabilities? Capabilities, out List<string> Properties, out Dictionary<string, int> Resources)
+		{
+			Properties = new List<string>();
+			Resources = new Dictionary<string, int>();
+
+			if (Capabilities != null && Capabilities.Devices.Count > 0)
+			{
+				RpcDeviceCapabilities Device = Capabilities.Devices[0];
+				if (Device.Properties != null)
+				{
+					Properties = new List<string>(Device.Properties);
+				}
+				if (Device.Resources != null)
+				{
+					Resources = new Dictionary<string, int>(Device.Resources);
+				}
+			}
+		}
+
 		/// <summary>
 		/// Creates a new session
 		/// </summary>
@@ -323,8 +344,11 @@ namespace HordeServer.Services
 				throw new StructuredRpcException(StatusCode.PermissionDenied, "User is not authenticated to create session for {AgentId}", Request.Name);
 			}
 
+			// Get the known properties for this agent
+			GetCapabilities(Request.Capabilities, out List<string> Properties, out Dictionary<string, int> Resources);
+
 			// Create a new session
-			Agent = await AgentService.CreateSessionAsync(Agent, Request.Status, new AgentCapabilities(Request.Capabilities), Request.Version);
+			Agent = await AgentService.CreateSessionAsync(Agent, Request.Status, Properties, Resources, Request.Version);
 			if (Agent == null)
 			{
 				throw new StructuredRpcException(StatusCode.NotFound, "Agent {AgentId} not found", Request.Name);
@@ -375,17 +399,18 @@ namespace HordeServer.Services
 					}
 
 					// Get the new capabilities of this agent
-					AgentCapabilities? Capabilities = null;
+					List<string>? Properties = null;
+					Dictionary<string, int>? Resources = null;
 					if (Request.Capabilities != null)
 					{
-						Capabilities = new AgentCapabilities(Request.Capabilities);
+						GetCapabilities(Request.Capabilities, out Properties, out Resources);
 					}
 
 					// Capture the current agent update index. This allows us to detect if anything has changed.
 					uint UpdateIndex = Agent.UpdateIndex;
 
 					// Update the agent session and return to the caller if anything changes
-					Agent = await AgentService.UpdateSessionAsync(Agent, Request.SessionId.ToObjectId(), Request.Status, Capabilities, Request.Leases);
+					Agent = await AgentService.UpdateSessionAsync(Agent, Request.SessionId.ToObjectId(), Request.Status, Properties, Resources, Request.Leases);
 					if (Agent != null && Agent.UpdateIndex == UpdateIndex && (Agent.Leases.Count > 0 || Agent.Status != AgentStatus.Stopping))
 					{
 						Agent = await AgentService.WaitForLeaseAsync(Agent, NextRequestTask);
diff --git a/Engine/Source/Programs/Horde/HordeServer/Services/TelemetryService.cs b/Engine/Source/Programs/Horde/HordeServer/Services/TelemetryService.cs
--- a/Engine/Source/Programs/Horde/HordeServer/Services/TelemetryService.cs
+++ b/Engine/Source/Programs/Horde/HordeServer/Services/TelemetryService.cs
@@ -87,7 +87,7 @@ namespace HordeServer.Services
 			List<ILease> Leases = await LeaseCollection.FindLeasesAsync(MinTime: MinTime);
 
 			// Find all the agents
-			Dictionary<AgentId, List<IPool>> AgentToPools = Agents.ToDictionary(x => x.Id, x => x.GetPools(Pools).ToList());
+			Dictionary<AgentId, List<PoolId>> AgentToPoolIds = Agents.ToDictionary(x => x.Id, x => x.GetPools().ToList());
 
 			// Generate all the telemetry data
 			DateTime BucketMinTime = MinTime;
@@ -98,28 +98,28 @@ namespace HordeServer.Services
 
 				NewUtilizationTelemetry Telemetry = new NewUtilizationTelemetry(BucketMinTime, BucketMaxTime);
 				Telemetry.NumAgents = Agents.Count;
-				foreach (IPool Pool in AgentToPools.Values.SelectMany(x => x))
+				foreach (PoolId PoolId in AgentToPoolIds.Values.SelectMany(x => x))
 				{
-					Telemetry.FindOrAddPool(Pool.Id).NumAgents++;
+					Telemetry.FindOrAddPool(PoolId).NumAgents++;
 				}
 				foreach (ILease Lease in Leases)
 				{
 					if (Lease.StartTime < BucketMaxTime && Lease.FinishTime >= BucketMinTime)
 					{
-						List<IPool>? LeasePools;
-						if (AgentToPools.TryGetValue(Lease.AgentId, out LeasePools))
+						List<PoolId>? LeasePools;
+						if (AgentToPoolIds.TryGetValue(Lease.AgentId, out LeasePools))
 						{
 							DateTime FinishTime = Lease.FinishTime ?? BucketMaxTime;
 							double Time = new TimeSpan(Math.Min(FinishTime.Ticks, BucketMaxTime.Ticks) - Math.Max(Lease.StartTime.Ticks, BucketMinTime.Ticks)).TotalHours;
 
-							foreach (IPool Pool in LeasePools)
+							foreach (PoolId PoolId in LeasePools)
 							{
-								NewPoolUtilizationTelemetry PoolTelemetry = Telemetry.FindOrAddPool(Pool.Id);
+								NewPoolUtilizationTelemetry PoolTelemetry = Telemetry.FindOrAddPool(PoolId);
 								if (Lease.PoolId == null || Lease.StreamId == null)
 								{
 									PoolTelemetry.AdminTime += Time;
 								}
-								else if (Pool.Id == Lease.PoolId)
+								else if (PoolId == Lease.PoolId)
 								{
 									PoolTelemetry.FindOrAddStream(Lease.StreamId.Value).Time += Time;
 								}
diff --git a/Engine/Source/Programs/Horde/HordeServer/Tasks/Impl/ActionTaskSource.cs b/Engine/Source/Programs/Horde/HordeServer/Tasks/Impl/ActionTaskSource.cs
--- a/Engine/Source/Programs/Horde/HordeServer/Tasks/Impl/ActionTaskSource.cs
+++ b/Engine/Source/Programs/Horde/HordeServer/Tasks/Impl/ActionTaskSource.cs
@@ -484,7 +484,7 @@ namespace HordeServer.Tasks.Impl
 				ActionTask.LogId = LogFile.Id.ToString();
 
 				byte[] Payload = Any.Pack(ActionTask).ToByteArray();
-				AgentLease Lease = new AgentLease(Operation.Id, LeaseName, null, null, LogFile.Id, LeaseState.Pending, Payload, new AgentRequirements(), null);
+				AgentLease Lease = new AgentLease(Operation.Id, LeaseName, null, null, LogFile.Id, LeaseState.Pending, Payload, null);
 
 				// Try to set the lease on the subscription. If it fails, the subscriber has already been allocated
 				Operation.ResultTaskSource = new TaskCompletionSource<ActionExecuteResult>();
diff --git a/Engine/Source/Programs/Horde/HordeServer/Tasks/Impl/ClusteredActionTaskSource.cs b/Engine/Source/Programs/Horde/HordeServer/Tasks/Impl/ClusteredActionTaskSource.cs
--- a/Engine/Source/Programs/Horde/HordeServer/Tasks/Impl/ClusteredActionTaskSource.cs
+++ b/Engine/Source/Programs/Horde/HordeServer/Tasks/Impl/ClusteredActionTaskSource.cs
@@ -361,7 +361,7 @@ namespace HordeServer.Tasks.Impl
 			
 			byte[] Payload = Any.Pack(ActionTask).ToByteArray();
 			ObjectId OpId = ObjectId.Parse(Op.Id);
-			AgentLease Lease = new AgentLease(OpId, LeaseName, null, null, LogFile.Id, LeaseState.Pending, Payload, new AgentRequirements(), null);
+			AgentLease Lease = new AgentLease(OpId, LeaseName, null, null, LogFile.Id, LeaseState.Pending, Payload, null);
 			RemoteExecOperationInternal OpInternal = new RemoteExecOperationInternal(OpId, Lease.Id);
 
 			lock (ActiveOperations)
diff --git a/Engine/Source/Programs/Horde/HordeServer/Tasks/Impl/ConformTaskSource.cs b/Engine/Source/Programs/Horde/HordeServer/Tasks/Impl/ConformTaskSource.cs
--- a/Engine/Source/Programs/Horde/HordeServer/Tasks/Impl/ConformTaskSource.cs
+++ b/Engine/Source/Programs/Horde/HordeServer/Tasks/Impl/ConformTaskSource.cs
@@ -179,7 +179,7 @@ namespace HordeServer.Tasks.Impl
 
 					byte[] Payload = Any.Pack(Task).ToByteArray();
 
-					AgentLease Lease = new AgentLease(LeaseId, "Updating workspaces", null, null, Log.Id, LeaseState.Pending, Payload, new AgentRequirements(), null);
+					AgentLease Lease = new AgentLease(LeaseId, "Updating workspaces", null, null, Log.Id, LeaseState.Pending, Payload, null);
 					return TaskSubscription.FromResult(Lease);
 				}
 			}
diff --git a/Engine/Source/Programs/Horde/HordeServer/Tasks/Impl/JobTaskSource.cs b/Engine/Source/Programs/Horde/HordeServer/Tasks/Impl/JobTaskSource.cs
--- a/Engine/Source/Programs/Horde/HordeServer/Tasks/Impl/JobTaskSource.cs
+++ b/Engine/Source/Programs/Horde/HordeServer/Tasks/Impl/JobTaskSource.cs
@@ -1,6 +1,7 @@
 // Copyright Epic Games, Inc. All Rights Reserved.
 
 using EpicGames.Core;
+using EpicGames.Horde.Common;
 using Google.Protobuf;
 using Google.Protobuf.Reflection;
 using Google.Protobuf.WellKnownTypes;
@@ -27,6 +28,7 @@ using System.Text;
 using System.Threading;
 using System.Threading.Tasks;
 using HordeServer.Services;
+using EpicGames.Horde.Compute;
 
 namespace HordeServer.Tasks.Impl
 {
@@ -58,15 +60,10 @@ namespace HordeServer.Tasks.Impl
 			/// </summary>
 			public int BatchIdx;
 
-			/// <summary>
-			/// Requirements for the agent to execute this item
-			/// </summary>
-			public AgentRequirements Requirements;
-
 			/// <summary>
 			/// The pool of machines to allocate from
 			/// </summary>
-			public IPool Pool;
+			public PoolId PoolId;
 
 			/// <summary>
 			/// The type of workspace that this item should run in
@@ -100,16 +97,14 @@ namespace HordeServer.Tasks.Impl
 			/// <param name="Stream">The stream containing this job</param>
 			/// <param name="Job">The job instance</param>
 			/// <param name="BatchIdx">The batch index to execute</param>
-			/// <param name="Pool">Unique id of the pool of machines to allocate from</param>
-			/// <param name="Requirements">Requirements for the agent to execute this batch</param>
+			/// <param name="PoolId">Unique id of the pool of machines to allocate from</param>
 			/// <param name="Workspace">The workspace that this job should run in</param>
-			public QueueItem(IStream Stream, IJob Job, int BatchIdx, IPool Pool, AgentRequirements Requirements, AgentWorkspace Workspace)
+			public QueueItem(IStream Stream, IJob Job, int BatchIdx, PoolId PoolId, AgentWorkspace Workspace)
 			{
 				this.Stream = Stream;
 				this.Job = Job;
 				this.BatchIdx = BatchIdx;
-				this.Pool = Pool;
-				this.Requirements = Requirements;
+				this.PoolId = PoolId;
 				this.Workspace = Workspace;
 			}
 		}
@@ -328,11 +323,6 @@ namespace HordeServer.Tasks.Impl
 		/// <inheritdoc/>
 		public MessageDescriptor Descriptor => ExecuteJobTask.Descriptor;
 
-		/// <summary>
-		/// The next server index that should be used to determine what Perforce server an agent should use.
-		/// </summary>
-		public uint NextServerIndex { get; set; } = 0;
-
 		/// <summary>
 		/// Constructor
 		/// </summary>
@@ -372,21 +362,20 @@ namespace HordeServer.Tasks.Impl
 		/// Gets an object containing the stats of the queue for diagnostic purposes.
 		/// </summary>
 		/// <returns>Status object</returns>
-		public async Task<object> GetStatus()
+		public object GetStatus()
 		{
-			List<IPool> AllPools = await PoolCollection.GetAsync();
 			lock (LockObject)
 			{
 				List<object> OutputItems = new List<object>();
 				foreach (QueueItem QueueItem in Queue)
 				{
-					OutputItems.Add(new { JobId = QueueItem.Job.Id.ToString(), BatchId = QueueItem.Batch.Id.ToString(), PoolId = QueueItem.Pool.Id.ToString(), Workspace = QueueItem.Workspace, Requirements = QueueItem.Requirements });
+					OutputItems.Add(new { JobId = QueueItem.Job.Id.ToString(), BatchId = QueueItem.Batch.Id.ToString(), PoolId = QueueItem.PoolId.ToString(), Workspace = QueueItem.Workspace });
 				}
 
 				List<object> OutputWaiters = new List<object>();
 				foreach (QueueWaiter Waiter in Waiters)
 				{
-					OutputWaiters.Add(new { Id = Waiter.Agent.Id.ToString(), Pools = Waiter.Agent.GetPools(AllPools).Select(x => x.Id.ToString()).ToList(), Workspaces = Waiter.Agent.Workspaces, Capabilities = Waiter.Agent.Capabilities });
+					OutputWaiters.Add(new { Id = Waiter.Agent.Id.ToString(), Pools = Waiter.Agent.GetPools().Select(x => x.ToString()).ToList(), Workspaces = Waiter.Agent.Workspaces });
 				}
 
 				return new { Items = OutputItems, Waiters = OutputWaiters };
@@ -436,7 +425,7 @@ namespace HordeServer.Tasks.Impl
 			HashSet<PoolId> OnlinePools = new HashSet<PoolId>(Agents.Where(x => x.IsSessionValid(UtcNow)).SelectMany(x => x.ExplicitPools));
 			foreach (IPool Pool in Pools)
 			{
-				if (Pool.Requirements != null && !OnlinePools.Contains(Pool.Id) && Agents.Any(x => x.IsSessionValid(UtcNow) && x.InPool(Pool)))
+				if (Pool.Condition != null && !OnlinePools.Contains(Pool.Id) && Agents.Any(x => x.IsSessionValid(UtcNow) && x.SatisfiesCondition(Pool.Condition)))
 				{
 					OnlinePools.Add(Pool.Id);
 				}
@@ -446,7 +435,7 @@ namespace HordeServer.Tasks.Impl
 			HashSet<PoolId> ValidPools = new HashSet<PoolId>(OnlinePools.Union(Agents.Where(x => !x.IsSessionValid(UtcNow)).SelectMany(x => x.ExplicitPools)));
 			foreach (IPool Pool in Pools)
 			{
-				if (Pool.Requirements != null && !ValidPools.Contains(Pool.Id) && Agents.Any(x => !x.IsSessionValid(UtcNow) && x.InPool(Pool)))
+				if (Pool.Condition != null && !ValidPools.Contains(Pool.Id) && Agents.Any(x => !x.IsSessionValid(UtcNow) && x.SatisfiesCondition(Pool.Condition)))
 				{
 					ValidPools.Add(Pool.Id);
 				}
@@ -523,9 +512,7 @@ namespace HordeServer.Tasks.Impl
 						}
 						else
 						{
-							AgentRequirements Requirements = new AgentRequirements();
-
-							QueueItem NewQueueItem = new QueueItem(Stream, NewJob, BatchIdx, Pool, Requirements, Workspace);
+							QueueItem NewQueueItem = new QueueItem(Stream, NewJob, BatchIdx, AgentType.Pool, Workspace);
 							NewQueue.Add(NewQueueItem);
 							NewBatchIdToQueueItem[(NewJob.Id, Batch.Id)] = NewQueueItem;
 						}
@@ -645,19 +632,15 @@ namespace HordeServer.Tasks.Impl
 		/// <returns></returns>
 		bool TryAssignItemToWaiter(QueueItem Item, QueueWaiter Waiter)
 		{
-			if (Item.AssignTask == null && Item.Batch.SessionId == null)
+			if (Item.AssignTask == null && Item.Batch.SessionId == null && Waiter.Agent.IsInPool(Item.PoolId))
 			{
-				List<AgentLeaseDevice>? LeasedDevices;
-				if (Waiter.Agent.TryCreateLease(Item.Pool, Item.Requirements, out LeasedDevices))
-				{
-					Task StartTask = new Task<Task>(() => TryCreateLeaseAsync(Item, Waiter, LeasedDevices));
-					Task ExecuteTask = StartTask.ContinueWith(Task => Task, TaskScheduler.Default);
+				Task StartTask = new Task<Task>(() => TryCreateLeaseAsync(Item, Waiter));
+				Task ExecuteTask = StartTask.ContinueWith(Task => Task, TaskScheduler.Default);
 
-					if (Interlocked.CompareExchange(ref Item.AssignTask, ExecuteTask, null) == null)
-					{
-						StartTask.Start(TaskScheduler.Default);
-						return true;
-					}
+				if (Interlocked.CompareExchange(ref Item.AssignTask, ExecuteTask, null) == null)
+				{
+					StartTask.Start(TaskScheduler.Default);
+					return true;
 				}
 			}
 			return false;
@@ -694,7 +677,7 @@ namespace HordeServer.Tasks.Impl
 								else
 								{
 									RemoveQueueItem(ExistingItem);
-									InsertQueueItem(Stream, Job, BatchIdx, ExistingItem.Pool, ExistingItem.Requirements, ExistingItem.Workspace);
+									InsertQueueItem(Stream, Job, BatchIdx, ExistingItem.PoolId, ExistingItem.Workspace);
 								}
 							}
 							continue;
@@ -710,14 +693,7 @@ namespace HordeServer.Tasks.Impl
 							AgentWorkspace? Workspace;
 							if (Stream.TryGetAgentWorkspace(AgentType, out Workspace))
 							{
-								AgentRequirements AgentRequirements = new AgentRequirements();
-
-								// Get the pool for this agent type
-								IPool? Pool;
-								if (CachedPoolIdToInstance.TryGetValue(AgentType.Pool, out Pool))
-								{
-									InsertQueueItem(Stream, Job, BatchIdx, Pool, AgentRequirements, Workspace);
-								}
+								InsertQueueItem(Stream, Job, BatchIdx, AgentType.Pool, Workspace);
 							}
 						}
 					}
@@ -743,14 +719,6 @@ namespace HordeServer.Tasks.Impl
 			}
 		}
 
-		/// <summary>
-		/// Gets the index that should be used to choose the Perforce server an agent should use for a lease.
-		/// </summary>
-		public uint GetPerforceServerIndex()
-		{
-			return NextServerIndex++;
-		}
-
 		/// <inheritdoc/>
 		public Task<ITaskListener?> SubscribeAsync(IAgent Agent)
 		{
@@ -767,9 +735,8 @@ namespace HordeServer.Tasks.Impl
 		/// </summary>
 		/// <param name="Item">The item to create a lease for</param>
 		/// <param name="Waiter">The agent waiting for work</param>
-		/// <param name="LeasedDevices">The devices to include in the lease</param>
 		/// <returns>New work to execute</returns>
-		private async Task<AgentLease?> TryCreateLeaseAsync(QueueItem Item, QueueWaiter Waiter, List<AgentLeaseDevice>? LeasedDevices)
+		private async Task<AgentLease?> TryCreateLeaseAsync(QueueItem Item, QueueWaiter Waiter)
 		{
 			IJob Job = Item.Job;
 			IJobStepBatch Batch = Job.Batches[Item.BatchIdx];
@@ -784,7 +751,7 @@ namespace HordeServer.Tasks.Impl
 
 			// Try to update the job with this agent id
 			ObjectId LogId = (await LogFileService.CreateLogFileAsync(Job.Id, Agent.SessionId, LogType.Json)).Id;
-			if (await Jobs.TryAssignLeaseAsync(Item.Job, Item.BatchIdx, Item.Pool.Id, Agent.Id, Agent.SessionId!.Value, LeaseId, LogId))
+			if (await Jobs.TryAssignLeaseAsync(Item.Job, Item.BatchIdx, Item.PoolId, Agent.Id, Agent.SessionId!.Value, LeaseId, LogId))
 			{
 				// Get the lease name
 				StringBuilder LeaseName = new StringBuilder($"{Item.Stream.Name} - ");
@@ -808,7 +775,7 @@ namespace HordeServer.Tasks.Impl
 					byte[] Payload = Any.Pack(Task).ToByteArray();
 
 					// Create the lease and try to set it on the waiter. If this fails, the waiter has already moved on, and the lease can be cancelled.
-					AgentLease Lease = new AgentLease(LeaseId, LeaseName.ToString(), Job.StreamId, Item.Pool.Id, LogId, LeaseState.Pending, Payload, Item.Requirements, LeasedDevices);
+					AgentLease Lease = new AgentLease(LeaseId, LeaseName.ToString(), Job.StreamId, Item.PoolId, LogId, LeaseState.Pending, Payload, null);
 					if (Waiter.LeaseSource.TrySetResult(new NewLeaseInfo(Lease)))
 					{
 						Logger.LogDebug("Assigned lease {LeaseId} to agent {AgentId}", LeaseId, Agent.Id);
@@ -1103,15 +1070,14 @@ namespace HordeServer.Tasks.Impl
 		/// <param name="Stream">The stream containing the job</param>
 		/// <param name="Job"></param>
 		/// <param name="BatchIdx"></param>
-		/// <param name="Pool">The pool to use</param>
-		/// <param name="Requirements">The agent requirements for this item</param>
+		/// <param name="PoolId">The pool to use</param>
 		/// <param name="Workspace">The workspace for this item to run in</param>
 		/// <returns></returns>
-		void InsertQueueItem(IStream Stream, IJob Job, int BatchIdx, IPool Pool, AgentRequirements Requirements, AgentWorkspace Workspace)
+		void InsertQueueItem(IStream Stream, IJob Job, int BatchIdx, PoolId PoolId, AgentWorkspace Workspace)
 		{
-			Logger.LogDebug("Adding queued job {JobId}, batch {BatchId} [Pool: {Pool}, Workspace: {Workspace}]", Job.Id, Job.Batches[BatchIdx].Id, Pool.Id, Workspace.Identifier);
+			Logger.LogDebug("Adding queued job {JobId}, batch {BatchId} [Pool: {Pool}, Workspace: {Workspace}]", Job.Id, Job.Batches[BatchIdx].Id, PoolId, Workspace.Identifier);
 
-			QueueItem NewItem = new QueueItem(Stream, Job, BatchIdx, Pool, Requirements, Workspace);
+			QueueItem NewItem = new QueueItem(Stream, Job, BatchIdx, PoolId, Workspace);
 			BatchIdToQueueItem[NewItem.Id] = NewItem;
 			Queue.Add(NewItem);
 
diff --git a/Engine/Source/Programs/Horde/HordeServer/Tasks/Impl/RestartTaskSource.cs b/Engine/Source/Programs/Horde/HordeServer/Tasks/Impl/RestartTaskSource.cs
--- a/Engine/Source/Programs/Horde/HordeServer/Tasks/Impl/RestartTaskSource.cs
+++ b/Engine/Source/Programs/Horde/HordeServer/Tasks/Impl/RestartTaskSource.cs
@@ -43,7 +43,7 @@ namespace HordeServer.Tasks.Impl
 
 			byte[] Payload = Any.Pack(Task).ToByteArray();
 
-			AgentLease? Lease = new AgentLease(ObjectId.GenerateNewId(), "Restart", null, null, Log.Id, LeaseState.Pending, Payload, new AgentRequirements(), null);
+			AgentLease? Lease = new AgentLease(ObjectId.GenerateNewId(), "Restart", null, null, Log.Id, LeaseState.Pending, Payload, null);
 			return TaskSubscription.FromResult(Lease);
 		}
 	}
diff --git a/Engine/Source/Programs/Horde/HordeServer/Tasks/Impl/ShutdownTaskSource.cs b/Engine/Source/Programs/Horde/HordeServer/Tasks/Impl/ShutdownTaskSource.cs
--- a/Engine/Source/Programs/Horde/HordeServer/Tasks/Impl/ShutdownTaskSource.cs
+++ b/Engine/Source/Programs/Horde/HordeServer/Tasks/Impl/ShutdownTaskSource.cs
@@ -43,7 +43,7 @@ namespace HordeServer.Tasks.Impl
 
 			byte[] Payload = Any.Pack(Task).ToByteArray();
 
-			AgentLease? Lease = new AgentLease(ObjectId.GenerateNewId(), "Shutdown", null, null, Log.Id, LeaseState.Pending, Payload, new AgentRequirements(), null);
+			AgentLease? Lease = new AgentLease(ObjectId.GenerateNewId(), "Shutdown", null, null, Log.Id, LeaseState.Pending, Payload, null);
 			return TaskSubscription.FromResult(Lease);
 		}
 	}
diff --git a/Engine/Source/Programs/Horde/HordeServer/Tasks/Impl/UpgradeTaskSource.cs b/Engine/Source/Programs/Horde/HordeServer/Tasks/Impl/UpgradeTaskSource.cs
--- a/Engine/Source/Programs/Horde/HordeServer/Tasks/Impl/UpgradeTaskSource.cs
+++ b/Engine/Source/Programs/Horde/HordeServer/Tasks/Impl/UpgradeTaskSource.cs
@@ -59,7 +59,7 @@ namespace HordeServer.Tasks.Impl
 						Payload = Any.Pack(Task).ToByteArray();
 					}
 
-					Lease = new AgentLease(ObjectId.GenerateNewId(), $"Upgrade to {RequiredVersion}", null, null, LogFile.Id, LeaseState.Pending, Payload, new AgentRequirements(), null);
+					Lease = new AgentLease(ObjectId.GenerateNewId(), $"Upgrade to {RequiredVersion}", null, null, LogFile.Id, LeaseState.Pending, Payload, null);
 				}
 				return TaskSubscription.FromResult(Lease);
 			}
diff --git a/Engine/Source/Programs/Horde/HordeServer/Utilities/AsyncCachedValue.cs b/Engine/Source/Programs/Horde/HordeServer/Utilities/AsyncCachedValue.cs
new file mode 100644
--- /dev/null
+++ b/Engine/Source/Programs/Horde/HordeServer/Utilities/AsyncCachedValue.cs
@@ -0,0 +1,123 @@
+// Copyright Epic Games, Inc. All Rights Reserved.
+
+using System;
+using System.Collections.Generic;
+using System.Diagnostics;
+using System.Linq;
+using System.Threading;
+using System.Threading.Tasks;
+
+namespace HordeServer.Utilities
+{
+	/// <summary>
+	/// Caches a value and asynchronously updates it after a period of time
+	/// </summary>
+	/// <typeparam name="T"></typeparam>
+	public class AsyncCachedValue<T>
+	{
+		class State
+		{
+			public readonly T Value;
+			Stopwatch Timer;
+			public Task<State>? Next;
+
+			public TimeSpan Elapsed => Timer.Elapsed;
+
+			public State(T Value)
+			{
+				this.Value = Value;
+				this.Timer = Stopwatch.StartNew();
+			}
+		}
+
+		/// <summary>
+		/// The current state
+		/// </summary>
+		Task<State>? Current = null;
+
+		/// <summary>
+		/// Generator for the new value
+		/// </summary>
+		Func<Task<T>> Generator;
+
+		/// <summary>
+		/// Time at which to start to refresh the value
+		/// </summary>
+		TimeSpan MinRefreshTime;
+
+		/// <summary>
+		/// Time at which to wait for the value to refresh
+		/// </summary>
+		TimeSpan MaxRefreshTime;
+
+		/// <summary>
+		/// Default constructor
+		/// </summary>
+		public AsyncCachedValue(Func<Task<T>> Generator, TimeSpan RefreshTime)
+			: this(Generator, RefreshTime * 0.75, RefreshTime)
+		{
+		}
+
+		/// <summary>
+		/// Default constructor
+		/// </summary>
+		public AsyncCachedValue(Func<Task<T>> Generator, TimeSpan MinRefreshTime, TimeSpan MaxRefreshTime)
+		{
+			this.Generator = Generator;
+			this.MinRefreshTime = MinRefreshTime;
+			this.MaxRefreshTime = MaxRefreshTime;
+		}
+
+		/// <summary>
+		/// Tries to get the current value
+		/// </summary>
+		/// <returns>The cached value, if valid</returns>
+		public Task<T> GetAsync()
+		{
+			return GetAsync(MaxRefreshTime);
+		}
+
+		Task<State> CreateUpdateTask(ref Task<State>? UpdateTask)
+		{
+			Task<State>? CurrentTask = UpdateTask;
+			while (CurrentTask == null)
+			{
+				TaskCompletionSource<State> NextTaskSource = new TaskCompletionSource<State>();
+				if (Interlocked.CompareExchange(ref UpdateTask, NextTaskSource.Task, null) == null)
+				{
+					Task.Run(async () => NextTaskSource.SetResult(new State(await Generator())));
+					return NextTaskSource.Task;
+				}
+				CurrentTask = UpdateTask;
+			}
+			return CurrentTask;
+		}
+
+		/// <summary>
+		/// Tries to get the current value
+		/// </summary>
+		/// <returns>The cached value, if valid</returns>
+		public async Task<T> GetAsync(TimeSpan MaxAge)
+		{
+			Task<State> CurrentCopy = CreateUpdateTask(ref Current);
+
+			State CurrentState = await CurrentCopy;
+			if (CurrentState.Elapsed > MaxAge)
+			{
+				CurrentState = await CreateUpdateTask(ref CurrentState.Next);
+				return CurrentState.Value;
+			}
+			if (CurrentState.Elapsed > MinRefreshTime)
+			{
+				_ = CreateUpdateTask(ref CurrentState.Next);
+			}
+			if (CurrentState.Next != null && CurrentState.Next.IsCompleted)
+			{
+				_ = Interlocked.CompareExchange(ref Current, CurrentState.Next, CurrentCopy);
+				return CurrentState.Next.Result.Value;
+			}
+
+			return CurrentState.Value;
+		}
+	}
+}
diff --git a/Engine/Source/Programs/Horde/HordeServerTests/AgentServiceTest.cs b/Engine/Source/Programs/Horde/HordeServerTests/AgentServiceTest.cs
--- a/Engine/Source/Programs/Horde/HordeServerTests/AgentServiceTest.cs
+++ b/Engine/Source/Programs/Horde/HordeServerTests/AgentServiceTest.cs
@@ -43,7 +43,7 @@ namespace HordeServerTests
 	        TestSetup TestSetup = await GetTestSetup();
 	        AgentService AgentService = TestSetup.AgentService;
 
-	        IAgent Agent = await AgentService.CreateSessionAsync(TestSetup.Fixture!.Agent1, AgentStatus.Ok, new AgentCapabilities(),
+	        IAgent Agent = await AgentService.CreateSessionAsync(TestSetup.Fixture!.Agent1, AgentStatus.Ok, new List<string>(), new Dictionary<string, int>(),
 		        "test");
 	        
 	        Assert.IsTrue(AgentService.AuthorizeSession(Agent, GetUser(Agent)));
diff --git a/Engine/Source/Programs/Horde/HordeServerTests/RemoteExec/ClusteredActionTaskSourceTest.cs b/Engine/Source/Programs/Horde/HordeServerTests/RemoteExec/ClusteredActionTaskSourceTest.cs
--- a/Engine/Source/Programs/Horde/HordeServerTests/RemoteExec/ClusteredActionTaskSourceTest.cs
+++ b/Engine/Source/Programs/Horde/HordeServerTests/RemoteExec/ClusteredActionTaskSourceTest.cs
@@ -28,9 +28,12 @@ namespace HordeServerTests.RemoteExec
 		public bool Ephemeral { get; set; }
 		public bool Deleted { get; set; }
 		public string? Version { get; set; }
+		public IReadOnlyList<string> Properties { get; set; } = new List<string>();
+		public IReadOnlyDictionary<string, int> Resources { get; set; } = new Dictionary<string, int>();
 		public StringId<AgentSoftwareChannels>? Channel { get; set; }
 		public string? LastUpgradeVersion { get; set; }
 		public DateTime? LastUpgradeTime { get; set; }
+		public IReadOnlyList<StringId<IPool>> DynamicPools { get; set; } = new List<StringId<IPool>>();
 		public IReadOnlyList<StringId<IPool>> ExplicitPools { get; set; } = new List<StringId<IPool>>();
 		public bool RequestConform { get; set; }
 		public bool RequestRestart { get; set; }
@@ -38,7 +41,6 @@ namespace HordeServerTests.RemoteExec
 		public IReadOnlyList<AgentWorkspace> Workspaces { get; set; } = new List<AgentWorkspace>();
 		public DateTime LastConformTime { get; set; }
 		public int? ConformAttemptCount { get; set; }
-		public AgentCapabilities Capabilities { get; set; } = new AgentCapabilities();
 		public IReadOnlyList<AgentLease> Leases { get; set; } = new List<AgentLease>();
 		public Acl? Acl { get; set; }
 		public DateTime UpdateTime { get; set; }
diff --git a/Engine/Source/Programs/Shared/EpicGames.Horde/Compute/Requirements.cs b/Engine/Source/Programs/Shared/EpicGames.Horde/Compute/Requirements.cs
--- a/Engine/Source/Programs/Shared/EpicGames.Horde/Compute/Requirements.cs
+++ b/Engine/Source/Programs/Shared/EpicGames.Horde/Compute/Requirements.cs
@@ -47,5 +47,23 @@ namespace EpicGames.Horde.Compute
 		{
 			this.Condition = Condition;
 		}
+
+		public override string ToString()
+		{
+			List<string> List = new List<string>();
+			if (Condition != null)
+			{
+				List.Add($"\"{Condition}\"");
+			}
+			foreach ((string Name, int Count) in Resources)
+			{
+				List.Add($"{Name}: {Count}");
+			}
+			if (Exclusive)
+			{
+				List.Add("Exclusive");
+			}
+			return String.Join(", ", List);
+		}
 	}
 }
