Horde: Initialize number of cores and available RAM to values returned as properties.

diff --git a/Engine/Source/Programs/Horde/HordeServer/Models/Agent.cs b/Engine/Source/Programs/Horde/HordeServer/Models/Agent.cs
--- a/Engine/Source/Programs/Horde/HordeServer/Models/Agent.cs
+++ b/Engine/Source/Programs/Horde/HordeServer/Models/Agent.cs
@@ -374,10 +374,25 @@ namespace HordeServer.Models
 	/// </summary>
 	static class KnownPropertyNames
 	{
+		/// <summary>
+		/// The operating system (Linux, MacOS, Windows)
+		/// </summary>
+		public const string OSFamily = "OSFamily";
+
 		/// <summary>
 		/// Pools that this agent belongs to
 		/// </summary>
 		public const string Pool = "Pool";
+
+		/// <summary>
+		/// Number of logical cores
+		/// </summary>
+		public const string LogicalCores = "LogicalCores";
+
+		/// <summary>
+		/// Amount of RAM, in GB
+		/// </summary>
+		public const string RAM = "RAM";
 	}
 
 	/// <summary>
diff --git a/Engine/Source/Programs/Horde/HordeServer/Services/RpcService.cs b/Engine/Source/Programs/Horde/HordeServer/Services/RpcService.cs
--- a/Engine/Source/Programs/Horde/HordeServer/Services/RpcService.cs
+++ b/Engine/Source/Programs/Horde/HordeServer/Services/RpcService.cs
@@ -290,6 +290,21 @@ namespace HordeServer.Services
 			}
 		}
 
+		static void CopyPropertyToResource(string Name, List<string> Properties, Dictionary<string, int> Resources)
+		{
+			foreach (string Property in Properties)
+			{
+				if (Property.Length > Name.Length && Property.StartsWith(Name, StringComparison.OrdinalIgnoreCase) && Property[Name.Length] == '=')
+				{
+					int Value;
+					if (int.TryParse(Property.AsSpan(Name.Length + 1), out Value))
+					{
+						Resources[Name] = Value;
+					}
+				}
+			}
+		}
+
 		static void GetCapabilities(RpcAgentCapabilities? Capabilities, out List<string> Properties, out Dictionary<string, int> Resources)
 		{
 			Properties = new List<string>();
@@ -301,10 +316,8 @@ namespace HordeServer.Services
 				if (Device.Properties != null)
 				{
 					Properties = new List<string>(Device.Properties);
-				}
-				if (Device.Resources != null)
-				{
-					Resources = new Dictionary<string, int>(Device.Resources);
+					CopyPropertyToResource(KnownPropertyNames.LogicalCores, Properties, Resources);
+					CopyPropertyToResource(KnownPropertyNames.RAM, Properties, Resources);
 				}
 			}
 		}
