Horde: Return copies of agent documents rather than mutating the originals.

diff --git a/Engine/Source/Programs/Horde/HordeServer/Collections/Impl/AgentCollection.cs b/Engine/Source/Programs/Horde/HordeServer/Collections/Impl/AgentCollection.cs
--- a/Engine/Source/Programs/Horde/HordeServer/Collections/Impl/AgentCollection.cs
+++ b/Engine/Source/Programs/Horde/HordeServer/Collections/Impl/AgentCollection.cs
@@ -163,12 +163,7 @@ namespace HordeServer.Collections.Impl
 			AgentDocument Agent = (AgentDocument)AgentInterface;
 
 			UpdateDefinition<AgentDocument> Update = Builders<AgentDocument>.Update.Set(x => x.Deleted, true);
-			if (await TryUpdateAsync(Agent, Update))
-			{
-				Agent.Deleted = true;
-				return Agent;
-			}
-			return null;
+			return await TryUpdateAsync(Agent, Update);
 		}
 
 		/// <inheritdoc/>
@@ -236,7 +231,7 @@ namespace HordeServer.Collections.Impl
 		/// <param name="Current">The document to update</param>
 		/// <param name="Update">The update definition</param>
 		/// <returns>True if the agent was updated</returns>
-		private async Task<bool> TryUpdateAsync(AgentDocument Current, UpdateDefinition<AgentDocument> Update)
+		private async Task<AgentDocument?> TryUpdateAsync(AgentDocument Current, UpdateDefinition<AgentDocument> Update)
 		{
 			uint PrevUpdateIndex = Current.UpdateIndex++;
 			Current.UpdateTime = DateTime.UtcNow;
@@ -244,8 +239,7 @@ namespace HordeServer.Collections.Impl
 			Expression<Func<AgentDocument, bool>> Filter = x => x.Id == Current.Id && x.UpdateIndex == PrevUpdateIndex;
 			UpdateDefinition<AgentDocument> UpdateWithIndex = Update.Set(x => x.UpdateIndex, Current.UpdateIndex).Set(x => x.UpdateTime, Current.UpdateTime);
 
-			UpdateResult Result = await Agents.UpdateOneAsync<AgentDocument>(Filter, UpdateWithIndex);
-			return Result.ModifiedCount == 1;
+			return await Agents.FindOneAndUpdateAsync<AgentDocument>(Filter, UpdateWithIndex, new FindOneAndUpdateOptions<AgentDocument, AgentDocument> { ReturnDocument = ReturnDocument.After });
 		}
 
 		/// <inheritdoc/>
@@ -313,41 +307,7 @@ namespace HordeServer.Collections.Impl
 			}
 
 			// Apply the update
-			if (Updates.Count != 0 && !await TryUpdateAsync(Agent, UpdateBuilder.Combine(Updates)))
-			{
-				return null;
-			}
-
-			// Update the document
-			if (Pools != null)
-			{
-				Agent.Pools = Pools;
-			}
-			if (Enabled != null)
-			{
-				Agent.Enabled = Enabled.Value;
-			}
-			if (RequestConform != null)
-			{
-				Agent.RequestConform = RequestConform.Value;
-			}
-			if (RequestRestart != null)
-			{
-				Agent.RequestRestart = RequestRestart.Value;
-			}
-			if (RequestShutdown != null)
-			{
-				Agent.RequestShutdown = RequestShutdown.Value;
-			}
-			if (Channel != null)
-			{
-				Agent.Channel = (Channel.Value == AgentSoftwareService.DefaultChannelName) ? null : Channel;
-			}
-			if (Acl != null)
-			{
-				Agent.Acl = Acl;
-			}
-			return Agent;
+			return await TryUpdateAsync(Agent, UpdateBuilder.Combine(Updates));
 		}
 
 		/// <inheritdoc/>
@@ -361,28 +321,23 @@ namespace HordeServer.Collections.Impl
 
 			if (Status != null && Agent.Status != Status.Value)
 			{
-				Agent.Status = Status.Value;
-				Updates.Add(UpdateBuilder.Set(x => x.Status, Agent.Status));
+				Updates.Add(UpdateBuilder.Set(x => x.Status, Status.Value));
 			}
 			if (SessionExpiresAt != null)
 			{
-				Agent.SessionExpiresAt = SessionExpiresAt.Value;
-				Updates.Add(UpdateBuilder.Set(x => x.SessionExpiresAt, Agent.SessionExpiresAt));
+				Updates.Add(UpdateBuilder.Set(x => x.SessionExpiresAt, SessionExpiresAt.Value));
 			}
 			if (Properties != null)
 			{
-				Agent.Properties = Properties.OrderBy(x => x, StringComparer.OrdinalIgnoreCase).ToList();
-				Updates.Add(UpdateBuilder.Set(x => x.Properties, Agent.Properties));
+				Updates.Add(UpdateBuilder.Set(x => x.Properties, Properties.OrderBy(x => x, StringComparer.OrdinalIgnoreCase).ToList()));
 			}
 			if (Resources != null)
 			{
-				Agent.Resources = new Dictionary<string, int>(Resources);
-				Updates.Add(UpdateBuilder.Set(x => x.Resources, Agent.Resources));
+				Updates.Add(UpdateBuilder.Set(x => x.Resources, new Dictionary<string, int>(Resources)));
 			}
 			if (DynamicPools != null && !DynamicPools.SequenceEqual(Agent.DynamicPools))
 			{
-				Agent.DynamicPools = new List<PoolId>(DynamicPools);
-				Updates.Add(UpdateBuilder.Set(x => x.DynamicPools, Agent.DynamicPools));
+				Updates.Add(UpdateBuilder.Set(x => x.DynamicPools, new List<PoolId>(DynamicPools)));
 			}
 			if (Leases != null)
 			{
@@ -393,34 +348,23 @@ namespace HordeServer.Collections.Impl
 						Any Payload = Any.Parser.ParseFrom(Lease.Payload.ToArray());
 						if (Payload.TryUnpack(out ConformTask ConformTask))
 						{
-							Agent.ConformAttemptCount = (Agent.ConformAttemptCount ?? 0) + 1;
-							Updates.Add(UpdateBuilder.Set(x => x.ConformAttemptCount, Agent.ConformAttemptCount));
-
-							Agent.LastConformTime = DateTime.UtcNow;
-							Updates.Add(UpdateBuilder.Set(x => x.LastConformTime, Agent.LastConformTime));
+							int NewConformAttemptCount = (Agent.ConformAttemptCount ?? 0) + 1;
+							Updates.Add(UpdateBuilder.Set(x => x.ConformAttemptCount, NewConformAttemptCount));
+							Updates.Add(UpdateBuilder.Set(x => x.LastConformTime, DateTime.UtcNow));
 						}
 						else if (Payload.TryUnpack(out UpgradeTask UpgradeTask))
 						{
-							Agent.LastUpgradeVersion = UpgradeTask.SoftwareId;
-							Updates.Add(UpdateBuilder.Set(x => x.LastUpgradeVersion, Agent.LastUpgradeVersion));
-
-							Agent.LastUpgradeTime = DateTime.UtcNow;
-							Updates.Add(UpdateBuilder.Set(x => x.LastUpgradeTime, Agent.LastUpgradeTime));
+							Updates.Add(UpdateBuilder.Set(x => x.LastUpgradeVersion, UpgradeTask.SoftwareId));
+							Updates.Add(UpdateBuilder.Set(x => x.LastUpgradeTime, DateTime.UtcNow));
 						}
 					}
 				}
 
-				Agent.Leases = Leases;
-				Updates.Add(UpdateBuilder.Set(x => x.Leases, Agent.Leases));
+				Updates.Add(UpdateBuilder.Set(x => x.Leases, Leases));
 			}
 
 			// Update the agent, and try to create new lease documents if we succeed
-			if (Updates.Count > 0 && !await TryUpdateAsync(Agent, UpdateBuilder.Combine(Updates)))
-			{
-				return null;
-			}
-
-			return Agent;
+			return await TryUpdateAsync(Agent, UpdateBuilder.Combine(Updates));
 		}
 
 		/// <inheritdoc/>
@@ -436,18 +380,7 @@ namespace HordeServer.Collections.Impl
 			Update = Update.Set(x => x.RequestConform, RequestConform);
 
 			// Update the agent
-			if (!await TryUpdateAsync(Agent, Update))
-			{
-				return null;
-			}
-
-			// Update the document
-			Agent.Workspaces = new List<AgentWorkspace>(Workspaces);
-			Agent.LastConformTime = LastConformTime;
-			Agent.ConformAttemptCount = null;
-			Agent.RequestConform = RequestConform;
-
-			return Agent;
+			return await TryUpdateAsync(Agent, Update);
 		}
 
 		/// <inheritdoc/>
@@ -475,23 +408,7 @@ namespace HordeServer.Collections.Impl
 			Updates.Add(UpdateBuilder.Unset(x => x.RequestShutdown));
 
 			// Apply the update
-			if (!await TryUpdateAsync(Agent, UpdateBuilder.Combine(Updates)))
-			{
-				return null;
-			}
-
-			// Update the document
-			Agent.SessionId = SessionId;
-			Agent.SessionExpiresAt = SessionExpiresAt;
-			Agent.Status = Status;
-			Agent.Leases = null;
-			Agent.Deleted = false;
-			Agent.Properties = NewProperties;
-			Agent.Resources = NewResources;
-			Agent.DynamicPools = NewDynamicPools;
-			Agent.Version = Version;
-
-			return Agent;
+			return await TryUpdateAsync(Agent, UpdateBuilder.Combine(Updates));
 		}
 
 		/// <inheritdoc/>
@@ -510,16 +427,7 @@ namespace HordeServer.Collections.Impl
 				Update = Update.Set(x => x.Deleted, Agent.Deleted);
 			}
 
-			if (!await TryUpdateAsync(Agent, Update))
-			{
-				return null;
-			}
-
-			Agent.SessionId = null;
-			Agent.SessionExpiresAt = null;
-			Agent.Deleted = bDeleted;
-			Agent.Leases = new List<AgentLease>();
-			return Agent;
+			return await TryUpdateAsync(Agent, Update);
 		}
 
 		/// <inheritdoc/>
@@ -535,13 +443,7 @@ namespace HordeServer.Collections.Impl
 			Leases.Add(NewLease);
 
 			UpdateDefinition<AgentDocument> Update = Builders<AgentDocument>.Update.Set(x => x.Leases, Leases);
-			if (!await TryUpdateAsync(Agent, Update))
-			{
-				return null;
-			}
-
-			Agent.Leases = Leases;
-			return Agent;
+			return await TryUpdateAsync(Agent, Update);
 		}
 
 		/// <inheritdoc/>
@@ -550,13 +452,7 @@ namespace HordeServer.Collections.Impl
 			AgentDocument Agent = (AgentDocument)AgentInterface;
 
 			UpdateDefinition<AgentDocument> Update = Builders<AgentDocument>.Update.Set(x => x.Leases![LeaseIdx].State, LeaseState.Cancelled);
-			if (!await TryUpdateAsync(Agent, Update))
-			{
-				return null;
-			}
-
-			Agent.Leases![LeaseIdx].State = LeaseState.Cancelled;
-			return Agent;
+			return await TryUpdateAsync(Agent, Update);
 		}
 	}
 }
