Horde: Add option for using fallback credentials for S3 client

See code comment for details.

diff --git a/Engine/Source/Programs/Horde/HordeServer/Startup.cs b/Engine/Source/Programs/Horde/HordeServer/Startup.cs
--- a/Engine/Source/Programs/Horde/HordeServer/Startup.cs
+++ b/Engine/Source/Programs/Horde/HordeServer/Startup.cs
@@ -416,6 +416,12 @@ namespace HordeServer
 			{
 				AwsOptions.Credentials = AssumeRoleWithWebIdentityCredentials.FromEnvironmentVariables();
 			}
+			else if(Settings.S3CredentialType == "Fallback")
+			{
+				// Using the fallback credentials from the AWS SDK, it will pick up credentials through a number of default mechanisms.
+				AwsOptions.Credentials = FallbackCredentialsFactory.GetCredentials();
+			}
+			
 			Services.AddSingleton(AwsOptions);
 
 			Services.AddSingleton(new StorageBackendSettings<PersistentLogStorage> { Type = Settings.ExternalStorageProviderType, BaseDir = Settings.LocalLogsDir, BucketName = Settings.S3LogBucketName });
