Horde: Speculative fix for exception trying to terminate sessions. MongoDB doesn't handle serialization of IReadOnlyList and IReadOnlyDictionary.

diff --git a/Engine/Source/Programs/Horde/HordeServer/Collections/Impl/SessionCollection.cs b/Engine/Source/Programs/Horde/HordeServer/Collections/Impl/SessionCollection.cs
--- a/Engine/Source/Programs/Horde/HordeServer/Collections/Impl/SessionCollection.cs
+++ b/Engine/Source/Programs/Horde/HordeServer/Collections/Impl/SessionCollection.cs
@@ -126,8 +126,8 @@ namespace HordeServer.Collections.Impl
 		{
 			UpdateDefinition<SessionDocument> Update = Builders<SessionDocument>.Update
 				.Set(x => x.FinishTime, FinishTime)
-				.Set(x => (IEnumerable<string>?)x.Properties, Properties)
-				.Set(x => (IEnumerable<KeyValuePair<string, int>>?)x.Resources, Resources);
+				.Set(x => x.Properties, new List<string>(Properties))
+				.Set(x => x.Resources, new Dictionary<string, int>(Resources));
 
 			return Sessions.FindOneAndUpdateAsync(x => x.Id == SessionId, Update);
 		}
