Horde: Speculative fix for invalid LeaseId objects.

diff --git a/Engine/Source/Programs/Horde/HordeServer/Compute/Impl/ComputeService.cs b/Engine/Source/Programs/Horde/HordeServer/Compute/Impl/ComputeService.cs
--- a/Engine/Source/Programs/Horde/HordeServer/Compute/Impl/ComputeService.cs
+++ b/Engine/Source/Programs/Horde/HordeServer/Compute/Impl/ComputeService.cs
@@ -101,7 +101,7 @@ namespace HordeServer.Compute.Impl
 		/// <inheritdoc/>
 		public ObjectId? LeaseId
 		{
-			get => (LeaseIdBytes == null)? (ObjectId?)null : new ObjectId(LeaseIdBytes);
+			get => (LeaseIdBytes == null || LeaseIdBytes.Length == 0)? (ObjectId?)null : new ObjectId(LeaseIdBytes);
 			set => LeaseIdBytes = value?.ToByteArray();
 		}
 
