Horde: Prevent no-op updates triggering modifications to agent documents, early return from UpdateSession() calls.

diff --git a/Engine/Source/Programs/Horde/HordeServer/Collections/Impl/AgentCollection.cs b/Engine/Source/Programs/Horde/HordeServer/Collections/Impl/AgentCollection.cs
--- a/Engine/Source/Programs/Horde/HordeServer/Collections/Impl/AgentCollection.cs
+++ b/Engine/Source/Programs/Horde/HordeServer/Collections/Impl/AgentCollection.cs
@@ -329,9 +329,13 @@ namespace HordeServer.Collections.Impl
 			}
 			if (Properties != null)
 			{
-				Updates.Add(UpdateBuilder.Set(x => x.Properties, Properties.OrderBy(x => x, StringComparer.OrdinalIgnoreCase).ToList()));
+				List<string> NewProperties = Properties.OrderBy(x => x, StringComparer.OrdinalIgnoreCase).ToList();
+				if (!AgentInterface.Properties.SequenceEqual(NewProperties, StringComparer.Ordinal))
+				{
+					Updates.Add(UpdateBuilder.Set(x => x.Properties, NewProperties));
+				}
 			}
-			if (Resources != null)
+			if (Resources != null && !ResourcesEqual(Resources, AgentInterface.Resources))
 			{
 				Updates.Add(UpdateBuilder.Set(x => x.Resources, new Dictionary<string, int>(Resources)));
 			}
@@ -363,10 +367,35 @@ namespace HordeServer.Collections.Impl
 				Updates.Add(UpdateBuilder.Set(x => x.Leases, Leases));
 			}
 
+			// If there are no new updates, return immediately. This is important for preventing UpdateSession calls from returning immediately.
+			if (Updates.Count == 0)
+			{
+				return Agent;
+			}
+
 			// Update the agent, and try to create new lease documents if we succeed
 			return await TryUpdateAsync(Agent, UpdateBuilder.Combine(Updates));
 		}
 
+		static bool ResourcesEqual(IReadOnlyDictionary<string, int> DictA, IReadOnlyDictionary<string, int> DictB)
+		{
+			if (DictA.Count != DictB.Count)
+			{
+				return false;
+			}
+
+			foreach (KeyValuePair<string, int> Pair in DictA)
+			{
+				int Value;
+				if (!DictB.TryGetValue(Pair.Key, out Value) || Value != Pair.Value)
+				{
+					return false;
+				}
+			}
+
+			return true;
+		}
+
 		/// <inheritdoc/>
 		public async Task<IAgent?> TryUpdateWorkspacesAsync(IAgent AgentInterface, List<AgentWorkspace> Workspaces, bool RequestConform)
 		{
