Horde: Remove S3 object getter from debug controller

Prevented DebugController from working as the injected IAmazonS3 wasn't available.

diff --git a/Engine/Source/Programs/Horde/HordeServer/Controllers/DebugController.cs b/Engine/Source/Programs/Horde/HordeServer/Controllers/DebugController.cs
--- a/Engine/Source/Programs/Horde/HordeServer/Controllers/DebugController.cs
+++ b/Engine/Source/Programs/Horde/HordeServer/Controllers/DebugController.cs
@@ -62,33 +62,17 @@ namespace HordeServer.Controllers
 
 		IDogStatsd DogStatsd;
 
-		IAmazonS3 AwsTest;
-
 		/// <summary>
 		/// Constructor
 		/// </summary>
 		/// <param name="RequestTrackerService"></param>
 		/// <param name="ApplicationLifetime"></param>
 		/// <param name="DogStatsd"></param>
-		/// <param name="AwsTest"></param>
-		public PublicDebugController(RequestTrackerService RequestTrackerService, IHostApplicationLifetime ApplicationLifetime, IDogStatsd DogStatsd, IAmazonS3 AwsTest)
+		public PublicDebugController(RequestTrackerService RequestTrackerService, IHostApplicationLifetime ApplicationLifetime, IDogStatsd DogStatsd)
 		{
 			this.RequestTrackerService = RequestTrackerService;
 			this.ApplicationLifetime = ApplicationLifetime;
 			this.DogStatsd = DogStatsd;
-			this.AwsTest = AwsTest;
-		}
-
-		/// <summary>
-		/// Returns an object from aws
-		/// </summary>
-		/// <returns>Http result</returns>
-		[HttpGet]
-		[Route("/api/v1/debug/awsdata")]
-		public async Task<ActionResult> GetAwsData([FromQuery] string BucketName, [FromQuery] string Path)
-		{
-			Stream Stream = await AwsTest.GetObjectStreamAsync(BucketName, Path, new Dictionary<string, object>());
-			return new FileStreamResult(Stream, new MediaTypeHeaderValue("application/octet-stream"));
 		}
 
 		/// <summary>
