Horde: Store the 'exclusive' flag on leases, and prevent other leases running at the same time.

diff --git a/Engine/Source/Programs/Horde/HordeServer/Compute/Impl/ComputeService.cs b/Engine/Source/Programs/Horde/HordeServer/Compute/Impl/ComputeService.cs
--- a/Engine/Source/Programs/Horde/HordeServer/Compute/Impl/ComputeService.cs
+++ b/Engine/Source/Programs/Horde/HordeServer/Compute/Impl/ComputeService.cs
@@ -219,18 +219,22 @@ namespace HordeServer.Compute.Impl
 			{
 				(IoHash RequirementsHash, ComputeTaskInfo TaskInfo) = Entry.Value;
 
-				ComputeTaskMessage ComputeTask = new ComputeTaskMessage();
-				ComputeTask.ChannelId = TaskInfo.ChannelId.ToString();
-				ComputeTask.NamespaceId = DefaultNamespaceId.ToString();
-				ComputeTask.RequirementsHash = new CbObjectAttachment(RequirementsHash);
-				ComputeTask.TaskHash = TaskInfo.TaskHash;
-
-				string LeaseName = $"Remote action ({TaskInfo.TaskHash})";
-				byte[] Payload = Any.Pack(ComputeTask).ToByteArray();
-
-				AgentLease Lease = new AgentLease(ObjectId.GenerateNewId(), LeaseName, null, null, null, LeaseState.Pending, Payload, null);
-				Logger.LogDebug("Created lease {LeaseId} for channel {ChannelId} task {TaskHash} req {RequirementsHash}", Lease.Id, ComputeTask.ChannelId, (CbObjectAttachment)ComputeTask.TaskHash, (CbObjectAttachment)ComputeTask.RequirementsHash);
-				return Lease;
+				Requirements? Requirements = await GetCachedRequirementsAsync(RequirementsHash);
+				if (Requirements != null)
+				{
+					ComputeTaskMessage ComputeTask = new ComputeTaskMessage();
+					ComputeTask.ChannelId = TaskInfo.ChannelId.ToString();
+					ComputeTask.NamespaceId = DefaultNamespaceId.ToString();
+					ComputeTask.RequirementsHash = new CbObjectAttachment(RequirementsHash);
+					ComputeTask.TaskHash = TaskInfo.TaskHash;
+
+					string LeaseName = $"Remote action ({TaskInfo.TaskHash})";
+					byte[] Payload = Any.Pack(ComputeTask).ToByteArray();
+
+					AgentLease Lease = new AgentLease(ObjectId.GenerateNewId(), LeaseName, null, null, null, LeaseState.Pending, null, Requirements.Exclusive, Payload);
+					Logger.LogDebug("Created lease {LeaseId} for channel {ChannelId} task {TaskHash} req {RequirementsHash}", Lease.Id, ComputeTask.ChannelId, (CbObjectAttachment)ComputeTask.TaskHash, (CbObjectAttachment)ComputeTask.RequirementsHash);
+					return Lease;
+				}
 			}
 			return null;
 		}
diff --git a/Engine/Source/Programs/Horde/HordeServer/Models/Agent.cs b/Engine/Source/Programs/Horde/HordeServer/Models/Agent.cs
--- a/Engine/Source/Programs/Horde/HordeServer/Models/Agent.cs
+++ b/Engine/Source/Programs/Horde/HordeServer/Models/Agent.cs
@@ -234,6 +234,11 @@ namespace HordeServer.Models
 		/// </summary>
 		public IReadOnlyDictionary<string, int>? Resources { get; set; }
 
+		/// <summary>
+		/// Whether the lease requires exclusive access to the agent
+		/// </summary>
+		public bool Exclusive { get; set; }
+
 		/// <summary>
 		/// For leases in the pending state, encodes an "any" protobuf containing the payload for the agent to execute the lease.
 		/// </summary>
@@ -257,9 +262,10 @@ namespace HordeServer.Models
 		/// <param name="PoolId"></param>
 		/// <param name="LogId">Unique id for the log</param>
 		/// <param name="State">State for the lease</param>
-		/// <param name="Payload">Encoded "any" protobuf describing the contents of the payload</param>
 		/// <param name="Resources">Resources required for this lease</param>
-		public AgentLease(ObjectId Id, string Name, StreamId? StreamId, PoolId? PoolId, ObjectId? LogId, LeaseState State, byte[]? Payload, IReadOnlyDictionary<string, int>? Resources)
+		/// <param name="Exclusive">Whether to reserve the entire device</param>
+		/// <param name="Payload">Encoded "any" protobuf describing the contents of the payload</param>
+		public AgentLease(ObjectId Id, string Name, StreamId? StreamId, PoolId? PoolId, ObjectId? LogId, LeaseState State, IReadOnlyDictionary<string, int>? Resources, bool Exclusive, byte[]? Payload)
 		{
 			this.Id = Id;
 			this.Name = Name;
@@ -267,8 +273,9 @@ namespace HordeServer.Models
 			this.PoolId = PoolId;
 			this.LogId = LogId;
 			this.State = State;
-			this.Payload = Payload;
 			this.Resources = Resources;
+			this.Exclusive = Exclusive;
+			this.Payload = Payload;
 			StartTime = DateTime.UtcNow;
 		}
 
@@ -649,6 +656,10 @@ namespace HordeServer.Models
 			{
 				return false;
 			}
+			if (Agent.Leases.Any(x => x.Exclusive))
+			{
+				return false;
+			}
 			if (Exclusive && Agent.Leases.Any())
 			{
 				return false;
diff --git a/Engine/Source/Programs/Horde/HordeServer/Tasks/Impl/ActionTaskSource.cs b/Engine/Source/Programs/Horde/HordeServer/Tasks/Impl/ActionTaskSource.cs
--- a/Engine/Source/Programs/Horde/HordeServer/Tasks/Impl/ActionTaskSource.cs
+++ b/Engine/Source/Programs/Horde/HordeServer/Tasks/Impl/ActionTaskSource.cs
@@ -484,7 +484,7 @@ namespace HordeServer.Tasks.Impl
 				ActionTask.LogId = LogFile.Id.ToString();
 
 				byte[] Payload = Any.Pack(ActionTask).ToByteArray();
-				AgentLease Lease = new AgentLease(Operation.Id, LeaseName, null, null, LogFile.Id, LeaseState.Pending, Payload, null);
+				AgentLease Lease = new AgentLease(Operation.Id, LeaseName, null, null, LogFile.Id, LeaseState.Pending, null, true, Payload);
 
 				// Try to set the lease on the subscription. If it fails, the subscriber has already been allocated
 				Operation.ResultTaskSource = new TaskCompletionSource<ActionExecuteResult>();
diff --git a/Engine/Source/Programs/Horde/HordeServer/Tasks/Impl/ClusteredActionTaskSource.cs b/Engine/Source/Programs/Horde/HordeServer/Tasks/Impl/ClusteredActionTaskSource.cs
--- a/Engine/Source/Programs/Horde/HordeServer/Tasks/Impl/ClusteredActionTaskSource.cs
+++ b/Engine/Source/Programs/Horde/HordeServer/Tasks/Impl/ClusteredActionTaskSource.cs
@@ -361,7 +361,7 @@ namespace HordeServer.Tasks.Impl
 			
 			byte[] Payload = Any.Pack(ActionTask).ToByteArray();
 			ObjectId OpId = ObjectId.Parse(Op.Id);
-			AgentLease Lease = new AgentLease(OpId, LeaseName, null, null, LogFile.Id, LeaseState.Pending, Payload, null);
+			AgentLease Lease = new AgentLease(OpId, LeaseName, null, null, LogFile.Id, LeaseState.Pending, null, false, Payload);
 			RemoteExecOperationInternal OpInternal = new RemoteExecOperationInternal(OpId, Lease.Id);
 
 			lock (ActiveOperations)
diff --git a/Engine/Source/Programs/Horde/HordeServer/Tasks/Impl/ConformTaskSource.cs b/Engine/Source/Programs/Horde/HordeServer/Tasks/Impl/ConformTaskSource.cs
--- a/Engine/Source/Programs/Horde/HordeServer/Tasks/Impl/ConformTaskSource.cs
+++ b/Engine/Source/Programs/Horde/HordeServer/Tasks/Impl/ConformTaskSource.cs
@@ -179,7 +179,7 @@ namespace HordeServer.Tasks.Impl
 
 					byte[] Payload = Any.Pack(Task).ToByteArray();
 
-					AgentLease Lease = new AgentLease(LeaseId, "Updating workspaces", null, null, Log.Id, LeaseState.Pending, Payload, null);
+					AgentLease Lease = new AgentLease(LeaseId, "Updating workspaces", null, null, Log.Id, LeaseState.Pending, null, true, Payload);
 					return TaskSubscription.FromResult(Lease);
 				}
 			}
diff --git a/Engine/Source/Programs/Horde/HordeServer/Tasks/Impl/JobTaskSource.cs b/Engine/Source/Programs/Horde/HordeServer/Tasks/Impl/JobTaskSource.cs
--- a/Engine/Source/Programs/Horde/HordeServer/Tasks/Impl/JobTaskSource.cs
+++ b/Engine/Source/Programs/Horde/HordeServer/Tasks/Impl/JobTaskSource.cs
@@ -775,7 +775,7 @@ namespace HordeServer.Tasks.Impl
 					byte[] Payload = Any.Pack(Task).ToByteArray();
 
 					// Create the lease and try to set it on the waiter. If this fails, the waiter has already moved on, and the lease can be cancelled.
-					AgentLease Lease = new AgentLease(LeaseId, LeaseName.ToString(), Job.StreamId, Item.PoolId, LogId, LeaseState.Pending, Payload, null);
+					AgentLease Lease = new AgentLease(LeaseId, LeaseName.ToString(), Job.StreamId, Item.PoolId, LogId, LeaseState.Pending, null, true, Payload);
 					if (Waiter.LeaseSource.TrySetResult(new NewLeaseInfo(Lease)))
 					{
 						Logger.LogDebug("Assigned lease {LeaseId} to agent {AgentId}", LeaseId, Agent.Id);
diff --git a/Engine/Source/Programs/Horde/HordeServer/Tasks/Impl/RestartTaskSource.cs b/Engine/Source/Programs/Horde/HordeServer/Tasks/Impl/RestartTaskSource.cs
--- a/Engine/Source/Programs/Horde/HordeServer/Tasks/Impl/RestartTaskSource.cs
+++ b/Engine/Source/Programs/Horde/HordeServer/Tasks/Impl/RestartTaskSource.cs
@@ -43,7 +43,7 @@ namespace HordeServer.Tasks.Impl
 
 			byte[] Payload = Any.Pack(Task).ToByteArray();
 
-			AgentLease? Lease = new AgentLease(ObjectId.GenerateNewId(), "Restart", null, null, Log.Id, LeaseState.Pending, Payload, null);
+			AgentLease? Lease = new AgentLease(ObjectId.GenerateNewId(), "Restart", null, null, Log.Id, LeaseState.Pending, null, true, Payload);
 			return TaskSubscription.FromResult(Lease);
 		}
 	}
diff --git a/Engine/Source/Programs/Horde/HordeServer/Tasks/Impl/ShutdownTaskSource.cs b/Engine/Source/Programs/Horde/HordeServer/Tasks/Impl/ShutdownTaskSource.cs
--- a/Engine/Source/Programs/Horde/HordeServer/Tasks/Impl/ShutdownTaskSource.cs
+++ b/Engine/Source/Programs/Horde/HordeServer/Tasks/Impl/ShutdownTaskSource.cs
@@ -43,7 +43,7 @@ namespace HordeServer.Tasks.Impl
 
 			byte[] Payload = Any.Pack(Task).ToByteArray();
 
-			AgentLease? Lease = new AgentLease(ObjectId.GenerateNewId(), "Shutdown", null, null, Log.Id, LeaseState.Pending, Payload, null);
+			AgentLease? Lease = new AgentLease(ObjectId.GenerateNewId(), "Shutdown", null, null, Log.Id, LeaseState.Pending, null, true, Payload);
 			return TaskSubscription.FromResult(Lease);
 		}
 	}
diff --git a/Engine/Source/Programs/Horde/HordeServer/Tasks/Impl/UpgradeTaskSource.cs b/Engine/Source/Programs/Horde/HordeServer/Tasks/Impl/UpgradeTaskSource.cs
--- a/Engine/Source/Programs/Horde/HordeServer/Tasks/Impl/UpgradeTaskSource.cs
+++ b/Engine/Source/Programs/Horde/HordeServer/Tasks/Impl/UpgradeTaskSource.cs
@@ -59,7 +59,7 @@ namespace HordeServer.Tasks.Impl
 						Payload = Any.Pack(Task).ToByteArray();
 					}
 
-					Lease = new AgentLease(ObjectId.GenerateNewId(), $"Upgrade to {RequiredVersion}", null, null, LogFile.Id, LeaseState.Pending, Payload, null);
+					Lease = new AgentLease(ObjectId.GenerateNewId(), $"Upgrade to {RequiredVersion}", null, null, LogFile.Id, LeaseState.Pending, null, true, Payload);
 				}
 				return TaskSubscription.FromResult(Lease);
 			}
