Horde: Add an LRU cache for CAS objects, backed by a memory-mapped file on disk. Cache reads are lock free. Writes require a lock but complete in constant time (and do not prohibit reads).

The cache is transactional, and is designed not to lose data if the process is terminated. The cache index and bulk store are kept separate, and blocks in the bulk store are not overwritten until the index has been flushed.

The age of items is tracked via "generations". An 8-bit generation counter is incremented once a certain size of items has been added to the cache, and older entries in the cache can be trimmed periodically by calling TrimAsync().

Since blocks of memory are allocated via memory mapped files and references are not visible to the garbage collector, clients can create locks over the cache which prevents trim operations from running.

diff --git a/Engine/Source/Programs/Horde/HordeCommon/HordeCommon.csproj b/Engine/Source/Programs/Horde/HordeCommon/HordeCommon.csproj
--- a/Engine/Source/Programs/Horde/HordeCommon/HordeCommon.csproj
+++ b/Engine/Source/Programs/Horde/HordeCommon/HordeCommon.csproj
@@ -3,6 +3,7 @@
   <PropertyGroup>
     <TargetFramework>netcoreapp3.1</TargetFramework>
     <Nullable>enable</Nullable>
+    <AllowUnsafeBlocks>true</AllowUnsafeBlocks>
   </PropertyGroup>
 
   <ItemGroup>
diff --git a/Engine/Source/Programs/Horde/HordeCommon/LruCache.cs b/Engine/Source/Programs/Horde/HordeCommon/LruCache.cs
new file mode 100644
--- /dev/null
+++ b/Engine/Source/Programs/Horde/HordeCommon/LruCache.cs
@@ -0,0 +1,967 @@
+// Copyright Epic Games, Inc. All Rights Reserved.
+
+using EpicGames.Core;
+using Microsoft.Extensions.Logging;
+using System;
+using System.Buffers;
+using System.Collections;
+using System.Collections.Concurrent;
+using System.Collections.Generic;
+using System.Diagnostics;
+using System.Diagnostics.CodeAnalysis;
+using System.IO;
+using System.IO.MemoryMappedFiles;
+using System.Linq;
+using System.Numerics;
+using System.Runtime.CompilerServices;
+using System.Text;
+using System.Threading;
+using System.Threading.Tasks;
+
+[assembly: InternalsVisibleTo("HordeServerTests")]
+
+namespace HordeCommon
+{
+	/// <summary>
+	/// A thread-safe LRU cache of hashed objects, backed by a file on disk. Reads are typically lock free. Writes require a lock but complete in constant time.
+	/// The cache is transactional, and will not lose data if the process is terminated. Entries are kept in the cache until explicitly trimmed by calling
+	/// <see cref="LruCache.TrimAsync(long)"/>.
+	/// </summary>
+	public class LruCache : IDisposable
+	{
+		const int MinBlockSizeLog2 = 6; // 64 bytes
+		const int MaxBlockSize = 1 << MaxBlockSizeLog2;
+		const int MaxBlockSizeLog2 = 12; // 4096 bytes
+		const int PageSize = 1 << PageSizeLog2;
+		const int PageSizeLog2 = MaxBlockSizeLog2;
+
+		public const int MaxObjectSize = MaxBlockSize;
+
+		const int CurrentVersion = 0;
+
+		/// <summary>
+		/// A page of data allocated to store objects
+		/// </summary>
+		class Page
+		{
+			public Memory<byte> Data { get; }
+
+			public byte BlockSizeLog2;
+			public int PageCount; // For a large item, this is the total number of full pages in the list. The full item size is LargeItemPageCount * 4kb + Item.LastBlockLength
+			public ulong FreeBitMap;
+			public int NextPageIdx;
+
+			public int BlockSize => 1 << BlockSizeLog2;
+			public ulong EmptyBitMap => ~0UL >> (64 - (1 << (PageSizeLog2 - BlockSizeLog2))); // Note: left-shift only uses bottom 5 bits, cannot use (1UL << (PageSize >> BlockSizeLog2)) -1
+				
+			public Page(Memory<byte> Data)
+			{
+				this.Data = Data;
+			}
+
+			public void Reset(int BlockSizeLog2, int PageCount, int NextPage)
+			{
+				this.BlockSizeLog2 = (byte)BlockSizeLog2;
+				this.PageCount = PageCount;
+				this.FreeBitMap = EmptyBitMap;
+				this.NextPageIdx = NextPage;
+			}
+		}
+
+		/// <summary>
+		/// Information about an allocated item
+		/// </summary>
+		readonly struct Item
+		{
+			const long LargeItemFlag = 1L << 63;
+			const long ValidFlag = 1L << 62;
+
+			const int FirstPageOffset = 0;
+			const int FirstPageLength = 32;
+			const int GenerationOffset = FirstPageOffset + FirstPageLength;
+			const int GenerationLength = 8;
+			const int TailIndexOffset = GenerationOffset + GenerationLength;
+			const int TailIndexLength = 6;
+			const int TailSizeOffset = TailIndexOffset + TailIndexLength;
+			const int TailSizeLength = 13;
+
+			public readonly long Data;
+
+			public int FirstPageIdx => (int)(Data >> FirstPageOffset);
+			public byte Generation => (byte)((Data >> GenerationOffset) & ((1 << GenerationLength) - 1));
+			public int TailIndex => (int)(Data >> TailIndexOffset) & ((1 << TailIndexLength) - 1);
+			public int TailSize => (int)(Data >> TailSizeOffset) & ((1 << TailSizeLength) - 1);
+			public bool IsLargeItem => (Data & LargeItemFlag) != 0;
+			public bool IsValid => (Data & ValidFlag) != 0;
+
+			public Item(long Data)
+			{
+				this.Data = Data;
+			}
+
+			public Item(int FirstPage, byte Generation, int TailIndex, int TailSize, bool LargeItem)
+			{
+				Data = ((long)FirstPage << FirstPageOffset) | ((long)Generation << GenerationOffset) | ((long)TailIndex << TailIndexOffset) | ((long)TailSize << TailSizeOffset) | ValidFlag;
+				if (LargeItem)
+				{
+					Data |= LargeItemFlag;
+				}
+			}
+		}
+
+		/// <summary>
+		/// A scoped view of an LRU cache. Prevents data from being moved while a client is still using it.
+		/// </summary>
+		public class View : IDisposable
+		{
+			internal readonly LruCache Outer;
+
+			internal View(LruCache Outer)
+			{
+				this.Outer = Outer;
+				lock (Outer.LockObject)
+				{
+					Outer.NumReaders++;
+				}
+			}
+
+			public ReadOnlyMemory<byte> Get(IoHash Hash) => Outer.Get(Hash);
+
+			public void Dispose()
+			{
+				lock (Outer.LockObject)
+				{
+					Outer.NumReaders--;
+				}
+			}
+		}
+
+		/// <summary>
+		/// Open-addressed hashtable used for looking up items by hash. Uses Robin-Hood hashing for performance.
+		/// </summary>
+		class HashTable
+		{
+			const ulong Multiplier = 11400714819323198485ul; // fibonnaci hashing
+
+			private byte[] Hashes;
+			private long[] Items;
+			public int NumItems { get; private set; }
+			public int MaxItems => 1 << MaxItemsLog2;
+			private int MaxItemsLog2;
+
+			public HashTable(int MaxItems)
+			{
+				MaxItemsLog2 = BitOperations.Log2((uint)MaxItems);
+				MaxItems = 1 << MaxItemsLog2;
+
+				Hashes = new byte[MaxItems * IoHash.NumBytes];
+				Items = new long[MaxItems];
+			}
+
+			Memory<byte> GetHashMemory(int Slot) => Hashes.AsMemory(Slot * IoHash.NumBytes, IoHash.NumBytes);
+
+			int GetTargetSlot(IoHash Hash) => (int)(((ulong)Hash.GetHashCode() * Multiplier) >> (64 - MaxItemsLog2));
+
+			public bool Add(IoHash Hash, Item Item)
+			{
+				if (NumItems < MaxItems)
+				{
+					NumItems++;
+					int TargetSlot = GetTargetSlot(Hash);
+					Add(Hash, Item.Data, TargetSlot, TargetSlot);
+					return true;
+				}
+				return false;
+			}
+
+			private void Add(IoHash Hash, long ItemData, int Slot, int TargetSlot)
+			{
+				for(; ;)
+				{
+					// Check if this slot is empty
+					IoHash CompareHash = GetHash(Slot);
+					if(CompareHash == IoHash.Zero)
+					{
+						break;
+					}
+
+					// If not, get the probe sequence length of the item currently in this slot
+					int CompareTargetSlot = GetTargetSlot(CompareHash);
+
+					// If this is a better fit, replace it. Update subsequent slots first to ensure consistent read ordering.
+					int NextSlot = (Slot + 1) & (MaxItems - 1);
+					if (Displace(Slot, TargetSlot, CompareTargetSlot))
+					{
+						Add(CompareHash, Items[Slot], NextSlot, CompareTargetSlot);
+						break;
+					}
+					Slot = NextSlot;
+				}
+
+				Hash.Memory.CopyTo(GetHashMemory(Slot));
+				Items[Slot] = ItemData;
+			}
+
+			public bool Displace(int Slot, int TargetSlot, int CompareTargetSlot)
+			{
+				int ProbeLen = (TargetSlot + MaxItems - Slot) & (MaxItems - 1);
+				int CompareProbeLen = (CompareTargetSlot + MaxItems - Slot) & (MaxItems - 1);
+				return ProbeLen > CompareProbeLen;
+			}
+
+			public IoHash GetHash(int Slot) => new IoHash(GetHashMemory(Slot));
+			public Item GetItem(int Slot) => new Item(Items[Slot]);
+
+			public int Find(IoHash Hash)
+			{
+				int TargetSlot = GetTargetSlot(Hash);
+				for (int Slot = TargetSlot;;Slot = (Slot + 1) & (MaxItems - 1))
+				{
+					IoHash CompareHash = GetHash(Slot);
+					if (CompareHash == IoHash.Zero)
+					{
+						break;
+					}
+					if (CompareHash == Hash)
+					{
+						return Slot;
+					}
+
+					int CompareTargetSlot = GetTargetSlot(CompareHash);
+					if (Displace(Slot, TargetSlot, CompareTargetSlot))
+					{
+						break;
+					}
+				}
+				return -1;
+			}
+
+			public bool TryUpdate(int Slot, Item OldItem, Item NewItem)
+			{
+				return Interlocked.CompareExchange(ref Items[Slot], NewItem.Data, OldItem.Data) == OldItem.Data;
+			}
+		}
+
+		object LockObject = new object();
+		int NumReaders = 0;
+
+		FileReference IndexFile;
+
+		MemoryMappedFile DataFile;
+		MemoryMappedViewAccessor DataFileViewAccessor;
+		MemoryMappedFileView DataFileView;
+
+		const int InvalidIdx = -1;
+
+		int NumFreePages;
+		int FreePageIdx;
+		int[] PageIdxWithFreeBlock = new int[MaxBlockSizeLog2 + 1];
+
+		Page[] Pages;
+		HashTable Lookup;
+
+		int Generation;
+		long[] GenerationSize;
+
+		long MaxBytesPerGeneration;
+
+		public int NumItems => Lookup.NumItems;
+		public int MaxItems => Lookup.MaxItems;
+
+		public long NumBytes { get; private set; }
+		public long NumBytesWithBlockSlack { get; private set; }
+		public long NumBytesWithPageSlack => (Pages.Length - NumFreePages) << PageSizeLog2;
+
+		public long MaxSize { get; }
+
+		/// <summary>
+		/// Constructor
+		/// </summary>
+		/// <param name="IndexFileRef">Path to the index file</param>
+		/// <param name="DataFileRef">Path to the cache file</param>
+		/// <param name="MaxItems">Maximum number of items that can be stored in the cache</param>
+		/// <param name="MaxSize">Size of the cache file. This will be preallocated.</param>
+		private LruCache(FileReference IndexFileRef, FileReference DataFileRef, int MaxItems, long MaxSize)
+		{
+			MaxSize = MaxSize & ~(PageSize - 1);
+
+			this.IndexFile = IndexFileRef;
+			this.Generation = 0;
+			this.GenerationSize = new long[256];
+			this.Lookup = new HashTable((int)(MaxItems * 1.4));
+			this.MaxSize = MaxSize;
+
+			DataFile = MemoryMappedFile.CreateFromFile(DataFileRef.FullName, FileMode.OpenOrCreate, null, MaxSize);
+			DataFileViewAccessor = DataFile.CreateViewAccessor(0, MaxSize);
+			DataFileView = new MemoryMappedFileView(DataFileViewAccessor);
+
+			MaxBytesPerGeneration = MaxSize / 256;
+
+			int NumPages = (int)(MaxSize >> PageSizeLog2);
+			Pages = new Page[NumPages];
+
+			const int MaxChunkSize = int.MaxValue & ~(PageSize - 1);
+			const int MaxPagesPerChunk = MaxChunkSize >> PageSizeLog2;
+
+			long ChunkOffset = 0;
+			for (int BaseIdx = 0; BaseIdx < NumPages; )
+			{
+				int NumPagesInBlock = Math.Min(NumPages - BaseIdx, MaxPagesPerChunk);
+				int ChunkSize = NumPagesInBlock << PageSizeLog2;
+
+				Memory<byte> ChunkData = DataFileView.GetMemory(ChunkOffset, ChunkSize);
+				for (int Idx = 0; Idx < NumPagesInBlock; Idx++)
+				{
+					Memory<byte> Data = ChunkData.Slice(Idx << MaxBlockSizeLog2, MaxBlockSize);
+					Pages[BaseIdx + Idx] = new Page(Data);
+				}
+
+				BaseIdx += NumPagesInBlock;
+				ChunkOffset += ChunkSize;
+			}
+
+			FreePageIdx = 0;
+			for (int Idx = 1; Idx < Pages.Length; Idx++)
+			{
+				Pages[Idx - 1].NextPageIdx = Idx;
+			}
+
+			for (int Idx = 0; Idx < PageIdxWithFreeBlock.Length; Idx++)
+			{
+				PageIdxWithFreeBlock[Idx] = InvalidIdx;
+			}
+
+			NumFreePages = Pages.Length;
+		}
+
+		/// <inheritdoc/>
+		public void Dispose()
+		{
+			DataFileView.Dispose();
+			DataFileViewAccessor.Dispose();
+			DataFile.Dispose();
+		}
+
+		/// <summary>
+		/// Creates a new cache
+		/// </summary>
+		/// <param name="IndexFile"></param>
+		/// <param name="DataFile"></param>
+		/// <param name="MaxItems"></param>
+		/// <param name="MaxSize"></param>
+		/// <returns></returns>
+		public static LruCache CreateNew(FileReference IndexFile, FileReference DataFile, int MaxItems, long MaxSize)
+		{
+			DirectoryReference.CreateDirectory(IndexFile.Directory);
+			DirectoryReference.CreateDirectory(DataFile.Directory);
+
+			FileReference.Delete(IndexFile);
+			FileReference.Delete(GetTransactionFile(IndexFile));
+			FileReference.Delete(DataFile);
+
+			return new LruCache(IndexFile, DataFile, MaxItems, MaxSize);
+		}
+
+		/// <summary>
+		/// Opens an existing cache and modify its settings
+		/// </summary>
+		/// <param name="IndexFile"></param>
+		/// <param name="DataFile"></param>
+		/// <param name="MaxItems"></param>
+		/// <param name="MaxSize"></param>
+		/// <returns></returns>
+		public static async Task<LruCache?> OpenAndModifyAsync(FileReference IndexFile, FileReference DataFile, int MaxItems, long MaxSize)
+		{
+			LruCache? Cache = await OpenAsync(IndexFile, DataFile);
+			if (Cache != null)
+			{
+				if (Cache.MaxItems != MaxItems || Cache.MaxSize != MaxSize)
+				{
+					FileReference NewIndexFile = new FileReference(IndexFile.FullName + ".new");
+					FileReference NewDataFile = new FileReference(DataFile.FullName + ".new");
+
+					using (LruCache NewCache = LruCache.CreateNew(NewIndexFile, NewDataFile, MaxItems, MaxSize))
+					{
+						for (int Idx = 0; Idx < Cache.Lookup.MaxItems; Idx++)
+						{
+							Item Item = NewCache.Lookup.GetItem(Idx);
+							if (Item.IsValid)
+							{
+								IoHash Hash = Cache.Lookup.GetHash(Idx);
+								NewCache.Add(Hash, NewCache.Get(Hash));
+							}
+						}
+					}
+
+					FileReference OldIndexFile = new FileReference(IndexFile.FullName + ".old");
+					FileReference OldDataFile = new FileReference(DataFile.FullName + ".old");
+
+					FileReference.Move(IndexFile, OldIndexFile, true);
+					FileReference.Move(DataFile, OldDataFile, true);
+
+					FileReference.Move(NewIndexFile, IndexFile, true);
+					FileReference.Move(NewDataFile, DataFile, true);
+
+					FileReference.Delete(OldIndexFile);
+					FileReference.Delete(OldDataFile);
+
+					Cache.Dispose();
+					Cache = await OpenAsync(IndexFile, DataFile);
+				}
+			}
+			return Cache;
+		}
+
+		/// <summary>
+		/// Loads the current cache index from disk
+		/// </summary>
+		public static async Task<LruCache> OpenAsync(FileReference IndexFile, FileReference DataFile)
+		{
+			LruCache? Cache = await TryOpenAsync(IndexFile, DataFile);
+			if (Cache == null)
+			{
+				throw new FileNotFoundException($"Unable to open {IndexFile}");
+			}
+			return Cache;
+		}
+
+		/// <summary>
+		/// Loads the current cache index from disk
+		/// </summary>
+		public static async Task<LruCache?> TryOpenAsync(FileReference IndexFile, FileReference DataFile)
+		{
+			// Check if the index exists. If not, check if a save of it was interrupted.
+			if (!FileReference.Exists(IndexFile))
+			{
+				FileReference TransactionIndexFile = new FileReference(IndexFile.FullName + ".tr");
+				if (!FileReference.Exists(TransactionIndexFile))
+				{
+					return null;
+				}
+				FileReference.Move(TransactionIndexFile, IndexFile, true);
+			}
+
+			byte[] Data = await FileReference.ReadAllBytesAsync(IndexFile);
+			MemoryReader Reader = new MemoryReader(Data);
+
+			byte Version = Reader.ReadUInt8();
+			if (Version > CurrentVersion)
+			{
+				throw new InvalidDataException($"Unable to load lru cache index of version {Version}");
+			}
+
+			int MaxItems = Reader.ReadInt32();
+			long MaxSize = Reader.ReadInt64();
+
+			LruCache Cache = new LruCache(IndexFile, DataFile, MaxItems, MaxSize);
+			Cache.Generation = Reader.ReadUInt8();
+
+			foreach (Page Page in Cache.Pages)
+			{
+				Page.PageCount = InvalidIdx;
+			}
+
+			int ItemCount = Reader.ReadInt32();
+			for (int Idx = 0; Idx < ItemCount; Idx++)
+			{
+				IoHash Hash = Reader.ReadIoHash();
+				Item Item = new Item(Reader.ReadInt64());
+
+				int PageIdx = Item.FirstPageIdx;
+
+				Page Page = Cache.Pages[PageIdx];
+				if (Item.IsLargeItem)
+				{
+					int PageCount = 0;
+
+					Page FirstPage = Page;
+					for (; ; )
+					{
+						int NextPageIdx = Reader.ReadInt32();
+						Page.Reset(MaxBlockSizeLog2, 0, NextPageIdx & 0x7fffffff);
+						Page.FreeBitMap = 0;
+						Page = Cache.Pages[Page.NextPageIdx];
+
+						PageCount++;
+
+						if (NextPageIdx >= 0)
+						{
+							break;
+						}
+					}
+
+					Cache.NumBytes += PageCount << PageSizeLog2;
+					Cache.NumBytesWithBlockSlack += PageCount << PageSizeLog2;
+
+					Page = FirstPage;
+					for (; PageCount > 0; PageCount--)
+					{
+						Page.PageCount = PageCount;
+						Page = Cache.Pages[Page.NextPageIdx];
+					}
+				}
+
+				if (Page.PageCount == InvalidIdx)
+				{
+					int BlockSizeLog2 = GetBlockSizeLog2(Item.TailSize);
+					Page.Reset(BlockSizeLog2, 0, InvalidIdx);
+				}
+
+				Page.FreeBitMap ^= 1UL << Item.TailIndex;
+
+				Cache.NumBytes += Item.TailSize;
+				Cache.NumBytesWithBlockSlack += Page.BlockSize;
+
+				Cache.Lookup.Add(Hash, Item);
+			}
+
+			// Fix up the lists of free and partially-free pages
+			Cache.NumFreePages = 0;
+			Cache.FreePageIdx = InvalidIdx;
+			for (int PageIdx = Cache.Pages.Length - 1; PageIdx >= 0; PageIdx--)
+			{
+				Page Page = Cache.Pages[PageIdx];
+				if (Page.PageCount == InvalidIdx)
+				{
+					Page.Reset(0, 0, Cache.FreePageIdx);
+					Page.NextPageIdx = Cache.FreePageIdx;
+					Cache.FreePageIdx = PageIdx;
+					Cache.NumFreePages++;
+				}
+				else if (Page.FreeBitMap != 0)
+				{
+					Page.NextPageIdx = Cache.PageIdxWithFreeBlock[Page.BlockSizeLog2];
+					Cache.FreePageIdx = PageIdx;
+				}
+			}
+			return Cache;
+		}
+
+		/// <summary>
+		/// Saves the cache index
+		/// </summary>
+		public async Task SaveAsync()
+		{
+			byte[] IndexData;
+			lock (LockObject)
+			{
+				int IndexSize = sizeof(byte) + sizeof(int) + sizeof(long) + sizeof(byte) + sizeof(int);
+				for (int Slot = 0; Slot < Lookup.MaxItems; Slot++)
+				{
+					Item Item = Lookup.GetItem(Slot);
+					if (Item.IsValid)
+					{
+						IndexSize += IoHash.NumBytes + sizeof(long);
+						if (Item.IsLargeItem)
+						{
+							IndexSize += Pages[Item.FirstPageIdx].PageCount * sizeof(int);
+						}
+					}
+				}
+
+				IndexData = new byte[IndexSize];
+
+				MemoryWriter Writer = new MemoryWriter(IndexData);
+				Writer.WriteUInt8(CurrentVersion);
+				Writer.WriteInt32(MaxItems);
+				Writer.WriteInt64(MaxSize);
+
+				Writer.WriteUInt8((byte)Generation);
+
+				Writer.WriteInt32(Lookup.NumItems);
+				for (int Slot = 0; Slot < Lookup.MaxItems; Slot++)
+				{
+					Item Item = Lookup.GetItem(Slot);
+					if (Item.IsValid)
+					{
+						Writer.WriteFixedLengthBytes(Lookup.GetHash(Slot).Span);
+						Writer.WriteInt64(Item.Data);
+
+						if (Item.IsLargeItem)
+						{
+							Page Page = Pages[Item.FirstPageIdx];
+							while (Page.PageCount > 1)
+							{
+								Writer.WriteUInt32((uint)Page.NextPageIdx | 0x80000000U);
+								Page = Pages[Page.NextPageIdx];
+							}
+							Writer.WriteInt32(Page.NextPageIdx);
+						}
+					}
+				}
+				Writer.CheckOffset(IndexSize);
+			}
+
+			// Write the new file to disk
+			FileReference TransactionFile = GetTransactionFile(IndexFile);
+			await FileReference.WriteAllBytesAsync(TransactionFile, IndexData);
+			FileReference.Move(TransactionFile, IndexFile, true);
+		}
+
+		/// <summary>
+		/// Update the generation counter. Useful for testing, where we can specify an explicit point to garbage collect to.
+		/// </summary>
+		public void NextGeneration()
+		{
+			lock (LockObject)
+			{
+				Generation = (Generation + 1) & 0xff;
+			}
+		}
+
+		/// <summary>
+		/// Adds an item to the cache
+		/// </summary>
+		/// <param name="Data"></param>
+		/// <returns></returns>
+		public IoHash Add(ReadOnlyMemory<byte> Data)
+		{
+			IoHash Hash = IoHash.Compute(Data.Span);
+			Add(Hash, Data);
+			return Hash;
+		}
+
+		static int GetBlockSizeLog2(int ItemSize)
+		{
+			int ItemSizeLog2 = Math.Max(BitOperations.Log2((uint)ItemSize), MinBlockSizeLog2);
+			if (ItemSize > (1 << ItemSizeLog2))
+			{
+				ItemSizeLog2++;
+			}
+			return ItemSizeLog2;
+		}
+
+		/// <summary>
+		/// Adds an item into the cache
+		/// </summary>
+		/// <param name="Hash">Has of the item</param>
+		/// <param name="Data">The data to add</param>
+		public void Add(IoHash Hash, ReadOnlyMemory<byte> Data)
+		{
+			lock (LockObject)
+			{
+				// Make sure we can add the hash
+				if (Lookup.NumItems >= Lookup.MaxItems)
+				{
+					return;
+				}
+				if (Lookup.Find(Hash) != -1)
+				{
+					return;
+				}
+
+				// Make sure there are enough full pages to store the data. The last page is always treated as a 'tail' block, even if it's exactly 4kb long.
+				int NumFullPages = Math.Max(Data.Length - 1, 0) >> PageSizeLog2;
+
+				// Get the size of the 'tail' page
+				int TailSize = Data.Length - (NumFullPages << PageSizeLog2);
+				int TailSizeLog2 = GetBlockSizeLog2(TailSize);
+
+				// Get the tail page. This may be null if there are no pages free of the given size.
+				int TailPageIdx = PageIdxWithFreeBlock[TailSizeLog2];
+
+				// Make sure there are enough pages available
+				int NumRequiredPages = NumFullPages;
+				if (TailPageIdx == InvalidIdx)
+				{
+					NumRequiredPages++;
+				}
+				if (NumFreePages < NumRequiredPages)
+				{
+					return;
+				}
+
+				// Allocate the tail page if necessary
+				if (TailPageIdx == -1)
+				{
+					TailPageIdx = FreePageIdx;
+					FreePageIdx = Pages[TailPageIdx].NextPageIdx;
+
+					PageIdxWithFreeBlock[TailSizeLog2] = TailPageIdx;
+					Pages[TailPageIdx].Reset(TailSizeLog2, 0, -1);
+				}
+
+				// Take an item from the first free span
+				Page TailPage = Pages[TailPageIdx];
+				int TailIndex = BitOperations.TrailingZeroCount(TailPage.FreeBitMap) & 63;
+				TailPage.FreeBitMap ^= 1UL << TailIndex;
+
+				// If the page is full, remove it from the free list
+				if (TailPage.FreeBitMap == 0UL)
+				{
+					PageIdxWithFreeBlock[TailSizeLog2] = TailPage.NextPageIdx;
+				}
+
+				// Copy the tail data to the cache
+				int TailOffset = TailIndex << TailSizeLog2;
+				Data.Slice(Data.Length - TailSize).CopyTo(TailPage.Data.Slice(TailOffset));
+
+				// Add all the other full pages in reverse order
+				int FirstPageIdx = TailPageIdx;
+				for (int Idx = NumFullPages - 1; Idx >= 0; Idx--)
+				{
+					int PageIdx = FreePageIdx;
+
+					Page Page = Pages[PageIdx];
+					FreePageIdx = Page.NextPageIdx;
+
+					Page.Reset(PageSizeLog2, NumFullPages - Idx, FirstPageIdx);
+					Page.FreeBitMap = 0;
+
+					Data.Slice(Idx << PageSizeLog2, PageSize).CopyTo(Page.Data);
+					FirstPageIdx = PageIdx;
+				}
+
+				// Create the new item
+				Item Item = new Item(FirstPageIdx, (byte)Generation, TailIndex, TailSize, NumFullPages > 0);
+				Lookup.Add(Hash, Item);
+
+				// Update the allocation stats
+				NumBytes += Data.Length;
+				NumBytesWithBlockSlack += (NumFullPages << PageSizeLog2) + (1 << TailSizeLog2);
+				NumFreePages -= NumRequiredPages;
+
+				// Update the generation
+				long NewGenerationSize = Interlocked.Add(ref GenerationSize[Generation], (NumFullPages << PageSizeLog2) + (1 << TailSizeLog2));
+				if (NewGenerationSize > MaxBytesPerGeneration)
+				{
+					Generation++;
+				}
+			}
+		}
+
+		/// <summary>
+		/// Inserts a page into a linked list, sorted by index
+		/// </summary>
+		/// <param name="FirstPageIdx"></param>
+		/// <param name="PageIdx"></param>
+		void AddPageToList(ref int FirstPageIdx, int PageIdx)
+		{
+			Page Page = Pages[PageIdx];
+			if (FirstPageIdx == InvalidIdx || PageIdx < FirstPageIdx)
+			{
+				Page.NextPageIdx = FirstPageIdx;
+				FirstPageIdx = PageIdx;
+			}
+			else
+			{
+				Page PrevPage = Pages[FirstPageIdx];
+				while (PrevPage.NextPageIdx != InvalidIdx && PageIdx < PrevPage.NextPageIdx)
+				{
+					PrevPage = Pages[PrevPage.NextPageIdx];
+				}
+
+				Page.NextPageIdx = PrevPage.NextPageIdx;
+				PrevPage.NextPageIdx = PageIdx;
+			}
+		}
+
+		/// <summary>
+		/// Removes a page from a linked list
+		/// </summary>
+		/// <param name="FirstPage"></param>
+		/// <param name="Page"></param>
+		void RemovePageFromList(ref int FirstPageIdx, int PageIdx)
+		{
+			Page Page = Pages[PageIdx];
+			if (FirstPageIdx == PageIdx)
+			{
+				FirstPageIdx = Page.NextPageIdx;
+			}
+			else if(FirstPageIdx != InvalidIdx)
+			{
+				Page PrevPage = Pages[FirstPageIdx];
+				for (; ;)
+				{
+					int NextPageIdx = PrevPage.NextPageIdx;
+					if (NextPageIdx == InvalidIdx)
+					{
+						break;
+					}
+					if (NextPageIdx == PageIdx)
+					{
+						PrevPage.NextPageIdx = Page.NextPageIdx;
+						break;
+					}
+					PrevPage = Pages[NextPageIdx];
+				}
+			}
+		}
+
+		/// <summary>
+		/// Creates a view of the cache which can be used for queries
+		/// </summary>
+		/// <returns>New view object</returns>
+		public View LockView()
+		{
+			return new View(this);
+		}
+
+		/// <summary>
+		/// Tries to read an item from the cache. Invoked via <see cref="View.Get(IoHash)"/>.
+		/// </summary>
+		/// <param name="Hash"></param>
+		/// <returns></returns>
+		private ReadOnlyMemory<byte> Get(IoHash Hash)
+		{
+			int Slot = Lookup.Find(Hash);
+			if (Slot == -1)
+			{
+				return ReadOnlyMemory<byte>.Empty;
+			}
+
+			Item Item = Lookup.GetItem(Slot);
+			Page FirstPage = Pages[Item.FirstPageIdx];
+
+			if (Item.IsLargeItem)
+			{
+				int ItemSize = (FirstPage.PageCount << PageSizeLog2) + Item.TailSize;
+
+				byte[] Buffer = new byte[ItemSize];
+
+				int Offset = 0;
+
+				Page NextPage = FirstPage;
+				for (int Idx = 0; Idx < FirstPage.PageCount; Idx++)
+				{
+					NextPage.Data.CopyTo(Buffer.AsMemory(Offset));
+					NextPage = Pages[NextPage.NextPageIdx];
+					Offset += PageSize;
+				}
+				NextPage!.Data.Slice(Item.TailIndex << NextPage.BlockSizeLog2, Item.TailSize).CopyTo(Buffer.AsMemory(Offset));
+
+				UpdateGeneration(Slot, Item, ItemSize);
+				return Buffer;
+			}
+			else
+			{
+				UpdateGeneration(Slot, Item, Item.TailSize);
+				return FirstPage.Data.Slice(Item.TailIndex << FirstPage.BlockSizeLog2, Item.TailSize);
+			}
+		}
+
+		void UpdateGeneration(int Slot, Item Item, int ItemSize)
+		{
+			for (; ; )
+			{
+				int GenerationCopy = Generation;
+				if (Item.Generation == GenerationCopy)
+				{
+					break;
+				}
+
+				Item NewItem = new Item(Item.FirstPageIdx, (byte)GenerationCopy, Item.TailIndex, Item.TailSize, Item.IsLargeItem);
+				if (Lookup.TryUpdate(Slot, Item, NewItem))
+				{
+					Interlocked.Add(ref GenerationSize[Item.Generation], -ItemSize);
+
+					long NewGenerationSize = Interlocked.Add(ref GenerationSize[Item.Generation], ItemSize);
+					if (NewGenerationSize > MaxBytesPerGeneration)
+					{
+						Interlocked.CompareExchange(ref Generation, (byte)(GenerationCopy + 1), GenerationCopy);
+					}
+
+					break;
+				}
+
+				Item = Lookup.GetItem(Slot);
+			}
+		}
+
+		/// <summary>
+		/// Trims the cache to the given size, allowing new objects to be allocated. 
+		/// </summary>
+		/// <param name="TargetSize"></param>
+		/// <returns></returns>
+		public async ValueTask<bool> TrimAsync(long TargetSize)
+		{
+			List<Item> FreeItems = new List<Item>();
+			lock (LockObject)
+			{
+				// Check that nobody is currently reading from the cache
+				if (NumReaders > 0)
+				{
+					return false;
+				}
+
+				// Figure out how many generations we're going to keep
+				long NewTotalSize = 0;
+				int KeepGenerations = 0;
+				while (KeepGenerations < 256 && NewTotalSize < TargetSize)
+				{
+					NewTotalSize += GenerationSize[(int)(byte)(Generation - KeepGenerations)];
+					KeepGenerations++;
+				}
+
+				// Create a new list of read-only items
+				HashTable NewItems = new HashTable(Lookup.MaxItems);
+				for (int Slot = 0; Slot < Lookup.MaxItems; Slot++)
+				{
+					Item Item = Lookup.GetItem(Slot);
+					if (Item.IsValid)
+					{
+						int Age = (byte)(Generation - Item.Generation);
+						if (Age < KeepGenerations)
+						{
+							NewItems.Add(Lookup.GetHash(Slot), Item);
+						}
+						else
+						{
+							FreeItems.Add(Item);
+						}
+					}
+				}
+				Lookup = NewItems;
+			}
+
+			// Save the new state. This will ensure consistency when we restart.
+			await SaveAsync();
+
+			// Release any pages that are no longer needed
+			lock (LockObject)
+			{
+				foreach (Item Item in FreeItems)
+				{
+					int PageIdx = Item.FirstPageIdx;
+					Page Page = Pages[PageIdx];
+
+					// Free any full pages
+					int PageCount = Page.PageCount;
+					while (Page.PageCount > 0)
+					{
+						AddPageToList(ref FreePageIdx, PageIdx);
+						PageIdx = Page.NextPageIdx;
+						Page = Pages[Page.NextPageIdx];
+					}
+
+					// Update the stats
+					NumBytes -= (PageCount << PageSizeLog2) + Item.TailSize;
+					NumBytesWithBlockSlack -= (PageCount << PageSizeLog2) + (1 << Page.BlockSizeLog2);
+
+					// If this is the first free item on the tail page, add it to the free list
+					if (Page.FreeBitMap == 0)
+					{
+						AddPageToList(ref PageIdxWithFreeBlock[Page.BlockSizeLog2], PageIdx);
+					}
+
+					// Mark this block as free
+					Page.FreeBitMap |= 1UL << Item.TailIndex;
+
+					// If this page is completely empty, add it to the free page list
+					if (Page.FreeBitMap == Page.EmptyBitMap)
+					{
+						RemovePageFromList(ref PageIdxWithFreeBlock[Page.BlockSizeLog2], PageIdx);
+						AddPageToList(ref FreePageIdx, PageIdx);
+						NumFreePages++;
+					}
+				}
+			}
+			return true;
+		}
+
+		static FileReference GetTransactionFile(FileReference IndexFile)
+		{
+			return new FileReference(IndexFile.FullName + ".tr");
+		}
+	}
+}
diff --git a/Engine/Source/Programs/Horde/HordeCommon/MemoryMappedFileView.cs b/Engine/Source/Programs/Horde/HordeCommon/MemoryMappedFileView.cs
new file mode 100644
--- /dev/null
+++ b/Engine/Source/Programs/Horde/HordeCommon/MemoryMappedFileView.cs
@@ -0,0 +1,76 @@
+// Copyright Epic Games, Inc. All Rights Reserved.
+
+using System;
+using System.Buffers;
+using System.Collections.Generic;
+using System.IO;
+using System.IO.MemoryMappedFiles;
+using System.Text;
+
+namespace HordeCommon
+{
+	/// <summary>
+	/// Implements an unmarshlled view of a memory mapped file
+	/// </summary>
+	unsafe class MemoryMappedFileView : IDisposable
+	{
+		sealed unsafe class MemoryWrapper : MemoryManager<byte>
+		{
+			private readonly byte* Pointer;
+			private readonly int Length;
+
+			public MemoryWrapper(byte* Pointer, int Length)
+			{
+				this.Pointer = Pointer;
+				this.Length = Length;
+			}
+
+			/// <inheritdoc/>
+			public override Span<byte> GetSpan() => new Span<byte>(Pointer, Length);
+
+			/// <inheritdoc/>
+			public override MemoryHandle Pin(int elementIndex) => new MemoryHandle(Pointer + elementIndex);
+
+			/// <inheritdoc/>
+			public override void Unpin() { }
+
+			/// <inheritdoc/>
+			protected override void Dispose(bool disposing) { }
+		}
+
+		MemoryMappedViewAccessor MemoryMappedViewAccessor;
+		byte* Data;
+
+		/// <summary>
+		/// Constructor
+		/// </summary>
+		/// <param name="MemoryMappedViewAccessor"></param>
+		public MemoryMappedFileView(MemoryMappedViewAccessor MemoryMappedViewAccessor)
+		{
+			this.MemoryMappedViewAccessor = MemoryMappedViewAccessor;
+			MemoryMappedViewAccessor.SafeMemoryMappedViewHandle.AcquirePointer(ref Data);
+		}
+
+		/// <summary>
+		/// Gets a memory object for the given range
+		/// </summary>
+		/// <param name="Offset"></param>
+		/// <param name="Length"></param>
+		/// <returns></returns>
+		public Memory<byte> GetMemory(long Offset, int Length)
+		{
+			MemoryWrapper Wrapper = new MemoryWrapper(Data + Offset, Length);
+			return Wrapper.Memory;
+		}
+
+		/// <inheritdoc/>
+		public void Dispose()
+		{
+			if (Data != null)
+			{
+				MemoryMappedViewAccessor.SafeMemoryMappedViewHandle.ReleasePointer();
+				Data = null;
+			}
+		}
+	}
+}
diff --git a/Engine/Source/Programs/Horde/HordeServerTests/LruCacheTests.cs b/Engine/Source/Programs/Horde/HordeServerTests/LruCacheTests.cs
new file mode 100644
--- /dev/null
+++ b/Engine/Source/Programs/Horde/HordeServerTests/LruCacheTests.cs
@@ -0,0 +1,314 @@
+// Copyright Epic Games, Inc. All Rights Reserved.
+
+using EpicGames.Core;
+using HordeCommon;
+using Microsoft.VisualStudio.TestTools.UnitTesting;
+using System;
+using System.Buffers.Binary;
+using System.Collections.Generic;
+using System.Text;
+using System.Threading.Tasks;
+
+namespace HordeServerTests
+{
+	[TestClass]
+	public class LruCacheTests : IDisposable
+	{
+		FileReference IndexFile;
+		FileReference DataFile;
+		LruCache Cache;
+
+		public LruCacheTests()
+		{
+			this.IndexFile = new FileReference("test.idx");
+			this.DataFile = new FileReference("test.dat");
+
+			Cache = LruCache.CreateNew(IndexFile, DataFile, 2048, 1024 * 1024);
+		}
+
+		public void Dispose()
+		{
+			Cache.Dispose();
+		}
+
+		[TestMethod]
+		public void EmptyObject()
+		{
+			IoHash Hash = Cache.Add(Array.Empty<byte>());
+			Assert.AreEqual(1, Cache.NumItems);
+
+			using (LruCache.View View = Cache.LockView())
+			{
+				byte[] Data = View.Get(Hash).ToArray();
+				Assert.AreEqual(0, Data.Length);
+			}
+		}
+
+		[TestMethod]
+		public void DuplicateObjects()
+		{
+			Cache.Add(new byte[] { 1, 2, 3 });
+			Cache.Add(new byte[] { 1, 2, 3 });
+			Cache.Add(new byte[] { 1, 2, 4 });
+			Assert.AreEqual(2, Cache.NumItems);
+			Assert.AreEqual(6, Cache.NumBytes);
+			Assert.AreEqual(128, Cache.NumBytesWithBlockSlack);
+			Assert.AreEqual(4096, Cache.NumBytesWithPageSlack);
+		}
+
+		static byte[] MakeBlockWithIndex(int Size, int Index)
+		{
+			byte[] Data = new byte[Size];
+			BinaryPrimitives.WriteInt32LittleEndian(Data, Index);
+			return Data;
+		}
+
+		[TestMethod]
+		public void LargeObjectPacking()
+		{
+			Cache.Add(MakeBlockWithIndex(1000, 1));
+			Assert.AreEqual(1, Cache.NumItems);
+			Assert.AreEqual(1000, Cache.NumBytes);
+			Assert.AreEqual(1024, Cache.NumBytesWithBlockSlack);
+			Assert.AreEqual(4096, Cache.NumBytesWithPageSlack);
+
+			Cache.Add(MakeBlockWithIndex(1000, 2));
+			Assert.AreEqual(2, Cache.NumItems);
+			Assert.AreEqual(2000, Cache.NumBytes);
+			Assert.AreEqual(2048, Cache.NumBytesWithBlockSlack);
+			Assert.AreEqual(4096, Cache.NumBytesWithPageSlack);
+
+			Cache.Add(MakeBlockWithIndex(1000, 3));
+			Assert.AreEqual(3, Cache.NumItems);
+			Assert.AreEqual(3000, Cache.NumBytes);
+			Assert.AreEqual(3072, Cache.NumBytesWithBlockSlack);
+			Assert.AreEqual(4096, Cache.NumBytesWithPageSlack);
+
+			Cache.Add(MakeBlockWithIndex(1000, 4));
+			Assert.AreEqual(4, Cache.NumItems);
+			Assert.AreEqual(4000, Cache.NumBytes);
+			Assert.AreEqual(4096, Cache.NumBytesWithBlockSlack);
+			Assert.AreEqual(4096, Cache.NumBytesWithPageSlack);
+
+			Cache.Add(MakeBlockWithIndex(1000, 5));
+			Assert.AreEqual(5, Cache.NumItems);
+			Assert.AreEqual(5000, Cache.NumBytes);
+			Assert.AreEqual(5120, Cache.NumBytesWithBlockSlack);
+			Assert.AreEqual(8192, Cache.NumBytesWithPageSlack);
+		}
+
+		[TestMethod]
+		public void SmallObjectPacking()
+		{
+			Cache.Add(MakeBlockWithIndex(64, 1));
+			Assert.AreEqual(1, Cache.NumItems);
+			Assert.AreEqual(64, Cache.NumBytes);
+			Assert.AreEqual(64, Cache.NumBytesWithBlockSlack);
+			Assert.AreEqual(4096, Cache.NumBytesWithPageSlack);
+
+			Cache.Add(MakeBlockWithIndex(64, 2));
+			Assert.AreEqual(2, Cache.NumItems);
+			Assert.AreEqual(128, Cache.NumBytes);
+			Assert.AreEqual(128, Cache.NumBytesWithBlockSlack);
+			Assert.AreEqual(4096, Cache.NumBytesWithPageSlack);
+
+			Cache.Add(MakeBlockWithIndex(64, 3));
+			Assert.AreEqual(3, Cache.NumItems);
+			Assert.AreEqual(192, Cache.NumBytes);
+			Assert.AreEqual(192, Cache.NumBytesWithBlockSlack);
+			Assert.AreEqual(4096, Cache.NumBytesWithPageSlack);
+		}
+
+		[TestMethod]
+		public void FillItems()
+		{
+			for (int Idx = 0; Idx < 2048; Idx++)
+			{
+				byte[] Data = MakeBlockWithIndex(64, Idx);
+				Cache.Add(Data);
+			}
+			Assert.AreEqual(2048, Cache.NumItems);
+
+			Cache.Add(new byte[] { 1, 2, 3 });
+			Assert.AreEqual(2048, Cache.NumItems);
+		}
+
+		[TestMethod]
+		public void FillCapacity()
+		{
+			for (int Idx = 0; Idx < 1024; Idx++)
+			{
+				byte[] Data = MakeBlockWithIndex(1024, Idx);
+				Cache.Add(Data);
+			}
+			Assert.AreEqual(1024, Cache.NumItems);
+
+			Cache.Add(new byte[] { 1, 2, 3 });
+			Assert.AreEqual(1024, Cache.NumItems);
+		}
+
+		[TestMethod]
+		public void HugeObjects()
+		{
+			byte[] Data = new byte[4096 + 4096 + 2000];
+
+			Random Random = new Random(0);
+			Random.NextBytes(Data);
+
+			IoHash Hash = Cache.Add(Data);
+			Assert.AreEqual(1, Cache.NumItems);
+			Assert.AreEqual(Data.Length, Cache.NumBytes);
+			Assert.AreEqual(4096 + 4096 + 2048, Cache.NumBytesWithBlockSlack);
+			Assert.AreEqual(4096 + 4096 + 4096, Cache.NumBytesWithPageSlack);
+
+			using LruCache.View View = Cache.LockView();
+			ReadOnlyMemory<byte> LookupData = View.Get(Hash);
+			Assert.IsTrue(Data.AsSpan().SequenceEqual(LookupData.Span));
+		}
+
+		[TestMethod]
+		public void HugeObjects2()
+		{
+			const int NumItems = 5;
+			for (int Idx = 0; Idx < NumItems; Idx++)
+			{
+				byte[] Data = MakeBlockWithIndex(2000 + (Idx * 2048), Idx);
+				Cache.Add(Data);
+			}
+
+			Assert.AreEqual(NumItems, Cache.NumItems);
+
+			using (LruCache.View View = Cache.LockView())
+			{
+				for (int Idx = 0; Idx < NumItems; Idx++)
+				{
+					byte[] Data = MakeBlockWithIndex(2000 + (Idx * 2048), Idx);
+					IoHash Hash = IoHash.Compute(Data);
+					ReadOnlyMemory<byte> OtherData = View.Get(Hash);
+					Assert.IsTrue(Data.AsSpan().SequenceEqual(OtherData.Span));
+				}
+			}
+		}
+
+		[TestMethod]
+		public async Task SmallObjectSerialization()
+		{
+			for (int Idx = 0; Idx < 1000; Idx++)
+			{
+				byte[] Data = MakeBlockWithIndex(4 + Idx, Idx);
+				Cache.Add(Data);
+			}
+
+			CheckSmallObjects();
+
+			Assert.AreEqual(1000, Cache.NumItems);
+			long NumBytes = Cache.NumBytes;
+			long NumBytesWithBlockSlack = Cache.NumBytesWithBlockSlack;
+			long NumBytesWithPageSlack = Cache.NumBytesWithPageSlack;
+
+			await Cache.SaveAsync();
+			Cache.Dispose();
+
+			Cache = await LruCache.OpenAsync(IndexFile, DataFile);
+
+			Assert.AreEqual(1000, Cache.NumItems);
+			Assert.AreEqual(NumBytes, Cache.NumBytes);
+			Assert.AreEqual(NumBytesWithBlockSlack, Cache.NumBytesWithBlockSlack);
+			Assert.AreEqual(NumBytesWithPageSlack, Cache.NumBytesWithPageSlack);
+
+			CheckSmallObjects();
+		}
+
+		void CheckSmallObjects()
+		{
+			using (LruCache.View View = Cache.LockView())
+			{
+				for (int Idx = 0; Idx < 1000; Idx++)
+				{
+					byte[] Data = MakeBlockWithIndex(4 + Idx, Idx);
+					IoHash Hash = IoHash.Compute(Data);
+					ReadOnlyMemory<byte> OtherData = View.Get(Hash);
+					Assert.IsTrue(Data.AsSpan().SequenceEqual(OtherData.Span));
+				}
+			}
+		}
+
+		[TestMethod]
+		public async Task LargeObjectSerialization()
+		{
+			const int NumItems = 10;
+			for (int Idx = 0; Idx < NumItems; Idx++)
+			{
+				byte[] Data = MakeBlockWithIndex(2000 + (Idx * 2048), Idx);
+				Cache.Add(Data);
+			}
+
+			CheckLargeObjects(NumItems);
+
+			Assert.AreEqual(NumItems, Cache.NumItems);
+			long NumBytes = Cache.NumBytes;
+			long NumBytesWithBlockSlack = Cache.NumBytesWithBlockSlack;
+			long NumBytesWithPageSlack = Cache.NumBytesWithPageSlack;
+
+			await Cache.SaveAsync();
+			Cache.Dispose();
+
+			Cache = await LruCache.OpenAsync(IndexFile, DataFile);
+
+			Assert.AreEqual(NumItems, Cache.NumItems);
+			Assert.AreEqual(NumBytes, Cache.NumBytes);
+			Assert.AreEqual(NumBytesWithBlockSlack, Cache.NumBytesWithBlockSlack);
+			Assert.AreEqual(NumBytesWithPageSlack, Cache.NumBytesWithPageSlack);
+
+			CheckLargeObjects(NumItems);
+		}
+
+		void CheckLargeObjects(int NumItems)
+		{
+			using (LruCache.View View = Cache.LockView())
+			{
+				for (int Idx = 0; Idx < NumItems; Idx++)
+				{
+					byte[] Data = MakeBlockWithIndex(2000 + (Idx * 2048), Idx);
+					IoHash Hash = IoHash.Compute(Data);
+					ReadOnlyMemory<byte> OtherData = View.Get(Hash);
+					Assert.IsTrue(Data.AsSpan().SequenceEqual(OtherData.Span));
+				}
+			}
+		}
+
+		[TestMethod]
+		public async Task ExpireItems()
+		{
+			for (int Idx = 0; Idx < 1000; Idx++)
+			{
+				byte[] Data = MakeBlockWithIndex(4 + Idx, Idx);
+				Cache.Add(Data);
+			}
+
+			long TrimSize = Cache.NumBytesWithBlockSlack;
+			Cache.NextGeneration();
+
+			List<IoHash> TestHashes = new List<IoHash>();
+			for (int Idx = 1000; Idx < 2000; Idx++)
+			{
+				byte[] Data = MakeBlockWithIndex(4 + Idx, Idx);
+				TestHashes.Add(Cache.Add(Data));
+			}
+
+			long DesiredSize = Cache.NumBytesWithBlockSlack - TrimSize;
+			await Cache.TrimAsync(DesiredSize);
+
+			Assert.AreEqual(DesiredSize, Cache.NumBytesWithBlockSlack);
+
+			using(LruCache.View View = Cache.LockView())
+			{
+				for (int Idx = 0; Idx < 100; Idx++)
+				{
+					ReadOnlyMemory<byte> Memory = View.Get(TestHashes[Idx]);
+					Assert.IsFalse(Memory.IsEmpty);
+				}
+			}
+		}
+	}
+}
