Horde: Prevent multiple job leases executing at once.

diff --git a/Engine/Source/Programs/Horde/HordeServer/Tasks/Impl/JobTaskSource.cs b/Engine/Source/Programs/Horde/HordeServer/Tasks/Impl/JobTaskSource.cs
--- a/Engine/Source/Programs/Horde/HordeServer/Tasks/Impl/JobTaskSource.cs
+++ b/Engine/Source/Programs/Horde/HordeServer/Tasks/Impl/JobTaskSource.cs
@@ -632,7 +632,7 @@ namespace HordeServer.Tasks.Impl
 		/// <returns></returns>
 		bool TryAssignItemToWaiter(QueueItem Item, QueueWaiter Waiter)
 		{
-			if (Item.AssignTask == null && Item.Batch.SessionId == null && Waiter.Agent.IsInPool(Item.PoolId))
+			if (Item.AssignTask == null && Item.Batch.SessionId == null && Waiter.Agent.Leases.Count == 0 && Waiter.Agent.IsInPool(Item.PoolId))
 			{
 				Task StartTask = new Task<Task>(() => TryCreateLeaseAsync(Item, Waiter));
 				Task ExecuteTask = StartTask.ContinueWith(Task => Task, TaskScheduler.Default);
