Horde: Have health readiness check adhere to app lifetime (SIGTERM)

diff --git a/Engine/Source/Programs/Horde/HordeServer/Controllers/LifetimeController.cs b/Engine/Source/Programs/Horde/HordeServer/Controllers/LifetimeController.cs
--- a/Engine/Source/Programs/Horde/HordeServer/Controllers/LifetimeController.cs
+++ b/Engine/Source/Programs/Horde/HordeServer/Controllers/LifetimeController.cs
@@ -28,24 +28,26 @@ namespace HordeServer.Controllers
 		}
 
 		/// <summary>
-		/// Readiness check for Kubernetes
-		/// If this return a non-successful HTTP response, Kubernetes will remove it from the load balancer,
-		/// preventing it from serving traffic.
+		/// Readiness check for server
+		/// 
+		/// When a SIGTERM has been received, the server is shutting down.
+		/// To communicate the server is stopping to load balancers/orchestrators, this route will return either 503 or 200.
+		/// The server will continue to serve requests, but it's assumed the load balancer has reacted and stopped sending
+		/// traffic by the time the server process exits.
 		/// </summary>
-		/// <returns>Ok if app is not stopping and all databases can be reached</returns>
+		/// <returns>Status code 503 is server is stopping, else 200 OK</returns>
 		[HttpGet]
 		[Route("/health/ready")]
-		public Task<ActionResult> K8sReadinessProbe()
+		public Task<ActionResult> ServerReadiness()
 		{
-			// Disabled for now
-			bool IsRunning = !LifetimeService.IsPreStopping;
-			bool IsMongoDBHealthy = true; //await LifetimeService.IsMongoDbConnectionHealthy();
-			bool IsRedisHealthy = true; //await LifetimeService.IsRedisConnectionHealthy();
-			int StatusCode = IsRunning && IsMongoDBHealthy && IsRedisHealthy ? 200 : 503;
-
-			string Content = $"IsRunning={IsRunning}\n";
-			Content += $"IsMongoDBHealthy={IsMongoDBHealthy}\n";
-			Content += $"IsRedisHealthy={IsRedisHealthy}\n";
+			int StatusCode = 200;
+			string Content = "ok";
+			
+			if (LifetimeService.IsPreStopping)
+			{
+				StatusCode = 503; // Service Unavailable
+				Content = "stopping";
+			}
 
 			return Task.FromResult<ActionResult>(new ContentResult { ContentType = "text/plain", StatusCode = StatusCode, Content = Content });
 		}
diff --git a/Engine/Source/Programs/Horde/HordeServer/Services/LifetimeService.cs b/Engine/Source/Programs/Horde/HordeServer/Services/LifetimeService.cs
--- a/Engine/Source/Programs/Horde/HordeServer/Services/LifetimeService.cs
+++ b/Engine/Source/Programs/Horde/HordeServer/Services/LifetimeService.cs
@@ -27,11 +27,6 @@ namespace HordeServer.Services
 		/// </summary>
 		ILogger<LifetimeService> Logger;
 
-		/// <summary>
-		/// Task which waits for a period of time and sets the stopping task
-		/// </summary>
-		Task WaitAndTriggerStoppingTask = Task.CompletedTask;
-
 		/// <summary>
 		/// Task source for the server stopping
 		/// </summary>
@@ -105,12 +100,24 @@ namespace HordeServer.Services
 		/// </summary>
 		void ApplicationStopping()
 		{
+			Logger.LogInformation("SIGTERM signal received");
+			IsPreStopping = true;
+			IsStopping = true;
+
+			PreStoppingTaskCompletionSource.TrySetResult(true);
+			StoppingTaskCompletionSource.TrySetResult(true);
+
+			int ShutdownDelayMs = 30 * 1000;
+			Logger.LogInformation($"Delaying shutdown by sleeping {ShutdownDelayMs} ms...");
+			Thread.Sleep(ShutdownDelayMs);
+			Logger.LogInformation($"Server process now shutting down...");
+			
+			/*
 			if (PreStoppingTaskCompletionSource.TrySetResult(true))
 			{
-				Logger.LogInformation("Setting pre-stop event");
-				IsPreStopping = true;
+				
 				WaitAndTriggerStoppingTask = Task.Run(() => ExecStoppingTask());
-				/*
+				
 				Logger.LogInformation("App is stopping. Waiting an initial {InitialDelay} secs before waiting on any requests...", (int)InitialStoppingDelay.TotalSeconds);
 				Thread.Sleep(InitialStoppingDelay);
 				Logger.LogInformation("Blocking shutdown for up to {MaxGraceTimeout} secs until all request have finished...", (int)RequestGracefulTimeout.TotalSeconds);
@@ -131,8 +138,9 @@ namespace HordeServer.Services
 				{
 					Logger.LogInformation("One or more requests did not finish within the grace period of {TimeTaken} secs. Shutdown will now resume with risk of interrupting those requests!", (DateTime.UtcNow - StartTime).TotalSeconds);
 					RequestTrackerService.LogRequestsInProgress();
-				}*/
+				}
 			}
+			*/
 		}
 
 		async Task ExecStoppingTask()
