Horde: Remove the last commit information from the stream document; we can store this in Redis instead.

diff --git a/Engine/Source/Programs/Horde/HordeServer/Collections/IStreamCollection.cs b/Engine/Source/Programs/Horde/HordeServer/Collections/IStreamCollection.cs
--- a/Engine/Source/Programs/Horde/HordeServer/Collections/IStreamCollection.cs
+++ b/Engine/Source/Programs/Horde/HordeServer/Collections/IStreamCollection.cs
@@ -84,14 +84,6 @@ namespace HordeServer.Collections
 		/// <returns>Async task object</returns>
 		Task<bool> TryUpdatePropertiesAsync(IStream Stream, string? NewName, int? NewOrder, string? NewNotificationChannel, string? NewNotificationChannelFilter, string? NewTriageChannel, List<StreamTab>? NewTabs, Dictionary<string, AgentType?>? NewAgentTypes, Dictionary<string, WorkspaceType?>? NewWorkspaceTypes, Dictionary<TemplateRefId, TemplateRef>? NewTemplateRefs, Dictionary<string, string?>? NewProperties, Acl? NewAcl, bool? UpdatePauseFields, DateTime? NewPausedUntil, string? NewPauseComment);
 
-		/// <summary>
-		/// Attempts to update the last commit time for a streams
-		/// </summary>
-		/// <param name="Stream"></param>
-		/// <param name="LastCommitTime"></param>
-		/// <returns></returns>
-		Task<bool> TryUpdateCommitTimeAsync(IStream Stream, DateTime LastCommitTime);
-
 		/// <summary>
 		/// Attempts to update the last trigger time for a schedule
 		/// </summary>
diff --git a/Engine/Source/Programs/Horde/HordeServer/Collections/Impl/StreamCollection.cs b/Engine/Source/Programs/Horde/HordeServer/Collections/Impl/StreamCollection.cs
--- a/Engine/Source/Programs/Horde/HordeServer/Collections/Impl/StreamCollection.cs
+++ b/Engine/Source/Programs/Horde/HordeServer/Collections/Impl/StreamCollection.cs
@@ -54,7 +54,6 @@ namespace HordeServer.Collections.Impl
 			public Dictionary<string, AgentType> AgentTypes { get; set; } = new Dictionary<string, AgentType>(StringComparer.Ordinal);
 			public Dictionary<string, WorkspaceType> WorkspaceTypes { get; set; } = new Dictionary<string, WorkspaceType>(StringComparer.Ordinal);
 			public Dictionary<TemplateRefId, TemplateRef> Templates { get; set; } = new Dictionary<TemplateRefId, TemplateRef>();
-			public DateTime? LastCommitTime { get; set; }
 			public DateTime? PausedUntil { get; set; }
 			public string? PauseComment { get; set; }
 			public Acl? Acl { get; set; }
@@ -452,29 +451,6 @@ namespace HordeServer.Collections.Impl
 			}
 		}
 
-		/// <inheritdoc/>
-		public async Task<bool> TryUpdateCommitTimeAsync(IStream StreamInterface, DateTime LastCommitTime)
-		{
-			StreamDocument Stream = (StreamDocument)StreamInterface;
-
-			// Get the update to apply
-			UpdateDefinition<StreamDocument> Update = Builders<StreamDocument>.Update
-				.Set(x => x.LastCommitTime, Stream.LastCommitTime)
-				.Set(x => x.UpdateIndex, Stream.UpdateIndex + 1);
-
-			// Apply the update
-			UpdateResult Result = await Streams.UpdateOneAsync(x => x.Id == Stream.Id && x.UpdateIndex == Stream.UpdateIndex, Update);
-			if (Result.ModifiedCount == 0)
-			{
-				return false;
-			}
-
-			// Update the input document
-			Stream.LastCommitTime = LastCommitTime;
-			Stream.UpdateIndex++;
-			return true;
-		}
-
 		/// <inheritdoc/>
 		public async Task<bool> TryUpdateScheduleTriggerAsync(IStream StreamInterface, TemplateRefId TemplateRefId, DateTimeOffset? LastTriggerTime, int? LastTriggerChange, List<ObjectId> NewActiveJobs)
 		{
diff --git a/Engine/Source/Programs/Horde/HordeServer/Models/Stream.cs b/Engine/Source/Programs/Horde/HordeServer/Models/Stream.cs
--- a/Engine/Source/Programs/Horde/HordeServer/Models/Stream.cs
+++ b/Engine/Source/Programs/Horde/HordeServer/Models/Stream.cs
@@ -557,11 +557,6 @@ namespace HordeServer.Models
 		/// </summary>
 		public IReadOnlyDictionary<TemplateRefId, TemplateRef> Templates { get; }
 
-		/// <summary>
-		/// Last time that we queried for commits
-		/// </summary>
-		public DateTime? LastCommitTime { get; }
-		
 		/// <summary>
 		/// Stream is paused for builds until specified time
 		/// </summary>
@@ -716,7 +711,6 @@ namespace HordeServer.Models
 			HordeCommon.Rpc.GetStreamResponse Response = new HordeCommon.Rpc.GetStreamResponse();
 			Response.Name = Stream.Name;
 			Response.AgentTypes.Add(Stream.AgentTypes.ToDictionary(x => x.Key, x => x.Value.ToRpcResponse()));
-			Response.LastCommitTime = Stream.LastCommitTime.HasValue? Timestamp.FromDateTime(Stream.LastCommitTime.Value) : new Timestamp();
 			return Response;
 		}
 
diff --git a/Engine/Source/Programs/Horde/HordeServer/Services/StreamService.cs b/Engine/Source/Programs/Horde/HordeServer/Services/StreamService.cs
--- a/Engine/Source/Programs/Horde/HordeServer/Services/StreamService.cs
+++ b/Engine/Source/Programs/Horde/HordeServer/Services/StreamService.cs
@@ -265,29 +265,6 @@ namespace HordeServer.Services
 			return Streams.GetPermissionsAsync(StreamId);
 		}
 
-		/// <summary>
-		/// Updates the query time on a stream object
-		/// </summary>
-		/// <param name="Stream">The stream to update</param>
-		/// <param name="LastCommitTime">The last query time</param>
-		/// <returns>New state of the commit stream</returns>
-		public async Task<IStream?> UpdateCommitTimeAsync(IStream? Stream, DateTime LastCommitTime)
-		{
-			// Try to update the commit stream with the new commits
-			while (Stream != null && (!Stream.LastCommitTime.HasValue || LastCommitTime > Stream.LastCommitTime))
-			{
-				// Try to update the document
-				if(await Streams.TryUpdateCommitTimeAsync(Stream, LastCommitTime))
-				{
-					break;
-				}
-
-				// Fetch the new stream document
-				Stream = await GetStreamAsync(Stream.Id);
-			}
-			return Stream;
-		}
-
 		/// <summary>
 		/// Determines if the user is authorized to perform an action on a particular stream
 		/// </summary>
