Horde: Fix conform task loop being ticked constantly.

diff --git a/Engine/Source/Programs/Horde/HordeServer/Tasks/Impl/ConformTaskSource.cs b/Engine/Source/Programs/Horde/HordeServer/Tasks/Impl/ConformTaskSource.cs
--- a/Engine/Source/Programs/Horde/HordeServer/Tasks/Impl/ConformTaskSource.cs
+++ b/Engine/Source/Programs/Horde/HordeServer/Tasks/Impl/ConformTaskSource.cs
@@ -53,7 +53,7 @@ namespace HordeServer.Tasks.Impl
 			this.PerforceLoadBalancer = PerforceLoadBalancer;
 			this.LogService = LogService;
 			this.Logger = Logger;
-			this.TickConformList = new ElectedTick(DatabaseService, new ObjectId("60afc5cf555a9a76aff0a50c"), CleanConformListAsync, TimeSpan.FromMinutes(0.0), Logger);
+			this.TickConformList = new ElectedTick(DatabaseService, new ObjectId("60afc5cf555a9a76aff0a50c"), CleanConformListAsync, TimeSpan.FromMinutes(1.0), Logger);
 		}
 
 		/// <inheritdoc/>
