Horde: Output a log message indicating that DataDog is enabled.

diff --git a/Engine/Source/Programs/Horde/HordeServer/Program.cs b/Engine/Source/Programs/Horde/HordeServer/Program.cs
--- a/Engine/Source/Programs/Horde/HordeServer/Program.cs
+++ b/Engine/Source/Programs/Horde/HordeServer/Program.cs
@@ -140,11 +140,6 @@ namespace HordeServer
 				LogDir = DirectoryReference.Combine(DataDir);
 			}
 
-			if (HordeSettings.WithDatadog)
-			{
-				GlobalTracer.Register(Datadog.Trace.OpenTracing.OpenTracingTracerFactory.CreateTracer());
-			}
-
 			Serilog.Log.Logger = new LoggerConfiguration()
 				.WithHordeConfig(HordeSettings)
 				.Enrich.FromLogContext()
@@ -154,6 +149,12 @@ namespace HordeServer
 				.ReadFrom.Configuration(Config)
 				.CreateLogger();
 
+			if (HordeSettings.WithDatadog)
+			{
+				Serilog.Log.Logger.Information("Enabling datadog tracing");
+				GlobalTracer.Register(Datadog.Trace.OpenTracing.OpenTracingTracerFactory.CreateTracer());
+			}
+
 			if (Arguments.HasOption("-UpdateSchemas"))
 			{
 				DirectoryReference? SchemaDir = Arguments.GetDirectoryReferenceOrDefault("-SchemaDir=", null);
