Add support for custom agent properties using envvars starting with horde.

#rnx

diff --git a/Engine/Source/Programs/Horde/HordeAgent/Services/WorkerService.cs b/Engine/Source/Programs/Horde/HordeAgent/Services/WorkerService.cs
--- a/Engine/Source/Programs/Horde/HordeAgent/Services/WorkerService.cs
+++ b/Engine/Source/Programs/Horde/HordeAgent/Services/WorkerService.cs
@@ -1,13 +1,16 @@
 // Copyright Epic Games, Inc. All Rights Reserved.
 
+using Amazon.EC2;
+using Amazon.EC2.Model;
+using Amazon.Util;
+using EpicGames.Core;
+using Google.Protobuf;
 using Google.Protobuf.WellKnownTypes;
 using Grpc.Core;
 using Grpc.Net.Client;
 using HordeAgent.Execution;
 using HordeAgent.Execution.Interfaces;
 using HordeAgent.Parser;
-using HordeAgent.Parser.Interfaces;
-using HordeAgent.Modes.Service;
 using HordeAgent.Utility;
 using HordeCommon;
 using HordeCommon.Rpc;
@@ -16,43 +19,25 @@ using HordeCommon.Rpc.Tasks;
 using Microsoft.Extensions.Hosting;
 using Microsoft.Extensions.Logging;
 using Microsoft.Extensions.Options;
+using OpenTracing;
+using OpenTracing.Util;
 using System;
-using System.Collections.Concurrent;
 using System.Collections.Generic;
 using System.Diagnostics;
+using System.Diagnostics.CodeAnalysis;
 using System.IO;
 using System.IO.Compression;
 using System.Linq;
 using System.Management;
-using System.Net;
-using System.Net.Http;
 using System.Net.Http.Headers;
-using System.Net.Sockets;
 using System.Reflection;
+using System.Runtime.CompilerServices;
 using System.Runtime.InteropServices;
-using System.Security.Cryptography;
-using System.Security.Cryptography.X509Certificates;
-using System.ServiceProcess;
 using System.Text;
-using System.Text.Json;
 using System.Threading;
 using System.Threading.Tasks;
 using System.Xml;
-using EpicGames.Core;
-using System.ComponentModel;
-using System.Diagnostics.CodeAnalysis;
-using System.Runtime.CompilerServices;
-using Serilog.Context;
-using Polly;
-using Microsoft.Extensions.Http;
-using Polly.Extensions.Http;
-using Amazon.Util;
-using Amazon.EC2;
-using Amazon.EC2.Model;
 using Status = Grpc.Core.Status;
-using Google.Protobuf;
-using OpenTracing.Util;
-using OpenTracing;
 
 [assembly: InternalsVisibleTo("HordeAgentTests")]
 
@@ -1701,6 +1686,13 @@ namespace HordeAgent.Services
 			}
 			PrimaryDevice.Properties.Add($"WorkingDir={WorkingDir}");
 
+			// Add any horde. env vars for custom properties.
+			ICollection<string> EnvVars = Environment.GetEnvironmentVariables().Keys as ICollection<string> ?? new List<string>();
+			foreach (string EnvVar in EnvVars.Where(x => x.StartsWith("horde.", StringComparison.InvariantCultureIgnoreCase)))
+			{
+				PrimaryDevice.Properties.Add($"{EnvVar}={Environment.GetEnvironmentVariable(EnvVar)}");
+			}
+
 			// Create the worker
 			AgentCapabilities Agent = new AgentCapabilities();
 			Agent.Devices.Add(PrimaryDevice);
