Fix how we gather the keys for adding horde. env vars as properties.

#rnx

diff --git a/Engine/Source/Programs/Horde/HordeAgent/Services/WorkerService.cs b/Engine/Source/Programs/Horde/HordeAgent/Services/WorkerService.cs
--- a/Engine/Source/Programs/Horde/HordeAgent/Services/WorkerService.cs
+++ b/Engine/Source/Programs/Horde/HordeAgent/Services/WorkerService.cs
@@ -1687,7 +1687,7 @@ namespace HordeAgent.Services
 			PrimaryDevice.Properties.Add($"WorkingDir={WorkingDir}");
 
 			// Add any horde. env vars for custom properties.
-			ICollection<string> EnvVars = Environment.GetEnvironmentVariables().Keys as ICollection<string> ?? new List<string>();
+			IEnumerable<string> EnvVars = Environment.GetEnvironmentVariables().Keys.Cast<string>();
 			foreach (string EnvVar in EnvVars.Where(x => x.StartsWith("horde.", StringComparison.InvariantCultureIgnoreCase)))
 			{
 				PrimaryDevice.Properties.Add($"{EnvVar}={Environment.GetEnvironmentVariable(EnvVar)}");
