Horde: Add dotTrace profiler support

Controlled via the DebugController for remote start and download of the snapshots.

diff --git a/Engine/Source/Programs/Horde/HordeServer/Controllers/DebugController.cs b/Engine/Source/Programs/Horde/HordeServer/Controllers/DebugController.cs
--- a/Engine/Source/Programs/Horde/HordeServer/Controllers/DebugController.cs
+++ b/Engine/Source/Programs/Horde/HordeServer/Controllers/DebugController.cs
@@ -36,6 +36,7 @@ using Microsoft.Extensions.Options;
 using StatsdClient;
 
 using Amazon.S3;
+using JetBrains.Profiler.SelfApi;
 
 namespace HordeServer.Controllers
 {
@@ -871,5 +872,59 @@ namespace HordeServer.Controllers
 
 			return new ContentResult { ContentType = "text/plain", StatusCode = (int)HttpStatusCode.OK, Content = Message };
 		}
+		
+		/// <summary>
+		/// Starts the profiler session
+		/// </summary>
+		/// <returns>Text message</returns>
+		[HttpGet]
+		[Route("/api/v1/debug/profiler/start")]
+		public async Task<ActionResult> StartProfiler()
+		{
+			await DotTrace.EnsurePrerequisiteAsync();
+
+			string SnapshotDir = Path.Join(Path.GetTempPath(), "horde-profiler-snapshots");
+			if (!Directory.Exists(SnapshotDir))
+			{
+				Directory.CreateDirectory(SnapshotDir);
+			}
+
+			var Config = new DotTrace.Config();
+			Config.SaveToDir(SnapshotDir);
+			DotTrace.Attach(Config);
+			DotTrace.StartCollectingData();
+			
+			return new ContentResult { ContentType = "text/plain", StatusCode = (int)HttpStatusCode.OK, Content = "Profiling session started. Using dir " + SnapshotDir };
+		}
+		
+		/// <summary>
+		/// Stops the profiler session
+		/// </summary>
+		/// <returns>Text message</returns>
+		[HttpGet]
+		[Route("/api/v1/debug/profiler/stop")]
+		public ActionResult StopProfiler()
+		{
+			DotTrace.SaveData();
+			DotTrace.Detach();
+			return new ContentResult { ContentType = "text/plain", StatusCode = (int)HttpStatusCode.OK, Content = "Profiling session stopped" };
+		}
+		
+		/// <summary>
+		/// Downloads the captured profiling snapshots
+		/// </summary>
+		/// <returns>A .zip file containing the profiling snapshots</returns>
+		[HttpGet]
+		[Route("/api/v1/debug/profiler/download")]
+		public ActionResult DownloadProfilingData()
+		{
+			string SnapshotZipFile = DotTrace.GetCollectedSnapshotFilesArchive(false);
+			if (!System.IO.File.Exists(SnapshotZipFile))
+			{
+				return NotFound("The generated snapshot .zip file was not found");
+			}
+			
+			return PhysicalFile(SnapshotZipFile, "application/zip", Path.GetFileName(SnapshotZipFile));
+		}
 	}
 }
diff --git a/Engine/Source/Programs/Horde/HordeServer/HordeServer.csproj b/Engine/Source/Programs/Horde/HordeServer/HordeServer.csproj
--- a/Engine/Source/Programs/Horde/HordeServer/HordeServer.csproj
+++ b/Engine/Source/Programs/Horde/HordeServer/HordeServer.csproj
@@ -25,6 +25,7 @@
       <PrivateAssets>all</PrivateAssets>
       <IncludeAssets>runtime; build; native; contentfiles; analyzers; buildtransitive</IncludeAssets>
     </PackageReference>
+    <PackageReference Include="JetBrains.Profiler.SelfApi" Version="2.1.2" />
     <PackageReference Include="Microsoft.AspNetCore.Authentication.Cookies" Version="2.2.0" />
     <PackageReference Include="Microsoft.AspNetCore.Authentication.JwtBearer" Version="3.1.5" />
     <PackageReference Include="Microsoft.AspNetCore.Authentication.OpenIdConnect" Version="3.1.5" />
