Horde: Add a more detailed error return path for failed compute tasks. Task result messages now include a numerical "Outcome" field and string "Detail" field providing additional context surrounding a task that was not executed.

diff --git a/Engine/Source/Programs/Horde/HordeAgent/Commands/Execution/ExecuteCommandV2.cs b/Engine/Source/Programs/Horde/HordeAgent/Commands/Execution/ExecuteCommandV2.cs
--- a/Engine/Source/Programs/Horde/HordeAgent/Commands/Execution/ExecuteCommandV2.cs
+++ b/Engine/Source/Programs/Horde/HordeAgent/Commands/Execution/ExecuteCommandV2.cs
@@ -272,9 +272,9 @@ namespace HordeAgent.Commands
 						GetTaskUpdateResponse Response = new GetTaskUpdateResponse();
 						Response.AgentId = RpcResponse.AgentId;
 						Response.LeaseId = RpcResponse.LeaseId;
-						Response.ResultHash = (RpcResponse.ResultHash == null)? (CbObjectAttachment?)null : (CbObjectAttachment)RpcResponse.ResultHash;
+						Response.Result = RpcResponse.Result?.Hash;
 						Response.State = RpcResponse.State;
-						Response.TaskHash = RpcResponse.TaskHash;
+						Response.TaskHash = RpcResponse.Task;
 						Response.Time = RpcResponse.Time.ToDateTime();
 
 						if (!await OnUpdateAsync(Response))
@@ -426,9 +426,9 @@ namespace HordeAgent.Commands
 				{
 					Logger.LogInformation("{OperationName}: Running on agent {AgentId} under lease {LeaseId}", (IoHash)(CbObjectAttachment)Response.TaskHash, Response.AgentId, Response.LeaseId);
 				}
-				if (Response.ResultHash != null)
+				if (Response.Result != null)
 				{
-					await HandleCompleteTask(ComputeClient, NamespaceId, Response.ResultHash.Value, Logger);
+					await HandleCompleteTask(ComputeClient, NamespaceId, Response.Result.Value, Logger);
 				}
 				return Response.State != ComputeTaskState.Complete;
 			});
diff --git a/Engine/Source/Programs/Horde/HordeAgent/Execution/ComputeTaskExecutor.cs b/Engine/Source/Programs/Horde/HordeAgent/Execution/ComputeTaskExecutor.cs
--- a/Engine/Source/Programs/Horde/HordeAgent/Execution/ComputeTaskExecutor.cs
+++ b/Engine/Source/Programs/Horde/HordeAgent/Execution/ComputeTaskExecutor.cs
@@ -104,12 +104,8 @@ namespace HordeAgent
 		/// <returns>The action result</returns>
 		public async Task<ComputeTaskResultMessage> ExecuteAsync(string LeaseId, ComputeTaskMessage ComputeTaskMessage, DirectoryReference SandboxDir, CancellationToken CancellationToken)
 		{
-			ComputeTask? Task = await BlobStore.GetObjectAsync<ComputeTask>(ComputeTaskMessage.NamespaceId, (CbObjectAttachment)ComputeTaskMessage.TaskHash);
-			if (Task == null)
-			{
-				throw new RpcException(new Grpc.Core.Status(StatusCode.NotFound, "Unable to find action message"));
-			}
-			Logger.LogInformation("Executing task {Hash} for lease ID {LeaseId}", (IoHash)(CbObjectAttachment)ComputeTaskMessage.TaskHash, LeaseId);
+			ComputeTask Task = await BlobStore.GetObjectAsync<ComputeTask>(ComputeTaskMessage.NamespaceId, ComputeTaskMessage.Task);
+			Logger.LogInformation("Executing task {Hash} for lease ID {LeaseId}", ComputeTaskMessage.Task.Hash, LeaseId);
 
 			DirectoryReference.CreateDirectory(SandboxDir);
 			FileUtils.ForceDeleteDirectoryContents(SandboxDir);
@@ -119,7 +115,6 @@ namespace HordeAgent
 			await SetupSandboxAsync(ComputeTaskMessage.NamespaceId, InputDirectory, SandboxDir);
 			DateTimeOffset InputFetchCompleted = DateTimeOffset.UtcNow;
 
-			ComputeTaskResultMessage ResultMessage;
 			using (ManagedProcessGroup ProcessGroup = new ManagedProcessGroup())
 			{
 				string FileName = FileReference.Combine(SandboxDir, Task.WorkingDirectory.ToString(), Task.Executable.ToString()).FullName;
@@ -169,40 +164,10 @@ namespace HordeAgent
 						Result.OutputHash = await PutOutput(ComputeTaskMessage.NamespaceId, SandboxDir, OutputFiles);
 					}
 
-					ResultMessage = new ComputeTaskResultMessage();
-					ResultMessage.OutputHash = (CbObjectAttachment)await BlobStore.PutObjectAsync(ComputeTaskMessage.NamespaceId, Result);
-
-
-					/*
-										Result.ExecutionMetadata.OutputUploadStartTimestamp = Timestamp.FromDateTimeOffset(DateTimeOffset.UtcNow);
-										foreach (var (FileRef, OutputPath) in ResolveOutputPaths(SandboxDir, Command.OutputPaths))
-										{
-											byte[] Bytes = await FileReference.ReadAllBytesAsync(FileRef);
-											Digest Digest = await Storage.PutBulkDataAsync(InstanceName, Bytes);
-
-											OutputFile OutputFileInfo = new OutputFile();
-											OutputFileInfo.Path = OutputPath;
-											OutputFileInfo.Digest = Digest;
-											Result.OutputFiles.Add(OutputFileInfo);
-
-											Logger.LogInformation("Uploaded {File} (digest: {Digest}, size: {Size})", FileRef, Digest.Hash, Digest.SizeBytes);
-										}
-										Result.ExecutionMetadata.OutputUploadCompletedTimestamp = Timestamp.FromDateTimeOffset(DateTimeOffset.UtcNow);
-										Result.ExecutionMetadata.WorkerCompletedTimestamp = Timestamp.FromDateTimeOffset(DateTimeOffset.UtcNow);
-
-										PostActionResultRequest PostResultRequest = new PostActionResultRequest();
-										PostResultRequest.LeaseId = LeaseId;
-										PostResultRequest.ActionDigest = ActionTask.Digest;
-										PostResultRequest.Result = Result;
-										await ActionRpc.PostActionResultAsync(PostResultRequest);
-
-										await CacheActionResultAsync(Action, ActionTask, Result);
-					*/
-					// Command.OutputPaths
-					// Command.OutputDirectories
+					CbObjectAttachment Attachment = await BlobStore.PutObjectAsync(ComputeTaskMessage.NamespaceId, Result);
+					return new ComputeTaskResultMessage(Attachment);
 				}
 			}
-			return ResultMessage;
 		}
 
 		/// <summary>
diff --git a/Engine/Source/Programs/Horde/HordeAgent/Services/WorkerService.cs b/Engine/Source/Programs/Horde/HordeAgent/Services/WorkerService.cs
--- a/Engine/Source/Programs/Horde/HordeAgent/Services/WorkerService.cs
+++ b/Engine/Source/Programs/Horde/HordeAgent/Services/WorkerService.cs
@@ -38,6 +38,8 @@ using System.Threading;
 using System.Threading.Tasks;
 using System.Xml;
 using Status = Grpc.Core.Status;
+using EpicGames.Horde.Storage;
+using EpicGames.Horde.Compute;
 
 [assembly: InternalsVisibleTo("HordeAgentTests")]
 
@@ -908,41 +910,44 @@ namespace HordeAgent.Services
 		{
 			Logger.LogInformation("Starting compute task (lease {LeaseId})", LeaseId);
 
-			DateTimeOffset ActionTaskStartTime = DateTimeOffset.UtcNow;
-			using (IRpcClientRef Client = await RpcConnection.GetClientRef(new RpcContext(), CancellationToken))
+			ComputeTaskResultMessage Result;
+			try
 			{
-				DirectoryReference LeaseDir = DirectoryReference.Combine(WorkingDir, "Compute", LeaseId);
-				DirectoryReference.CreateDirectory(LeaseDir);
-
-				ComputeTaskExecutor Executor = new ComputeTaskExecutor(Client.Channel, Logger);
-				try
-				{
-					Logger.LogInformation("Running executor");
-					ComputeTaskResultMessage Result = await Executor.ExecuteAsync(LeaseId, ComputeTask, LeaseDir, CancellationToken);
-					return new LeaseResult(Result.ToByteArray());
-				}
-				catch (Exception re)
-				{
-/*					Logger.LogError(re, "Unhandled exception during remote execution");
-					PostActionResultRequest PostResultRequest = new PostActionResultRequest();
-					PostResultRequest.LeaseId = LeaseId;
-					PostResultRequest.ActionDigest = ComputeTask.Digest;
-					PostResultRequest.Error = new RpcException(new Status(StatusCode.Unknown, "Unhandled exception during remote execution: " + re.Message), re.Message).Status;
-					await Executor.ActionRpc.PostActionResultAsync(PostResultRequest);
-*/
-					throw re;
-				}
-				finally
+				DateTimeOffset ActionTaskStartTime = DateTimeOffset.UtcNow;
+				using (IRpcClientRef Client = await RpcConnection.GetClientRef(new RpcContext(), CancellationToken))
 				{
+					DirectoryReference LeaseDir = DirectoryReference.Combine(WorkingDir, "Compute", LeaseId);
+					DirectoryReference.CreateDirectory(LeaseDir);
+
+					ComputeTaskExecutor Executor = new ComputeTaskExecutor(Client.Channel, Logger);
 					try
 					{
-						DirectoryReference.Delete(LeaseDir, true);
+						Result = await Executor.ExecuteAsync(LeaseId, ComputeTask, LeaseDir, CancellationToken);
 					}
-					catch
+					finally
 					{
+						try
+						{
+							DirectoryReference.Delete(LeaseDir, true);
+						}
+						catch
+						{
+						}
 					}
 				}
 			}
+			catch (BlobNotFoundException Ex)
+			{
+				Logger.LogError(Ex, "Blob not found: {Hash}", Ex.Hash);
+				Result = new ComputeTaskResultMessage(ComputeTaskOutcome.BlobNotFound, Ex.Hash.ToString());
+			}
+			catch (Exception Ex)
+			{
+				Logger.LogError(Ex, "Exception while executing compute task");
+				Result = new ComputeTaskResultMessage(ComputeTaskOutcome.Exception, Ex.ToString());
+			}
+
+			return new LeaseResult(Result.ToByteArray());
 		}
 
 		/// <summary>
diff --git a/Engine/Source/Programs/Horde/HordeServer/Compute/Impl/ComputeController.cs b/Engine/Source/Programs/Horde/HordeServer/Compute/Impl/ComputeController.cs
--- a/Engine/Source/Programs/Horde/HordeServer/Compute/Impl/ComputeController.cs
+++ b/Engine/Source/Programs/Horde/HordeServer/Compute/Impl/ComputeController.cs
@@ -86,12 +86,14 @@ namespace HordeServer.Compute.Impl
 			{
 				GetTaskUpdateResponse Update = new GetTaskUpdateResponse();
 
-				Update.TaskHash = Result.TaskHash;
+				Update.TaskHash = Result.Task;
 				Update.Time = Result.Time;
 				Update.State = Result.State;
-				Update.ResultHash = Result.ResultHash;
+				Update.Outcome = Result.Outcome;
+				Update.Result = Result.Result;
 				Update.AgentId = Result.AgentId?.ToString();
 				Update.LeaseId = Result.LeaseId?.ToString();
+				Update.Detail = Result.Detail;
 
 				Response.Updates.Add(Update);
 			}
diff --git a/Engine/Source/Programs/Shared/EpicGames.Horde.Client/Storage/BlobStoreExtensions.cs b/Engine/Source/Programs/Shared/EpicGames.Horde.Client/Storage/BlobStoreExtensions.cs
--- a/Engine/Source/Programs/Shared/EpicGames.Horde.Client/Storage/BlobStoreExtensions.cs
+++ b/Engine/Source/Programs/Shared/EpicGames.Horde.Client/Storage/BlobStoreExtensions.cs
@@ -11,6 +11,27 @@ using System.Threading.Tasks;
 
 namespace EpicGames.Horde.Storage
 {
+	/// <summary>
+	/// Exception thrown when a blob is not found in the object store
+	/// </summary>
+	public sealed class BlobNotFoundException : Exception
+	{
+		/// <summary>
+		/// Hash of the missing blob
+		/// </summary>
+		public IoHash Hash { get; }
+
+		/// <summary>
+		/// Constructor
+		/// </summary>
+		/// <param name="Hash"></param>
+		public BlobNotFoundException(IoHash Hash)
+			: base("Blob {Hash} was not found in the object store")
+		{
+			this.Hash = Hash;
+		}
+	}
+
 	/// <summary>
 	/// Extension methods for <see cref="BlobStore.BlobStoreClient"/>
 	/// </summary>
@@ -38,7 +59,7 @@ namespace EpicGames.Horde.Storage
 				}
 			}
 
-			throw new RpcException(new Status(StatusCode.NotFound, $"Unable to fetch blob {Hash}"));
+			throw new BlobNotFoundException(Hash);
 		}
 
 		/// <summary>
@@ -64,7 +85,7 @@ namespace EpicGames.Horde.Storage
 		/// <param name="NamespaceId">Namespace for the data</param>
 		/// <param name="Hash">Hash of the blob</param>
 		/// <returns>Byte array of blob data</returns>
-		public static async Task<T> GetObjectAsync<T>(this BlobStore.BlobStoreClient BlobStore, string NamespaceId, IoHash Hash) where T : class
+		public static async Task<T> GetObjectAsync<T>(this BlobStore.BlobStoreClient BlobStore, string NamespaceId, CbObjectAttachment Hash) where T : class
 		{
 			byte[] Data = await GetBlobAsync(BlobStore, NamespaceId, Hash);
 			return CbSerializer.Deserialize<T>(new CbField(Data));
@@ -117,10 +138,10 @@ namespace EpicGames.Horde.Storage
 		/// <param name="NamespaceId">Namespace for the data</param>
 		/// <param name="Item">The item to write</param>
 		/// <returns>Byte array of blob data</returns>
-		public static Task<IoHash> PutObjectAsync<T>(this BlobStore.BlobStoreClient BlobStore, string NamespaceId, T Item)
+		public static async Task<CbObjectAttachment> PutObjectAsync<T>(this BlobStore.BlobStoreClient BlobStore, string NamespaceId, T Item)
 		{
 			ReadOnlyMemory<byte> Data = CbSerializer.Serialize(Item).GetView();
-			return PutBlobAsync(BlobStore, NamespaceId, Data);
+			return await PutBlobAsync(BlobStore, NamespaceId, Data);
 		}
 	}
 }
diff --git a/Engine/Source/Programs/Shared/EpicGames.Horde/Common/CbObjectAttachmentWrapper.cs b/Engine/Source/Programs/Shared/EpicGames.Horde/Common/CbObjectAttachmentWrapper.cs
--- a/Engine/Source/Programs/Shared/EpicGames.Horde/Common/CbObjectAttachmentWrapper.cs
+++ b/Engine/Source/Programs/Shared/EpicGames.Horde/Common/CbObjectAttachmentWrapper.cs
@@ -15,22 +15,32 @@ namespace EpicGames.Horde.Common
 	/// </summary>
 	public partial class CbObjectAttachmentWrapper
 	{
+		/// <summary>
+		/// The hash value
+		/// </summary>
+		public IoHash Hash => new IoHash(Data.ToByteArray());
+
+		/// <summary>
+		/// Returns the object attachment 
+		/// </summary>
+		public CbObjectAttachment Attachment => new CbObjectAttachment(Hash);
+
 		/// <summary>
 		/// Constructor
 		/// </summary>
 		/// <param name="Hash"></param>
 		public CbObjectAttachmentWrapper(IoHash Hash)
 		{
-			this.Hash = ByteString.CopyFrom(Hash.Memory.ToArray());
+			this.Data = ByteString.CopyFrom(Hash.Memory.ToArray());
 		}
 
 		/// <summary>
 		/// Convert from an IoHashWrapper to an IoHash
 		/// </summary>
 		/// <param name="Hash"></param>
-		public static implicit operator CbObjectAttachment(CbObjectAttachmentWrapper Attachment)
+		public static implicit operator CbObjectAttachment(CbObjectAttachmentWrapper Wrapper)
 		{
-			return new CbObjectAttachment(new IoHash(Attachment.Hash.ToByteArray()));
+			return Wrapper.Attachment;
 		}
 
 		/// <summary>
diff --git a/Engine/Source/Programs/Shared/EpicGames.Horde/Compute/Api.cs b/Engine/Source/Programs/Shared/EpicGames.Horde/Compute/Api.cs
--- a/Engine/Source/Programs/Shared/EpicGames.Horde/Compute/Api.cs
+++ b/Engine/Source/Programs/Shared/EpicGames.Horde/Compute/Api.cs
@@ -58,9 +58,17 @@ namespace EpicGames.Horde.Compute
 		[CbField("s")]
 		public ComputeTaskState State { get; set; }
 
-		/// <inheritdoc cref="GetTaskUpdatesRpcResponse.ResultHash"/>
+		/// <inheritdoc cref="GetTaskUpdatesRpcResponse.Outcome"/>
+		[CbField("o")]
+		public ComputeTaskOutcome Outcome { get; set; }
+
+		/// <inheritdoc cref="GetTaskUpdatesRpcResponse.Detail"/>
+		[CbField("d")]
+		public string? Detail { get; set; }
+
+		/// <inheritdoc cref="GetTaskUpdatesRpcResponse.Result"/>
 		[CbField("r")]
-		public CbObjectAttachment? ResultHash { get; set; }
+		public CbObjectAttachment? Result { get; set; }
 
 		/// <inheritdoc cref="GetTaskUpdatesRpcResponse.AgentId"/>
 		[CbField("a")]
