Horde: Missing files.

diff --git a/Engine/Source/Programs/Horde/HordeServer/Compute/IComputeService.cs b/Engine/Source/Programs/Horde/HordeServer/Compute/IComputeService.cs
--- a/Engine/Source/Programs/Horde/HordeServer/Compute/IComputeService.cs
+++ b/Engine/Source/Programs/Horde/HordeServer/Compute/IComputeService.cs
@@ -37,7 +37,7 @@ namespace HordeServer.Compute
 		/// <summary>
 		/// The input hash
 		/// </summary>
-		public CbObjectAttachment TaskHash { get; }
+		public CbObjectAttachment Task { get; }
 
 		/// <summary>
 		/// Timestamp for the event
@@ -49,6 +49,11 @@ namespace HordeServer.Compute
 		/// </summary>
 		public ComputeTaskState State { get; }
 
+		/// <summary>
+		/// Indicates the outcome of a task
+		/// </summary>
+		public ComputeTaskOutcome Outcome { get; }
+
 		/// <summary>
 		/// When transitioning to the executing state, includes the agent name
 		/// </summary>
@@ -62,7 +67,12 @@ namespace HordeServer.Compute
 		/// <summary>
 		/// Hash of the result, if complete
 		/// </summary>
-		public CbObjectAttachment? ResultHash { get; }
+		public CbObjectAttachment? Result { get; }
+
+		/// <summary>
+		/// Additional information for the given outcome
+		/// </summary>
+		public string? Detail { get; }
 	}
 
 	/// <summary>
diff --git a/Engine/Source/Programs/Horde/HordeServer/Compute/Impl/ComputeService.cs b/Engine/Source/Programs/Horde/HordeServer/Compute/Impl/ComputeService.cs
--- a/Engine/Source/Programs/Horde/HordeServer/Compute/Impl/ComputeService.cs
+++ b/Engine/Source/Programs/Horde/HordeServer/Compute/Impl/ComputeService.cs
@@ -73,7 +73,7 @@ namespace HordeServer.Compute.Impl
 	{
 		/// <inheritdoc/>
 		[CbField("h")]
-		public CbObjectAttachment TaskHash { get; set; }
+		public CbObjectAttachment Task { get; set; }
 
 		/// <inheritdoc/>
 		[CbField("t")]
@@ -83,6 +83,10 @@ namespace HordeServer.Compute.Impl
 		[CbField("s")]
 		public ComputeTaskState State { get; set; }
 
+		/// <inheritdoc/>
+		[CbField("o")]
+		public ComputeTaskOutcome Outcome { get; set; }
+
 		/// <inheritdoc cref="IComputeTaskStatus.AgentId"/>
 		[CbField("a")]
 		public Utf8String AgentId { get; set; }
@@ -93,7 +97,11 @@ namespace HordeServer.Compute.Impl
 
 		/// <inheritdoc/>
 		[CbField("r")]
-		public CbObjectAttachment? ResultHash { get; set; }
+		public CbObjectAttachment? Result { get; set; }
+
+		/// <inheritdoc/>
+		[CbField("d")]
+		public string? Detail { get; set; }
 
 		/// <inheritdoc/>
 		AgentId? IComputeTaskStatus.AgentId => AgentId.IsEmpty ? (AgentId?)null : new AgentId(AgentId.ToString());
@@ -109,9 +117,9 @@ namespace HordeServer.Compute.Impl
 		{
 		}
 
-		public ComputeTaskStatus(CbObjectAttachment TaskHash, ComputeTaskState State, AgentId? AgentId, ObjectId? LeaseId)
+		public ComputeTaskStatus(CbObjectAttachment Task, ComputeTaskState State, AgentId? AgentId, ObjectId? LeaseId)
 		{
-			this.TaskHash = TaskHash;
+			this.Task = Task;
 			this.Time = DateTime.UtcNow;
 			this.State = State;
 			this.AgentId = (AgentId == null)? Utf8String.Empty : AgentId.Value.ToString();
@@ -175,7 +183,8 @@ namespace HordeServer.Compute.Impl
 					}
 
 					ComputeTaskStatus Status = new ComputeTaskStatus(ComputeTask.TaskHash, ComputeTaskState.Complete, null, null);
-					Logger.LogInformation("Compute task cancelled (task: {TaskHash}, channel: {ChannelId})", ComputeTask.TaskHash, ComputeTask.ChannelId);
+					Status.Outcome = ComputeTaskOutcome.Expired;
+					Logger.LogInformation("Compute task expired (task: {TaskHash}, channel: {ChannelId})", ComputeTask.TaskHash, ComputeTask.ChannelId);
 					await MessageQueue.PostAsync(ComputeTask.ChannelId.ToString(), Status);
 				}
 			}
@@ -225,14 +234,14 @@ namespace HordeServer.Compute.Impl
 					ComputeTaskMessage ComputeTask = new ComputeTaskMessage();
 					ComputeTask.ChannelId = TaskInfo.ChannelId.ToString();
 					ComputeTask.NamespaceId = DefaultNamespaceId.ToString();
-					ComputeTask.RequirementsHash = new CbObjectAttachment(RequirementsHash);
-					ComputeTask.TaskHash = TaskInfo.TaskHash;
+					ComputeTask.Requirements = new CbObjectAttachment(RequirementsHash);
+					ComputeTask.Task = TaskInfo.TaskHash;
 
 					string LeaseName = $"Remote action ({TaskInfo.TaskHash})";
 					byte[] Payload = Any.Pack(ComputeTask).ToByteArray();
 
 					AgentLease Lease = new AgentLease(ObjectId.GenerateNewId(), LeaseName, null, null, null, LeaseState.Pending, null, Requirements.Exclusive, Payload);
-					Logger.LogDebug("Created lease {LeaseId} for channel {ChannelId} task {TaskHash} req {RequirementsHash}", Lease.Id, ComputeTask.ChannelId, (CbObjectAttachment)ComputeTask.TaskHash, (CbObjectAttachment)ComputeTask.RequirementsHash);
+					Logger.LogDebug("Created lease {LeaseId} for channel {ChannelId} task {TaskHash} req {RequirementsHash}", Lease.Id, ComputeTask.ChannelId, ComputeTask.Task.Hash, ComputeTask.Requirements.Hash);
 					return Lease;
 				}
 			}
@@ -243,8 +252,8 @@ namespace HordeServer.Compute.Impl
 		public Task CancelLeaseAsync(IAgent Agent, ObjectId LeaseId, Any Payload)
 		{
 			ComputeTaskMessage Message = Payload.Unpack<ComputeTaskMessage>();
-			ComputeTaskInfo TaskInfo = new ComputeTaskInfo(Message.TaskHash, new ChannelId(Message.ChannelId));
-			return TaskScheduler.EnqueueAsync((CbObjectAttachment)Message.RequirementsHash, TaskInfo);
+			ComputeTaskInfo TaskInfo = new ComputeTaskInfo(Message.Task, new ChannelId(Message.ChannelId));
+			return TaskScheduler.EnqueueAsync(Message.Requirements.Hash, TaskInfo);
 		}
 
 		/// <inheritdoc/>
@@ -263,9 +272,9 @@ namespace HordeServer.Compute.Impl
 		public Task OnLeaseStartedAsync(IAgent Agent, ObjectId LeaseId, Any Payload)
 		{
 			ComputeTaskMessage ComputeTask = Payload.Unpack<ComputeTaskMessage>();
-			Logger.LogInformation("Compute lease started (lease: {LeaseId}, task: {TaskHash}, agent: {AgentId}, channel: {ChannelId})", LeaseId, (CbObjectAttachment)ComputeTask.TaskHash, Agent.Id, ComputeTask.ChannelId);
+			Logger.LogInformation("Compute lease started (lease: {LeaseId}, task: {TaskHash}, agent: {AgentId}, channel: {ChannelId})", LeaseId, ComputeTask.Task.Hash, Agent.Id, ComputeTask.ChannelId);
 
-			ComputeTaskStatus Status = new ComputeTaskStatus(ComputeTask.TaskHash, ComputeTaskState.Executing, Agent.Id, LeaseId);
+			ComputeTaskStatus Status = new ComputeTaskStatus(ComputeTask.Task, ComputeTaskState.Executing, Agent.Id, LeaseId);
 			return MessageQueue.PostAsync(ComputeTask.ChannelId, Status);
 		}
 
@@ -273,15 +282,31 @@ namespace HordeServer.Compute.Impl
 		{
 			ComputeTaskMessage ComputeTask = Payload.Unpack<ComputeTaskMessage>();
 
-			ComputeTaskResultMessage Result = ComputeTaskResultMessage.Parser.ParseFrom(Output.ToArray());
+			ComputeTaskResultMessage Message = ComputeTaskResultMessage.Parser.ParseFrom(Output.ToArray());
 
-			ComputeTaskStatus Status = new ComputeTaskStatus(ComputeTask.TaskHash, ComputeTaskState.Complete, Agent.Id, LeaseId);
-			if (Result.OutputHash != null)
+			ComputeTaskStatus Status = new ComputeTaskStatus(ComputeTask.Task, ComputeTaskState.Complete, Agent.Id, LeaseId);
+			if (Message.Result != null)
+			{
+				Status.Result = Message.Result;
+			}
+			else if (Message.Outcome != ComputeTaskOutcome.Success)
+			{
+				(Status.Outcome, Status.Detail) = (Message.Outcome, Message.Detail);
+			}
+			else if (Outcome == LeaseOutcome.Failed)
+			{
+				Status.Outcome = ComputeTaskOutcome.Failed;
+			}
+			else if (Outcome == LeaseOutcome.Cancelled)
+			{
+				Status.Outcome = ComputeTaskOutcome.Cancelled;
+			}
+			else
 			{
-				Status.ResultHash = Result.OutputHash;
+				Status.Outcome = ComputeTaskOutcome.NoResult;
 			}
 
-			Logger.LogInformation("Compute lease finished (lease: {LeaseId}, task: {TaskHash}, agent: {AgentId}, channel: {ChannelId}, result: {ResultHash})", LeaseId, (CbObjectAttachment)ComputeTask.TaskHash, Agent.Id, ComputeTask.ChannelId, Status.ResultHash?.Hash ?? IoHash.Zero);
+			Logger.LogInformation("Compute lease finished (lease: {LeaseId}, task: {TaskHash}, agent: {AgentId}, channel: {ChannelId}, result: {ResultHash}, outcome: {Outcome})", LeaseId, ComputeTask.Task.Hash, Agent.Id, ComputeTask.ChannelId, Status.Result?.Hash ?? IoHash.Zero, Status.Outcome);
 			return MessageQueue.PostAsync(ComputeTask.ChannelId, Status);
 		}
 
@@ -366,10 +391,10 @@ namespace HordeServer.Compute.Impl
 						foreach (IComputeTaskStatus Update in Updates)
 						{
 							GetTaskUpdatesRpcResponse Response = new GetTaskUpdatesRpcResponse();
-							Response.TaskHash = Update.TaskHash;
+							Response.Task = Update.Task;
 							Response.Time = Timestamp.FromDateTime(Update.Time);
 							Response.State = Update.State;
-							Response.ResultHash = Update.ResultHash;
+							Response.Result = Update.Result;
 							await ResponseStream.WriteAsync(Response);
 						}
 					}
