Horde: Missing file.

diff --git a/Engine/Source/Programs/Horde/HordeCommon/ProtoMethods.cs b/Engine/Source/Programs/Horde/HordeCommon/ProtoMethods.cs
new file mode 100644
--- /dev/null
+++ b/Engine/Source/Programs/Horde/HordeCommon/ProtoMethods.cs
@@ -0,0 +1,22 @@
+using EpicGames.Horde.Compute;
+using EpicGames.Serialization;
+using System;
+using System.Collections.Generic;
+using System.Text;
+
+namespace HordeCommon.Rpc.Tasks
+{
+	partial class ComputeTaskResultMessage
+	{
+		public ComputeTaskResultMessage(CbObjectAttachment Result)
+		{
+			this.Result = Result;
+		}
+
+		public ComputeTaskResultMessage(ComputeTaskOutcome Outcome, string? Detail = null)
+		{
+			this.Outcome = Outcome;
+			this.Detail = Detail;
+		}
+	}
+}
