Horde: Missing copyright notice.

diff --git a/Engine/Source/Programs/Horde/HordeCommon/ProtoMethods.cs b/Engine/Source/Programs/Horde/HordeCommon/ProtoMethods.cs
--- a/Engine/Source/Programs/Horde/HordeCommon/ProtoMethods.cs
+++ b/Engine/Source/Programs/Horde/HordeCommon/ProtoMethods.cs
@@ -1,3 +1,5 @@
+// Copyright Epic Games, Inc. All Rights Reserved.
+
 using EpicGames.Horde.Compute;
 using EpicGames.Serialization;
 using System;
