Horde: First version of Perforce metadata mirroring. Replicates the state of each stream into the CAS as a merkle tree, allowing faster metadata updates on clients.

Not yet enabled by default due to memory issues.

diff --git a/Engine/Source/Programs/Horde/HordeServer/Commits/ICommit.cs b/Engine/Source/Programs/Horde/HordeServer/Commits/ICommit.cs
--- a/Engine/Source/Programs/Horde/HordeServer/Commits/ICommit.cs
+++ b/Engine/Source/Programs/Horde/HordeServer/Commits/ICommit.cs
@@ -63,6 +63,11 @@ namespace HordeServer.Commits
 		/// Date/time that change was committed
 		/// </summary>
 		public DateTime DateUtc { get; }
+
+		/// <summary>
+		/// Name of a reference describing the tree at this changelist
+		/// </summary>
+		public string? TreeRef { get; }
 	}
 
 	/// <summary>
@@ -94,6 +99,18 @@ namespace HordeServer.Commits
 		/// <inheritdoc cref="ICommit.DateUtc"/>
 		public DateTime DateUtc { get; set; }
 
+		/// <inheritdoc cref="ICommit.TreeRef"/>
+		public string? TreeRef { get; set; }
+
+		/// <summary>
+		/// Constructor
+		/// </summary>
+		public NewCommit(ICommit Commit)
+			: this(Commit.StreamId, Commit.Change, Commit.OriginalChange, Commit.AuthorId, Commit.OwnerId, Commit.Description, Commit.BasePath, Commit.DateUtc)
+		{
+			this.TreeRef = Commit.TreeRef;
+		}
+
 		/// <summary>
 		/// Constructor
 		/// </summary>
diff --git a/Engine/Source/Programs/Horde/HordeServer/Commits/ICommitCollection.cs b/Engine/Source/Programs/Horde/HordeServer/Commits/ICommitCollection.cs
--- a/Engine/Source/Programs/Horde/HordeServer/Commits/ICommitCollection.cs
+++ b/Engine/Source/Programs/Horde/HordeServer/Commits/ICommitCollection.cs
@@ -42,4 +42,27 @@ namespace HordeServer.Commits
 		/// <returns></returns>
 		Task<List<ICommit>> FindCommitsAsync(StreamId StreamId, int? MinChange = null, int? MaxChange = null, int? Index = null, int? Count = null);
 	}
+
+	/// <summary>
+	/// Extension methods for <see cref="ICommitCollection"/>
+	/// </summary>
+	static class CommitCollectionExtensions
+	{
+		/// <summary>
+		/// Gets a commit from a stream by changelist number
+		/// </summary>
+		/// <param name="CommitCollection">The commit collection</param>
+		/// <param name="StreamId"></param>
+		/// <param name="Change"></param>
+		/// <returns></returns>
+		public static async Task<ICommit?> GetCommitAsync(this ICommitCollection CommitCollection, StreamId StreamId, int Change)
+		{
+			List<ICommit> Commits = await CommitCollection.FindCommitsAsync(StreamId, Change, Change);
+			if (Commits.Count == 0)
+			{
+				return null;
+			}
+			return Commits[0];
+		}
+	}
 }
diff --git a/Engine/Source/Programs/Horde/HordeServer/Commits/Impl/CommitCollection.cs b/Engine/Source/Programs/Horde/HordeServer/Commits/Impl/CommitCollection.cs
--- a/Engine/Source/Programs/Horde/HordeServer/Commits/Impl/CommitCollection.cs
+++ b/Engine/Source/Programs/Horde/HordeServer/Commits/Impl/CommitCollection.cs
@@ -47,6 +47,9 @@ namespace HordeServer.Commits.Impl
 			[BsonElement("p")]
 			public string BasePath { get; set; } = String.Empty;
 
+			[BsonElement("tr")]
+			public string? TreeRef { get; set; }
+
 			[BsonElement("t")]
 			public DateTime DateUtc { get; set; }
 
@@ -72,6 +75,7 @@ namespace HordeServer.Commits.Impl
 				}
 				this.Description = NewCommit.Description;
 				this.BasePath = NewCommit.BasePath;
+				this.TreeRef = NewCommit.TreeRef;
 				this.DateUtc = NewCommit.DateUtc;
 			}
 		}
diff --git a/Engine/Source/Programs/Horde/HordeServer/Commits/Impl/CommitService.cs b/Engine/Source/Programs/Horde/HordeServer/Commits/Impl/CommitService.cs
--- a/Engine/Source/Programs/Horde/HordeServer/Commits/Impl/CommitService.cs
+++ b/Engine/Source/Programs/Horde/HordeServer/Commits/Impl/CommitService.cs
@@ -22,11 +22,18 @@ using System.Collections.Concurrent;
 using System.Globalization;
 using Microsoft.Extensions.Hosting;
 using EpicGames.Perforce.Managed;
+using Microsoft.Extensions.Caching.Memory;
+using System.Threading.Channels;
+using StackExchange.Redis;
+using EpicGames.Perforce;
+using EpicGames.Redis;
+using System.Diagnostics;
 
 namespace HordeServer.Commits.Impl
 {
 	using P4 = Perforce.P4;
 	using NamespaceId = StringId<INamespace>;
+	using BucketId = StringId<IBucket>;
 	using StreamId = StringId<IStream>;
 
 	/// <summary>
@@ -34,184 +41,299 @@ namespace HordeServer.Commits.Impl
 	/// </summary>
 	class CommitService : ICommitService, IHostedService, IDisposable
 	{
-		class StreamReplicationState
-		{
-			[BsonElement("s")]
-			public StreamId StreamId { get; set; }
-
-			[BsonElement("min")]
-			public int MinChange { get; set; } // Inclusive
+		bool EnableUpdates = false;
 
-			[BsonElement("max")]
-			public int MaxChange { get; set; } // Exclusive
-		}
-
-		[SingletonDocument("60f59cfa4cc1474f19bcc161")]
-		class ReplicationState : SingletonBase
+		/// <summary>
+		/// Metadata about a stream that needs to have commits mirrored
+		/// </summary>
+		class StreamInfo
 		{
-			public List<StreamReplicationState> Streams { get; set; } = new List<StreamReplicationState>();
+			public IStream Stream { get; set; }
+			public ViewMap View { get; }
+
+			public StreamInfo(IStream Stream, ViewMap View)
+			{
+				this.Stream = Stream;
+				this.View = View;
+			}
 		}
 
 		NamespaceId NamespaceId { get; } = new NamespaceId("default");
+		BucketId BucketId { get; } = new BucketId("commits");
 
+		// Redis
+		IDatabase Redis;
+
+		// Collections
 		ICommitCollection CommitCollection;
 		IObjectCollection ObjectCollection;
+		IRefCollection RefCollection;
 		IStreamCollection StreamCollection;
 		IPerforceService PerforceService;
 		IUserCollection UserCollection;
 		ISingletonDocument<Globals> GlobalsDocument;
-		ISingletonDocument<ReplicationState> StateDocument;
 		ILogger<CommitService> Logger;
-		ElectedTick Ticker;
+		ElectedTick UpdateCommitsTicker;
+
+		const int MaxBackgroundTasks = 2;
 
 		/// <summary>
 		/// Constructor
 		/// </summary>
-		public CommitService(ICommitCollection CommitCollection, IObjectCollection ObjectCollection, IStreamCollection StreamCollection, IPerforceService PerforceService, IUserCollection UserCollection, DatabaseService DatabaseService, ILogger<CommitService> Logger)
+		public CommitService(IDatabase Redis, ICommitCollection CommitCollection, IObjectCollection ObjectCollection, IRefCollection RefCollection, IStreamCollection StreamCollection, IPerforceService PerforceService, IUserCollection UserCollection, DatabaseService DatabaseService, ILogger<CommitService> Logger)
 		{
+			this.Redis = Redis;
+			this.RedisDirtyStreams = new RedisSet<StreamId>(Redis, RedisBaseKey.Append("streams"));
+			this.RedisReservations = new RedisSortedSet<StreamId>(Redis, RedisBaseKey.Append("reservations"));
+
 			this.CommitCollection = CommitCollection;
 			this.ObjectCollection = ObjectCollection;
+			this.RefCollection = RefCollection;
 			this.StreamCollection = StreamCollection;
 			this.PerforceService = PerforceService;
 			this.UserCollection = UserCollection;
 			this.GlobalsDocument = new SingletonDocument<Globals>(DatabaseService);
-			this.StateDocument = new SingletonDocument<ReplicationState>(DatabaseService);
 			this.Logger = Logger;
-			this.Ticker = new ElectedTick(DatabaseService, new ObjectId("60f866c49e7268f71803b6e0"), TickReplicationAsync, TimeSpan.FromSeconds(30.0), Logger);
+			this.UpdateCommitsTicker = new ElectedTick(DatabaseService, new ObjectId("60f866c49e7268f71803b6ef"), UpdateCommitsAsync, TimeSpan.FromSeconds(30.0), Logger);
 		}
 
 		/// <inheritdoc/>
 		public void Dispose()
 		{
-			Ticker.Dispose();
+			UpdateCommitsTicker.Dispose();
 		}
 
 		/// <inheritdoc/>
-		public Task StartAsync(CancellationToken cancellationToken)
+		public async Task StartAsync(CancellationToken CancellationToken)
 		{
-			Ticker.Start();
-			return Task.CompletedTask;
+			if (EnableUpdates)
+			{
+				UpdateCommitsTicker.Start();
+				await StartTreeUpdatesAsync();
+			}
 		}
 
 		/// <inheritdoc/>
-		public async Task StopAsync(CancellationToken cancellationToken)
+		public async Task StopAsync(CancellationToken CancellationToken)
 		{
-			await Ticker.StopAsync();
+			if (EnableUpdates)
+			{
+				await StopTreeUpdatesAsync();
+				await UpdateCommitsTicker.StopAsync();
+			}
 		}
 
+		#region Commit updates
+
 		/// <summary>
-		/// Updates commits from Perforce
+		/// Polls Perforce for submitted changes
 		/// </summary>
-		/// <param name="StoppingToken"></param>
+		/// <param name="CancellationToken"></param>
 		/// <returns></returns>
-		async Task TickReplicationAsync(CancellationToken StoppingToken)
+		async Task UpdateCommitsAsync(CancellationToken CancellationToken)
 		{
-			Globals Globals = await GlobalsDocument.GetAsync();
-			ReplicationState State = await StateDocument.GetAsync();
-			bool bUpdateState = false;
-
-			// Find the last 100 changes from any stream that does not currently have any changes registered
 			List<IStream> Streams = await StreamCollection.FindAllAsync();
+			foreach (IGrouping<string, IStream> Group in Streams.GroupBy(x => x.ClusterName, StringComparer.OrdinalIgnoreCase))
+			{
+				try
+				{
+					await UpdateCommitsForClusterAsync(Group.Key, Group);
+				}
+				catch (Exception Ex)
+				{
+					Logger.LogError(Ex, "Error while updating cluster {ClusterName}", Group.Key);
+				}
+			}
+		}
+
+		/// <summary>
+		/// Updates all the streams on a particular server cluster
+		/// </summary>
+		/// <param name="ClusterName"></param>
+		/// <param name="Streams"></param>
+		/// <returns></returns>
+		async Task UpdateCommitsForClusterAsync(string ClusterName, IEnumerable<IStream> Streams)
+		{
+			// Create a connection to the server
+			IPerforceConnection? Connection = await PerforceService.GetServiceUserConnection(ClusterName);
+			if (Connection == null)
+			{
+				throw new PerforceException($"Unable to create cluster connection for {ClusterName}");
+			}
+
+			// Figure out the case settings for the server
+			InfoRecord ServerInfo = await Connection.GetInfoAsync(InfoOptions.ShortOutput);
+
+			// Get the view for each stream
+			List<StreamInfo> StreamInfoList = new List<StreamInfo>();
 			foreach (IStream Stream in Streams)
 			{
-				StreamReplicationState? StreamState = State.Streams.FirstOrDefault(x => x.StreamId == Stream.Id);
-				if (StreamState == null)
+				ViewMap View = await GetStreamViewAsync(Connection, Stream.Name);
+				StreamInfoList.Add(new StreamInfo(Stream, View));
+			}
+
+			// Find all the depot roots
+			Dictionary<string, string> DepotRoots = new Dictionary<string, string>(ServerInfo.PathComparer);
+			foreach (ViewMap View in StreamInfoList.Select(x => x.View))
+			{
+				foreach (ViewMapEntry Entry in View.Entries)
 				{
-					List<ChangeSummary> Changes = await PerforceService.GetChangesAsync(Stream.ClusterName, Stream.Name, null, null, 10, null);
-					if (Changes.Count > 0)
+					if (Entry.SourcePrefix.Length >= 3)
 					{
-						StreamState = new StreamReplicationState();
-						StreamState.StreamId = Stream.Id;
-						StreamState.MinChange = Changes.Min(x => x.Number);
-						StreamState.MaxChange = Changes.Max(x => x.Number + 1);
+						int SlashIdx = Entry.SourcePrefix.IndexOf('/', 2);
+						string Depot = Entry.SourcePrefix.Substring(0, SlashIdx + 1);
 
-						foreach (ChangeSummary Change in Changes)
+						string? DepotRoot;
+						if (DepotRoots.TryGetValue(Depot, out DepotRoot))
 						{
-							await AddCommitAsync(Stream, Change.Number);
+							DepotRoot = GetCommonPrefix(DepotRoot, Entry.SourcePrefix);
 						}
-
-						State.Streams.Add(StreamState);
-						if (!await StateDocument.TryUpdateAsync(State))
+						else
 						{
-							return;
+							DepotRoot = Entry.SourcePrefix;
 						}
+
+						int LastSlashIdx = DepotRoot.LastIndexOf('/');
+						DepotRoots[Depot] = DepotRoot.Substring(0, LastSlashIdx + 1);
 					}
 				}
 			}
 
-			// Update all the clusters
-			foreach (IGrouping<string, IStream> Cluster in Streams.GroupBy(x => x.ClusterName, StringComparer.OrdinalIgnoreCase))
+			// Find the minimum changelist number to query
+			int MinChange = int.MaxValue;
+			foreach (StreamInfo StreamInfo in StreamInfoList)
+			{
+				RedisList<int> Changes = RedisStreamChanges(StreamInfo.Stream.Id);
+
+				int FirstChange = await Changes.GetByIndexAsync(-1);
+				if (FirstChange == 0)
+				{
+					FirstChange = await GetFirstCommitToReplicateAsync(Connection, StreamInfo.View, ServerInfo.PathComparison);
+				}
+				else
+				{
+					FirstChange++;
+				}
+
+				if (FirstChange != 0)
+				{
+					MinChange = Math.Min(MinChange, FirstChange);
+				}
+			}
+
+			// Find all the changes to consider
+			SortedSet<int> ChangeNumbers = new SortedSet<int>();
+			foreach (string DepotRoot in DepotRoots.Values)
+			{
+				List<ChangesRecord> Changes = await Connection.GetChangesAsync(ChangesOptions.None, null, MinChange, -1, ChangeStatus.Submitted, null, $"{DepotRoot}...");
+				ChangeNumbers.UnionWith(Changes.Select(x => x.Number));
+			}
+
+			// Add the changes in order
+			List<string> RelativePaths = new List<string>();
+			foreach (int ChangeNumber in ChangeNumbers)
 			{
-				// Get the first change to query from
-				int? MinChange = null;
-				foreach (IStream Stream in Cluster)
+				DescribeRecord DescribeRecord = await Connection.DescribeAsync(ChangeNumber);
+				foreach (StreamInfo StreamInfo in StreamInfoList)
 				{
-					StreamReplicationState? StreamState = State.Streams.FirstOrDefault(x => x.StreamId == Stream.Id);
-					if (StreamState != null)
+					IStream Stream = StreamInfo.Stream;
+
+					string? BasePath = GetBasePath(DescribeRecord, StreamInfo.View, ServerInfo.PathComparison);
+					if (BasePath != null)
 					{
-						if (MinChange == null || StreamState.MaxChange < MinChange.Value)
-						{
-							MinChange = StreamState.MaxChange + 1;
-						}
+						await AddCommitAsync(Stream, DescribeRecord, BasePath);
+
+						RedisList<int> StreamCommitsKey = RedisStreamChanges(Stream.Id);
+						await StreamCommitsKey.RightPushAsync(DescribeRecord.Number);
+
+						await RedisDirtyStreams.AddAsync(Stream.Id);
+						await Redis.PublishAsync(RedisUpdateChannel, Stream.Id);
 					}
 				}
+			}
+		}
 
-				// Query all the changes
-				List<ChangeSummary> Changes = await PerforceService.GetChangesAsync(Cluster.Key, MinChange, 100);
-				if(Changes.Count > 0)
+		/// <summary>
+		/// Find the first commit to replicate from a branch
+		/// </summary>
+		/// <param name="Connection"></param>
+		/// <param name="View"></param>
+		/// <param name="Comparison"></param>
+		/// <returns></returns>
+		static async Task<int> GetFirstCommitToReplicateAsync(IPerforceConnection Connection, ViewMap View, StringComparison Comparison)
+		{
+			int MinChange = 0;
+
+			List<string> RootPaths = View.GetRootPaths(Comparison);
+			foreach (string RootPath in RootPaths)
+			{
+				IList<ChangesRecord> PathChanges = await Connection.GetChangesAsync(ChangesOptions.None, 20, ChangeStatus.Submitted, $"{RootPath}...");
+				MinChange = Math.Max(MinChange, PathChanges.Min(x => x.Number));
+			}
+
+			return MinChange;
+		}
+
+		/// <summary>
+		/// Find the base path for a change within a stream
+		/// </summary>
+		/// <param name="Changelist">Information about the changelist</param>
+		/// <param name="View">Mapping from depot syntax to client view</param>
+		/// <param name="Comparison">Path comparison type for the server</param>
+		/// <returns>The base path for all files in the change</returns>
+		static string? GetBasePath(DescribeRecord Changelist, ViewMap View, StringComparison Comparison)
+		{
+			string? BasePath = null;
+			foreach (DescribeFileRecord File in Changelist.Files)
+			{
+				string? StreamFile;
+				if (View.TryMapFile(File.DepotFile, Comparison, out StreamFile))
 				{
-					int MaxChange = Changes.Max(x => x.Number) + 1;
-					foreach (IStream Stream in Cluster)
+					if (BasePath == null)
 					{
-						StreamReplicationState? StreamState = State.Streams.FirstOrDefault(x => x.StreamId == Stream.Id);
-						if (StreamState != null && MaxChange > StreamState.MaxChange)
-						{
-							using (IStreamView StreamView = await PerforceService.GetStreamViewAsync(Cluster.Key, Stream.Name))
-							{
-								foreach (ChangeSummary Change in Changes)
-								{
-									if (StreamView.TryGetStreamPath(Change.Path, out _)) // TODO: Need *possible* match with wildcard, not exact match
-									{
-										await AddCommitAsync(Stream, Change.Number);
-									}
-								}
-							}
-							StreamState.MaxChange = MaxChange;
-							bUpdateState = true;
-						}
+						BasePath = StreamFile;
+					}
+					else
+					{
+						BasePath = GetCommonPrefix(BasePath, StreamFile);
 					}
 				}
 			}
+			return BasePath;
+		}
 
-			if(bUpdateState)
+		/// <summary>
+		/// Gets the common prefix between two stringsc
+		/// </summary>
+		/// <param name="A"></param>
+		/// <param name="B"></param>
+		/// <returns></returns>
+		static string GetCommonPrefix(string A, string B)
+		{
+			int Index = 0;
+			while (Index < A.Length && Index < B.Length && A[Index] == B[Index])
 			{
-				await StateDocument.TryUpdateAsync(State);
+				Index++;
 			}
+			return A.Substring(0, Index);
 		}
 
 		/// <summary>
 		/// Adds metadata for a particular commit
 		/// </summary>
 		/// <param name="Stream"></param>
-		/// <param name="Number"></param>
+		/// <param name="Changelist"></param>
+		/// <param name="BasePath"></param>
 		/// <returns></returns>
-		async Task AddCommitAsync(IStream Stream, int Number)
+		async Task<ICommit> AddCommitAsync(IStream Stream, DescribeRecord Changelist, string BasePath)
 		{
-			ChangeDetails Details = await PerforceService.GetChangeDetailsAsync(Stream.ClusterName, Stream.Name, Number);
+			IUser Author = await UserCollection.FindOrAddUserByLoginAsync(Changelist.User);
+			IUser Owner = (await ParseRobomergeOwnerAsync(Changelist.Description)) ?? Author;
 
-			IUser Author = await UserCollection.FindOrAddUserByLoginAsync(Details.Author);
-			int? OriginalChange = ParseRobomergeSource(Details.Description);
-
-			IUser Owner = Author;
-
-			string? OriginalAuthor = ParseRobomergeOwner(Details.Description);
-			if (OriginalAuthor != null)
-			{
-				Owner = await UserCollection.FindOrAddUserByLoginAsync(OriginalAuthor);
-			}
+			int OriginalChange = ParseRobomergeSource(Changelist.Description) ?? Changelist.Number;
 
-			NewCommit Commit = new NewCommit(Stream.Id, Details.Number, OriginalChange ?? Details.Number, Author.Id, Owner.Id, Details.Description, Details.Path, Details.Date);
-			await CommitCollection.AddOrReplaceAsync(Commit);
+			NewCommit NewCommit = new NewCommit(Stream.Id, Changelist.Number, OriginalChange, Author.Id, Owner.Id, Changelist.Description, BasePath, Changelist.Time);
+			return await CommitCollection.AddOrReplaceAsync(NewCommit);
 		}
 
 		/// <inheritdoc/>
@@ -244,6 +366,16 @@ namespace HordeServer.Commits.Impl
 			}
 		}
 
+		async Task<IUser?> ParseRobomergeOwnerAsync(string Description)
+		{
+			string? OriginalAuthor = ParseRobomergeOwner(Description);
+			if (OriginalAuthor != null)
+			{
+				return await UserCollection.FindOrAddUserByLoginAsync(OriginalAuthor);
+			}
+			return null;
+		}
+
 		/// <summary>
 		/// Attempts to parse the Robomerge owner from this commit information
 		/// </summary>
@@ -262,24 +394,302 @@ namespace HordeServer.Commits.Impl
 				return null;
 			}
 		}
+		#endregion
+
+		#region Tree Updates
+
+		RedisKey RedisBaseKey { get; } = new RedisKey("commits/");
+		RedisChannel<StreamId> RedisUpdateChannel { get; } = new RedisChannel<StreamId>("commits/streams");
+		readonly RedisSet<StreamId> RedisDirtyStreams;
+		readonly RedisSortedSet<StreamId> RedisReservations;
+		Channel<StreamId>? RedisUpdateStreams;
+		RedisChannelSubscription<StreamId>? RedisUpdateSubscription;
+		Task? StreamUpdateTask;
+		RedisList<int> RedisStreamChanges(StreamId StreamId) => new RedisList<int>(Redis, RedisBaseKey.Append($"stream/{StreamId}/changes"));
+
+		async Task StartTreeUpdatesAsync()
+		{
+			RedisUpdateStreams = Channel.CreateUnbounded<StreamId>();
+			RedisUpdateSubscription = await Redis.Multiplexer.GetSubscriber().SubscribeAsync(RedisUpdateChannel, (_, StreamId) => RedisUpdateStreams.Writer.TryWrite(StreamId));
+			StreamUpdateTask = Task.Run(() => UpdateTreesAsync());
+		}
+
+		async Task StopTreeUpdatesAsync()
+		{
+			if (RedisUpdateStreams != null)
+			{
+				RedisUpdateStreams.Writer.Complete();
+				RedisUpdateStreams = null;
+			}
+			if (RedisUpdateSubscription != null)
+			{
+				await RedisUpdateSubscription.DisposeAsync();
+				RedisUpdateSubscription = null;
+			}
+			await StreamUpdateTask!;
+		}
+
+		async Task UpdateTreesAsync()
+		{
+			List<(StreamId, Task)> BackgroundTasks = new List<(StreamId, Task)>(MaxBackgroundTasks);
+			while (!RedisUpdateStreams!.Reader.Completion.IsCompleted || BackgroundTasks.Count > 0)
+			{
+				try
+				{
+					await UpdateStreamsInternalAsync(BackgroundTasks);
+				}
+				catch (Exception Ex)
+				{
+					Logger.LogError(Ex, "Exception when updating commit service");
+				}
+			}
+		}
+
+		async Task UpdateStreamsInternalAsync(List<(StreamId, Task)> BackgroundTasks)
+		{
+			// Remove any complete background tasks
+			for (int Idx = BackgroundTasks.Count - 1; Idx >= 0; Idx--)
+			{
+				(_, Task BackgroundTask) = BackgroundTasks[Idx];
+				if (BackgroundTask.IsCompleted)
+				{
+					await BackgroundTask;
+					BackgroundTasks.RemoveAt(Idx);
+				}
+			}
+
+			// Create a list of events to wait for
+			List<Task> WaitTasks = new List<Task>(BackgroundTasks.Select(x => x.Item2));
+			if (!RedisUpdateStreams!.Reader.Completion.IsCompleted)
+			{
+				WaitTasks.Add(RedisUpdateStreams.Reader.WaitToReadAsync().AsTask());
+
+				// If we have spare slots for executing background tasks, check if there are any dirty streams
+				if (BackgroundTasks.Count < MaxBackgroundTasks)
+				{
+					// Expire any reservations that are no longer valid
+					DateTime UtcNow = DateTime.UtcNow;
+					await RedisReservations.RemoveRangeByScoreAsync(Double.NegativeInfinity, UtcNow.Ticks);
+
+					// Find the streams that we can wait for
+					HashSet<StreamId> CheckStreams = new HashSet<StreamId>(await RedisDirtyStreams.MembersAsync());
+					CheckStreams.ExceptWith(BackgroundTasks.Select(x => x.Item1));
+
+					if (CheckStreams.Count > 0)
+					{
+						// Try to start new background tasks
+						double NewScore = (UtcNow + TimeSpan.FromSeconds(30.0)).Ticks;
+						foreach (StreamId CheckStream in CheckStreams)
+						{
+							if (await RedisReservations.AddAsync(CheckStream, NewScore, When.NotExists))
+							{
+								Task NewTask = Task.Run(() => UpdateStreamAsync(CheckStream));
+								BackgroundTasks.Add((CheckStream, NewTask));
+
+								if (BackgroundTasks.Count == MaxBackgroundTasks)
+								{
+									break;
+								}
+							}
+						}
+						CheckStreams.ExceptWith(BackgroundTasks.Select(x => x.Item1));
+
+						// If we still have spare tasks, check how long we should wait for a reservation to expire
+						if (BackgroundTasks.Count < MaxBackgroundTasks)
+						{
+							long WaitTime = long.MaxValue;
+							await foreach (SortedSetEntry<StreamId> Entry in RedisReservations.ScanAsync())
+							{
+								if (CheckStreams.Contains(Entry.Element))
+								{
+									WaitTime = Math.Min(WaitTime, (long)Entry.Score - UtcNow.Ticks);
+								}
+							}
+							if (WaitTime < int.MaxValue)
+							{
+								WaitTasks.Add(Task.Delay(new TimeSpan(Math.Max(WaitTime, 0))));
+							}
+						}
+					}
+				}
+			}
+
+			// Wait until any task has completed
+			if (WaitTasks.Count > 0)
+			{
+				await Task.WhenAny(WaitTasks);
+			}
+		}
+
+		async Task UpdateStreamAsync(StreamId StreamId)
+		{
+			Stopwatch Timer = Stopwatch.StartNew();
+
+			RedisList<int> StreamChanges = RedisStreamChanges(StreamId);
+			for (; ; )
+			{
+				// Update the stream, updating the reservation every 30 seconds
+				Task InternalTask = Task.Run(() => UpdateStreamInternalAsync(StreamId, StreamChanges));
+				while (!InternalTask.IsCompleted)
+				{
+					Task DelayTask = Task.Delay(TimeSpan.FromSeconds(15.0));
+					if (await Task.WhenAny(InternalTask, DelayTask) == DelayTask)
+					{
+						DateTime NewTime = DateTime.UtcNow + TimeSpan.FromSeconds(30.0);
+						await RedisReservations.AddAsync(StreamId, NewTime.Ticks);
+						Logger.LogInformation("Extending reservation for tree update of {StreamId} (elapsed: {Time}s)", StreamId, (int)Timer.Elapsed.TotalSeconds);
+					}
+				}
+				await InternalTask;
+
+				// Remove this stream from the dirty list if it's empty
+				ITransaction Transaction = Redis.CreateTransaction();
+				Transaction.AddCondition(Condition.ListLengthLessThan(StreamChanges.Key, 2));
+				_ = Transaction.With(RedisDirtyStreams).RemoveAsync(StreamId);
+				if (await Transaction.ExecuteAsync())
+				{
+					break;
+				}
+			}
+			await RedisReservations.RemoveAsync(StreamId);
+		}
+
+		async Task UpdateStreamInternalAsync(StreamId StreamId, RedisList<int> Changes)
+		{
+			IStream? Stream = await StreamCollection.GetAsync(StreamId);
+			if (Stream == null)
+			{
+				return;
+			}
+
+			NativePerforceConnection? Perforce = await PerforceService.GetServiceUserConnection(Stream.ClusterName);
+			if (Perforce == null)
+			{
+				return;
+			}
+
+			InfoRecord ServerInfo = await Perforce.GetInfoAsync(InfoOptions.ShortOutput);
+
+			ViewMap View = await GetStreamViewAsync(Perforce, "//UE5/Dev-Main-Minimal");// Stream.Name);
+
+			ICommit? Commit = null;
+			StreamTreeRef? TreeRef = null;
+			for (; ; )
+			{
+				// Get the first two changelists to update
+				int[] Values = await Changes.RangeAsync(0, 1);
+				if (Values.Length < 2)
+				{
+					break;
+				}
+
+				// Get the previous tree if we don't have it from the previous iteration
+				if (Commit == null || Commit.Change != Values[0])
+				{
+					Commit = await CommitCollection.GetCommitAsync(Stream.Id, Values[0]);
+					if (Commit == null)
+					{
+						TreeRef = null;
+					}
+					else if (Commit.TreeRef == null)
+					{
+						TreeRef = await AddTreeRefAsync(Perforce, Stream, View, null, null, Commit, ServerInfo.PathComparison);
+					}
+					else
+					{
+						TreeRef = await GetTreeRefAsync(Commit);
+					}
+				}
+
+				// Copy the previous commit and tree ref
+				ICommit? PrevCommit = Commit;
+				StreamTreeRef? PrevTreeRef = TreeRef;
+
+				// Move to the next commit
+				Commit = await CommitCollection.GetCommitAsync(StreamId, Values[1]);
+				if (Commit == null)
+				{
+					TreeRef = null;
+				}
+				else
+				{
+					TreeRef = await AddTreeRefAsync(Perforce, Stream, View, PrevCommit, PrevTreeRef, Commit, ServerInfo.PathComparison);
+				}
+
+				// Remove the first change number from this queue
+				await Changes.LeftPopAsync();
+			}
+		}
+
+		async Task<StreamTreeRef?> GetTreeRefAsync(ICommit Commit)
+		{
+			if (Commit.TreeRef == null)
+			{
+				return null;
+			}
+
+			IRef? Ref = await RefCollection.GetAsync(NamespaceId, BucketId, Commit.TreeRef);
+			if (Ref == null)
+			{
+				return null;
+			}
+
+			return CbSerializer.Deserialize<StreamTreeRef>(Ref.Value.AsField());
+		}
+
+		async Task<StreamTreeRef?> AddTreeRefAsync(IPerforceConnection Perforce, IStream Stream, ViewMap View, ICommit? PrevCommit, StreamTreeRef? PrevTreeRef, ICommit Commit, StringComparison PathComparison)
+		{
+			StreamTreeRef TreeRef;
+			if (PrevCommit == null || PrevTreeRef == null)
+			{
+				Logger.LogInformation("Performing snapshot of {StreamId} at CL {Change}", Commit.StreamId, Commit.Change);
+				TreeRef = await CreateTreeAsync(Perforce, View, Commit.Change, PathComparison);
+			}
+			else
+			{
+				Logger.LogInformation("Performing incremental update of {StreamId} from CL {PrevChange} to CL {Change}", Commit.StreamId, PrevCommit.Change, Commit.Change);
+				TreeRef = await UpdateTreeAsync(PrevTreeRef, Perforce, View, Commit.Change, PathComparison);
+			}
+
+			string RefName = $"tree_{Commit.StreamId}_{Commit.Change}";
+
+			CbWriter Writer = new CbWriter();
+			Writer.BeginObject();
+			TreeRef.Write(Writer, Stream.Name);
+			Writer.EndObject();
+
+			List<IoHash> MissingHashes = await RefCollection.SetAsync(NamespaceId, BucketId, RefName, Writer.ToObject());
+			if (MissingHashes.Count > 0)
+			{
+				Logger.LogWarning("Missing hashes when attempting to add ref: {MissingHashes}", String.Join(", ", MissingHashes.Select(x => x.ToString())));
+				return null;
+			}
+
+			NewCommit NewCommit = new NewCommit(Commit);
+			NewCommit.TreeRef = RefName;
+			await CommitCollection.AddOrReplaceAsync(NewCommit);
+
+			return TreeRef;
+		}
+
+		#endregion
 
 		#region Snapshot generation
 
 		/// <summary>
-		/// Number of files to query at a time before recursing down the tree
+		/// Number of files to query at a time before recursing down the tree. Tuning this value allows finding a threshold between request/response lengths and memory usage.
 		/// </summary>
-		const int BatchSize = 16384;
+		const int BatchSize = 256 * 1024;
 
 		/// <summary>
 		/// Creates a snapshot of a stream at a particular changelist
 		/// </summary>
-		/// <param name="Repository">The Perforce repository to mirror</param>
+		/// <param name="Connection">The Perforce repository to mirror</param>
 		/// <param name="View">View for the stream</param>
 		/// <param name="Change">Changelist to replicate</param>
-		async Task<StreamTreeRef> CreateTreeAsync(P4.Repository Repository, ViewMap View, int Change)
+		/// <param name="Comparison">Comparison type for paths</param>
+		async Task<StreamTreeRef> CreateTreeAsync(IPerforceConnection Connection, ViewMap View, int Change, StringComparison Comparison)
 		{
-			StringComparison Comparison = Repository.Server.Metadata.CaseSensitive ? StringComparison.Ordinal : StringComparison.OrdinalIgnoreCase;
-
 			// In order to optimize for the common case where we're adding multiple files to the same directory, we keep an open list of
 			// StreamTreeBuilder objects down to the last modified node. Whenever we need to move to a different node, we write the 
 			// previous path to the object store.
@@ -298,7 +708,7 @@ namespace HordeServer.Commits.Impl
 						FilteredView.Entries.Add(Entry);
 					}
 				}
-				await AddDepotDirToSnapshotAsync(Repository, RootPath, Change, FilteredView, Comparison, TreePath);
+				await AddDepotDirToSnapshotAsync(Connection, RootPath, Change, FilteredView, Comparison, TreePath);
 			}
 
 			// Encode the tree
@@ -308,35 +718,40 @@ namespace HordeServer.Commits.Impl
 		/// <summary>
 		/// Adds all files under a directory in the depot to a snapshot
 		/// </summary>
-		/// <param name="Repository">The perforce repository</param>
+		/// <param name="Connection">The perforce connection</param>
 		/// <param name="DepotDir">Depot directory to recurse through, with a trailing slash</param>
 		/// <param name="Change">Changelist number to capture</param>
 		/// <param name="View">View for the stream</param>
 		/// <param name="Comparison">Comparison type to use for names</param>
 		/// <param name="TreePath">The last computed path through the tree</param>
-		async Task AddDepotDirToSnapshotAsync(P4.Repository Repository, string DepotDir, int Change, ViewMap View, StringComparison Comparison, List<StreamTreeBuilder> TreePath)
+		async Task AddDepotDirToSnapshotAsync(IPerforceConnection Connection, string DepotDir, int Change, ViewMap View, StringComparison Comparison, List<StreamTreeBuilder> TreePath)
 		{
+			const string Filter = "^headAction=delete ^headAction=move/delete";
+			const string Fields = "depotFile,fileSize,digest,headType,headRev";
+
 			if (View.MayMatchAnyFilesInSubDirectory(DepotDir, Comparison))
 			{
-				P4.GetFileMetaDataCmdOptions Options = new P4.GetFileMetaDataCmdOptions(P4.GetFileMetadataCmdFlags.FileSize, null, null, BatchSize, null, null, null);
-				IList<P4.FileMetaData> Files = Repository.GetFileMetaData(Options, P4.FileSpec.DepotSpec($"{DepotDir}...", new P4.ChangelistIdVersion(Change)));
-				if (Files != null && Files.Count < BatchSize)
+				List<FStatRecord> Files = await Connection.FStatAsync(-1, -1, Filter, Fields, BatchSize, FStatOptions.IncludeFileSizes, $"{DepotDir}...@{Change}");
+				if (Files.Count < BatchSize)
 				{
+					Logger.LogInformation("Added {DepotDir}... to snapshot", DepotDir);
 					await AddDepotFilesToSnapshotAsync(Files, View, Comparison, TreePath);
 					return;
 				}
 
-				IList<string> Dirs = Repository.GetDepotDirs(new P4.GetDepotDirsCmdOptions(P4.GetDepotDirsCmdFlags.None, null), $"{DepotDir}*@{Change}");
-				foreach (string Dir in Dirs)
+				Files.Clear(); // Don't keep around in memory
+				Logger.LogInformation("Querying subdirectories of {DepotDir} (wasted results)", DepotDir);
+
+				IList<DirsRecord> Dirs = await Connection.GetDirsAsync(DirsOptions.None, null, $"{DepotDir}*@{Change}");
+				foreach (DirsRecord Dir in Dirs)
 				{
-					await AddDepotDirToSnapshotAsync(Repository, Dir + "/", Change, View, Comparison, TreePath);
+					await AddDepotDirToSnapshotAsync(Connection, Dir.Dir + "/", Change, View, Comparison, TreePath);
 				}
 			}
 
 			if (View.MayMatchAnyFilesInDirectory(DepotDir, Comparison))
 			{
-				P4.GetFileMetaDataCmdOptions Options = new P4.GetFileMetaDataCmdOptions(P4.GetFileMetadataCmdFlags.FileSize, null, null, -1, null, null, null);
-				IList<P4.FileMetaData> Files = Repository.GetFileMetaData(Options, P4.FileSpec.DepotSpec($"{DepotDir}*", new P4.ChangelistIdVersion(Change)));
+				List<FStatRecord> Files = await Connection.FStatAsync(-1, -1, Filter, Fields, -1, FStatOptions.IncludeFileSizes, $"{DepotDir}*@{Change}");
 				await AddDepotFilesToSnapshotAsync(Files, View, Comparison, TreePath);
 			}
 		}
@@ -348,19 +763,16 @@ namespace HordeServer.Commits.Impl
 		/// <param name="View">View for the stream</param>
 		/// <param name="Comparison">Comparison type to use for names</param>
 		/// <param name="TreePath">The last computed path through the tree</param>
-		async Task AddDepotFilesToSnapshotAsync(IList<P4.FileMetaData>? Files, ViewMap View, StringComparison Comparison, List<StreamTreeBuilder> TreePath)
+		async Task AddDepotFilesToSnapshotAsync(List<FStatRecord> Files, ViewMap View, StringComparison Comparison, List<StreamTreeBuilder> TreePath)
 		{
-			if (Files != null)
+			foreach (FStatRecord File in Files)
 			{
-				foreach (P4.FileMetaData File in Files)
+				if (!String.IsNullOrEmpty(File.Digest))
 				{
-					if (File.Digest != null)
+					string? TargetFile;
+					if (File.DepotFile != null && View.TryMapFile(File.DepotFile, Comparison, out TargetFile))
 					{
-						string? TargetFile;
-						if (View.TryMapFile(File.DepotPath.Path, Comparison, out TargetFile))
-						{
-							await AddFileAsync(TargetFile, File, TreePath);
-						}
+						await AddFileAsync(TargetFile, File, TreePath);
 					}
 				}
 			}
@@ -372,7 +784,7 @@ namespace HordeServer.Commits.Impl
 		/// <param name="Path">The stream path for the file</param>
 		/// <param name="MetaData">The file metadata</param>
 		/// <param name="TreePath">The last computed path through the tree</param>
-		async Task AddFileAsync(string Path, P4.FileMetaData MetaData, List<StreamTreeBuilder> TreePath)
+		async Task AddFileAsync(string Path, FStatRecord MetaData, List<StreamTreeBuilder> TreePath)
 		{
 			string[] Segments = Path.Split('/');
 
@@ -404,35 +816,35 @@ namespace HordeServer.Commits.Impl
 		/// <summary>
 		/// Creates a snapshot for the given stream, given an initial state and list of modified files
 		/// </summary>
-		/// <param name="TreeRef">The tree to update</param>
-		/// <param name="Repository"></param>
-		/// <param name="View"></param>
+		/// <param name="PrevTreeRef"></param>
+		/// <param name="Perforce"></param>
 		/// <param name="Change"></param>
+		/// <param name="View"></param>
+		/// <param name="Comparison"></param>
 		/// <returns></returns>
-		Task<StreamTreeRef> UpdateTreeAsync(StreamTreeRef TreeRef, P4.Repository Repository, ViewMap View, int Change)
+		async Task<StreamTreeRef> UpdateTreeAsync(StreamTreeRef PrevTreeRef, IPerforceConnection Perforce, ViewMap View, int Change, StringComparison Comparison)
 		{
-			P4.Changelist Changelist = Repository.GetChangelist(Change);
-			StringComparison Comparison = Repository.Server.Metadata.CaseSensitive ? StringComparison.Ordinal : StringComparison.OrdinalIgnoreCase;
-			return UpdateTreeAsync(Changelist.Files, View, Comparison, TreeRef);
+			List<FStatRecord> Records = await Perforce.FStatAsync(-1, Change, null, null, -1, FStatOptions.IncludeFileSizes, FileSpecList.Empty);
+			return await UpdateTreeAsync(PrevTreeRef, Records, View, Comparison);
 		}
 
 		/// <summary>
 		/// Creates a snapshot for the given stream, given an initial state and list of modified files
 		/// </summary>
+		/// <param name="PrevTreeRef"></param>
 		/// <param name="Files"></param>
 		/// <param name="View"></param>
 		/// <param name="Comparison"></param>
-		/// <param name="PrevTreeRef"></param>
 		/// <returns></returns>
-		async Task<StreamTreeRef> UpdateTreeAsync(IList<P4.FileMetaData> Files, ViewMap View, StringComparison Comparison, StreamTreeRef PrevTreeRef)
+		async Task<StreamTreeRef> UpdateTreeAsync(StreamTreeRef PrevTreeRef, IList<FStatRecord> Files, ViewMap View, StringComparison Comparison)
 		{
 			StreamTreeBuilder Tree = await ReadTreeAsync(PrevTreeRef);
 
 			// Update the tree
-			foreach (P4.FileMetaData File in Files)
+			foreach (FStatRecord File in Files)
 			{
 				string? LocalPathStr;
-				if (View.TryMapFile(File.DepotPath.Path, Comparison, out LocalPathStr))
+				if (File.DepotFile != null && View.TryMapFile(File.DepotFile, Comparison, out LocalPathStr))
 				{
 					Utf8String LocalPath = LocalPathStr;
 
@@ -452,7 +864,7 @@ namespace HordeServer.Commits.Impl
 					}
 
 					Utf8String FileName = LocalPath;
-					if (File.HeadRev > 0)
+					if (File.Action != FileAction.Delete && File.Action != FileAction.MoveDelete)
 					{
 						Node.NameToFile[FileName] = CreateStreamFile(File);
 					}
@@ -470,21 +882,31 @@ namespace HordeServer.Commits.Impl
 			return await EncodeTreeAsync(Tree);
 		}
 
-		#endregion
+#endregion
 
 		#region Misc tree manipulation
 
 		/// <summary>
 		/// Gets the view for a stream
 		/// </summary>
-		/// <param name="Repository">The Perforce repository to mirror</param>
-		/// <param name="Stream">Name of the stream to be mirrored</param>
-		/// <param name="Change">Changelist to replicate</param>
-		static ViewMap GetStreamView(P4.Repository Repository, string Stream, int Change)
+		/// <param name="Connection">The Perforce connection</param>
+		/// <param name="StreamName">Name of the stream to be mirrored</param>
+		static async Task<ViewMap> GetStreamViewAsync(IPerforceConnection Connection, string StreamName)
 		{
-			// Get the stream spec
-			P4.Stream StreamSpec = Repository.GetStream($"{Stream}@{Change}", new P4.Options { { "-v", "" } });
-			return new ViewMap(StreamSpec.View);
+			StreamRecord Record = await Connection.GetStreamAsync(StreamName, true);
+
+			ViewMap View = new ViewMap();
+			foreach (string ViewLine in Record.View)
+			{
+				Match Match = Regex.Match(ViewLine, "^(-?)([^ ]+) *([^ ]+)$");
+				if (!Match.Success)
+				{
+					throw new PerforceException($"Unable to parse stream view: '{ViewLine}'");
+				}
+				View.Entries.Add(new ViewMapEntry(Match.Groups[1].Length == 0, Match.Groups[2].Value, Match.Groups[3].Value));
+			}
+
+			return View;
 		}
 
 		/// <summary>
@@ -576,10 +998,10 @@ namespace HordeServer.Commits.Impl
 		/// </summary>
 		/// <param name="MetaData"></param>
 		/// <returns></returns>
-		static StreamFile CreateStreamFile(P4.FileMetaData MetaData)
+		static StreamFile CreateStreamFile(FStatRecord MetaData)
 		{
-			P4.FileType FileType = MetaData.Type ?? MetaData.HeadType;
-			return new StreamFile(MetaData.DepotPath.Path, MetaData.FileSize, new FileContentId(Md5Hash.Parse(MetaData.Digest), FileType.ToString()), MetaData.HeadRev);
+			string FileType = MetaData.Type ?? MetaData.HeadType ?? String.Empty;
+			return new StreamFile(MetaData.DepotFile!, MetaData.FileSize, new FileContentId(Md5Hash.Parse(MetaData.Digest!), FileType.ToString()), MetaData.HeadRevision);
 		}
 
 		/// <summary>
diff --git a/Engine/Source/Programs/Horde/HordeServer/Models/Globals.cs b/Engine/Source/Programs/Horde/HordeServer/Models/Globals.cs
--- a/Engine/Source/Programs/Horde/HordeServer/Models/Globals.cs
+++ b/Engine/Source/Programs/Horde/HordeServer/Models/Globals.cs
@@ -375,7 +375,7 @@ namespace HordeServer.Models
 		/// </summary>
 		/// <param name="Name">Name of the cluster</param>
 		/// <returns></returns>
-		public PerforceCluster FindPerforceCluster(string? Name)
+		public PerforceCluster? FindPerforceCluster(string? Name)
 		{
 			List<PerforceCluster> Clusters = PerforceClusters;
 			if (Clusters.Count == 0)
diff --git a/Engine/Source/Programs/Horde/HordeServer/Services/IPerforceService.cs b/Engine/Source/Programs/Horde/HordeServer/Services/IPerforceService.cs
--- a/Engine/Source/Programs/Horde/HordeServer/Services/IPerforceService.cs
+++ b/Engine/Source/Programs/Horde/HordeServer/Services/IPerforceService.cs
@@ -58,6 +58,13 @@ namespace HordeServer.Services
 	/// </summary>
 	public interface IPerforceService
 	{
+		/// <summary>
+		/// 
+		/// </summary>
+		/// <param name="ClusterName"></param>
+		/// <returns></returns>
+		public Task<NativePerforceConnection?> GetServiceUserConnection(string? ClusterName);
+
 		/// <summary>
 		/// Gets the definition of a stream
 		/// </summary>
diff --git a/Engine/Source/Programs/Horde/HordeServer/Services/Impl/PerforceService.cs b/Engine/Source/Programs/Horde/HordeServer/Services/Impl/PerforceService.cs
--- a/Engine/Source/Programs/Horde/HordeServer/Services/Impl/PerforceService.cs
+++ b/Engine/Source/Programs/Horde/HordeServer/Services/Impl/PerforceService.cs
@@ -1,6 +1,7 @@
 // Copyright Epic Games, Inc. All Rights Reserved.
 
 using EpicGames.Core;
+using EpicGames.Perforce;
 using HordeServer.Models;
 using HordeServer.Utilities;
 using Microsoft.Extensions.Logging;
@@ -38,7 +39,7 @@ namespace HordeServer.Services
 	/// <summary>
 	/// P4API implementation of the Perforce service
 	/// </summary>
-	public class PerforceService : IPerforceService
+	class PerforceService : IPerforceService
 	{
 		class CachedTicketInfo
 		{
@@ -150,6 +151,46 @@ namespace HordeServer.Services
 			return Repository;
 		}
 
+		/// <summary>
+		/// 
+		/// </summary>
+		/// <param name="ClusterName"></param>
+		/// <returns></returns>
+		public async Task<NativePerforceConnection?> GetServiceUserConnection(string? ClusterName)
+		{
+			if(ClusterName != null)
+			{
+				Globals Globals = await CachedGlobals.GetCached();
+				PerforceCluster? Cluster = Globals.FindPerforceCluster(ClusterName);
+				if(Cluster != null)
+				{
+					string? UserName = null;
+					string? Password = null;
+					if (Cluster.ServiceAccount != null)
+					{
+						PerforceCredentials? Credentials = Cluster.Credentials.FirstOrDefault(x => x.UserName.Equals(Cluster.ServiceAccount, StringComparison.OrdinalIgnoreCase));
+						if (Credentials == null)
+						{
+							throw new Exception($"No credentials defined for {Cluster.ServiceAccount} on {Cluster.Name}");
+						}
+					}
+
+					PerforceSettings Settings = new PerforceSettings();
+					Settings.User = UserName;
+					Settings.Client = "__DOES_NOT_EXIST__";
+
+					NativePerforceConnection NativeConnection = new NativePerforceConnection(Logger);
+					await NativeConnection.ConnectAsync(Settings);
+					if (Password != null)
+					{
+						await NativeConnection.LoginAsync(Password);
+					}
+					return NativeConnection;
+				}
+			}
+			return null;
+		}
+
 		async Task<P4.Repository> GetServiceUserConnection(PerforceCluster Cluster)
 		{
 			IPerforceServer Server = await SelectServer(Cluster);
diff --git a/Engine/Source/Programs/Horde/HordeServer/Utilities/ElectedTick.cs b/Engine/Source/Programs/Horde/HordeServer/Utilities/ElectedTick.cs
--- a/Engine/Source/Programs/Horde/HordeServer/Utilities/ElectedTick.cs
+++ b/Engine/Source/Programs/Horde/HordeServer/Utilities/ElectedTick.cs
@@ -41,7 +41,7 @@ namespace HordeServer.Utilities
 		/// <param name="Interval">Interval for calling the tick function</param>
 		/// <param name="Logger">The log service</param>
 		public ElectedTick(DatabaseService DatabaseService, ObjectId ServiceId, Func<CancellationToken, Task> TickAsync, TimeSpan Interval, ILogger Logger)
-			: this(DatabaseService, ServiceId, async (Token) => { await TickAsync(Token); return DateTime.UtcNow + Interval; }, Logger)
+			: this(DatabaseService, ServiceId, Interval, Interval, async (Token) => { await TickAsync(Token); return DateTime.UtcNow + Interval; }, Logger)
 		{
 		}
 
diff --git a/Engine/Source/Programs/Horde/HordeServerTests/Stubs/Services/PerforceServiceStub.cs b/Engine/Source/Programs/Horde/HordeServerTests/Stubs/Services/PerforceServiceStub.cs
--- a/Engine/Source/Programs/Horde/HordeServerTests/Stubs/Services/PerforceServiceStub.cs
+++ b/Engine/Source/Programs/Horde/HordeServerTests/Stubs/Services/PerforceServiceStub.cs
@@ -1,5 +1,6 @@
 // Copyright Epic Games, Inc. All Rights Reserved.
 
+using EpicGames.Perforce;
 using HordeServer.Models;
 using HordeServer.Services;
 using System;
@@ -36,6 +37,11 @@ namespace HordeServerTests.Stubs.Services
 
 		public Dictionary<string, SortedDictionary<int, ChangeDetails>> Changes { get; } = new Dictionary<string, SortedDictionary<int, ChangeDetails>>(StringComparer.OrdinalIgnoreCase);
 
+		public Task<NativePerforceConnection?> GetServiceUserConnection(string? ClusterName)
+		{
+			throw new NotImplementedException();
+		}
+
 		public void AddChange(string StreamName, int Number, string Author, string Description, IEnumerable<string> Files)
 		{
 			SortedDictionary<int, ChangeDetails>? StreamChanges;
diff --git a/Engine/Source/Programs/Shared/EpicGames.Redis/RedisSortedSet.cs b/Engine/Source/Programs/Shared/EpicGames.Redis/RedisSortedSet.cs
--- a/Engine/Source/Programs/Shared/EpicGames.Redis/RedisSortedSet.cs
+++ b/Engine/Source/Programs/Shared/EpicGames.Redis/RedisSortedSet.cs
@@ -160,9 +160,9 @@ namespace EpicGames.Redis
 		}
 
 		/// <inheritdoc cref="IDatabaseAsync.SortedSetRemoveRangeByScoreAsync(RedisKey, double, double, Exclude, CommandFlags)"/>
-		public async Task<long> RemoveRangeByScoreAsync(double Start, double Stop, Exclude Exclude = Exclude.None, CommandFlags Flags = CommandFlags.None)
+		public Task<long> RemoveRangeByScoreAsync(double Start, double Stop, Exclude Exclude = Exclude.None, CommandFlags Flags = CommandFlags.None)
 		{
-			return await Database.SortedSetRemoveRangeByScoreAsync(Key, Start, Stop, Exclude, Flags);
+			return Database.SortedSetRemoveRangeByScoreAsync(Key, Start, Stop, Exclude, Flags);
 		}
 
 		/// <inheritdoc cref="IDatabaseAsync.SortedSetRemoveRangeByValueAsync(RedisKey, RedisValue, RedisValue, Exclude, CommandFlags)"/>
@@ -181,6 +181,13 @@ namespace EpicGames.Redis
 				yield return new SortedSetEntry<TElement>(Entry);
 			}
 		}
+
+		/// <inheritdoc cref="IDatabaseAsync.SortedSetScoreAsync(RedisKey, RedisValue, CommandFlags)"/>
+		public Task<double?> ScoreAsync(TElement Member, CommandFlags Flags = CommandFlags.None)
+		{
+			RedisValue Value = RedisSerializer.Serialize(Member);
+			return Database.SortedSetScoreAsync(Key, Value, Flags);
+		}
 	}
 
 	/// <summary>
