Horde - UnrealAutomatedTests - add support for skipped tests, device details and metadata

diff --git a/Engine/Source/Programs/Horde/HordeDashboard/src/Plugins/UnrealAutomatedTests/models/UnrealAutomatedTests.ts b/Engine/Source/Programs/Horde/HordeDashboard/src/Plugins/UnrealAutomatedTests/models/UnrealAutomatedTests.ts
--- a/Engine/Source/Programs/Horde/HordeDashboard/src/Plugins/UnrealAutomatedTests/models/UnrealAutomatedTests.ts
+++ b/Engine/Source/Programs/Horde/HordeDashboard/src/Plugins/UnrealAutomatedTests/models/UnrealAutomatedTests.ts
@@ -6,6 +6,7 @@ export enum TestState {
     SuccessWithWarnings = "SuccessWithWarnings",
     NotRun = "NotRun",
     InProcess = "InProcess",
+    Skipped = "Skipped",
     Unknown = "Unknown",
 }
 
@@ -33,6 +34,7 @@ export type TestDetails = {
     TestDisplayName: string;
     FullTestPath: string;
     State: string;
+    DeviceInstance: string;
     Errors: number;
     Warnings: number;
     Entries: TestEntry[];
@@ -43,13 +45,28 @@ export type TestResult = {
     TestDisplayName: string;
     FullTestPath: string;
     State: string;
+    DeviceInstance: string;
     ArtifactName: string;
     Errors: number;
     Warnings: number;
 }
 
+export type TestDevice = {
+    DeviceName: string;
+    Instance: string;
+    Platform: string;
+    OSVersion: string;
+    Model: string;
+    GPU: string;
+    CPUModel: string;
+    RAMInGB: number;
+    RenderMode: string;
+    RHI: string;
+}
+
+export type Metadata = {[Key in string]: string}
+
 export type TestPassSummary = {
-    ClientDescriptor: string;
     ReportURL: string;
     FailedCount: number;
     NotRunCount: number;
@@ -59,6 +76,8 @@ export type TestPassSummary = {
     SucceededWithWarningsCount: number;
     TotalDurationSeconds: number;
     Tests: TestResult[];
+    Devices: TestDevice[];
+    Metadata: Metadata;
 }
 
 export type TestStateHistoryItem = {
