Horde: Prevent new steps being added to a batch that isn't in a running state.

diff --git a/Engine/Source/Programs/Horde/HordeServer/Services/RpcService.cs b/Engine/Source/Programs/Horde/HordeServer/Services/RpcService.cs
--- a/Engine/Source/Programs/Horde/HordeServer/Services/RpcService.cs
+++ b/Engine/Source/Programs/Horde/HordeServer/Services/RpcService.cs
@@ -600,6 +600,10 @@ namespace HordeServer.Services
 
 			// Find the batch being executed
 			IJobStepBatch Batch = AuthorizeBatch(Job, Request.BatchId.ToSubResourceId(), Context);
+			if (Batch.State != Api.JobStepBatchState.Starting && Batch.State != Api.JobStepBatchState.Running)
+			{
+				return new BeginStepResponse { State = BeginStepResponse.Types.Result.Complete };
+			}
 
 			// Figure out which step to execute next
 			IJobStep? Step;
