Horde: Reduce memory usage when querying workspace state (pt 1)

* CbWriter instances may now be cleared, and can cache allocated buffers between uses.
* Records may now be returned using IAsyncEnumerable<T> allowing them to be processed as they are received, and removing the need to allocate large numbers of temporary objects.

diff --git a/Engine/Source/Programs/Horde/HordeServer/Commits/Impl/CommitService.cs b/Engine/Source/Programs/Horde/HordeServer/Commits/Impl/CommitService.cs
--- a/Engine/Source/Programs/Horde/HordeServer/Commits/Impl/CommitService.cs
+++ b/Engine/Source/Programs/Horde/HordeServer/Commits/Impl/CommitService.cs
@@ -676,11 +676,6 @@ namespace HordeServer.Commits.Impl
 
 		#region Snapshot generation
 
-		/// <summary>
-		/// Number of files to query at a time before recursing down the tree. Tuning this value allows finding a threshold between request/response lengths and memory usage.
-		/// </summary>
-		const int BatchSize = 256 * 1024;
-
 		/// <summary>
 		/// Creates a snapshot of a stream at a particular changelist
 		/// </summary>
@@ -690,89 +685,96 @@ namespace HordeServer.Commits.Impl
 		/// <param name="Comparison">Comparison type for paths</param>
 		async Task<StreamTreeRef> CreateTreeAsync(IPerforceConnection Connection, ViewMap View, int Change, StringComparison Comparison)
 		{
-			// In order to optimize for the common case where we're adding multiple files to the same directory, we keep an open list of
-			// StreamTreeBuilder objects down to the last modified node. Whenever we need to move to a different node, we write the 
-			// previous path to the object store.
-			List<StreamTreeBuilder> TreePath = new List<StreamTreeBuilder>();
-			TreePath.Add(new StreamTreeBuilder());
-
-			// Generate a snapshot using all the root paths
-			List<string> RootPaths = View.GetRootPaths(Comparison);
-			foreach (string RootPath in RootPaths)
+			// Optimize the view by removing unnecessary inclusions of the same path
+			View = new ViewMap(View);
+			for (int Idx = 0; Idx < View.Entries.Count; Idx++)
 			{
-				ViewMap FilteredView = new ViewMap();
-				foreach (ViewMapEntry Entry in View.Entries)
+				ViewMapEntry Entry = View.Entries[Idx];
+				if (Entry.Include && Entry.IsPathWildcard())
 				{
-					if (Entry.SourcePrefix.StartsWith(RootPath, Comparison))
+					while (Idx + 1 < View.Entries.Count && View.Entries[Idx + 1].Include && View.Entries[Idx + 1].SourcePrefix.StartsWith(Entry.SourcePrefix, Comparison))
 					{
-						FilteredView.Entries.Add(Entry);
+						View.Entries.RemoveAt(Idx + 1);
 					}
 				}
-				await AddDepotDirToSnapshotAsync(Connection, RootPath, Change, FilteredView, Comparison, TreePath);
 			}
 
-			// Encode the tree
-			return await EncodeTreeAsync(TreePath[0]);
-		}
+			// Delegate to write an object
+			CbWriter Writer = new CbWriter();
+			Action<IoHash, CbObject> WriteObject = (x, y) => { };
 
-		/// <summary>
-		/// Adds all files under a directory in the depot to a snapshot
-		/// </summary>
-		/// <param name="Connection">The perforce connection</param>
-		/// <param name="DepotDir">Depot directory to recurse through, with a trailing slash</param>
-		/// <param name="Change">Changelist number to capture</param>
-		/// <param name="View">View for the stream</param>
-		/// <param name="Comparison">Comparison type to use for names</param>
-		/// <param name="TreePath">The last computed path through the tree</param>
-		async Task AddDepotDirToSnapshotAsync(IPerforceConnection Connection, string DepotDir, int Change, ViewMap View, StringComparison Comparison, List<StreamTreeBuilder> TreePath)
-		{
-			const string Filter = "^headAction=delete ^headAction=move/delete";
-			const string Fields = "depotFile,fileSize,digest,headType,headRev";
+			// In order to optimize for the common case where we're adding multiple files to the same directory, we keep an open list of
+			// StreamTreeBuilder objects down to the last modified node. Whenever we need to move to a different node, we write the 
+			// previous path to the object store.
+			List<StreamTreeBuilder> TreePath = new List<StreamTreeBuilder>();
+			TreePath.Add(new StreamTreeBuilder());
 
-			if (View.MayMatchAnyFilesInSubDirectory(DepotDir, Comparison))
+			for(int Idx = 0; Idx < View.Entries.Count; Idx++)
 			{
-				List<FStatRecord> Files = await Connection.FStatAsync(-1, -1, Filter, Fields, BatchSize, FStatOptions.IncludeFileSizes, $"{DepotDir}...@{Change}");
-				if (Files.Count < BatchSize)
+				ViewMapEntry Entry = View.Entries[Idx];
+				if (Entry.Include)
 				{
-					Logger.LogInformation("Added {DepotDir}... to snapshot", DepotDir);
-					await AddDepotFilesToSnapshotAsync(Files, View, Comparison, TreePath);
-					return;
-				}
+					ViewMap FilteredView = new ViewMap();
+					FilteredView.Entries.Add(Entry);
 
-				Files.Clear(); // Don't keep around in memory
-				Logger.LogInformation("Querying subdirectories of {DepotDir} (wasted results)", DepotDir);
+					for (int OtherIdx = Idx + 1; OtherIdx < View.Entries.Count; OtherIdx++)
+					{
+						ViewMapEntry OtherEntry = View.Entries[OtherIdx];
+						if (!OtherEntry.Include)
+						{
+							if (OtherEntry.SourcePrefix.StartsWith(Entry.SourcePrefix, Comparison) || Entry.SourcePrefix.StartsWith(OtherEntry.SourcePrefix, Comparison))
+							{
+								FilteredView.Entries.Add(OtherEntry);
+							}
+						}
+					}
 
-				IList<DirsRecord> Dirs = await Connection.GetDirsAsync(DirsOptions.None, null, $"{DepotDir}*@{Change}");
-				foreach (DirsRecord Dir in Dirs)
-				{
-					await AddDepotDirToSnapshotAsync(Connection, Dir.Dir + "/", Change, View, Comparison, TreePath);
+					Logger.LogInformation("Adding {Source}@{Change}", Entry.Source, Change);
+					await AddDepotFilesToSnapshotAsync(Connection, $"{Entry.Source}@{Change}", View, Comparison, TreePath, Writer, WriteObject);
+
+					GC.Collect();
 				}
 			}
 
-			if (View.MayMatchAnyFilesInDirectory(DepotDir, Comparison))
-			{
-				List<FStatRecord> Files = await Connection.FStatAsync(-1, -1, Filter, Fields, -1, FStatOptions.IncludeFileSizes, $"{DepotDir}*@{Change}");
-				await AddDepotFilesToSnapshotAsync(Files, View, Comparison, TreePath);
-			}
+			// Encode the tree
+			return TreePath[0].Encode(Writer, WriteObject);
 		}
 
 		/// <summary>
 		/// Adds a set of files in the depot to a snapshot
 		/// </summary>
-		/// <param name="Files">The files to add</param>
+		/// <param name="Connection"></param>
+		/// <param name="FileSpec">The files to add</param>
 		/// <param name="View">View for the stream</param>
-		/// <param name="Comparison">Comparison type to use for names</param>
+		/// <param name="Comparer">Comparison type to use for names</param>
 		/// <param name="TreePath">The last computed path through the tree</param>
-		async Task AddDepotFilesToSnapshotAsync(List<FStatRecord> Files, ViewMap View, StringComparison Comparison, List<StreamTreeBuilder> TreePath)
+		/// <param name="Writer"></param>
+		/// <param name="WriteObject">Delegate to write an object to the store</param>
+		async Task AddDepotFilesToSnapshotAsync(IPerforceConnection Connection, string FileSpec, ViewMap View, StringComparison Comparison, List<StreamTreeBuilder> TreePath, CbWriter Writer, Action<IoHash, CbObject> WriteObject)
 		{
-			foreach (FStatRecord File in Files)
+			const string Filter = "^headAction=delete ^headAction=move/delete";
+			const string Fields = "depotFile,fileSize,digest,headType,headRev";
+
+			List<string> Arguments = new List<string>();
+			Arguments.Add("-Ol");
+			Arguments.Add("-Op");
+			Arguments.Add("-Os");
+			Arguments.Add("-F");
+			Arguments.Add(Filter);
+			Arguments.Add("-T");
+			Arguments.Add(String.Join(",", Fields));
+			Arguments.Add(FileSpec);
+
+			IAsyncEnumerable<PerforceResponse<FStatRecord>> Responses = Connection.StreamCommandAsync<FStatRecord>("fstat", Arguments);
+			await foreach (PerforceResponse<FStatRecord> Response in Responses)
 			{
+				FStatRecord File = Response.Data;
 				if (!String.IsNullOrEmpty(File.Digest))
 				{
 					string? TargetFile;
 					if (File.DepotFile != null && View.TryMapFile(File.DepotFile, Comparison, out TargetFile))
 					{
-						await AddFileAsync(TargetFile, File, TreePath);
+						await AddFileAsync(TargetFile, File, TreePath, Writer, WriteObject);
 					}
 				}
 			}
@@ -784,28 +786,36 @@ namespace HordeServer.Commits.Impl
 		/// <param name="Path">The stream path for the file</param>
 		/// <param name="MetaData">The file metadata</param>
 		/// <param name="TreePath">The last computed path through the tree</param>
-		async Task AddFileAsync(string Path, FStatRecord MetaData, List<StreamTreeBuilder> TreePath)
+		/// <param name="Writer"></param>
+		/// <param name="WriteObject"></param>
+		async Task AddFileAsync(string Path, FStatRecord MetaData, List<StreamTreeBuilder> TreePath, CbWriter Writer, Action<IoHash, CbObject> WriteObject)
 		{
 			string[] Segments = Path.Split('/');
 
 			StreamTreeBuilder Node = TreePath[0];
-			for (int Idx = 0; Idx < Segments.Length - 1; Idx++)
-			{
-				Utf8String Segment = Segments[Idx];
 
-				StreamTreeBuilder NextNode = await FindOrAddTreeAsync(Node, Segment);
-				if (Idx + 1 < TreePath.Count && TreePath[Idx + 1] != NextNode)
-				{
-					await CollapseTreeAsync(Node, TreePath[Idx + 1]);
-				}
-				if (Idx + 1 >= TreePath.Count)
+			int SegmentIdx = 0;
+
+			// Match as many nodes as we can from the existing path
+			for (; SegmentIdx < Segments.Length - 1; SegmentIdx++)
+			{
+				Utf8String Segment = Segments[SegmentIdx];
+				if (!Node.NameToTreeBuilder.TryGetValue(Segment, out StreamTreeBuilder? NextNode))
 				{
-					TreePath.Add(NextNode);
+					break;
 				}
-
 				Node = NextNode;
 			}
 
+			// Collapse the rest of the current path
+			CollapseTree(Node, Writer, WriteObject);
+
+			// Expand the rest of the path
+			for (; SegmentIdx < Segments.Length - 1; SegmentIdx++)
+			{
+				Node = await FindOrAddTreeAsync(Node, Segments[SegmentIdx]);
+			}
+
 			Node.NameToFile[Segments[Segments.Length - 1]] = CreateStreamFile(MetaData);
 		}
 
@@ -840,6 +850,9 @@ namespace HordeServer.Commits.Impl
 		{
 			StreamTreeBuilder Tree = await ReadTreeAsync(PrevTreeRef);
 
+			// Delegate to write an object
+			Action<IoHash, CbObject> WriteObject = (x, y) => { };
+
 			// Update the tree
 			foreach (FStatRecord File in Files)
 			{
@@ -879,7 +892,7 @@ namespace HordeServer.Commits.Impl
 				}
 			}
 
-			return await EncodeTreeAsync(Tree);
+			return Tree.Encode(new CbWriter(), WriteObject);
 		}
 
 #endregion
@@ -958,39 +971,18 @@ namespace HordeServer.Commits.Impl
 		/// <summary>
 		/// Collapse a subtree under a given child
 		/// </summary>
-		/// <param name="Parent"></param>
-		/// <param name="Child"></param>
-		/// <returns></returns>
-		async Task CollapseTreeAsync(StreamTreeBuilder Parent, StreamTreeBuilder Child)
+		/// <param name="Node">The node to collapse</param>
+		/// <param name="Writer"></param>
+		/// <param name="WriteObject"></param>
+		static void CollapseTree(StreamTreeBuilder Node, CbWriter Writer, Action<IoHash, CbObject> WriteObject)
 		{
-			foreach ((Utf8String Name, StreamTreeBuilder Node) in Parent.NameToTreeBuilder)
+			foreach ((Utf8String Name, StreamTreeBuilder Child) in Node.NameToTreeBuilder)
 			{
-				if (Node == Child)
-				{
-					StreamTreeRef ChildRef = await EncodeTreeAsync(Child);
-					Parent.NameToTreeBuilder.Remove(Name);
-					Parent.NameToTree.Add(Name, ChildRef);
-					break;
-				}
+				CollapseTree(Child, Writer, WriteObject);
+				StreamTreeRef ChildRef = Child.Encode(Writer, WriteObject);
+				Node.NameToTree.Add(Name, ChildRef);
 			}
-		}
-
-		/// <summary>
-		/// Writes a subtree to the object store
-		/// </summary>
-		/// <param name="TreeBuilder">Root of the tree to write</param>
-		/// <returns>New tree reference for the serialized tree</returns>
-		async Task<StreamTreeRef> EncodeTreeAsync(StreamTreeBuilder TreeBuilder)
-		{
-			Dictionary<IoHash, CbObject> Objects = new Dictionary<IoHash, CbObject>();
-			StreamTreeRef Ref = TreeBuilder.Encode(Objects);
-
-			foreach ((IoHash Hash, CbObject Object) in Objects)
-			{
-				await ObjectCollection.AddAsync(NamespaceId, Hash, Object);
-			}
-
-			return Ref;
+			Node.NameToTreeBuilder.Clear();
 		}
 
 		/// <summary>
diff --git a/Engine/Source/Programs/Horde/HordeServer/Utilities/ViewMap.cs b/Engine/Source/Programs/Horde/HordeServer/Utilities/ViewMap.cs
--- a/Engine/Source/Programs/Horde/HordeServer/Utilities/ViewMap.cs
+++ b/Engine/Source/Programs/Horde/HordeServer/Utilities/ViewMap.cs
@@ -27,6 +27,15 @@ namespace HordeServer.Utilities
 		{
 		}
 
+		/// <summary>
+		/// Constructor
+		/// </summary>
+		/// <param name="Other"></param>
+		public ViewMap(ViewMap Other)
+		{
+			Entries.AddRange(Other.Entries);
+		}
+
 		/// <summary>
 		/// Constructor
 		/// </summary>
@@ -218,6 +227,16 @@ namespace HordeServer.Utilities
 		/// </summary>
 		public string TargetSuffix { get; }
 
+		/// <summary>
+		/// The full source path
+		/// </summary>
+		public string Source => (Wildcard == null) ? SourcePrefix : $"{SourcePrefix}{Wildcard}{SourceSuffix}";
+
+		/// <summary>
+		/// The full target path
+		/// </summary>
+		public string Target => (Wildcard == null) ? TargetPrefix : $"{TargetPrefix}{Wildcard}{TargetSuffix}";
+
 		/// <summary>
 		/// Tests if the entry has a file wildcard ('*')
 		/// </summary>
@@ -230,6 +249,15 @@ namespace HordeServer.Utilities
 		/// <returns>True if the entry has a path wildcard</returns>
 		public bool IsPathWildcard() => Wildcard != null && Wildcard.Length == 3;
 
+		/// <summary>
+		/// Constructor
+		/// </summary>
+		/// <param name="Other"></param>
+		public ViewMapEntry(ViewMapEntry Other)
+			: this(Other.Include, Other.Wildcard, Other.SourcePrefix, Other.SourceSuffix, Other.TargetPrefix, Other.TargetSuffix)
+		{
+		}
+
 		/// <summary>
 		/// Constructor
 		/// </summary>
diff --git a/Engine/Source/Programs/Shared/EpicGames.Perforce.Managed/StreamSnapshotFromMemory.cs b/Engine/Source/Programs/Shared/EpicGames.Perforce.Managed/StreamSnapshotFromMemory.cs
--- a/Engine/Source/Programs/Shared/EpicGames.Perforce.Managed/StreamSnapshotFromMemory.cs
+++ b/Engine/Source/Programs/Shared/EpicGames.Perforce.Managed/StreamSnapshotFromMemory.cs
@@ -50,7 +50,7 @@ namespace EpicGames.Perforce.Managed
 		public StreamSnapshotFromMemory(StreamTreeBuilder Builder)
 		{
 			Dictionary<IoHash, CbObject> HashToTree = new Dictionary<IoHash, CbObject>();
-			this.Root = Builder.Encode(HashToTree);
+			this.Root = Builder.Encode(new CbWriter(), (Hash, Object) => HashToTree[Hash] = Object);
 			this.HashToTree = HashToTree;
 		}
 
diff --git a/Engine/Source/Programs/Shared/EpicGames.Perforce.Managed/StreamTree.cs b/Engine/Source/Programs/Shared/EpicGames.Perforce.Managed/StreamTree.cs
--- a/Engine/Source/Programs/Shared/EpicGames.Perforce.Managed/StreamTree.cs
+++ b/Engine/Source/Programs/Shared/EpicGames.Perforce.Managed/StreamTree.cs
@@ -324,14 +324,14 @@ namespace EpicGames.Perforce.Managed
 		/// </summary>
 		/// <param name="BasePath"></param>
 		/// <returns></returns>
-		public CbObject ToCbObject(Utf8String BasePath)
+		public CbObject ToCbObject(CbWriter Writer, Utf8String BasePath)
 		{
 			if (BasePath.EndsWith("/"))
 			{
 				throw new ArgumentException("BasePath must not end in a slash", nameof(BasePath));
 			}
 
-			CbWriter Writer = new CbWriter();
+			Writer.Clear();
 			Writer.BeginObject();
 			Write(Writer, BasePath);
 			Writer.EndObject();
diff --git a/Engine/Source/Programs/Shared/EpicGames.Perforce.Managed/StreamTreeBuilder.cs b/Engine/Source/Programs/Shared/EpicGames.Perforce.Managed/StreamTreeBuilder.cs
--- a/Engine/Source/Programs/Shared/EpicGames.Perforce.Managed/StreamTreeBuilder.cs
+++ b/Engine/Source/Programs/Shared/EpicGames.Perforce.Managed/StreamTreeBuilder.cs
@@ -57,16 +57,17 @@ namespace EpicGames.Perforce.Managed
 		/// <summary>
 		/// Encodes the current tree state and returns a reference to it
 		/// </summary>
-		/// <param name="Objects">Dictionary of encoded objects</param>
+		/// <param name="Writer">A writer instance to use to create the object</param>
+		/// <param name="WriteObject"></param>
 		/// <returns></returns>
-		public StreamTreeRef Encode(Dictionary<IoHash, CbObject> Objects)
+		public StreamTreeRef Encode(CbWriter Writer, Action<IoHash, CbObject> WriteObject)
 		{
-			StreamTreeRef? Ref = EncodeInternal(Objects);
+			StreamTreeRef? Ref = EncodeInternal(Writer, WriteObject);
 			if (Ref == null)
 			{
-				CbObject Object = new StreamTree().ToCbObject(Utf8String.Empty);
+				CbObject Object = new StreamTree().ToCbObject(Writer, Utf8String.Empty);
 				IoHash ObjectHash = Object.GetHash();
-				Objects[ObjectHash] = Object;
+				WriteObject(ObjectHash, Object);
 				Ref = new StreamTreeRef(Utf8String.Empty, ObjectHash);
 			}
 			return Ref;
@@ -77,12 +78,12 @@ namespace EpicGames.Perforce.Managed
 		/// </summary>
 		/// <param name="Objects">Dictionary of encoded objects</param>
 		/// <returns></returns>
-		public StreamTreeRef? EncodeInternal(Dictionary<IoHash, CbObject> Objects)
+		public StreamTreeRef? EncodeInternal(CbWriter Writer, Action<IoHash, CbObject> WriteObject)
 		{
 			// Recursively serialize all the child items
 			foreach ((Utf8String SubTreeName, StreamTreeBuilder SubTree) in NameToTreeBuilder)
 			{
-				StreamTreeRef? SubTreeRef = SubTree.EncodeInternal(Objects);
+				StreamTreeRef? SubTreeRef = SubTree.EncodeInternal(Writer, WriteObject);
 				if (SubTreeRef == null)
 				{
 					NameToTreeBuilder.Remove(SubTreeName);
@@ -103,10 +104,10 @@ namespace EpicGames.Perforce.Managed
 			// Find the common base path for all items in this tree.
 			StreamTree Tree = new StreamTree(NameToFile, NameToTree);
 			Utf8String BasePath = Tree.FindBasePath();
-			CbObject PackedObject = Tree.ToCbObject(BasePath);
+			CbObject PackedObject = Tree.ToCbObject(Writer, BasePath);
 
 			IoHash ObjectHash = PackedObject.GetHash();
-			Objects[ObjectHash] = PackedObject;
+			WriteObject(ObjectHash, PackedObject);
 
 			return new StreamTreeRef(BasePath, ObjectHash);
 		}
diff --git a/Engine/Source/Programs/Shared/EpicGames.Perforce/IPerforceOutput.cs b/Engine/Source/Programs/Shared/EpicGames.Perforce/IPerforceOutput.cs
--- a/Engine/Source/Programs/Shared/EpicGames.Perforce/IPerforceOutput.cs
+++ b/Engine/Source/Programs/Shared/EpicGames.Perforce/IPerforceOutput.cs
@@ -7,6 +7,7 @@ using System.Collections.Generic;
 using System.Diagnostics.CodeAnalysis;
 using System.IO;
 using System.Linq;
+using System.Runtime.CompilerServices;
 using System.Text;
 using System.Threading;
 using System.Threading.Tasks;
@@ -154,6 +155,64 @@ namespace EpicGames.Perforce
 			return Result.ToString();
 		}
 
+		/// <summary>
+		/// Read a list of responses from the child process
+		/// </summary>
+		/// <param name="Perforce">The response to read from</param>
+		/// <param name="StatRecordType">The type of stat record to parse</param>
+		/// <param name="CancellationToken">Cancellation token for the read</param>
+		/// <returns>Async task</returns>
+		public static async IAsyncEnumerable<PerforceResponse> ReadStreamingResponsesAsync(this IPerforceOutput Perforce, Type? StatRecordType, [EnumeratorCancellation] CancellationToken CancellationToken)
+		{
+			CachedRecordInfo? StatRecordInfo = (StatRecordType == null) ? null : PerforceReflection.GetCachedRecordInfo(StatRecordType);
+
+			List<PerforceResponse> Responses = new List<PerforceResponse>();
+
+			// Read all the records into a list
+			long ParsedLen = 0;
+			long MaxParsedLen = 0;
+			while (await Perforce.ReadAsync(CancellationToken))
+			{
+				// Check for the whole message not being a marshalled python object, and produce a better response in that scenario
+				ReadOnlyMemory<byte> Data = Perforce.Data;
+				if (Data.Length > 0 && Responses.Count == 0 && Data.Span[0] != '{')
+				{
+					throw new PerforceException("Unexpected response from server (expected '{'):{0}", FormatDataAsString(Data.Span));
+				}
+
+				// Parse the responses from the current buffer
+				int BufferPos = 0;
+				for (; ; )
+				{
+					int NewBufferPos = BufferPos;
+					if (!TryReadResponse(Data.Span, ref NewBufferPos, StatRecordInfo, out PerforceResponse? Response))
+					{
+						MaxParsedLen = ParsedLen + NewBufferPos;
+						break;
+					}
+					if (Response.Error == null || Response.Error.Generic != PerforceGenericCode.Empty)
+					{
+						yield return Response;
+					}
+					BufferPos = NewBufferPos;
+				}
+
+				// Discard all the data that we've processed
+				Perforce.Discard(BufferPos);
+				ParsedLen += BufferPos;
+			}
+
+			// If the stream is complete but we couldn't parse a response from the server, treat it as an error
+			if (Perforce.Data.Length > 0)
+			{
+				long DumpOffset = Math.Max(MaxParsedLen - 32, ParsedLen);
+				int SliceOffset = (int)(DumpOffset - ParsedLen);
+				string StrDump = FormatDataAsString(Perforce.Data.Span.Slice(SliceOffset));
+				string HexDump = FormatDataAsHexDump(Perforce.Data.Span.Slice(SliceOffset, Math.Min(1024, Perforce.Data.Length - SliceOffset)));
+				throw new PerforceException("Unparsable data at offset {0}+{1}/{2}.\nString data from offset {3}:{4}\nHex data from offset {3}:{5}", ParsedLen, MaxParsedLen - ParsedLen, ParsedLen + Perforce.Data.Length, DumpOffset, StrDump, HexDump);
+			}
+		}
+
 		/// <summary>
 		/// Read a list of responses from the child process
 		/// </summary>
diff --git a/Engine/Source/Programs/Shared/EpicGames.Perforce/PerforceConnection.cs b/Engine/Source/Programs/Shared/EpicGames.Perforce/PerforceConnection.cs
--- a/Engine/Source/Programs/Shared/EpicGames.Perforce/PerforceConnection.cs
+++ b/Engine/Source/Programs/Shared/EpicGames.Perforce/PerforceConnection.cs
@@ -9,6 +9,7 @@ using System.Diagnostics;
 using System.IO;
 using System.Linq;
 using System.Reflection;
+using System.Runtime.CompilerServices;
 using System.Text;
 using System.Text.RegularExpressions;
 using System.Threading;
@@ -335,6 +336,50 @@ namespace EpicGames.Perforce
 			}
 		}
 
+		/// <summary>
+		/// Execute a command and parse the response
+		/// </summary>
+		/// <param name="Perforce">The Perforce connection</param>
+		/// <param name="Command">Command to execute</param>
+		/// <param name="Arguments">Arguments for the command</param>
+		/// <param name="FileArguments">File arguments for the command</param>
+		/// <param name="InputData">Input data to pass to Perforce</param>
+		/// <param name="StatRecordType">The type of records to return for "stat" responses</param>
+		/// <param name="CancellationToken">Token used to cancel the operation</param>
+		/// <returns>List of objects returned by the server</returns>
+		public static async IAsyncEnumerable<PerforceResponse> StreamCommandAsync(this IPerforceConnection Perforce, string Command, IReadOnlyList<string> Arguments, IReadOnlyList<string>? FileArguments, byte[]? InputData, Type? StatRecordType, [EnumeratorCancellation] CancellationToken CancellationToken)
+		{
+			await using (IPerforceOutput Output = await Perforce.CommandAsync(Command, Arguments, FileArguments, InputData))
+			{
+				await foreach (PerforceResponse Response in Output.ReadStreamingResponsesAsync(StatRecordType, CancellationToken))
+				{
+					yield return Response;
+				}
+			}
+		}
+
+		/// <summary>
+		/// Execute a command and parse the response
+		/// </summary>
+		/// <param name="Perforce">The Perforce connection</param>
+		/// <param name="Command">Command to execute</param>
+		/// <param name="Arguments">Arguments for the command</param>
+		/// <param name="FileArguments">File arguments for the command</param>
+		/// <param name="InputData">Input data to pass to Perforce</param>
+		/// <param name="CancellationToken">Token used to cancel the operation</param>
+		/// <returns>List of objects returned by the server</returns>
+		public static async IAsyncEnumerable<PerforceResponse<T>> StreamCommandAsync<T>(this IPerforceConnection Perforce, string Command, IReadOnlyList<string> Arguments, IReadOnlyList<string>? FileArguments = null, byte[]? InputData = null, [EnumeratorCancellation] CancellationToken CancellationToken = default) where T : class
+		{
+			await using (IPerforceOutput Output = await Perforce.CommandAsync(Command, Arguments, FileArguments, InputData))
+			{
+				Type StatRecordType = typeof(T);
+				await foreach (PerforceResponse Response in Output.ReadStreamingResponsesAsync(StatRecordType, CancellationToken))
+				{
+					yield return new PerforceResponse<T>(Response);
+				}
+			}
+		}
+
 		/// <summary>
 		/// Execute a command and parse the response
 		/// </summary>
diff --git a/Engine/Source/Programs/Shared/EpicGames.Serialization/CompactBinaryWriter.cs b/Engine/Source/Programs/Shared/EpicGames.Serialization/CompactBinaryWriter.cs
--- a/Engine/Source/Programs/Shared/EpicGames.Serialization/CompactBinaryWriter.cs
+++ b/Engine/Source/Programs/Shared/EpicGames.Serialization/CompactBinaryWriter.cs
@@ -51,14 +51,23 @@ namespace EpicGames.Serialization
 			public int Offset; // Offset to insert the length/count
 			public int Length; // Excludes the size of this field's headers, and child fields' headers.
 			public int Count;
-			public List<Scope>? Children;
+			public List<Scope> Children = new List<Scope>();
 			public int SizeOfChildHeaders; // Sum of additional headers for child items, recursively.
 
 			public Scope(CbFieldType FieldType, CbFieldType UniformFieldType, int Offset)
+			{
+				Reset(FieldType, UniformFieldType, Offset);
+			}
+
+			public void Reset(CbFieldType FieldType, CbFieldType UniformFieldType, int Offset)
 			{
 				this.FieldType = FieldType;
 				this.UniformFieldType = UniformFieldType;
 				this.Offset = Offset;
+				this.Length = 0;
+				this.Count = 0;
+				this.Children.Clear();
+				this.SizeOfChildHeaders = 0;
 			}
 		}
 
@@ -74,8 +83,15 @@ namespace EpicGames.Serialization
 
 			public Chunk(int Offset, int MaxLength)
 			{
-				this.Offset = Offset;
 				this.Data = new byte[MaxLength];
+				Reset(Offset);
+			}
+
+			public void Reset(int Offset)
+			{
+				this.Offset = Offset;
+				this.Length = 0;
+				this.Scopes.Clear();
 			}
 		}
 
@@ -86,6 +102,8 @@ namespace EpicGames.Serialization
 		Chunk CurrentChunk => Chunks[Chunks.Count - 1];
 		Scope CurrentScope => OpenScopes.Peek();
 		int CurrentOffset;
+		List<Chunk> FreeChunks = new List<Chunk>();
+		List<Scope> FreeScopes = new List<Scope>();
 
 		/// <summary>
 		/// Constructor
@@ -105,6 +123,86 @@ namespace EpicGames.Serialization
 			OpenScopes.Push(new Scope(CbFieldType.Array, CbFieldType.None, 0));
 		}
 
+		/// <summary>
+		/// 
+		/// </summary>
+		public void Clear()
+		{
+			foreach (Chunk Chunk in Chunks)
+			{
+				FreeChunk(Chunk);
+			}
+
+			CurrentOffset = 0;
+
+			Chunks.Clear();
+			Chunks.Add(AllocChunk(0, DefaultChunkSize));
+
+			OpenScopes.Clear();
+			OpenScopes.Push(AllocScope(CbFieldType.Array, CbFieldType.None, 0));
+		}
+
+		/// <summary>
+		/// Allocate a new chunk object
+		/// </summary>
+		/// <param name="Offset">Offset of the chunk</param>
+		/// <param name="MaxLength">Maximum length of the chunk</param>
+		/// <returns>New chunk object</returns>
+		Chunk AllocChunk(int Offset, int MaxLength)
+		{
+			for(int Idx = FreeChunks.Count - 1; Idx >= 0; Idx--)
+			{
+				Chunk Chunk = FreeChunks[Idx];
+				if (Chunk.Data.Length >= MaxLength)
+				{
+					FreeChunks.RemoveAt(Idx);
+					Chunk.Reset(Offset);
+					return Chunk;
+				}
+			}
+			return new Chunk(Offset, MaxLength);
+		}
+
+		/// <summary>
+		/// Adds a chunk to the free list
+		/// </summary>
+		/// <param name="Chunk"></param>
+		void FreeChunk(Chunk Chunk)
+		{
+			// Add the scopes to the free list
+			FreeScopes.AddRange(Chunk.Scopes);
+			Chunk.Scopes.Clear();
+
+			// Insert it into the free list, sorted by descending size
+			for (int Idx = 0; ; Idx++)
+			{
+				if (Idx == FreeChunks.Count || Chunk.Data.Length >= FreeChunks[Idx].Data.Length)
+				{
+					FreeChunks.Insert(Idx, Chunk);
+					break;
+				}
+			}
+		}
+
+		/// <summary>
+		/// Allocate a scope object
+		/// </summary>
+		/// <param name="FieldType"></param>
+		/// <param name="UniformFieldType"></param>
+		/// <param name="Offset"></param>
+		/// <returns></returns>
+		Scope AllocScope(CbFieldType FieldType, CbFieldType UniformFieldType, int Offset)
+		{
+			if (FreeScopes.Count > 0)
+			{
+				Scope Scope = FreeScopes[FreeScopes.Count - 1];
+				Scope.Reset(FieldType, UniformFieldType, Offset);
+				FreeScopes.RemoveAt(FreeScopes.Count - 1);
+				return Scope;
+			}
+			return new Scope(FieldType, UniformFieldType, Offset);
+		}
+
 		/// <summary>
 		/// Ensure that a block of contiguous memory of the given length is available in the output buffer
 		/// </summary>
@@ -116,7 +214,7 @@ namespace EpicGames.Serialization
 			if (LastChunk.Length + Length > LastChunk.Data.Length)
 			{
 				int ChunkSize = Math.Max(Length, DefaultChunkSize);
-				LastChunk = new Chunk(CurrentOffset, ChunkSize);
+				LastChunk = AllocChunk(CurrentOffset, ChunkSize);
 				Chunks.Add(LastChunk);
 			}
 
@@ -133,8 +231,7 @@ namespace EpicGames.Serialization
 		/// <param name="UniformFieldType"></param>
 		void PushScope(CbFieldType FieldType, CbFieldType UniformFieldType)
 		{
-			Scope NewScope = new Scope(FieldType, UniformFieldType, CurrentOffset);
-			CurrentScope.Children ??= new List<Scope>();
+			Scope NewScope = AllocScope(FieldType, UniformFieldType, CurrentOffset);
 			CurrentScope.Children.Add(NewScope);
 			OpenScopes.Push(NewScope);
 
@@ -682,24 +779,21 @@ namespace EpicGames.Serialization
 		static int ComputeSizeOfChildHeaders(Scope Scope)
 		{
 			int SizeOfChildHeaders = 0;
-			if (Scope.Children != null)
+			foreach (Scope ChildScope in Scope.Children)
 			{
-				foreach (Scope ChildScope in Scope.Children)
+				switch (ChildScope.FieldType)
 				{
-					switch (ChildScope.FieldType)
-					{
-						case CbFieldType.Object:
-						case CbFieldType.UniformObject:
-							SizeOfChildHeaders += ChildScope.SizeOfChildHeaders + VarInt.Measure(ChildScope.Length + ChildScope.SizeOfChildHeaders);
-							break;
-						case CbFieldType.Array:
-						case CbFieldType.UniformArray:
-							int ArrayCountLength = VarInt.Measure(ChildScope.Count);
-							SizeOfChildHeaders += ChildScope.SizeOfChildHeaders + VarInt.Measure(ChildScope.Length + ChildScope.SizeOfChildHeaders + ArrayCountLength) + ArrayCountLength;
-							break;
-						default:
-							throw new InvalidOperationException();
-					}
+					case CbFieldType.Object:
+					case CbFieldType.UniformObject:
+						SizeOfChildHeaders += ChildScope.SizeOfChildHeaders + VarInt.Measure(ChildScope.Length + ChildScope.SizeOfChildHeaders);
+						break;
+					case CbFieldType.Array:
+					case CbFieldType.UniformArray:
+						int ArrayCountLength = VarInt.Measure(ChildScope.Count);
+						SizeOfChildHeaders += ChildScope.SizeOfChildHeaders + VarInt.Measure(ChildScope.Length + ChildScope.SizeOfChildHeaders + ArrayCountLength) + ArrayCountLength;
+						break;
+					default:
+						throw new InvalidOperationException();
 				}
 			}
 			return SizeOfChildHeaders;
