Horde: Reduce memory usage when capturing workspace (part 2)

* Added UTF8 native version of fstat record to avoid unnecessary conversions
* Avoid conversion to/from UTF8 strings when converting Perforce records into CB objects
* Fixed implicit conversion of string to UTF8 string in argument check
* Converted ViewMap to use UTF8 strings

diff --git a/Engine/Source/Programs/Horde/HordeServer/Commits/Impl/CommitService.cs b/Engine/Source/Programs/Horde/HordeServer/Commits/Impl/CommitService.cs
--- a/Engine/Source/Programs/Horde/HordeServer/Commits/Impl/CommitService.cs
+++ b/Engine/Source/Programs/Horde/HordeServer/Commits/Impl/CommitService.cs
@@ -173,7 +173,7 @@ namespace HordeServer.Commits.Impl
 			}
 
 			// Find all the depot roots
-			Dictionary<string, string> DepotRoots = new Dictionary<string, string>(ServerInfo.PathComparer);
+			Dictionary<Utf8String, Utf8String> DepotRoots = new Dictionary<Utf8String, Utf8String>(ServerInfo.Utf8PathComparer);
 			foreach (ViewMap View in StreamInfoList.Select(x => x.View))
 			{
 				foreach (ViewMapEntry Entry in View.Entries)
@@ -181,9 +181,9 @@ namespace HordeServer.Commits.Impl
 					if (Entry.SourcePrefix.Length >= 3)
 					{
 						int SlashIdx = Entry.SourcePrefix.IndexOf('/', 2);
-						string Depot = Entry.SourcePrefix.Substring(0, SlashIdx + 1);
+						Utf8String Depot = Entry.SourcePrefix.Slice(0, SlashIdx + 1);
 
-						string? DepotRoot;
+						Utf8String DepotRoot;
 						if (DepotRoots.TryGetValue(Depot, out DepotRoot))
 						{
 							DepotRoot = GetCommonPrefix(DepotRoot, Entry.SourcePrefix);
@@ -194,7 +194,7 @@ namespace HordeServer.Commits.Impl
 						}
 
 						int LastSlashIdx = DepotRoot.LastIndexOf('/');
-						DepotRoots[Depot] = DepotRoot.Substring(0, LastSlashIdx + 1);
+						DepotRoots[Depot] = DepotRoot.Slice(0, LastSlashIdx + 1);
 					}
 				}
 			}
@@ -208,7 +208,7 @@ namespace HordeServer.Commits.Impl
 				int FirstChange = await Changes.GetByIndexAsync(-1);
 				if (FirstChange == 0)
 				{
-					FirstChange = await GetFirstCommitToReplicateAsync(Connection, StreamInfo.View, ServerInfo.PathComparison);
+					FirstChange = await GetFirstCommitToReplicateAsync(Connection, StreamInfo.View, ServerInfo.Utf8PathComparer);
 				}
 				else
 				{
@@ -223,7 +223,7 @@ namespace HordeServer.Commits.Impl
 
 			// Find all the changes to consider
 			SortedSet<int> ChangeNumbers = new SortedSet<int>();
-			foreach (string DepotRoot in DepotRoots.Values)
+			foreach (Utf8String DepotRoot in DepotRoots.Values)
 			{
 				List<ChangesRecord> Changes = await Connection.GetChangesAsync(ChangesOptions.None, null, MinChange, -1, ChangeStatus.Submitted, null, $"{DepotRoot}...");
 				ChangeNumbers.UnionWith(Changes.Select(x => x.Number));
@@ -238,10 +238,10 @@ namespace HordeServer.Commits.Impl
 				{
 					IStream Stream = StreamInfo.Stream;
 
-					string? BasePath = GetBasePath(DescribeRecord, StreamInfo.View, ServerInfo.PathComparison);
-					if (BasePath != null)
+					Utf8String BasePath = GetBasePath(DescribeRecord, StreamInfo.View, ServerInfo.Utf8PathComparer);
+					if (!BasePath.IsEmpty)
 					{
-						await AddCommitAsync(Stream, DescribeRecord, BasePath);
+						await AddCommitAsync(Stream, DescribeRecord, BasePath.ToString());
 
 						RedisList<int> StreamCommitsKey = RedisStreamChanges(Stream.Id);
 						await StreamCommitsKey.RightPushAsync(DescribeRecord.Number);
@@ -258,14 +258,14 @@ namespace HordeServer.Commits.Impl
 		/// </summary>
 		/// <param name="Connection"></param>
 		/// <param name="View"></param>
-		/// <param name="Comparison"></param>
+		/// <param name="Comparer"></param>
 		/// <returns></returns>
-		static async Task<int> GetFirstCommitToReplicateAsync(IPerforceConnection Connection, ViewMap View, StringComparison Comparison)
+		static async Task<int> GetFirstCommitToReplicateAsync(IPerforceConnection Connection, ViewMap View, Utf8StringComparer Comparer)
 		{
 			int MinChange = 0;
 
-			List<string> RootPaths = View.GetRootPaths(Comparison);
-			foreach (string RootPath in RootPaths)
+			List<Utf8String> RootPaths = View.GetRootPaths(Comparer);
+			foreach (Utf8String RootPath in RootPaths)
 			{
 				IList<ChangesRecord> PathChanges = await Connection.GetChangesAsync(ChangesOptions.None, 20, ChangeStatus.Submitted, $"{RootPath}...");
 				MinChange = Math.Max(MinChange, PathChanges.Min(x => x.Number));
@@ -279,17 +279,17 @@ namespace HordeServer.Commits.Impl
 		/// </summary>
 		/// <param name="Changelist">Information about the changelist</param>
 		/// <param name="View">Mapping from depot syntax to client view</param>
-		/// <param name="Comparison">Path comparison type for the server</param>
+		/// <param name="Comparer">Path comparison type for the server</param>
 		/// <returns>The base path for all files in the change</returns>
-		static string? GetBasePath(DescribeRecord Changelist, ViewMap View, StringComparison Comparison)
+		static Utf8String GetBasePath(DescribeRecord Changelist, ViewMap View, Utf8StringComparer Comparer)
 		{
-			string? BasePath = null;
+			Utf8String BasePath = default;
 			foreach (DescribeFileRecord File in Changelist.Files)
 			{
-				string? StreamFile;
-				if (View.TryMapFile(File.DepotFile, Comparison, out StreamFile))
+				Utf8String StreamFile;
+				if (View.TryMapFile(File.DepotFile, Comparer, out StreamFile))
 				{
-					if (BasePath == null)
+					if (BasePath.IsEmpty)
 					{
 						BasePath = StreamFile;
 					}
@@ -308,7 +308,7 @@ namespace HordeServer.Commits.Impl
 		/// <param name="A"></param>
 		/// <param name="B"></param>
 		/// <returns></returns>
-		static string GetCommonPrefix(string A, string B)
+		static Utf8String GetCommonPrefix(Utf8String A, Utf8String B)
 		{
 			int Index = 0;
 			while (Index < A.Length && Index < B.Length && A[Index] == B[Index])
@@ -593,7 +593,7 @@ namespace HordeServer.Commits.Impl
 					}
 					else if (Commit.TreeRef == null)
 					{
-						TreeRef = await AddTreeRefAsync(Perforce, Stream, View, null, null, Commit, ServerInfo.PathComparison);
+						TreeRef = await AddTreeRefAsync(Perforce, Stream, View, null, null, Commit, ServerInfo.Utf8PathComparer);
 					}
 					else
 					{
@@ -613,7 +613,7 @@ namespace HordeServer.Commits.Impl
 				}
 				else
 				{
-					TreeRef = await AddTreeRefAsync(Perforce, Stream, View, PrevCommit, PrevTreeRef, Commit, ServerInfo.PathComparison);
+					TreeRef = await AddTreeRefAsync(Perforce, Stream, View, PrevCommit, PrevTreeRef, Commit, ServerInfo.Utf8PathComparer);
 				}
 
 				// Remove the first change number from this queue
@@ -637,18 +637,18 @@ namespace HordeServer.Commits.Impl
 			return CbSerializer.Deserialize<StreamTreeRef>(Ref.Value.AsField());
 		}
 
-		async Task<StreamTreeRef?> AddTreeRefAsync(IPerforceConnection Perforce, IStream Stream, ViewMap View, ICommit? PrevCommit, StreamTreeRef? PrevTreeRef, ICommit Commit, StringComparison PathComparison)
+		async Task<StreamTreeRef?> AddTreeRefAsync(IPerforceConnection Perforce, IStream Stream, ViewMap View, ICommit? PrevCommit, StreamTreeRef? PrevTreeRef, ICommit Commit, Utf8StringComparer PathComparer)
 		{
 			StreamTreeRef TreeRef;
 			if (PrevCommit == null || PrevTreeRef == null)
 			{
 				Logger.LogInformation("Performing snapshot of {StreamId} at CL {Change}", Commit.StreamId, Commit.Change);
-				TreeRef = await CreateTreeAsync(Perforce, View, Commit.Change, PathComparison);
+				TreeRef = await CreateTreeAsync(Perforce, View, Commit.Change, PathComparer);
 			}
 			else
 			{
 				Logger.LogInformation("Performing incremental update of {StreamId} from CL {PrevChange} to CL {Change}", Commit.StreamId, PrevCommit.Change, Commit.Change);
-				TreeRef = await UpdateTreeAsync(PrevTreeRef, Perforce, View, Commit.Change, PathComparison);
+				TreeRef = await UpdateTreeAsync(PrevTreeRef, Perforce, View, Commit.Change, PathComparer);
 			}
 
 			string RefName = $"tree_{Commit.StreamId}_{Commit.Change}";
@@ -682,8 +682,8 @@ namespace HordeServer.Commits.Impl
 		/// <param name="Connection">The Perforce repository to mirror</param>
 		/// <param name="View">View for the stream</param>
 		/// <param name="Change">Changelist to replicate</param>
-		/// <param name="Comparison">Comparison type for paths</param>
-		async Task<StreamTreeRef> CreateTreeAsync(IPerforceConnection Connection, ViewMap View, int Change, StringComparison Comparison)
+		/// <param name="PathComparer">Comparison type for paths</param>
+		async Task<StreamTreeRef> CreateTreeAsync(IPerforceConnection Connection, ViewMap View, int Change, Utf8StringComparer PathComparer)
 		{
 			// Optimize the view by removing unnecessary inclusions of the same path
 			View = new ViewMap(View);
@@ -692,7 +692,7 @@ namespace HordeServer.Commits.Impl
 				ViewMapEntry Entry = View.Entries[Idx];
 				if (Entry.Include && Entry.IsPathWildcard())
 				{
-					while (Idx + 1 < View.Entries.Count && View.Entries[Idx + 1].Include && View.Entries[Idx + 1].SourcePrefix.StartsWith(Entry.SourcePrefix, Comparison))
+					while (Idx + 1 < View.Entries.Count && View.Entries[Idx + 1].Include && View.Entries[Idx + 1].SourcePrefix.StartsWith(Entry.SourcePrefix, PathComparer))
 					{
 						View.Entries.RemoveAt(Idx + 1);
 					}
@@ -722,7 +722,7 @@ namespace HordeServer.Commits.Impl
 						ViewMapEntry OtherEntry = View.Entries[OtherIdx];
 						if (!OtherEntry.Include)
 						{
-							if (OtherEntry.SourcePrefix.StartsWith(Entry.SourcePrefix, Comparison) || Entry.SourcePrefix.StartsWith(OtherEntry.SourcePrefix, Comparison))
+							if (OtherEntry.SourcePrefix.StartsWith(Entry.SourcePrefix, PathComparer) || Entry.SourcePrefix.StartsWith(OtherEntry.SourcePrefix, PathComparer))
 							{
 								FilteredView.Entries.Add(OtherEntry);
 							}
@@ -730,7 +730,7 @@ namespace HordeServer.Commits.Impl
 					}
 
 					Logger.LogInformation("Adding {Source}@{Change}", Entry.Source, Change);
-					await AddDepotFilesToSnapshotAsync(Connection, $"{Entry.Source}@{Change}", View, Comparison, TreePath, Writer, WriteObject);
+					await AddDepotFilesToSnapshotAsync(Connection, $"{Entry.Source}@{Change}", View, PathComparer, TreePath, Writer, WriteObject);
 
 					GC.Collect();
 				}
@@ -740,6 +740,30 @@ namespace HordeServer.Commits.Impl
 			return TreePath[0].Encode(Writer, WriteObject);
 		}
 
+		/// <summary>
+		/// Record returned by the fstat command
+		/// </summary>
+		public class Utf8FStatRecord
+		{
+			[PerforceTag("depotFile", Optional = true)]
+			public Utf8String DepotFile;
+
+			[PerforceTag("headType", Optional = true)]
+			public Utf8String HeadType;
+
+			[PerforceTag("headRev", Optional = true)]
+			public int HeadRevision;
+
+			[PerforceTag("digest", Optional = true)]
+			public Utf8String Digest;
+
+			[PerforceTag("fileSize", Optional = true)]
+			public long FileSize;
+
+			[PerforceTag("type", Optional = true)]
+			public Utf8String Type;
+		}
+
 		/// <summary>
 		/// Adds a set of files in the depot to a snapshot
 		/// </summary>
@@ -750,7 +774,7 @@ namespace HordeServer.Commits.Impl
 		/// <param name="TreePath">The last computed path through the tree</param>
 		/// <param name="Writer"></param>
 		/// <param name="WriteObject">Delegate to write an object to the store</param>
-		async Task AddDepotFilesToSnapshotAsync(IPerforceConnection Connection, string FileSpec, ViewMap View, StringComparison Comparison, List<StreamTreeBuilder> TreePath, CbWriter Writer, Action<IoHash, CbObject> WriteObject)
+		async Task AddDepotFilesToSnapshotAsync(IPerforceConnection Connection, string FileSpec, ViewMap View, Utf8StringComparer Comparer, List<StreamTreeBuilder> TreePath, CbWriter Writer, Action<IoHash, CbObject> WriteObject)
 		{
 			const string Filter = "^headAction=delete ^headAction=move/delete";
 			const string Fields = "depotFile,fileSize,digest,headType,headRev";
@@ -765,16 +789,18 @@ namespace HordeServer.Commits.Impl
 			Arguments.Add(String.Join(",", Fields));
 			Arguments.Add(FileSpec);
 
-			IAsyncEnumerable<PerforceResponse<FStatRecord>> Responses = Connection.StreamCommandAsync<FStatRecord>("fstat", Arguments);
-			await foreach (PerforceResponse<FStatRecord> Response in Responses)
+			List<Utf8String> Segments = new List<Utf8String>();
+
+			IAsyncEnumerable<PerforceResponse<Utf8FStatRecord>> Responses = Connection.StreamCommandAsync<Utf8FStatRecord>("fstat", Arguments);
+			await foreach (PerforceResponse<Utf8FStatRecord> Response in Responses)
 			{
-				FStatRecord File = Response.Data;
-				if (!String.IsNullOrEmpty(File.Digest))
+				Utf8FStatRecord File = Response.Data;
+				if (File.Digest != Utf8String.Empty)
 				{
-					string? TargetFile;
-					if (File.DepotFile != null && View.TryMapFile(File.DepotFile, Comparison, out TargetFile))
+					Utf8String TargetFile;
+					if (File.DepotFile != Utf8String.Empty && View.TryMapFile(File.DepotFile, Comparer, out TargetFile))
 					{
-						await AddFileAsync(TargetFile, File, TreePath, Writer, WriteObject);
+						await AddFileAsync(TargetFile, File, TreePath, Segments, Writer, WriteObject);
 					}
 				}
 			}
@@ -786,18 +812,33 @@ namespace HordeServer.Commits.Impl
 		/// <param name="Path">The stream path for the file</param>
 		/// <param name="MetaData">The file metadata</param>
 		/// <param name="TreePath">The last computed path through the tree</param>
+		/// <param name="Segments"></param>
 		/// <param name="Writer"></param>
 		/// <param name="WriteObject"></param>
-		async Task AddFileAsync(string Path, FStatRecord MetaData, List<StreamTreeBuilder> TreePath, CbWriter Writer, Action<IoHash, CbObject> WriteObject)
+		async Task AddFileAsync(Utf8String Path, Utf8FStatRecord MetaData, List<StreamTreeBuilder> TreePath, List<Utf8String> Segments, CbWriter Writer, Action<IoHash, CbObject> WriteObject)
 		{
-			string[] Segments = Path.Split('/');
-
 			StreamTreeBuilder Node = TreePath[0];
 
-			int SegmentIdx = 0;
+			// Split the path into segments
+			Segments.Clear();
+			for (int Pos = 0; ;)
+			{
+				int NextPos = Path.IndexOf('/', Pos);
+				if (NextPos == -1)
+				{
+					Segments.Add(Path.Substring(Pos));
+					break;
+				}
+				else
+				{
+					Segments.Add(Path.Substring(Pos, NextPos - Pos));
+					Pos = NextPos + 1;
+				}
+			}
 
 			// Match as many nodes as we can from the existing path
-			for (; SegmentIdx < Segments.Length - 1; SegmentIdx++)
+			int SegmentIdx = 0;
+			for(; SegmentIdx < Segments.Count; SegmentIdx++)
 			{
 				Utf8String Segment = Segments[SegmentIdx];
 				if (!Node.NameToTreeBuilder.TryGetValue(Segment, out StreamTreeBuilder? NextNode))
@@ -811,12 +852,15 @@ namespace HordeServer.Commits.Impl
 			CollapseTree(Node, Writer, WriteObject);
 
 			// Expand the rest of the path
-			for (; SegmentIdx < Segments.Length - 1; SegmentIdx++)
+			for(; SegmentIdx < Segments.Count; SegmentIdx++)
 			{
-				Node = await FindOrAddTreeAsync(Node, Segments[SegmentIdx]);
+				Utf8String Segment = Segments[SegmentIdx];
+				Node = await FindOrAddTreeAsync(Node, Segment);
 			}
 
-			Node.NameToFile[Segments[Segments.Length - 1]] = CreateStreamFile(MetaData);
+			// Add the filename
+			Utf8String FileName = Segments[Segments.Count - 1];
+			Node.NameToFile[FileName] = CreateStreamFile(MetaData);
 		}
 
 		#endregion
@@ -830,12 +874,12 @@ namespace HordeServer.Commits.Impl
 		/// <param name="Perforce"></param>
 		/// <param name="Change"></param>
 		/// <param name="View"></param>
-		/// <param name="Comparison"></param>
+		/// <param name="PathComparer"></param>
 		/// <returns></returns>
-		async Task<StreamTreeRef> UpdateTreeAsync(StreamTreeRef PrevTreeRef, IPerforceConnection Perforce, ViewMap View, int Change, StringComparison Comparison)
+		async Task<StreamTreeRef> UpdateTreeAsync(StreamTreeRef PrevTreeRef, IPerforceConnection Perforce, ViewMap View, int Change, Utf8StringComparer PathComparer)
 		{
 			List<FStatRecord> Records = await Perforce.FStatAsync(-1, Change, null, null, -1, FStatOptions.IncludeFileSizes, FileSpecList.Empty);
-			return await UpdateTreeAsync(PrevTreeRef, Records, View, Comparison);
+			return await UpdateTreeAsync(PrevTreeRef, Records, View, PathComparer);
 		}
 
 		/// <summary>
@@ -844,9 +888,9 @@ namespace HordeServer.Commits.Impl
 		/// <param name="PrevTreeRef"></param>
 		/// <param name="Files"></param>
 		/// <param name="View"></param>
-		/// <param name="Comparison"></param>
+		/// <param name="PathComparer"></param>
 		/// <returns></returns>
-		async Task<StreamTreeRef> UpdateTreeAsync(StreamTreeRef PrevTreeRef, IList<FStatRecord> Files, ViewMap View, StringComparison Comparison)
+		async Task<StreamTreeRef> UpdateTreeAsync(StreamTreeRef PrevTreeRef, IList<FStatRecord> Files, ViewMap View, Utf8StringComparer PathComparer)
 		{
 			StreamTreeBuilder Tree = await ReadTreeAsync(PrevTreeRef);
 
@@ -856,8 +900,8 @@ namespace HordeServer.Commits.Impl
 			// Update the tree
 			foreach (FStatRecord File in Files)
 			{
-				string? LocalPathStr;
-				if (File.DepotFile != null && View.TryMapFile(File.DepotFile, Comparison, out LocalPathStr))
+				Utf8String LocalPathStr;
+				if (File.DepotFile != null && View.TryMapFile(File.DepotFile, PathComparer, out LocalPathStr))
 				{
 					Utf8String LocalPath = LocalPathStr;
 
@@ -985,6 +1029,25 @@ namespace HordeServer.Commits.Impl
 			Node.NameToTreeBuilder.Clear();
 		}
 
+		/// <summary>
+		/// Creates a StreamFile from a FileMetaData object
+		/// </summary>
+		/// <param name="MetaData"></param>
+		/// <returns></returns>
+		static StreamFile CreateStreamFile(Utf8FStatRecord MetaData)
+		{
+			Utf8String FileType = Utf8String.Empty;
+			if (MetaData.Type != Utf8String.Empty)
+			{
+				FileType = MetaData.Type;
+			}
+			else if (MetaData.HeadType != Utf8String.Empty)
+			{
+				FileType = MetaData.HeadType;
+			}
+			return new StreamFile(MetaData.DepotFile!, MetaData.FileSize, new FileContentId(Md5Hash.Parse(MetaData.Digest!), FileType), MetaData.HeadRevision);
+		}
+
 		/// <summary>
 		/// Creates a StreamFile from a FileMetaData object
 		/// </summary>
diff --git a/Engine/Source/Programs/Horde/HordeServer/Utilities/ViewMap.cs b/Engine/Source/Programs/Horde/HordeServer/Utilities/ViewMap.cs
--- a/Engine/Source/Programs/Horde/HordeServer/Utilities/ViewMap.cs
+++ b/Engine/Source/Programs/Horde/HordeServer/Utilities/ViewMap.cs
@@ -1,9 +1,11 @@
 // Copyright Epic Games, Inc. All Rights Reserved.
 
+using EpicGames.Core;
 using System;
 using System.Collections.Generic;
 using System.Diagnostics;
 using System.Diagnostics.CodeAnalysis;
+using System.Text;
 using System.Text.RegularExpressions;
 
 namespace HordeServer.Utilities
@@ -18,13 +20,14 @@ namespace HordeServer.Utilities
 		/// <summary>
 		/// List of entries making up the view
 		/// </summary>
-		public List<ViewMapEntry> Entries { get; } = new List<ViewMapEntry>();
+		public List<ViewMapEntry> Entries { get; }
 
 		/// <summary>
 		/// Default constructor
 		/// </summary>
 		public ViewMap()
 		{
+			Entries = new List<ViewMapEntry>();
 		}
 
 		/// <summary>
@@ -33,7 +36,7 @@ namespace HordeServer.Utilities
 		/// <param name="Other"></param>
 		public ViewMap(ViewMap Other)
 		{
-			Entries.AddRange(Other.Entries);
+			Entries = new List<ViewMapEntry>(Other.Entries);
 		}
 
 		/// <summary>
@@ -41,6 +44,7 @@ namespace HordeServer.Utilities
 		/// </summary>
 		/// <param name="ViewMap">Peforce viewmap definition</param>
 		public ViewMap(P4.ViewMap ViewMap)
+			: this()
 		{
 			foreach (P4.MapEntry Entry in ViewMap)
 			{
@@ -54,7 +58,7 @@ namespace HordeServer.Utilities
 		/// <param name="File">The file to test</param>
 		/// <param name="Comparison">The comparison type</param>
 		/// <returns>True if the file is included in the view</returns>
-		public bool MatchFile(string File, StringComparison Comparison)
+		public bool MatchFile(Utf8String File, Utf8StringComparer Comparison)
 		{
 			bool Included = false;
 			foreach (ViewMapEntry Entry in Entries)
@@ -74,7 +78,7 @@ namespace HordeServer.Utilities
 		/// <param name="Comparison">The comparison type</param>
 		/// <param name="TargetFile"></param>
 		/// <returns></returns>
-		public bool TryMapFile(string SourceFile, StringComparison Comparison, [NotNullWhen(true)] out string? TargetFile)
+		public bool TryMapFile(Utf8String SourceFile, Utf8StringComparer Comparison, out Utf8String TargetFile)
 		{
 			ViewMapEntry? MapEntry = null;
 			foreach (ViewMapEntry Entry in Entries)
@@ -92,82 +96,24 @@ namespace HordeServer.Utilities
 			}
 			else
 			{
-				TargetFile = null;
+				TargetFile = Utf8String.Empty;
 				return false;
 			}
 		}
 
-		/// <summary>
-		/// Tests whether the view includes any files directly under the given directory (not including any subdirectories).
-		/// </summary>
-		/// <param name="Dir">The directory to test; should end with a Slash.</param>
-		/// <param name="Comparison">The comparison type</param>
-		/// <returns>True if the view may include any files in the given directory</returns>
-		public bool MayMatchAnyFilesInDirectory(string Dir, StringComparison Comparison)
-		{
-			bool MayInclude = false;
-			foreach (ViewMapEntry Entry in Entries)
-			{
-				if (Entry.Include)
-				{
-					if (Entry.MayMatchAnyFilesInDirectory(Dir, Comparison))
-					{
-						MayInclude = true;
-					}
-				}
-				else
-				{
-					if (Entry.MatchAllFilesInDirectory(Dir, Comparison))
-					{
-						MayInclude = false;
-					}
-				}
-			}
-			return MayInclude;
-		}
-
-		/// <summary>
-		/// Tests whether the view may include any files in subdirectories of the given path
-		/// </summary>
-		/// <param name="Dir">The directory to test; should end with a Slash.</param>
-		/// <param name="Comparison">The comparison type</param>
-		/// <returns>True if the pattern may match files under the given directory</returns>
-		public bool MayMatchAnyFilesInSubDirectory(string Dir, StringComparison Comparison)
-		{
-			bool MayInclude = false;
-			foreach (ViewMapEntry Entry in Entries)
-			{
-				if (Entry.Include)
-				{
-					if (Entry.MayMatchAnyFilesInSubDirectory(Dir, Comparison))
-					{
-						MayInclude = true;
-					}
-				}
-				else
-				{
-					if (Entry.MatchAllFilesInDirectory(Dir, Comparison))
-					{
-						MayInclude = false;
-					}
-				}
-			}
-			return MayInclude;
-		}
-
 		/// <summary>
 		/// Gets the root paths from the view entries
 		/// </summary>
 		/// <returns></returns>
-		public List<string> GetRootPaths(StringComparison Comparison)
+		public List<Utf8String> GetRootPaths(Utf8StringComparer Comparison)
 		{
-			List<string> RootPaths = new List<string>();
+			List<Utf8String> RootPaths = new List<Utf8String>();
 			foreach (ViewMapEntry Entry in Entries)
 			{
 				if (Entry.Include)
 				{
 					int LastSlashIdx = Entry.SourcePrefix.LastIndexOf('/');
-					string RootPath = Entry.SourcePrefix.Substring(0, LastSlashIdx + 1);
+					Utf8String RootPath = Entry.SourcePrefix.Slice(0, LastSlashIdx + 1);
 
 					for (int Idx = 0; ; Idx++)
 					{
@@ -205,49 +151,49 @@ namespace HordeServer.Utilities
 		/// <summary>
 		/// The wildcard string - either '*' or '...'
 		/// </summary>
-		public string? Wildcard { get; }
+		public Utf8String Wildcard { get; }
 
 		/// <summary>
 		/// The source part of the pattern before the wildcard
 		/// </summary>
-		public string SourcePrefix { get; }
+		public Utf8String SourcePrefix { get; }
 
 		/// <summary>
 		/// The source part of the pattern after the wildcard. Perforce does not permit a slash to be in this part of the mapping.
 		/// </summary>
-		public string SourceSuffix { get; }
+		public Utf8String SourceSuffix { get; }
 
 		/// <summary>
 		/// The target mapping for the pattern before the wildcard
 		/// </summary>
-		public string TargetPrefix { get; }
+		public Utf8String TargetPrefix { get; }
 
 		/// <summary>
 		/// The target mapping for the pattern after the wildcard
 		/// </summary>
-		public string TargetSuffix { get; }
+		public Utf8String TargetSuffix { get; }
 
 		/// <summary>
-		/// The full source path
+		/// The full source pattern
 		/// </summary>
-		public string Source => (Wildcard == null) ? SourcePrefix : $"{SourcePrefix}{Wildcard}{SourceSuffix}";
+		public Utf8String Source => $"{SourcePrefix}{Wildcard}{SourceSuffix}";
 
 		/// <summary>
-		/// The full target path
+		/// The full target pattern
 		/// </summary>
-		public string Target => (Wildcard == null) ? TargetPrefix : $"{TargetPrefix}{Wildcard}{TargetSuffix}";
+		public Utf8String Target => $"{TargetPrefix}{Wildcard}{TargetSuffix}";
 
 		/// <summary>
 		/// Tests if the entry has a file wildcard ('*')
 		/// </summary>
 		/// <returns>True if the entry has a file wildcard</returns>
-		public bool IsFileWildcard() => Wildcard != null && Wildcard.Length == 1;
+		public bool IsFileWildcard() => Wildcard.Length == 1;
 
 		/// <summary>
 		/// Tests if the entry has a path wildcard ('...')
 		/// </summary>
 		/// <returns>True if the entry has a path wildcard</returns>
-		public bool IsPathWildcard() => Wildcard != null && Wildcard.Length == 3;
+		public bool IsPathWildcard() => Wildcard.Length == 3;
 
 		/// <summary>
 		/// Constructor
@@ -268,23 +214,30 @@ namespace HordeServer.Utilities
 		{
 			this.Include = Include;
 
-			Match Match = Regex.Match(Source, @"^(.*)(\*|\.\.\.)(.*)$");
+			Match Match = Regex.Match(Source, @"^(.*)(\*|\.\.\.|%%1)(.*)$");
 			if (Match.Success)
 			{
+				string WildcardStr = Match.Groups[2].Value;
+
 				SourcePrefix = Match.Groups[1].Value;
 				SourceSuffix = Match.Groups[3].Value;
 				Wildcard = Match.Groups[2].Value;
 
-				int OtherIdx = Target.IndexOf(Wildcard, StringComparison.Ordinal);
+				int OtherIdx = Target.IndexOf(WildcardStr, StringComparison.Ordinal);
 				TargetPrefix = Target.Substring(0, OtherIdx);
 				TargetSuffix = Target.Substring(OtherIdx + Wildcard.Length);
+
+				if (WildcardStr.Equals("%%1", StringComparison.Ordinal))
+				{
+					Wildcard = "*";
+				}
 			}
 			else
 			{
 				SourcePrefix = Source;
-				SourceSuffix = String.Empty;
+				SourceSuffix = Utf8String.Empty;
 				TargetPrefix = Target;
-				TargetSuffix = String.Empty;
+				TargetSuffix = Utf8String.Empty;
 			}
 		}
 
@@ -297,7 +250,7 @@ namespace HordeServer.Utilities
 		/// <param name="SourceSuffix"></param>
 		/// <param name="TargetPrefix"></param>
 		/// <param name="TargetSuffix"></param>
-		public ViewMapEntry(bool Include, string? Wildcard, string SourcePrefix, string SourceSuffix, string TargetPrefix, string TargetSuffix)
+		public ViewMapEntry(bool Include, Utf8String Wildcard, Utf8String SourcePrefix, Utf8String SourceSuffix, Utf8String TargetPrefix, Utf8String TargetSuffix)
 		{
 			this.Include = Include;
 			this.Wildcard = Wildcard;
@@ -321,10 +274,10 @@ namespace HordeServer.Utilities
 		/// </summary>
 		/// <param name="SourceFile"></param>
 		/// <returns></returns>
-		public string MapFile(string SourceFile)
+		public Utf8String MapFile(Utf8String SourceFile)
 		{
 			int Count = SourceFile.Length - SourceSuffix.Length - SourcePrefix.Length;
-			return TargetPrefix + SourceFile.Substring(SourcePrefix.Length, Count) + TargetSuffix;
+			return TargetPrefix + SourceFile.Slice(SourcePrefix.Length, Count) + TargetSuffix;
 		}
 
 		/// <summary>
@@ -333,11 +286,11 @@ namespace HordeServer.Utilities
 		/// <param name="Path">Path to the file</param>
 		/// <param name="Comparison">The comparison type</param>
 		/// <returns>True if the path matches the entry</returns>
-		public bool MatchFile(string Path, StringComparison Comparison)
+		public bool MatchFile(Utf8String Path, Utf8StringComparer Comparison)
 		{
-			if (Wildcard == null)
+			if (Wildcard.Length == 0)
 			{
-				return Path.Equals(SourcePrefix, Comparison);
+				return Comparison.Compare(Path, SourcePrefix) == 0;
 			}
 			else
 			{
@@ -345,7 +298,7 @@ namespace HordeServer.Utilities
 				{
 					return false;
 				}
-				if (IsFileWildcard() && Path.IndexOf('/', SourcePrefix.Length, Path.Length - SourceSuffix.Length - SourcePrefix.Length) != -1)
+				if (IsFileWildcard() && Path.Slice(SourcePrefix.Length, Path.Length - SourceSuffix.Length - SourcePrefix.Length).IndexOf('/') != -1)
 				{
 					return false;
 				}
@@ -353,60 +306,16 @@ namespace HordeServer.Utilities
 			}
 		}
 
-		/// <summary>
-		/// Determine if the pattern may match any files directly under the given directory
-		/// </summary>
-		/// <param name="Dir">The directory to test</param>
-		/// <param name="Comparison">The comparison type</param>
-		/// <returns>True if the pattern may match a file in this directory</returns>
-		public bool MayMatchAnyFilesInDirectory(string Dir, StringComparison Comparison)
+		/// <inheritdoc/>
+		public override string ToString()
 		{
-			Debug.Assert(Dir.EndsWith("/", StringComparison.Ordinal));
-			if (Dir.StartsWith(SourcePrefix, Comparison))
+			StringBuilder Builder = new StringBuilder();
+			if (!Include)
 			{
-				if (IsPathWildcard() || Dir.IndexOf('/', SourcePrefix.Length) == -1)
-				{
-					return true;
-				}
+				Builder.Append("-");
 			}
-			return false;
-		}
-
-		/// <summary>
-		/// Determine if the pattern may match any files in subdirectories of the given directory
-		/// </summary>
-		/// <param name="Dir">The directory to test</param>
-		/// <param name="Comparison">The comparison type</param>
-		/// <returns>True if the pattern may match a file in this directory</returns>
-		public bool MayMatchAnyFilesInSubDirectory(string Dir, StringComparison Comparison)
-		{
-			Debug.Assert(Dir.EndsWith("/", StringComparison.Ordinal));
-			if (Dir.StartsWith(SourcePrefix, Comparison))
-			{
-				if (IsPathWildcard())
-				{
-					return true;
-				}
-			}
-			else
-			{
-				if (SourcePrefix.StartsWith(Dir, Comparison))
-				{
-					return true;
-				}
-			}
-			return false;
-		}
-
-		/// <summary>
-		/// Determine if the pattern matches all files in the given directory
-		/// </summary>
-		/// <param name="Dir">The directory to test</param>
-		/// <param name="Comparison">The comparison type</param>
-		/// <returns>True if the pattern may match a file in this directory</returns>
-		public bool MatchAllFilesInDirectory(string Dir, StringComparison Comparison)
-		{
-			return Wildcard != null && Wildcard.Length == 3 && Dir.StartsWith(SourcePrefix, Comparison);
+			Builder.Append($"{SourcePrefix}{Wildcard}{SourceSuffix} {TargetPrefix}{Wildcard}{TargetSuffix}");
+			return Builder.ToString();
 		}
 	}
 }
diff --git a/Engine/Source/Programs/Shared/EpicGames.Core/Utf8String.cs b/Engine/Source/Programs/Shared/EpicGames.Core/Utf8String.cs
--- a/Engine/Source/Programs/Shared/EpicGames.Core/Utf8String.cs
+++ b/Engine/Source/Programs/Shared/EpicGames.Core/Utf8String.cs
@@ -156,6 +156,24 @@ namespace EpicGames.Core
 			}
 		}
 
+		/// <inheritdoc cref="String.IndexOf(char, int)"/>
+		public int IndexOf(char Char, int Index) => IndexOf(Char, Index, Length - Index);
+
+		/// <inheritdoc cref="String.IndexOf(char, int, int)"/>
+		public int IndexOf(char Char, int Index, int Count)
+		{
+			int Result;
+			if (Char < 0x80)
+			{
+				Result = Span.Slice(Index, Count).IndexOf((byte)Char);
+			}
+			else
+			{
+				Result = Span.Slice(Index, Count).IndexOf(Encoding.UTF8.GetBytes(new[] { Char }));
+			}
+			return (Result == -1) ? -1 : Result + Index;
+		}
+
 		/// <inheritdoc cref="String.IndexOf(string)"/>
 		public int IndexOf(Utf8String String)
 		{
@@ -175,6 +193,25 @@ namespace EpicGames.Core
 			return -1;
 		}
 
+		/// <inheritdoc cref="String.LastIndexOf(char)"/>
+		public int LastIndexOf(byte Char)
+		{
+			return Span.IndexOf(Char);
+		}
+
+		/// <inheritdoc cref="String.LastIndexOf(char)"/>
+		public int LastIndexOf(char Char)
+		{
+			if (Char < 0x80)
+			{
+				return Span.IndexOf((byte)Char);
+			}
+			else
+			{
+				return Span.IndexOf(Encoding.UTF8.GetBytes(new[] { Char }));
+			}
+		}
+
 		/// <summary>
 		/// Tests if this string starts with another string
 		/// </summary>
@@ -217,26 +254,24 @@ namespace EpicGames.Core
 			return Length >= Other.Length && Comparer.Equals(Slice(Length - Other.Length), Other);
 		}
 
-		/// <summary>
-		/// Slices this string at the given start position
-		/// </summary>
-		/// <param name="Start">Start position</param>
-		/// <returns>String starting at the given position</returns>
-		public Utf8String Slice(int Start)
+		/// <inheritdoc cref="Substring(int)"/>
+		public Utf8String Slice(int Start) => Substring(Start);
+
+		/// <inheritdoc cref="Substring(int, int)"/>
+		public Utf8String Slice(int Start, int Count) => Substring(Start, Count);
+
+		/// <inheritdoc cref="String.Substring(int)"/>
+		public Utf8String Substring(int Start)
 		{
 			return new Utf8String(Memory.Slice(Start));
 		}
 
-		/// <summary>
-		/// Slices this string at the given start position and length
-		/// </summary>
-		/// <param name="Start">Start position</param>
-		/// <returns>String starting at the given position</returns>
-		public Utf8String Slice(int Start, int Count)
+		/// <inheritdoc cref="String.Substring(int, int)"/>
+		public Utf8String Substring(int Start, int Count)
 		{
 			return new Utf8String(Memory.Slice(Start, Count));
 		}
-
+		
 		/// <summary>
 		/// Tests if this string is equal to the other object
 		/// </summary>
@@ -305,6 +340,15 @@ namespace EpicGames.Core
 		/// <returns></returns>
 		public static Utf8String operator +(Utf8String A, Utf8String B)
 		{
+			if (A.Length == 0)
+			{
+				return B;
+			}
+			if (B.Length == 0)
+			{
+				return A;
+			}
+
 			byte[] Buffer = new byte[A.Length + B.Length];
 			A.Span.CopyTo(Buffer);
 			B.Span.CopyTo(Buffer.AsSpan(A.Length));
diff --git a/Engine/Source/Programs/Shared/EpicGames.Perforce.Managed/StreamTree.cs b/Engine/Source/Programs/Shared/EpicGames.Perforce.Managed/StreamTree.cs
--- a/Engine/Source/Programs/Shared/EpicGames.Perforce.Managed/StreamTree.cs
+++ b/Engine/Source/Programs/Shared/EpicGames.Perforce.Managed/StreamTree.cs
@@ -326,7 +326,7 @@ namespace EpicGames.Perforce.Managed
 		/// <returns></returns>
 		public CbObject ToCbObject(CbWriter Writer, Utf8String BasePath)
 		{
-			if (BasePath.EndsWith("/"))
+			if (BasePath.Length > 0 && BasePath[BasePath.Length - 1] == '/')
 			{
 				throw new ArgumentException("BasePath must not end in a slash", nameof(BasePath));
 			}
diff --git a/Engine/Source/Programs/Shared/EpicGames.Perforce/IPerforceOutput.cs b/Engine/Source/Programs/Shared/EpicGames.Perforce/IPerforceOutput.cs
--- a/Engine/Source/Programs/Shared/EpicGames.Perforce/IPerforceOutput.cs
+++ b/Engine/Source/Programs/Shared/EpicGames.Perforce/IPerforceOutput.cs
@@ -2,6 +2,7 @@
 
 using EpicGames.Core;
 using System;
+using System.Buffers.Binary;
 using System.Buffers.Text;
 using System.Collections.Generic;
 using System.Diagnostics.CodeAnalysis;
@@ -185,7 +186,7 @@ namespace EpicGames.Perforce
 				for (; ; )
 				{
 					int NewBufferPos = BufferPos;
-					if (!TryReadResponse(Data.Span, ref NewBufferPos, StatRecordInfo, out PerforceResponse? Response))
+					if (!TryReadResponse(Data, ref NewBufferPos, StatRecordInfo, out PerforceResponse? Response))
 					{
 						MaxParsedLen = ParsedLen + NewBufferPos;
 						break;
@@ -243,7 +244,7 @@ namespace EpicGames.Perforce
 				for (; ; )
 				{
 					int NewBufferPos = BufferPos;
-					if (!TryReadResponse(Data.Span, ref NewBufferPos, StatRecordInfo, out PerforceResponse? Response))
+					if (!TryReadResponse(Data, ref NewBufferPos, StatRecordInfo, out PerforceResponse? Response))
 					{
 						MaxParsedLen = ParsedLen + NewBufferPos;
 						break;
@@ -296,7 +297,7 @@ namespace EpicGames.Perforce
 				for (; ; )
 				{
 					int InitialBufferPos = BufferPos;
-					if (!ReadRecord(Data.Span, ref BufferPos, Record.Rows))
+					if (!ReadRecord(Data, ref BufferPos, Record.Rows))
 					{
 						BufferPos = InitialBufferPos;
 						break;
@@ -320,16 +321,17 @@ namespace EpicGames.Perforce
 		/// <param name="BufferPos">Current read position within the buffer</param>
 		/// <param name="Rows">List of rows to read into</param>
 		/// <returns>True if a record could be read; false if more data is required</returns>
-		static bool ReadRecord(ReadOnlySpan<byte> Buffer, ref int BufferPos, List<KeyValuePair<Utf8String, PerforceValue>> Rows)
+		static bool ReadRecord(ReadOnlyMemory<byte> Buffer, ref int BufferPos, List<KeyValuePair<Utf8String, PerforceValue>> Rows)
 		{
 			Rows.Clear();
+			ReadOnlySpan<byte> BufferSpan = Buffer.Span;
 
 			// Check we can read the initial record marker
 			if (BufferPos >= Buffer.Length)
 			{
 				return false;
 			}
-			if (Buffer[BufferPos] != '{')
+			if (BufferSpan[BufferPos] != '{')
 			{
 				throw new PerforceException("Invalid record start");
 			}
@@ -346,14 +348,14 @@ namespace EpicGames.Perforce
 				}
 
 				// If this is the end of the record, break out
-				byte KeyType = Buffer[BufferPos++];
+				byte KeyType = Buffer.Span[BufferPos++];
 				if (KeyType == '0')
 				{
 					break;
 				}
 				else if (KeyType != 's')
 				{
-					throw new PerforceException("Unexpected key field type while parsing marshalled output ({0}) - expected 's', got: {1}", (int)KeyType, FormatDataAsHexDump(Buffer.Slice(BufferPos - 1)));
+					throw new PerforceException("Unexpected key field type while parsing marshalled output ({0}) - expected 's', got: {1}", (int)KeyType, FormatDataAsHexDump(Buffer.Slice(BufferPos - 1).Span));
 				}
 
 				// Read the tag
@@ -368,7 +370,7 @@ namespace EpicGames.Perforce
 
 				// Read the value type
 				byte ValueType;
-				if (!TryReadByte(Buffer, ref BufferPos, out ValueType))
+				if (!TryReadByte(BufferSpan, ref BufferPos, out ValueType))
 				{
 					return false;
 				}
@@ -387,7 +389,7 @@ namespace EpicGames.Perforce
 				else if (ValueType == 'i')
 				{
 					int Integer;
-					if (!TryReadInt(Buffer, ref BufferPos, out Integer))
+					if (!TryReadInt(BufferSpan, ref BufferPos, out Integer))
 					{
 						return false;
 					}
@@ -399,7 +401,7 @@ namespace EpicGames.Perforce
 				}
 
 				// Construct the response object with the record
-				Rows.Add(KeyValuePair.Create(Key, Value));
+				Rows.Add(KeyValuePair.Create(Key.Clone(), Value));
 			}
 			return true;
 		}
@@ -412,7 +414,7 @@ namespace EpicGames.Perforce
 		/// <param name="StatRecordInfo">The type of record expected to parse from the response</param>
 		/// <param name="Response">Receives the response object on success</param>
 		/// <returns>True if a response was read, false if the buffer needs more data</returns>
-		static bool TryReadResponse(ReadOnlySpan<byte> Buffer, ref int BufferPos, CachedRecordInfo? StatRecordInfo, [NotNullWhen(true)] out PerforceResponse? Response)
+		static bool TryReadResponse(ReadOnlyMemory<byte> Buffer, ref int BufferPos, CachedRecordInfo? StatRecordInfo, [NotNullWhen(true)] out PerforceResponse? Response)
 		{
 			if (BufferPos + RecordPrefix.Length + 4 > Buffer.Length)
 			{
@@ -420,8 +422,8 @@ namespace EpicGames.Perforce
 				return false;
 			}
 
-			ReadOnlySpan<byte> Prefix = Buffer.Slice(BufferPos, RecordPrefix.Length);
-			if (!Prefix.SequenceEqual(RecordPrefix))
+			ReadOnlyMemory<byte> Prefix = Buffer.Slice(BufferPos, RecordPrefix.Length);
+			if (!Prefix.Span.SequenceEqual(RecordPrefix))
 			{
 				throw new PerforceException("Expected 'code' field at the start of record");
 			}
@@ -466,7 +468,7 @@ namespace EpicGames.Perforce
 			}
 
 			// Skip over the record terminator
-			if (BufferPos >= Buffer.Length || Buffer[BufferPos] != '0')
+			if (BufferPos >= Buffer.Length || Buffer.Span[BufferPos] != '0')
 			{
 				throw new PerforceException("Unexpected record terminator");
 			}
@@ -486,7 +488,7 @@ namespace EpicGames.Perforce
 		/// <param name="RecordInfo">Reflection information for the type being serialized into.</param>
 		/// <param name="Record">Receives the record on success</param>
 		/// <returns>The parsed object.</returns>
-		static bool TryReadTypedRecord(ReadOnlySpan<byte> Buffer, ref int BufferPos, Utf8String RequiredSuffix, CachedRecordInfo RecordInfo, [NotNullWhen(true)] out object? Record)
+		static bool TryReadTypedRecord(ReadOnlyMemory<byte> Buffer, ref int BufferPos, Utf8String RequiredSuffix, CachedRecordInfo RecordInfo, [NotNullWhen(true)] out object? Record)
 		{
 			// Create a bitmask for all the required tags
 			ulong RequiredTagsBitMask = 0;
@@ -499,6 +501,7 @@ namespace EpicGames.Perforce
 			}
 
 			// Get the record info, and parse it into the object
+			ReadOnlySpan<byte> BufferSpan = Buffer.Span;
 			for (; ; )
 			{
 				// Check that we've got a string field
@@ -509,14 +512,14 @@ namespace EpicGames.Perforce
 				}
 
 				// If this is the end of the record, break out
-				byte KeyType = Buffer[BufferPos];
+				byte KeyType = BufferSpan[BufferPos];
 				if (KeyType == '0')
 				{
 					break;
 				}
 				else if (KeyType != 's')
 				{
-					throw new PerforceException("Unexpected key field type while parsing marshalled output ({0}) - expected 's', got: {1}", (int)KeyType, FormatDataAsHexDump(Buffer.Slice(BufferPos)));
+					throw new PerforceException("Unexpected key field type while parsing marshalled output ({0}) - expected 's', got: {1}", (int)KeyType, FormatDataAsHexDump(Buffer.Slice(BufferPos).Span));
 				}
 
 				// Capture the initial buffer position, in case we have to roll back
@@ -648,11 +651,13 @@ namespace EpicGames.Perforce
 		/// <param name="NewRecord">The new record</param>
 		/// <param name="TagInfo">The current tag</param>
 		/// <returns></returns>
-		static bool TryReadValue(ReadOnlySpan<byte> Buffer, ref int BufferPos, object NewRecord, CachedTagInfo? TagInfo)
+		static bool TryReadValue(ReadOnlyMemory<byte> Buffer, ref int BufferPos, object NewRecord, CachedTagInfo? TagInfo)
 		{
+			ReadOnlySpan<byte> BufferSpan = Buffer.Span;
+
 			// Read the value type
 			byte ValueType;
-			if (!TryReadByte(Buffer, ref BufferPos, out ValueType))
+			if (!TryReadByte(BufferSpan, ref BufferPos, out ValueType))
 			{
 				return false;
 			}
@@ -673,7 +678,7 @@ namespace EpicGames.Perforce
 			else if (ValueType == 'i')
 			{
 				int Integer;
-				if (!TryReadInt(Buffer, ref BufferPos, out Integer))
+				if (!TryReadInt(BufferSpan, ref BufferPos, out Integer))
 				{
 					return false;
 				}
@@ -730,28 +735,6 @@ namespace EpicGames.Perforce
 			return true;
 		}
 
-		/// <summary>
-		/// Attempts to read a string with type indicator from the buffer
-		/// </summary>
-		/// <param name="Buffer">The buffer to read from</param>
-		/// <param name="BufferPos">Current read position within the buffer</param>
-		/// <param name="String">Receives the value that was read</param>
-		/// <returns>True if a string was read from the buffer, false if there was not enough data</returns>
-		static bool TryReadStringWithType(ReadOnlySpan<byte> Buffer, ref int BufferPos, out Utf8String String)
-		{
-			byte ValueType;
-			if (!TryReadByte(Buffer, ref BufferPos, out ValueType))
-			{
-				String = new Utf8String();
-				return false;
-			}
-			if (ValueType != 's')
-			{
-				throw new PerforceException("Expected string value, got '{0}'", ValueType);
-			}
-			return TryReadString(Buffer, ref BufferPos, out String);
-		}
-
 		/// <summary>
 		/// Attempts to read a string from the buffer
 		/// </summary>
@@ -759,22 +742,22 @@ namespace EpicGames.Perforce
 		/// <param name="BufferPos">Current read position within the buffer</param>
 		/// <param name="String">Receives the value that was read</param>
 		/// <returns>True if a string was read from the buffer, false if there was not enough data</returns>
-		static bool TryReadString(ReadOnlySpan<byte> Buffer, ref int BufferPos, out Utf8String String)
+		static bool TryReadString(ReadOnlyMemory<byte> Buffer, ref int BufferPos, out Utf8String String)
 		{
 			int Length;
-			if (!TryReadInt(Buffer, ref BufferPos, out Length))
+			if (!TryReadInt(Buffer.Span, ref BufferPos, out Length))
 			{
-				String = new Utf8String();
+				String = Utf8String.Empty;
 				return false;
 			}
 
 			if (BufferPos + Length > Buffer.Length)
 			{
-				String = new Utf8String();
+				String = Utf8String.Empty;
 				return false;
 			}
 
-			String = new Utf8String(Buffer.Slice(BufferPos, Length).ToArray());
+			String = new Utf8String(Buffer.Slice(BufferPos, Length));
 			BufferPos += Length;
 			return true;
 		}
diff --git a/Engine/Source/Programs/Shared/EpicGames.Perforce/PerforceReflection.cs b/Engine/Source/Programs/Shared/EpicGames.Perforce/PerforceReflection.cs
--- a/Engine/Source/Programs/Shared/EpicGames.Perforce/PerforceReflection.cs
+++ b/Engine/Source/Programs/Shared/EpicGames.Perforce/PerforceReflection.cs
@@ -449,6 +449,10 @@ namespace EpicGames.Perforce
 						{
 							TagInfo.SetFromString = (Obj, String) => FieldCopy.SetValue(Obj, ParseString(String));
 						}
+						else if (FieldType == typeof(Utf8String))
+						{
+							TagInfo.SetFromString = (Obj, String) => FieldCopy.SetValue(Obj, String.Clone());
+						}
 						else if (FieldType.IsEnum)
 						{
 							CachedEnumInfo EnumInfo = GetCachedEnumInfo(FieldType);
diff --git a/Engine/Source/Programs/Shared/EpicGames.Serialization/CompactBinary.cs b/Engine/Source/Programs/Shared/EpicGames.Serialization/CompactBinary.cs
--- a/Engine/Source/Programs/Shared/EpicGames.Serialization/CompactBinary.cs
+++ b/Engine/Source/Programs/Shared/EpicGames.Serialization/CompactBinary.cs
@@ -1402,8 +1402,10 @@ namespace EpicGames.Serialization
 		/// <param name="Hasher"></param>
 		void AppendHash(Blake3.Hasher Hasher)
 		{
-			byte[] Data = new byte[] { (byte)CbFieldUtils.GetSerializedType(TypeWithFlags) };
+			Span<byte> Data = stackalloc byte[1];
+			Data[0] = (byte)CbFieldUtils.GetSerializedType(TypeWithFlags);
 			Hasher.Update(Data);
+
 			Hasher.Update(GetViewNoType().Span);
 		}
 
