Horde: List all the currently active jobs when unable to start a schedule due to the max limit being reached.

diff --git a/Engine/Source/Programs/Horde/HordeServer/Services/ScheduleService.cs b/Engine/Source/Programs/Horde/HordeServer/Services/ScheduleService.cs
--- a/Engine/Source/Programs/Horde/HordeServer/Services/ScheduleService.cs
+++ b/Engine/Source/Programs/Horde/HordeServer/Services/ScheduleService.cs
@@ -242,7 +242,11 @@ namespace HordeServer.Services
 			// Check we're not already at the maximum number of allowed jobs
 			if (Schedule.MaxActive != 0 && NumActiveJobs >= Schedule.MaxActive)
 			{
-				Logger.LogInformation("Skipping trigger of {StreamName} template {TemplateRefId} - already have maximum number of jobs running ({NumJobs})", Stream.Id, TemplateRefId, Schedule.MaxActive);
+				Logger.LogInformation("Skipping trigger of {StreamId} template {TemplateRefId} - already have maximum number of jobs running ({NumJobs})", Stream.Id, TemplateRefId, Schedule.MaxActive);
+				foreach (ObjectId JobId in Schedule.ActiveJobs)
+				{
+					Logger.LogInformation("Active job for {StreamId} template {TemplateRefId}: {JobId}", Stream.Id, TemplateRefId, JobId);
+				}
 				return;
 			}
 
