Horde: Fix files being incorrectly added to their own nodes in tree snapshots.

diff --git a/Engine/Source/Programs/Horde/HordeServer/Commits/Impl/CommitService.cs b/Engine/Source/Programs/Horde/HordeServer/Commits/Impl/CommitService.cs
--- a/Engine/Source/Programs/Horde/HordeServer/Commits/Impl/CommitService.cs
+++ b/Engine/Source/Programs/Horde/HordeServer/Commits/Impl/CommitService.cs
@@ -655,7 +655,11 @@ namespace HordeServer.Commits.Impl
 
 			CbWriter Writer = new CbWriter();
 			Writer.BeginObject();
-			TreeRef.Write(Writer, Stream.Name);
+			if (TreeRef.Path != Stream.Name)
+			{
+				Writer.WriteString("path", TreeRef.Path);
+			}
+			TreeRef.Write(Writer);
 			Writer.EndObject();
 
 			List<IoHash> MissingHashes = await RefCollection.SetAsync(NamespaceId, BucketId, RefName, Writer.ToObject());
@@ -706,9 +710,7 @@ namespace HordeServer.Commits.Impl
 			// In order to optimize for the common case where we're adding multiple files to the same directory, we keep an open list of
 			// StreamTreeBuilder objects down to the last modified node. Whenever we need to move to a different node, we write the 
 			// previous path to the object store.
-			List<StreamTreeBuilder> TreePath = new List<StreamTreeBuilder>();
-			TreePath.Add(new StreamTreeBuilder());
-
+			StreamTreeBuilder RootNode = new StreamTreeBuilder();
 			for(int Idx = 0; Idx < View.Entries.Count; Idx++)
 			{
 				ViewMapEntry Entry = View.Entries[Idx];
@@ -730,14 +732,12 @@ namespace HordeServer.Commits.Impl
 					}
 
 					Logger.LogInformation("Adding {Source}@{Change}", Entry.Source, Change);
-					await AddDepotFilesToSnapshotAsync(Connection, $"{Entry.Source}@{Change}", View, PathComparer, TreePath, Writer, WriteObject);
+					await AddDepotFilesToSnapshotAsync(Connection, $"{Entry.Source}@{Change}", View, PathComparer, RootNode, Writer, WriteObject);
 
 					GC.Collect();
 				}
 			}
-
-			// Encode the tree
-			return TreePath[0].Encode(Writer, WriteObject);
+			return RootNode.Encode(Writer, WriteObject);
 		}
 
 		/// <summary>
@@ -771,10 +771,10 @@ namespace HordeServer.Commits.Impl
 		/// <param name="FileSpec">The files to add</param>
 		/// <param name="View">View for the stream</param>
 		/// <param name="Comparer">Comparison type to use for names</param>
-		/// <param name="TreePath">The last computed path through the tree</param>
+		/// <param name="RootNode">The root node of the tree</param>
 		/// <param name="Writer"></param>
 		/// <param name="WriteObject">Delegate to write an object to the store</param>
-		async Task AddDepotFilesToSnapshotAsync(IPerforceConnection Connection, string FileSpec, ViewMap View, Utf8StringComparer Comparer, List<StreamTreeBuilder> TreePath, CbWriter Writer, Action<IoHash, CbObject> WriteObject)
+		async Task AddDepotFilesToSnapshotAsync(IPerforceConnection Connection, string FileSpec, ViewMap View, Utf8StringComparer Comparer, StreamTreeBuilder RootNode, CbWriter Writer, Action<IoHash, CbObject> WriteObject)
 		{
 			const string Filter = "^headAction=delete ^headAction=move/delete";
 			const string Fields = "depotFile,fileSize,digest,headType,headRev";
@@ -800,7 +800,7 @@ namespace HordeServer.Commits.Impl
 					Utf8String TargetFile;
 					if (File.DepotFile != Utf8String.Empty && View.TryMapFile(File.DepotFile, Comparer, out TargetFile))
 					{
-						await AddFileAsync(TargetFile, File, TreePath, Segments, Writer, WriteObject);
+						await AddFileAsync(TargetFile, File, RootNode, Segments, Writer, WriteObject);
 					}
 				}
 			}
@@ -811,13 +811,13 @@ namespace HordeServer.Commits.Impl
 		/// </summary>
 		/// <param name="Path">The stream path for the file</param>
 		/// <param name="MetaData">The file metadata</param>
-		/// <param name="TreePath">The last computed path through the tree</param>
+		/// <param name="RootNode">The last computed path through the tree</param>
 		/// <param name="Segments"></param>
 		/// <param name="Writer"></param>
 		/// <param name="WriteObject"></param>
-		async Task AddFileAsync(Utf8String Path, Utf8FStatRecord MetaData, List<StreamTreeBuilder> TreePath, List<Utf8String> Segments, CbWriter Writer, Action<IoHash, CbObject> WriteObject)
+		async Task AddFileAsync(Utf8String Path, Utf8FStatRecord MetaData, StreamTreeBuilder RootNode, List<Utf8String> Segments, CbWriter Writer, Action<IoHash, CbObject> WriteObject)
 		{
-			StreamTreeBuilder Node = TreePath[0];
+			StreamTreeBuilder Node = RootNode;
 
 			// Split the path into segments
 			Segments.Clear();
@@ -838,7 +838,7 @@ namespace HordeServer.Commits.Impl
 
 			// Match as many nodes as we can from the existing path
 			int SegmentIdx = 0;
-			for(; SegmentIdx < Segments.Count; SegmentIdx++)
+			for(; SegmentIdx < Segments.Count - 1; SegmentIdx++)
 			{
 				Utf8String Segment = Segments[SegmentIdx];
 				if (!Node.NameToTreeBuilder.TryGetValue(Segment, out StreamTreeBuilder? NextNode))
@@ -852,7 +852,7 @@ namespace HordeServer.Commits.Impl
 			CollapseTree(Node, Writer, WriteObject);
 
 			// Expand the rest of the path
-			for(; SegmentIdx < Segments.Count; SegmentIdx++)
+			for(; SegmentIdx < Segments.Count - 1; SegmentIdx++)
 			{
 				Utf8String Segment = Segments[SegmentIdx];
 				Node = await FindOrAddTreeAsync(Node, Segment);
