Horde: Generalize the delegate used to serialize StreamTree objects, to allow serializing directly into existing buffers without constructing intermediate CbObject representations.

diff --git a/Engine/Source/Programs/Horde/HordeServer/Commits/Impl/CommitService.cs b/Engine/Source/Programs/Horde/HordeServer/Commits/Impl/CommitService.cs
--- a/Engine/Source/Programs/Horde/HordeServer/Commits/Impl/CommitService.cs
+++ b/Engine/Source/Programs/Horde/HordeServer/Commits/Impl/CommitService.cs
@@ -639,20 +639,22 @@ namespace HordeServer.Commits.Impl
 
 		async Task<StreamTreeRef?> AddTreeRefAsync(IPerforceConnection Perforce, IStream Stream, ViewMap View, ICommit? PrevCommit, StreamTreeRef? PrevTreeRef, ICommit Commit, Utf8StringComparer PathComparer)
 		{
-			StreamTreeRef TreeRef;
+			StreamTree Tree;
 			if (PrevCommit == null || PrevTreeRef == null)
 			{
 				Logger.LogInformation("Performing snapshot of {StreamId} at CL {Change}", Commit.StreamId, Commit.Change);
-				TreeRef = await CreateTreeAsync(Perforce, View, Commit.Change, PathComparer);
+				Tree = await CreateTreeAsync(Perforce, View, Commit.Change, PathComparer);
 			}
 			else
 			{
 				Logger.LogInformation("Performing incremental update of {StreamId} from CL {PrevChange} to CL {Change}", Commit.StreamId, PrevCommit.Change, Commit.Change);
-				TreeRef = await UpdateTreeAsync(PrevTreeRef, Perforce, View, Commit.Change, PathComparer);
+				Tree = await UpdateTreeAsync(PrevTreeRef, Perforce, View, Commit.Change, PathComparer);
 			}
 
 			string RefName = $"tree_{Commit.StreamId}_{Commit.Change}";
 
+			StreamTreeRef TreeRef = new StreamTreeRef(Tree.Path, Tree.ToCbObject().GetHash());
+
 			CbWriter Writer = new CbWriter();
 			Writer.BeginObject();
 			if (TreeRef.Path != Stream.Name)
@@ -687,7 +689,7 @@ namespace HordeServer.Commits.Impl
 		/// <param name="View">View for the stream</param>
 		/// <param name="Change">Changelist to replicate</param>
 		/// <param name="PathComparer">Comparison type for paths</param>
-		async Task<StreamTreeRef> CreateTreeAsync(IPerforceConnection Connection, ViewMap View, int Change, Utf8StringComparer PathComparer)
+		async Task<StreamTree> CreateTreeAsync(IPerforceConnection Connection, ViewMap View, int Change, Utf8StringComparer PathComparer)
 		{
 			// Optimize the view by removing unnecessary inclusions of the same path
 			View = new ViewMap(View);
@@ -704,8 +706,7 @@ namespace HordeServer.Commits.Impl
 			}
 
 			// Delegate to write an object
-			CbWriter Writer = new CbWriter();
-			Action<IoHash, CbObject> WriteObject = (x, y) => { };
+			Func<StreamTree, IoHash> WriteObject = x => x.ToCbObject().GetHash();
 
 			// In order to optimize for the common case where we're adding multiple files to the same directory, we keep an open list of
 			// StreamTreeBuilder objects down to the last modified node. Whenever we need to move to a different node, we write the 
@@ -732,12 +733,12 @@ namespace HordeServer.Commits.Impl
 					}
 
 					Logger.LogInformation("Adding {Source}@{Change}", Entry.Source, Change);
-					await AddDepotFilesToSnapshotAsync(Connection, $"{Entry.Source}@{Change}", View, PathComparer, RootNode, Writer, WriteObject);
+					await AddDepotFilesToSnapshotAsync(Connection, $"{Entry.Source}@{Change}", View, PathComparer, RootNode, WriteObject);
 
 					GC.Collect();
 				}
 			}
-			return RootNode.Encode(Writer, WriteObject);
+			return RootNode.Encode(WriteObject);
 		}
 
 		/// <summary>
@@ -772,9 +773,8 @@ namespace HordeServer.Commits.Impl
 		/// <param name="View">View for the stream</param>
 		/// <param name="Comparer">Comparison type to use for names</param>
 		/// <param name="RootNode">The root node of the tree</param>
-		/// <param name="Writer"></param>
 		/// <param name="WriteObject">Delegate to write an object to the store</param>
-		async Task AddDepotFilesToSnapshotAsync(IPerforceConnection Connection, string FileSpec, ViewMap View, Utf8StringComparer Comparer, StreamTreeBuilder RootNode, CbWriter Writer, Action<IoHash, CbObject> WriteObject)
+		async Task AddDepotFilesToSnapshotAsync(IPerforceConnection Connection, string FileSpec, ViewMap View, Utf8StringComparer Comparer, StreamTreeBuilder RootNode, Func<StreamTree, IoHash> WriteObject)
 		{
 			const string Filter = "^headAction=delete ^headAction=move/delete";
 			const string Fields = "depotFile,fileSize,digest,headType,headRev";
@@ -800,7 +800,7 @@ namespace HordeServer.Commits.Impl
 					Utf8String TargetFile;
 					if (File.DepotFile != Utf8String.Empty && View.TryMapFile(File.DepotFile, Comparer, out TargetFile))
 					{
-						await AddFileAsync(TargetFile, File, RootNode, Segments, Writer, WriteObject);
+						await AddFileAsync(TargetFile, File, RootNode, Segments, WriteObject);
 					}
 				}
 			}
@@ -813,9 +813,8 @@ namespace HordeServer.Commits.Impl
 		/// <param name="MetaData">The file metadata</param>
 		/// <param name="RootNode">The last computed path through the tree</param>
 		/// <param name="Segments"></param>
-		/// <param name="Writer"></param>
 		/// <param name="WriteObject"></param>
-		async Task AddFileAsync(Utf8String Path, Utf8FStatRecord MetaData, StreamTreeBuilder RootNode, List<Utf8String> Segments, CbWriter Writer, Action<IoHash, CbObject> WriteObject)
+		async Task AddFileAsync(Utf8String Path, Utf8FStatRecord MetaData, StreamTreeBuilder RootNode, List<Utf8String> Segments, Func<StreamTree, IoHash> WriteObject)
 		{
 			StreamTreeBuilder Node = RootNode;
 
@@ -849,7 +848,7 @@ namespace HordeServer.Commits.Impl
 			}
 
 			// Collapse the rest of the current path
-			CollapseTree(Node, Writer, WriteObject);
+			Node.EncodeChildren(WriteObject);
 
 			// Expand the rest of the path
 			for(; SegmentIdx < Segments.Count - 1; SegmentIdx++)
@@ -876,7 +875,7 @@ namespace HordeServer.Commits.Impl
 		/// <param name="View"></param>
 		/// <param name="PathComparer"></param>
 		/// <returns></returns>
-		async Task<StreamTreeRef> UpdateTreeAsync(StreamTreeRef PrevTreeRef, IPerforceConnection Perforce, ViewMap View, int Change, Utf8StringComparer PathComparer)
+		async Task<StreamTree> UpdateTreeAsync(StreamTreeRef PrevTreeRef, IPerforceConnection Perforce, ViewMap View, int Change, Utf8StringComparer PathComparer)
 		{
 			List<FStatRecord> Records = await Perforce.FStatAsync(-1, Change, null, null, -1, FStatOptions.IncludeFileSizes, FileSpecList.Empty);
 			return await UpdateTreeAsync(PrevTreeRef, Records, View, PathComparer);
@@ -890,13 +889,10 @@ namespace HordeServer.Commits.Impl
 		/// <param name="View"></param>
 		/// <param name="PathComparer"></param>
 		/// <returns></returns>
-		async Task<StreamTreeRef> UpdateTreeAsync(StreamTreeRef PrevTreeRef, IList<FStatRecord> Files, ViewMap View, Utf8StringComparer PathComparer)
+		async Task<StreamTree> UpdateTreeAsync(StreamTreeRef PrevTreeRef, IList<FStatRecord> Files, ViewMap View, Utf8StringComparer PathComparer)
 		{
 			StreamTreeBuilder Tree = await ReadTreeAsync(PrevTreeRef);
 
-			// Delegate to write an object
-			Action<IoHash, CbObject> WriteObject = (x, y) => { };
-
 			// Update the tree
 			foreach (FStatRecord File in Files)
 			{
@@ -936,7 +932,7 @@ namespace HordeServer.Commits.Impl
 				}
 			}
 
-			return Tree.Encode(new CbWriter(), WriteObject);
+			return Tree.Encode(SubTree => SubTree.ToCbObject().GetHash());
 		}
 
 #endregion
@@ -1009,24 +1005,7 @@ namespace HordeServer.Commits.Impl
 			{
 				throw new Exception($"Missing object {TreeRef.Hash} referenced by tree for {TreeRef.Path}");
 			}
-			return new StreamTreeBuilder(new StreamTree(Object, TreeRef.Path + "/"));
-		}
-
-		/// <summary>
-		/// Collapse a subtree under a given child
-		/// </summary>
-		/// <param name="Node">The node to collapse</param>
-		/// <param name="Writer"></param>
-		/// <param name="WriteObject"></param>
-		static void CollapseTree(StreamTreeBuilder Node, CbWriter Writer, Action<IoHash, CbObject> WriteObject)
-		{
-			foreach ((Utf8String Name, StreamTreeBuilder Child) in Node.NameToTreeBuilder)
-			{
-				CollapseTree(Child, Writer, WriteObject);
-				StreamTreeRef ChildRef = Child.Encode(Writer, WriteObject);
-				Node.NameToTree.Add(Name, ChildRef);
-			}
-			Node.NameToTreeBuilder.Clear();
+			return new StreamTreeBuilder(new StreamTree(TreeRef.Path, Object));
 		}
 
 		/// <summary>
@@ -1077,7 +1056,7 @@ namespace HordeServer.Commits.Impl
 		/// <returns></returns>
 		async Task PrintSnapshotInternalAsync(StreamTreeRef TreeRef, string Prefix)
 		{
-			StreamTree Tree = new StreamTree(await ObjectCollection.GetRequiredObjectAsync(NamespaceId, TreeRef.Hash), TreeRef.Path);
+			StreamTree Tree = new StreamTree(TreeRef.Path, await ObjectCollection.GetRequiredObjectAsync(NamespaceId, TreeRef.Hash));
 
 			foreach ((Utf8String Name, StreamFile File) in Tree.NameToFile)
 			{
diff --git a/Engine/Source/Programs/Shared/EpicGames.Perforce.Managed/StreamSnapshotFromMemory.cs b/Engine/Source/Programs/Shared/EpicGames.Perforce.Managed/StreamSnapshotFromMemory.cs
--- a/Engine/Source/Programs/Shared/EpicGames.Perforce.Managed/StreamSnapshotFromMemory.cs
+++ b/Engine/Source/Programs/Shared/EpicGames.Perforce.Managed/StreamSnapshotFromMemory.cs
@@ -50,14 +50,29 @@ namespace EpicGames.Perforce.Managed
 		public StreamSnapshotFromMemory(StreamTreeBuilder Builder)
 		{
 			Dictionary<IoHash, CbObject> HashToTree = new Dictionary<IoHash, CbObject>();
-			this.Root = Builder.Encode(new CbWriter(), (Hash, Object) => HashToTree[Hash] = Object);
+			this.Root = Builder.EncodeRef(Tree => EncodeObject(Tree, HashToTree));
 			this.HashToTree = HashToTree;
 		}
 
+		/// <summary>
+		/// Serialize to a compact binary object
+		/// </summary>
+		/// <param name="BasePath"></param>
+		/// <returns></returns>
+		static IoHash EncodeObject(StreamTree Tree, Dictionary<IoHash, CbObject> HashToTree)
+		{
+			CbObject Object = Tree.ToCbObject();
+
+			IoHash Hash = Object.GetHash();
+			HashToTree[Hash] = Object;
+
+			return Hash;
+		}
+
 		/// <inheritdoc/>
 		public override StreamTree Lookup(StreamTreeRef Ref)
 		{
-			return new StreamTree(HashToTree[Ref.Hash], Ref.Path);
+			return new StreamTree(Ref.Path, HashToTree[Ref.Hash]);
 		}
 
 		/// <summary>
diff --git a/Engine/Source/Programs/Shared/EpicGames.Perforce.Managed/StreamTree.cs b/Engine/Source/Programs/Shared/EpicGames.Perforce.Managed/StreamTree.cs
--- a/Engine/Source/Programs/Shared/EpicGames.Perforce.Managed/StreamTree.cs
+++ b/Engine/Source/Programs/Shared/EpicGames.Perforce.Managed/StreamTree.cs
@@ -157,6 +157,11 @@ namespace EpicGames.Perforce.Managed
 	/// </summary>
 	public class StreamTree
 	{
+		/// <summary>
+		/// The path to this tree
+		/// </summary>
+		public Utf8String Path { get; }
+
 		/// <summary>
 		/// Map of name to file within the directory
 		/// </summary>
@@ -184,8 +189,11 @@ namespace EpicGames.Perforce.Managed
 		/// <summary>
 		/// Default constructor
 		/// </summary>
-		public StreamTree(Dictionary<Utf8String, StreamFile> NameToFile, Dictionary<Utf8String, StreamTreeRef> NameToTree)
+		public StreamTree(Utf8String Path, Dictionary<Utf8String, StreamFile> NameToFile, Dictionary<Utf8String, StreamTreeRef> NameToTree)
 		{
+			CheckPath(Path);
+			this.Path = Path;
+
 			this.NameToFile = NameToFile;
 			this.NameToTree = NameToTree;
 		}
@@ -193,18 +201,19 @@ namespace EpicGames.Perforce.Managed
 		/// <summary>
 		/// Deserialize a tree from a compact binary object
 		/// </summary>
-		/// <param name="Object"></param>
-		/// <param name="BasePath"></param>
-		public StreamTree(CbObject Object, Utf8String BasePath)
+		public StreamTree(Utf8String Path, CbObject Object)
 		{
+			CheckPath(Path);
+			this.Path = Path;
+
 			CbArray FileArray = Object[FilesField].AsArray();
 			foreach (CbField FileField in FileArray)
 			{
 				CbObject FileObject = FileField.AsObject();
 
 				Utf8String Name = FileObject[NameField].AsString();
-				Utf8String Path = ReadPath(FileObject, BasePath, Name);
-				StreamFile File = new StreamFile(Path, FileObject);
+				Utf8String FilePath = ReadPath(FileObject, Path, Name);
+				StreamFile File = new StreamFile(FilePath, FileObject);
 
 				NameToFile.Add(Name, File);
 			}
@@ -215,74 +224,19 @@ namespace EpicGames.Perforce.Managed
 				CbObject TreeObject = TreeField.AsObject();
 
 				Utf8String Name = TreeObject[NameField].AsString();
-				Utf8String Path = ReadPath(TreeObject, BasePath, Name);
-				StreamTreeRef Tree = new StreamTreeRef(Path, TreeObject);
+				Utf8String TreePath = ReadPath(TreeObject, Path, Name);
+				StreamTreeRef Tree = new StreamTreeRef(TreePath, TreeObject);
 
 				NameToTree.Add(Name, Tree);
 			}
 		}
 
-		/// <summary>
-		/// Finds the most common base path for items in this tree
-		/// </summary>
-		/// <returns>The most common base path</returns>
-		public Utf8String FindBasePath()
-		{
-			Dictionary<Utf8String, int> BasePathToCount = new Dictionary<Utf8String, int>();
-			foreach ((Utf8String Name, StreamFile File) in NameToFile)
-			{
-				AddBasePath(BasePathToCount, File.Path, Name);
-			}
-			foreach ((Utf8String Name, StreamTreeRef Tree) in NameToTree)
-			{
-				AddBasePath(BasePathToCount, Tree.Path, Name);
-			}
-			return (BasePathToCount.Count == 0) ? Utf8String.Empty : BasePathToCount.MaxBy(x => x.Value).Key;
-		}
-
-		/// <summary>
-		/// Adds the base path of the given item to the count of similar items
-		/// </summary>
-		/// <param name="BasePathToCount"></param>
-		/// <param name="Path"></param>
-		/// <param name="Name"></param>
-		static void AddBasePath(Dictionary<Utf8String, int> BasePathToCount, Utf8String Path, Utf8String Name)
-		{
-			if (Path.EndsWith(Name) && Path[^(Name.Length + 1)] == '/')
-			{
-				Utf8String BasePath = Path[..^(Name.Length + 1)];
-				BasePathToCount.TryGetValue(BasePath, out int Count);
-				BasePathToCount[BasePath] = Count + 1;
-			}
-		}
-
-		/// <summary>
-		/// Gets the default path for a child path
-		/// </summary>
-		/// <param name="BasePath"></param>
-		/// <param name="Name"></param>
-		/// <returns></returns>
-		static Utf8String GetDefaultPath(Utf8String BasePath, Utf8String Name)
-		{
-			byte[] Data = new byte[BasePath.Length + 1 + Name.Length];
-			BasePath.Span.CopyTo(Data);
-			Data[BasePath.Length] = (byte)'/';
-			Name.Span.CopyTo(Data.AsSpan(BasePath.Length + 1));
-			return new Utf8String(Data);
-		}
-
 		/// <summary>
 		/// Serialize to a compact binary object
 		/// </summary>
 		/// <param name="Writer"></param>
-		/// <param name="BasePath"></param>
-		public void Write(CbWriter Writer, Utf8String BasePath)
+		public void Write(CbWriter Writer)
 		{
-			if (BasePath.Length > 0 && BasePath[BasePath.Length - 1] == '/')
-			{
-				throw new ArgumentException("BasePath must not end in a slash", nameof(BasePath));
-			}
-
 			if (NameToFile.Count > 0)
 			{
 				Writer.BeginArray(FilesField);
@@ -290,7 +244,7 @@ namespace EpicGames.Perforce.Managed
 				{
 					Writer.BeginObject();
 					Writer.WriteString(NameField, Name);
-					WritePath(Writer, File.Path, BasePath, Name);
+					WritePath(Writer, File.Path, Path, Name);
 					File.Write(Writer);
 					Writer.EndObject();
 				}
@@ -304,7 +258,7 @@ namespace EpicGames.Perforce.Managed
 				{
 					Writer.BeginObject();
 					Writer.WriteString(NameField, Name);
-					WritePath(Writer, Tree.Path, BasePath, Name);
+					WritePath(Writer, Tree.Path, Path, Name);
 					Tree.Write(Writer);
 					Writer.EndObject();
 				}
@@ -349,20 +303,27 @@ namespace EpicGames.Perforce.Managed
 		}
 
 		/// <summary>
-		/// Serialize to a compact binary object
+		/// Checks that a base path does not have a trailing slash
 		/// </summary>
-		/// <param name="BasePath"></param>
-		/// <returns></returns>
-		public CbObject ToCbObject(CbWriter Writer, Utf8String BasePath)
+		/// <param name="Path"></param>
+		/// <exception cref="ArgumentException"></exception>
+		static void CheckPath(Utf8String Path)
 		{
-			if (BasePath.Length > 0 && BasePath[BasePath.Length - 1] == '/')
+			if (Path.Length > 0 && Path[Path.Length - 1] == '/')
 			{
-				throw new ArgumentException("BasePath must not end in a slash", nameof(BasePath));
+				throw new ArgumentException("BasePath must not end in a slash", nameof(Path));
 			}
+		}
 
-			Writer.Clear();
+		/// <summary>
+		/// Convert to a compact binary object
+		/// </summary>
+		/// <returns></returns>
+		public CbObject ToCbObject()
+		{
+			CbWriter Writer = new CbWriter();
 			Writer.BeginObject();
-			Write(Writer, BasePath);
+			Write(Writer);
 			Writer.EndObject();
 			return Writer.ToObject();
 		}
diff --git a/Engine/Source/Programs/Shared/EpicGames.Perforce.Managed/StreamTreeBuilder.cs b/Engine/Source/Programs/Shared/EpicGames.Perforce.Managed/StreamTreeBuilder.cs
--- a/Engine/Source/Programs/Shared/EpicGames.Perforce.Managed/StreamTreeBuilder.cs
+++ b/Engine/Source/Programs/Shared/EpicGames.Perforce.Managed/StreamTreeBuilder.cs
@@ -57,59 +57,72 @@ namespace EpicGames.Perforce.Managed
 		/// <summary>
 		/// Encodes the current tree state and returns a reference to it
 		/// </summary>
-		/// <param name="Writer">A writer instance to use to create the object</param>
-		/// <param name="WriteObject"></param>
+		/// <param name="Objects">Dictionary of encoded objects</param>
 		/// <returns></returns>
-		public StreamTreeRef Encode(CbWriter Writer, Action<IoHash, CbObject> WriteObject)
+		public StreamTree Encode(Func<StreamTree, IoHash> WriteTree)
 		{
-			StreamTreeRef? Ref = EncodeInternal(Writer, WriteObject);
-			if (Ref == null)
+			// Recursively serialize all the child items
+			EncodeChildren(WriteTree);
+
+			// Find the common base path for all items in this tree.
+			Dictionary<Utf8String, int> BasePathToCount = new Dictionary<Utf8String, int>();
+			foreach ((Utf8String Name, StreamFile File) in NameToFile)
 			{
-				CbObject Object = new StreamTree().ToCbObject(Writer, Utf8String.Empty);
-				IoHash ObjectHash = Object.GetHash();
-				WriteObject(ObjectHash, Object);
-				Ref = new StreamTreeRef(Utf8String.Empty, ObjectHash);
+				AddBasePath(BasePathToCount, File.Path, Name);
 			}
-			return Ref;
+			foreach ((Utf8String Name, StreamTreeRef Tree) in NameToTree)
+			{
+				AddBasePath(BasePathToCount, Tree.Path, Name);
+			}
+
+			// Create the new tree
+			Utf8String BasePath = (BasePathToCount.Count == 0) ? Utf8String.Empty : BasePathToCount.MaxBy(x => x.Value).Key;
+			return new StreamTree(BasePath, NameToFile, NameToTree);
 		}
 
 		/// <summary>
-		/// Encodes the current tree state and returns a reference to it
+		/// Encodes a StreamTreeRef from this tree
 		/// </summary>
-		/// <param name="Objects">Dictionary of encoded objects</param>
-		/// <returns></returns>
-		public StreamTreeRef? EncodeInternal(CbWriter Writer, Action<IoHash, CbObject> WriteObject)
+		/// <param name="WriteTree"></param>
+		/// <returns>The new tree ref</returns>
+		public StreamTreeRef EncodeRef(Func<StreamTree, IoHash> WriteTree)
 		{
-			// Recursively serialize all the child items
-			foreach ((Utf8String SubTreeName, StreamTreeBuilder SubTree) in NameToTreeBuilder)
+			StreamTree Tree = Encode(WriteTree);
+			return new StreamTreeRef(Tree.Path, WriteTree(Tree));
+		}
+
+		/// <summary>
+		/// Collapses all of the builders underneath this node
+		/// </summary>
+		/// <param name="WriteTree"></param>
+		public void EncodeChildren(Func<StreamTree, IoHash> WriteTree)
+		{
+			foreach ((Utf8String SubTreeName, StreamTreeBuilder SubTreeBuilder) in NameToTreeBuilder)
 			{
-				StreamTreeRef? SubTreeRef = SubTree.EncodeInternal(Writer, WriteObject);
-				if (SubTreeRef == null)
+				StreamTree SubTree = SubTreeBuilder.Encode(WriteTree);
+				if (SubTree.NameToFile.Count > 0 || SubTree.NameToTree.Count > 0)
 				{
-					NameToTreeBuilder.Remove(SubTreeName);
-				}
-				else
-				{
-					NameToTree[SubTreeName] = SubTreeRef;
+					IoHash Hash = WriteTree(SubTree);
+					NameToTree[SubTreeName] = new StreamTreeRef(SubTree.Path, Hash);
 				}
 			}
 			NameToTreeBuilder.Clear();
+		}
 
-			// If if's empty, don't create this tree at all
-			if (IsEmpty)
+		/// <summary>
+		/// Adds the base path of the given item to the count of similar items
+		/// </summary>
+		/// <param name="BasePathToCount"></param>
+		/// <param name="Path"></param>
+		/// <param name="Name"></param>
+		static void AddBasePath(Dictionary<Utf8String, int> BasePathToCount, Utf8String Path, Utf8String Name)
+		{
+			if (Path.EndsWith(Name) && Path[^(Name.Length + 1)] == '/')
 			{
-				return null;
+				Utf8String BasePath = Path[..^(Name.Length + 1)];
+				BasePathToCount.TryGetValue(BasePath, out int Count);
+				BasePathToCount[BasePath] = Count + 1;
 			}
-
-			// Find the common base path for all items in this tree.
-			StreamTree Tree = new StreamTree(NameToFile, NameToTree);
-			Utf8String BasePath = Tree.FindBasePath();
-			CbObject PackedObject = Tree.ToCbObject(Writer, BasePath);
-
-			IoHash ObjectHash = PackedObject.GetHash();
-			WriteObject(ObjectHash, PackedObject);
-
-			return new StreamTreeRef(BasePath, ObjectHash);
 		}
 	}
 }
diff --git a/Engine/Source/Programs/Shared/EpicGames.Perforce.Managed/StreamTreeReader.cs b/Engine/Source/Programs/Shared/EpicGames.Perforce.Managed/StreamTreeReader.cs
--- a/Engine/Source/Programs/Shared/EpicGames.Perforce.Managed/StreamTreeReader.cs
+++ b/Engine/Source/Programs/Shared/EpicGames.Perforce.Managed/StreamTreeReader.cs
@@ -47,7 +47,7 @@ namespace EpicGames.Perforce.Managed
 		/// <inheritdoc/>
 		public override Task<StreamTree> ReadAsync(StreamTreeRef Ref)
 		{
-			return Task.FromResult(new StreamTree(HashToTree[Ref.Hash], Ref.Path));
+			return Task.FromResult(new StreamTree(Ref.Path, HashToTree[Ref.Hash]));
 		}
 	}
 }
