Horde: Enable serilog analyzer and fix a warning.

diff --git a/Engine/Source/Programs/Horde/HordeServer/HordeServer.csproj b/Engine/Source/Programs/Horde/HordeServer/HordeServer.csproj
--- a/Engine/Source/Programs/Horde/HordeServer/HordeServer.csproj
+++ b/Engine/Source/Programs/Horde/HordeServer/HordeServer.csproj
@@ -1,4 +1,4 @@
-<Project Sdk="Microsoft.NET.Sdk.Web">
+﻿<Project Sdk="Microsoft.NET.Sdk.Web">
 
   <PropertyGroup>
     <TargetFramework>netcoreapp3.1</TargetFramework>
@@ -46,6 +46,7 @@
     <PackageReference Include="Serilog.Extensions.Hosting" Version="3.1.0" />
     <PackageReference Include="Serilog.Sinks.Console" Version="3.1.1" />
     <PackageReference Include="Serilog.Sinks.File" Version="4.1.0" />
+    <PackageReference Include="SerilogAnalyzer" Version="0.15.0" />
     <PackageReference Include="SharpZipLib" Version="1.2.0" />
     <PackageReference Include="StackExchange.Redis" Version="2.2.4" />
     <PackageReference Include="Swashbuckle.AspNetCore" Version="5.5.0" />
diff --git a/Engine/Source/Programs/Horde/HordeServer/Services/Impl/PerforceService.cs b/Engine/Source/Programs/Horde/HordeServer/Services/Impl/PerforceService.cs
--- a/Engine/Source/Programs/Horde/HordeServer/Services/Impl/PerforceService.cs
+++ b/Engine/Source/Programs/Horde/HordeServer/Services/Impl/PerforceService.cs
@@ -481,7 +481,7 @@ namespace HordeServer.Services
 		{
 			lock (P4LogLock)
 			{
-				Serilog.Log.Information("Perforce: " + Log);
+				Serilog.Log.Information("Perforce: {Message}", Log);
 			}
 
 		}
