Horde: Fix some warnings building against NET 5.0.

diff --git a/Engine/Source/Programs/Horde/HordeServer/Api/Server.cs b/Engine/Source/Programs/Horde/HordeServer/Api/Server.cs
--- a/Engine/Source/Programs/Horde/HordeServer/Api/Server.cs
+++ b/Engine/Source/Programs/Horde/HordeServer/Api/Server.cs
@@ -34,7 +34,7 @@ namespace HordeServer.Api
         {
 
             FileVersionInfo VersionInfo = FileVersionInfo.GetVersionInfo(Assembly.GetExecutingAssembly().Location);		
-			ServerVersion = VersionInfo.ProductVersion;			
+			ServerVersion = VersionInfo.ProductVersion ?? String.Empty;			
 			OSDescription = RuntimeInformation.OSDescription;
 			this.SingleInstance = SingleInstance;
 		}
diff --git a/Engine/Source/Programs/Horde/HordeServer/Authentication/AnonymousHandler.cs b/Engine/Source/Programs/Horde/HordeServer/Authentication/AnonymousHandler.cs
--- a/Engine/Source/Programs/Horde/HordeServer/Authentication/AnonymousHandler.cs
+++ b/Engine/Source/Programs/Horde/HordeServer/Authentication/AnonymousHandler.cs
@@ -37,11 +37,13 @@ namespace HordeServer.Authentication
 
 		protected override async Task<AuthenticateResult> HandleAuthenticateAsync()
 		{
-			Claim[] Claims =
+			List<Claim> Claims = new List<Claim>();
+			Claims.Add(new Claim(ClaimTypes.Name, AuthenticationScheme));
+
+			if (Options.AdminClaimType != null && Options.AdminClaimValue != null)
 			{
-				new Claim(ClaimTypes.Name, AuthenticationScheme),
-				new Claim(Options.AdminClaimType, Options.AdminClaimValue)
-			};
+				Claims.Add(new Claim(Options.AdminClaimType, Options.AdminClaimValue));
+			}
 
 			IUser User = await UserCollection.FindOrAddUserByLoginAsync("anonymous", "Anonymous", "[REDACTED_EMAIL]");
 
diff --git a/Engine/Source/Programs/Horde/HordeServer/Authentication/OktaHandler.cs b/Engine/Source/Programs/Horde/HordeServer/Authentication/OktaHandler.cs
--- a/Engine/Source/Programs/Horde/HordeServer/Authentication/OktaHandler.cs
+++ b/Engine/Source/Programs/Horde/HordeServer/Authentication/OktaHandler.cs
@@ -43,15 +43,15 @@ namespace HordeServer.Authentication
 			{
 				if (Identity.FindFirst(ClaimTypes.Name) == null)
 				{
-					Identity.AddClaim(new Claim(ClaimTypes.Name, UserElement.ToString()));
+					Identity.AddClaim(new Claim(ClaimTypes.Name, UserElement.ToString()!));
 				}
 				if (Identity.FindFirst(HordeClaimTypes.User) == null)
 				{
-					Identity.AddClaim(new Claim(HordeClaimTypes.User, UserElement.ToString()));
+					Identity.AddClaim(new Claim(HordeClaimTypes.User, UserElement.ToString()!));
 				}
 				if (Identity.FindFirst(HordeClaimTypes.PerforceUser) == null)
 				{
-					Identity.AddClaim(new Claim(HordeClaimTypes.PerforceUser, UserElement.ToString()));
+					Identity.AddClaim(new Claim(HordeClaimTypes.PerforceUser, UserElement.ToString()!));
 				}
 			}
 
@@ -60,14 +60,14 @@ namespace HordeServer.Authentication
 			{
 				for (int Idx = 0; Idx < GroupsElement.GetArrayLength(); Idx++)
 				{
-					Identity.AddClaim(new Claim(ClaimTypes.Role, GroupsElement[Idx].ToString()));
+					Identity.AddClaim(new Claim(ClaimTypes.Role, GroupsElement[Idx].ToString()!));
 				}
 			}
 
 			JsonElement EmailElement;
 			if (UserInfo.TryGetProperty("email", out EmailElement) && Identity.FindFirst(ClaimTypes.Email) == null)
 			{
-				Identity.AddClaim(new Claim(ClaimTypes.Email, EmailElement.GetString()));
+				Identity.AddClaim(new Claim(ClaimTypes.Email, EmailElement.ToString()!));
 			}
 		}
 
@@ -80,7 +80,11 @@ namespace HordeServer.Authentication
 				return Result;
 			}
 
-			ClaimsIdentity Identity = (ClaimsIdentity)Result.Principal.Identity;
+			ClaimsIdentity? Identity = (ClaimsIdentity?)Result.Principal?.Identity;
+			if (Identity == null)
+			{
+				return HandleRequestResult.Fail("No identity specified");
+			}
 
 			string Login = Identity.FindFirst(ClaimTypes.Name)!.Value;
 			string? Name = Identity.FindFirst("name")?.Value;
diff --git a/Engine/Source/Programs/Horde/HordeServer/Collections/Impl/AgentCollection.cs b/Engine/Source/Programs/Horde/HordeServer/Collections/Impl/AgentCollection.cs
--- a/Engine/Source/Programs/Horde/HordeServer/Collections/Impl/AgentCollection.cs
+++ b/Engine/Source/Programs/Horde/HordeServer/Collections/Impl/AgentCollection.cs
@@ -113,7 +113,7 @@ namespace HordeServer.Collections.Impl
 			IReadOnlyList<PoolId> IAgent.ExplicitPools => Pools;
 			IReadOnlyList<AgentWorkspace> IAgent.Workspaces => Workspaces;
 			IReadOnlyList<AgentLease> IAgent.Leases => Leases ?? EmptyLeases;
-			IReadOnlyList<string> IAgent.Properties => Properties ?? Capabilities.Devices.FirstOrDefault()?.Properties.ToList() ?? new List<string>();
+			IReadOnlyList<string> IAgent.Properties => Properties ?? Capabilities.Devices.FirstOrDefault()?.Properties?.ToList() ?? new List<string>();
 			IReadOnlyDictionary<string, int> IAgent.Resources => Resources ?? Capabilities.Devices.FirstOrDefault()?.Resources ?? new Dictionary<string, int>();
 
 			[BsonConstructor]
diff --git a/Engine/Source/Programs/Horde/HordeServer/Collections/Impl/DeviceCollection.cs b/Engine/Source/Programs/Horde/HordeServer/Collections/Impl/DeviceCollection.cs
--- a/Engine/Source/Programs/Horde/HordeServer/Collections/Impl/DeviceCollection.cs
+++ b/Engine/Source/Programs/Horde/HordeServer/Collections/Impl/DeviceCollection.cs
@@ -559,7 +559,7 @@ namespace HordeServer.Collections.Impl
 
 			foreach (DeviceRequestData Data in Request)
 			{
-				DeviceDocument Device = PoolDevices.FirstOrDefault(A =>
+				DeviceDocument? Device = PoolDevices.FirstOrDefault(A =>
 				{
 
 					if (Allocated.Contains(A.Id) || A.PlatformId != Data.PlatformId)
diff --git a/Engine/Source/Programs/Horde/HordeServer/Collections/Impl/GraphCollection.cs b/Engine/Source/Programs/Horde/HordeServer/Collections/Impl/GraphCollection.cs
--- a/Engine/Source/Programs/Horde/HordeServer/Collections/Impl/GraphCollection.cs
+++ b/Engine/Source/Programs/Horde/HordeServer/Collections/Impl/GraphCollection.cs
@@ -425,8 +425,8 @@ namespace HordeServer.Collections.Impl
 				// Trim the cache
 				while (CachedGraphs.Count > MaxGraphs)
 				{
-					ContentHash RemoveHash = CachedGraphs.OrderBy(x => x.Value.LastAccessTime).Select(x => x.Key).FirstOrDefault();
-					if (RemoveHash == ContentHash.Empty)
+					ContentHash? RemoveHash = CachedGraphs.OrderBy(x => x.Value.LastAccessTime).Select(x => x.Key).FirstOrDefault();
+					if (RemoveHash == null || RemoveHash == ContentHash.Empty)
 					{
 						break;
 					}
diff --git a/Engine/Source/Programs/Horde/HordeServer/Collections/Impl/SoftwareCollection.cs b/Engine/Source/Programs/Horde/HordeServer/Collections/Impl/SoftwareCollection.cs
--- a/Engine/Source/Programs/Horde/HordeServer/Collections/Impl/SoftwareCollection.cs
+++ b/Engine/Source/Programs/Horde/HordeServer/Collections/Impl/SoftwareCollection.cs
@@ -113,13 +113,16 @@ namespace HordeServer.Collections.Impl
 							using MemoryStream MemoryStream = new MemoryStream();
 							InputEntryStream.CopyTo(MemoryStream);
 
-							Dictionary<string, Dictionary<string, object>> Document = JsonSerializer.Deserialize<Dictionary<string, Dictionary<string, object>>>(MemoryStream.ToArray());
-							Document["Horde"]["Version"] = Version;
+							Dictionary<string, Dictionary<string, object>>? Document = JsonSerializer.Deserialize<Dictionary<string, Dictionary<string, object>>>(MemoryStream.ToArray());
+							if (Document != null)
+							{
+								Document["Horde"]["Version"] = Version;
 
-							using Utf8JsonWriter Writer = new Utf8JsonWriter(OutputEntryStream, new JsonWriterOptions { Indented = true });
-							JsonSerializer.Serialize<Dictionary<string, Dictionary<string, object>>>(Writer, Document, new JsonSerializerOptions { WriteIndented = true });
+								using Utf8JsonWriter Writer = new Utf8JsonWriter(OutputEntryStream, new JsonWriterOptions { Indented = true });
+								JsonSerializer.Serialize<Dictionary<string, Dictionary<string, object>>>(Writer, Document, new JsonSerializerOptions { WriteIndented = true });
 
-							bWrittenClientId = true;
+								bWrittenClientId = true;
+							}
 						}
 						else
 						{
diff --git a/Engine/Source/Programs/Horde/HordeServer/Controllers/AccountController.cs b/Engine/Source/Programs/Horde/HordeServer/Controllers/AccountController.cs
--- a/Engine/Source/Programs/Horde/HordeServer/Controllers/AccountController.cs
+++ b/Engine/Source/Programs/Horde/HordeServer/Controllers/AccountController.cs
@@ -60,9 +60,9 @@ namespace HordeServer.Controllers
 		{
 			StringBuilder Content = new StringBuilder();
 			Content.Append($"<html><style>{StyleSheet}</style><h1>Horde Server</h1>");
-			if (User.Identity.IsAuthenticated)
+			if (User.Identity?.IsAuthenticated ?? false)
 			{
-				Content.Append($"<p>User <b>{User.Identity.Name}</b> is logged in. <a href=\"/account/logout\">Log out</a></p>");
+				Content.Append($"<p>User <b>{User.Identity?.Name}</b> is logged in. <a href=\"/account/logout\">Log out</a></p>");
 				if (await AclService.AuthorizeAsync(AclAction.AdminWrite, User))
 				{
 					Content.Append("<p>");
@@ -74,7 +74,7 @@ namespace HordeServer.Controllers
 					Content.Append("<a href=\"/api/v1/admin/chainedjobtoken\">Get chained job token</a><br/>");
 					Content.Append("</p>");
 				}
-				Content.Append($"<p>Claims for {User.Identity.Name}:");
+				Content.Append($"<p>Claims for {User.Identity?.Name}:");
 				Content.Append("<table>");
 				foreach (System.Security.Claims.Claim Claim in User.Claims)
 				{
diff --git a/Engine/Source/Programs/Horde/HordeServer/Controllers/AdminController.cs b/Engine/Source/Programs/Horde/HordeServer/Controllers/AdminController.cs
--- a/Engine/Source/Programs/Horde/HordeServer/Controllers/AdminController.cs
+++ b/Engine/Source/Programs/Horde/HordeServer/Controllers/AdminController.cs
@@ -198,7 +198,7 @@ namespace HordeServer.Controllers
 			}
 
 			List<AclClaim> Claims = new List<AclClaim>();
-			Claims.Add(new AclClaim(ClaimTypes.Name, User.Identity.Name ?? "Unknown"));
+			Claims.Add(new AclClaim(ClaimTypes.Name, User.Identity?.Name ?? "Unknown"));
 			Claims.Add(AclService.AgentRegistrationClaim);
 
 			return AclService.IssueBearerToken(Claims, null);
@@ -218,7 +218,7 @@ namespace HordeServer.Controllers
 			}
 
 			List<AclClaim> Claims = new List<AclClaim>();
-			Claims.Add(new AclClaim(ClaimTypes.Name, User.Identity.Name ?? "Unknown"));
+			Claims.Add(new AclClaim(ClaimTypes.Name, User.Identity?.Name ?? "Unknown"));
 			Claims.Add(AclService.UploadSoftwareClaim);
 
 			return AclService.IssueBearerToken(Claims, null);
@@ -238,7 +238,7 @@ namespace HordeServer.Controllers
 			}
 
 			List<AclClaim> Claims = new List<AclClaim>();
-			Claims.Add(new AclClaim(ClaimTypes.Name, User.Identity.Name ?? "Unknown"));
+			Claims.Add(new AclClaim(ClaimTypes.Name, User.Identity?.Name ?? "Unknown"));
 			Claims.Add(AclService.DownloadSoftwareClaim);
 
 			return AclService.IssueBearerToken(Claims, null);
@@ -258,7 +258,7 @@ namespace HordeServer.Controllers
 			}
 
 			List<AclClaim> Claims = new List<AclClaim>();
-			Claims.Add(new AclClaim(ClaimTypes.Name, User.Identity.Name ?? "Unknown"));
+			Claims.Add(new AclClaim(ClaimTypes.Name, User.Identity?.Name ?? "Unknown"));
 			Claims.Add(AclService.ConfigureProjectsClaim);
 
 			return AclService.IssueBearerToken(Claims, null);
diff --git a/Engine/Source/Programs/Horde/HordeServer/Controllers/AgentSoftwareController.cs b/Engine/Source/Programs/Horde/HordeServer/Controllers/AgentSoftwareController.cs
--- a/Engine/Source/Programs/Horde/HordeServer/Controllers/AgentSoftwareController.cs
+++ b/Engine/Source/Programs/Horde/HordeServer/Controllers/AgentSoftwareController.cs
@@ -123,7 +123,7 @@ namespace HordeServer.Controllers
 				Data = MemoryStream.ToArray();
 			}
 
-			await AgentSoftwareService.SetArchiveAsync(new AgentSoftwareChannelName(Name), User.Identity.Name, Data);
+			await AgentSoftwareService.SetArchiveAsync(new AgentSoftwareChannelName(Name), User.Identity?.Name, Data);
 			return Ok();
 		}
 
diff --git a/Engine/Source/Programs/Horde/HordeServer/Controllers/ArtifactsController.cs b/Engine/Source/Programs/Horde/HordeServer/Controllers/ArtifactsController.cs
--- a/Engine/Source/Programs/Horde/HordeServer/Controllers/ArtifactsController.cs
+++ b/Engine/Source/Programs/Horde/HordeServer/Controllers/ArtifactsController.cs
@@ -408,7 +408,7 @@ namespace HordeServer.Controllers
 								{
 									foreach (IJobStepBatch Batch in Job.Batches)
 									{
-										IJobStep Step;
+										IJobStep? Step;
 										if (Batch.TryGetStep(Artifact.StepId.Value, out Step))
 										{
 											StepName = Graph.Groups[Batch.GroupIdx].Nodes[Step.NodeIdx].Name;
diff --git a/Engine/Source/Programs/Horde/HordeServer/Controllers/DevicesController.cs b/Engine/Source/Programs/Horde/HordeServer/Controllers/DevicesController.cs
--- a/Engine/Source/Programs/Horde/HordeServer/Controllers/DevicesController.cs
+++ b/Engine/Source/Programs/Horde/HordeServer/Controllers/DevicesController.cs
@@ -97,7 +97,7 @@ namespace HordeServer.Controllers
 			{
 				ModelId = new string(DeviceRequest.ModelId);
 
-				if (Platform.Models.FirstOrDefault(x => x == ModelId) == null)
+				if (Platform.Models?.FirstOrDefault(x => x == ModelId) == null)
 				{
 					return BadRequest($"Bad model id {ModelId} for platform {Platform.Id} on request");
 				}
@@ -219,7 +219,7 @@ namespace HordeServer.Controllers
 			{
 				ModelIdValue = new string(Update.ModelId);
 
-				if (Platform.Models.FirstOrDefault(x => x == ModelIdValue) == null)
+				if (Platform.Models?.FirstOrDefault(x => x == ModelIdValue) == null)
 				{
 					return BadRequest($"Bad model id {Update.ModelId} for platform {Platform.Id} on request");
 				}
@@ -379,7 +379,7 @@ namespace HordeServer.Controllers
 			foreach (IDevicePlatform Platform in Platforms)
 			{
 				// @todo: ACL per platform
-				Responses.Add(new GetDevicePlatformResponse(Platform.Id.ToString(), Platform.Name, Platform.Models.ToArray()));
+				Responses.Add(new GetDevicePlatformResponse(Platform.Id.ToString(), Platform.Name, Platform.Models?.ToArray() ?? Array.Empty<string>()));
 			}
 
 			return Responses;
@@ -485,7 +485,7 @@ namespace HordeServer.Controllers
 				{
 					foreach (string Model in DeviceRequest.IncludeModels)
 					{
-						if (Platform.Models.FirstOrDefault(x => x == Model) == null)
+						if (Platform.Models?.FirstOrDefault(x => x == Model) == null)
 						{
 							return BadRequest($"Unknown model {Model} for platform {DeviceRequest.PlatformId} on device reservation request");
 						}
@@ -498,7 +498,7 @@ namespace HordeServer.Controllers
 				{
 					foreach (string Model in DeviceRequest.ExcludeModels)
 					{
-						if (Platform.Models.FirstOrDefault(x => x == Model) == null)
+						if (Platform.Models?.FirstOrDefault(x => x == Model) == null)
 						{
 							return BadRequest($"Unknown model {Model} for platform {DeviceRequest.PlatformId} on device reservation request");
 						}
diff --git a/Engine/Source/Programs/Horde/HordeServer/Controllers/JobsController.cs b/Engine/Source/Programs/Horde/HordeServer/Controllers/JobsController.cs
--- a/Engine/Source/Programs/Horde/HordeServer/Controllers/JobsController.cs
+++ b/Engine/Source/Programs/Horde/HordeServer/Controllers/JobsController.cs
@@ -298,7 +298,7 @@ namespace HordeServer.Controllers
 			string? AbortedByUser = null;
 			if (Request.Aborted ?? false)
 			{
-				AbortedByUser = User.Identity.Name;
+				AbortedByUser = User.Identity?.Name;
 			}
 
 			if (!await JobService.UpdateJobAsync(Job, Name: Request.Name, Priority: Request.Priority, AutoSubmit: Request.AutoSubmit, AbortedByUser: AbortedByUser, Arguments: Request.Arguments))
@@ -825,7 +825,7 @@ namespace HordeServer.Controllers
 				return NotFound();
 			}
 
-			IJobStepBatch Batch = Job.Batches.FirstOrDefault(x => x.Id == BatchIdValue);
+			IJobStepBatch? Batch = Job.Batches.FirstOrDefault(x => x.Id == BatchIdValue);
 			if (Batch == null)
 			{
 				return NotFound();
@@ -947,7 +947,7 @@ namespace HordeServer.Controllers
 			// Check permissions for updating this step. Only the agent executing the step can modify the state of it.
 			if (Request.State != JobStepState.Unspecified || Request.Outcome != JobStepOutcome.Unspecified)
 			{
-				IJobStepBatch Batch = Job.Batches.FirstOrDefault(x => x.Id == BatchIdValue);
+				IJobStepBatch? Batch = Job.Batches.FirstOrDefault(x => x.Id == BatchIdValue);
 				if (Batch == null)
 				{
 					return NotFound();
@@ -972,8 +972,8 @@ namespace HordeServer.Controllers
 				}
 			}
 
-			string? RetryByUser = (Request.Retry.HasValue && Request.Retry.Value) ? (User.Identity.Name ?? "Anonymous") : null;
-			string? AbortByUser = (Request.AbortRequested.HasValue && Request.AbortRequested.Value) ? (User.Identity.Name ?? "Anonymous") : null;
+			string? RetryByUser = (Request.Retry.HasValue && Request.Retry.Value) ? (User.Identity?.Name ?? "Anonymous") : null;
+			string? AbortByUser = (Request.AbortRequested.HasValue && Request.AbortRequested.Value) ? (User.Identity?.Name ?? "Anonymous") : null;
 
 			try
 			{
@@ -1405,7 +1405,7 @@ namespace HordeServer.Controllers
 			}
 
 			INotificationSubscription? Subscription;
-			if (!Job.LabelIdxToTriggerId.ContainsKey(LabelIndex) || Job.LabelIdxToTriggerId[LabelIndex] == null)
+			if (!Job.LabelIdxToTriggerId.ContainsKey(LabelIndex))
 			{
 				Subscription = null;
 			}
diff --git a/Engine/Source/Programs/Horde/HordeServer/Controllers/SoftwareController.cs b/Engine/Source/Programs/Horde/HordeServer/Controllers/SoftwareController.cs
--- a/Engine/Source/Programs/Horde/HordeServer/Controllers/SoftwareController.cs
+++ b/Engine/Source/Programs/Horde/HordeServer/Controllers/SoftwareController.cs
@@ -64,7 +64,7 @@ namespace HordeServer.Controllers
 			ObjectId Id;
 			using (System.IO.Stream Stream = File.OpenReadStream())
 			{
-				Id = await SoftwareService.CreateSoftwareAsync(Stream, User.Identity.Name, Default);
+				Id = await SoftwareService.CreateSoftwareAsync(Stream, User.Identity?.Name, Default);
 			}
 			return new CreateSoftwareResponse(Id.ToString());
 		}
diff --git a/Engine/Source/Programs/Horde/HordeServer/IssueHandlers/Impl/CompileIssueHandler.cs b/Engine/Source/Programs/Horde/HordeServer/IssueHandlers/Impl/CompileIssueHandler.cs
--- a/Engine/Source/Programs/Horde/HordeServer/IssueHandlers/Impl/CompileIssueHandler.cs
+++ b/Engine/Source/Programs/Horde/HordeServer/IssueHandlers/Impl/CompileIssueHandler.cs
@@ -60,7 +60,7 @@ namespace HordeServer.IssueHandlers.Impl
 					{
 						if (Property.NameEquals("file") && Property.Value.ValueKind == JsonValueKind.String)
 						{
-							AddSourceFile(SourceFiles, Property.Value.GetString());
+							AddSourceFile(SourceFiles, Property.Value.GetString()!);
 						}
 						if (Property.Value.HasStringProperty("type", "SourceFile") && Property.Value.TryGetStringProperty("relativePath", out string? Value))
 						{
diff --git a/Engine/Source/Programs/Horde/HordeServer/Models/Job.cs b/Engine/Source/Programs/Horde/HordeServer/Models/Job.cs
--- a/Engine/Source/Programs/Horde/HordeServer/Models/Job.cs
+++ b/Engine/Source/Programs/Horde/HordeServer/Models/Job.cs
@@ -253,7 +253,7 @@ namespace HordeServer.Models
 		/// <param name="StepId">The step id</param>
 		/// <param name="Step">On success, receives the step object</param>
 		/// <returns>True if the step was found</returns>
-		public static bool TryGetStep(this IJobStepBatch Batch, SubResourceId StepId, [NotNullWhen(true)] out IJobStep Step)
+		public static bool TryGetStep(this IJobStepBatch Batch, SubResourceId StepId, [NotNullWhen(true)] out IJobStep? Step)
 		{
 			Step = Batch.Steps.FirstOrDefault(x => x.Id == StepId);
 			return Step != null;
diff --git a/Engine/Source/Programs/Horde/HordeServer/Notifications/Impl/NotificationService.cs b/Engine/Source/Programs/Horde/HordeServer/Notifications/Impl/NotificationService.cs
--- a/Engine/Source/Programs/Horde/HordeServer/Notifications/Impl/NotificationService.cs
+++ b/Engine/Source/Programs/Horde/HordeServer/Notifications/Impl/NotificationService.cs
@@ -150,7 +150,7 @@ namespace HordeServer.Notifications.Impl
 			ObjectId? UserId = User.GetUserId();
 			if (UserId == null)
 			{
-				Logger.LogWarning("Unable to find username for principal {User}", User.Identity.Name);
+				Logger.LogWarning("Unable to find username for principal {User}", User.Identity?.Name);
 				return false;
 			}
 
diff --git a/Engine/Source/Programs/Horde/HordeServer/Notifications/Impl/SlackNotificationSink.cs b/Engine/Source/Programs/Horde/HordeServer/Notifications/Impl/SlackNotificationSink.cs
--- a/Engine/Source/Programs/Horde/HordeServer/Notifications/Impl/SlackNotificationSink.cs
+++ b/Engine/Source/Programs/Horde/HordeServer/Notifications/Impl/SlackNotificationSink.cs
@@ -823,7 +823,7 @@ namespace HordeServer.Notifications.Impl
 				return;
 			}
 
-			Match Match = Regex.Match(Action.Value, @"^issue_(\d+)_([a-zA-Z]+)_([a-fA-F0-9]{24})$");
+			Match Match = Regex.Match(Action.Value ?? String.Empty, @"^issue_(\d+)_([a-zA-Z]+)_([a-fA-F0-9]{24})$");
 			if (!Match.Success)
 			{
 				Logger.LogWarning("Could not match format of button action: {Action}", Action.Value);
@@ -1104,7 +1104,7 @@ namespace HordeServer.Notifications.Impl
 			HttpResponseMessage ResponseMessage = await Client.SendAsync(GetUserIdRequest);
 			byte[] ResponseData = await ResponseMessage.Content.ReadAsByteArrayAsync();
 
-			UserResponse UserResponse = JsonSerializer.Deserialize<UserResponse>(ResponseData);
+			UserResponse UserResponse = JsonSerializer.Deserialize<UserResponse>(ResponseData)!;
 			if(!UserResponse.Ok || UserResponse.User == null)
 			{
 				Logger.LogWarning("Unable to find Slack user id for {Email}: {Response}", Email, Encoding.UTF8.GetString(ResponseData));
@@ -1211,7 +1211,7 @@ namespace HordeServer.Notifications.Impl
 						HttpResponseMessage Response = await Client.SendAsync(SendMessageRequest);
 						byte[] ResponseBytes = await Response.Content.ReadAsByteArrayAsync();
 
-						TResponse ResponseObject = JsonSerializer.Deserialize<TResponse>(ResponseBytes);
+						TResponse ResponseObject = JsonSerializer.Deserialize<TResponse>(ResponseBytes)!;
 						if(!ResponseObject.Ok)
 						{
 							Logger.LogError("Failed to send Slack message ({Error}). Request: {Request}. Response: {Response}", ResponseObject.Error, RequestJson, Encoding.UTF8.GetString(ResponseBytes));
@@ -1263,12 +1263,12 @@ namespace HordeServer.Notifications.Impl
 			using HttpClient Client = new HttpClient();
 			Client.DefaultRequestHeaders.Add("Authorization", $"Bearer {Settings.SlackSocketToken}");
 
-			using FormUrlEncodedContent Content = new FormUrlEncodedContent(Array.Empty<KeyValuePair<string, string>>());
+			using FormUrlEncodedContent Content = new FormUrlEncodedContent(Array.Empty<KeyValuePair<string?, string?>>());
 			HttpResponseMessage Response = await Client.PostAsync(new Uri("https://slack.com/api/apps.connections.open"), Content, StoppingToken);
 
 			byte[] ResponseData = await Response.Content.ReadAsByteArrayAsync();
 
-			SocketResponse SocketResponse = JsonSerializer.Deserialize<SocketResponse>(ResponseData);
+			SocketResponse SocketResponse = JsonSerializer.Deserialize<SocketResponse>(ResponseData)!;
 			if (!SocketResponse.Ok)
 			{
 				Logger.LogWarning("Unable to get websocket url: {Response}", Encoding.UTF8.GetString(ResponseData));
@@ -1316,7 +1316,7 @@ namespace HordeServer.Notifications.Impl
 
 				// Get the message data
 				Logger.LogInformation("Slack event: {Message}", Encoding.UTF8.GetString(Buffer, 0, Length));
-				EventMessage EventMessage = JsonSerializer.Deserialize<EventMessage>(Buffer.AsSpan(0, Length));
+				EventMessage EventMessage = JsonSerializer.Deserialize<EventMessage>(Buffer.AsSpan(0, Length))!;
 
 				// Acknowledge the message
 				if (EventMessage.EnvelopeId != null)
diff --git a/Engine/Source/Programs/Horde/HordeServer/Services/ConfigService.cs b/Engine/Source/Programs/Horde/HordeServer/Services/ConfigService.cs
--- a/Engine/Source/Programs/Horde/HordeServer/Services/ConfigService.cs
+++ b/Engine/Source/Programs/Horde/HordeServer/Services/ConfigService.cs
@@ -452,7 +452,7 @@ namespace HordeServer.Services
 			JsonSerializerOptions Options = new JsonSerializerOptions();
 			Startup.ConfigureJsonSerializer(Options);
 
-			return JsonSerializer.Deserialize<T>(Data, Options);
+			return JsonSerializer.Deserialize<T>(Data, Options)!;
 		}
 
 		Task<byte[]> ReadDataAsync(Uri ConfigPath)
diff --git a/Engine/Source/Programs/Horde/HordeServer/Services/DatabaseService.cs b/Engine/Source/Programs/Horde/HordeServer/Services/DatabaseService.cs
--- a/Engine/Source/Programs/Horde/HordeServer/Services/DatabaseService.cs
+++ b/Engine/Source/Programs/Horde/HordeServer/Services/DatabaseService.cs
@@ -570,7 +570,7 @@ namespace HordeServer.Services
 		/// <param name="Chain"></param>
 		/// <param name="SslPolicyErrors"></param>
 		/// <returns>True if the certificate is allowed, false otherwise</returns>
-		bool CertificateValidationCallBack(object Sender, X509Certificate Certificate, X509Chain Chain, SslPolicyErrors SslPolicyErrors)
+		bool CertificateValidationCallBack(object Sender, X509Certificate? Certificate, X509Chain? Chain, SslPolicyErrors SslPolicyErrors)
 		{
 			// If the certificate is a valid, signed certificate, return true.
 			if (SslPolicyErrors == SslPolicyErrors.None)
diff --git a/Engine/Source/Programs/Horde/HordeServer/Services/DeviceService.cs b/Engine/Source/Programs/Horde/HordeServer/Services/DeviceService.cs
--- a/Engine/Source/Programs/Horde/HordeServer/Services/DeviceService.cs
+++ b/Engine/Source/Programs/Horde/HordeServer/Services/DeviceService.cs
@@ -292,7 +292,7 @@ namespace HordeServer.Services
 							IGraph Graph = await JobService.GetGraphAsync(Job)!;
 
 							SubResourceId StepIdValue = SubResourceId.Parse(StepId);
-							IJobStepBatch Batch = Job.Batches.FirstOrDefault(B => B.Steps.FirstOrDefault(S => S.Id == StepIdValue) != null);
+							IJobStepBatch? Batch = Job.Batches.FirstOrDefault(B => B.Steps.FirstOrDefault(S => S.Id == StepIdValue) != null);
 							if (Batch != null)
 							{
 								Step = Batch.Steps.FirstOrDefault(S => S.Id == StepIdValue)!;
diff --git a/Engine/Source/Programs/Horde/HordeServer/Services/Impl/PerforceService.cs b/Engine/Source/Programs/Horde/HordeServer/Services/Impl/PerforceService.cs
--- a/Engine/Source/Programs/Horde/HordeServer/Services/Impl/PerforceService.cs
+++ b/Engine/Source/Programs/Horde/HordeServer/Services/Impl/PerforceService.cs
@@ -890,7 +890,7 @@ namespace HordeServer.Services
 
 				foreach (string Path in Paths)
 				{
-					P4.FileMetaData Meta = Files.FirstOrDefault(File => File.DepotPath.Path == Path);
+					P4.FileMetaData? Meta = Files.FirstOrDefault(File => File.DepotPath?.Path == Path);
 
 					if (Meta == null)
 					{
diff --git a/Engine/Source/Programs/Horde/HordeServer/Services/JobService.cs b/Engine/Source/Programs/Horde/HordeServer/Services/JobService.cs
--- a/Engine/Source/Programs/Horde/HordeServer/Services/JobService.cs
+++ b/Engine/Source/Programs/Horde/HordeServer/Services/JobService.cs
@@ -669,7 +669,7 @@ namespace HordeServer.Services
 
 			foreach (IJobStepRef Step in Steps)
 			{
-				if (Step.StartTimeUtc != null && Step.FinishTimeUtc != null)
+				if (Step.FinishTimeUtc != null)
 				{
 					WaitTimeSum += Step.BatchWaitTime;
 					InitTimeSum += Step.BatchInitTime;
diff --git a/Engine/Source/Programs/Horde/HordeServer/Services/LogFileService.cs b/Engine/Source/Programs/Horde/HordeServer/Services/LogFileService.cs
--- a/Engine/Source/Programs/Horde/HordeServer/Services/LogFileService.cs
+++ b/Engine/Source/Programs/Horde/HordeServer/Services/LogFileService.cs
@@ -841,7 +841,7 @@ namespace HordeServer.Services
 				JsonElement MessageElement;
 				if (Data.TryGetProperty("renderedMessage", out MessageElement) || Data.TryGetProperty("message", out MessageElement))
 				{
-					Message = MessageElement.GetString();
+					Message = MessageElement.GetString() ?? "(Invalid)";
 				}
 				else
 				{
diff --git a/Engine/Source/Programs/Horde/HordeServer/Services/PerforceLoadBalancer.cs b/Engine/Source/Programs/Horde/HordeServer/Services/PerforceLoadBalancer.cs
--- a/Engine/Source/Programs/Horde/HordeServer/Services/PerforceLoadBalancer.cs
+++ b/Engine/Source/Programs/Horde/HordeServer/Services/PerforceLoadBalancer.cs
@@ -478,7 +478,7 @@ namespace HordeServer.Services
 			}
 
 			List<string> HostNames = new List<string>();
-			if (!Server.ResolveDns || IPAddress.TryParse(InitialHostName, out IPAddress HostAddress))
+			if (!Server.ResolveDns || IPAddress.TryParse(InitialHostName, out IPAddress? HostAddress))
 			{
 				HostNames.Add(InitialHostName);
 			}
@@ -576,7 +576,7 @@ namespace HordeServer.Services
 				{
 					if (Element.TryGetProperty("output", out JsonElement Output))
 					{
-						string Status = Output.GetString();
+						string Status = Output.GetString() ?? String.Empty;
 						switch(Status)
 						{
 							case "green":
diff --git a/Engine/Source/Programs/Horde/HordeServer/Startup.cs b/Engine/Source/Programs/Horde/HordeServer/Startup.cs
--- a/Engine/Source/Programs/Horde/HordeServer/Startup.cs
+++ b/Engine/Source/Programs/Horde/HordeServer/Startup.cs
@@ -224,7 +224,12 @@ namespace HordeServer
 		{
 			public override ObjectId Read(ref Utf8JsonReader Reader, Type TypeToConvert, JsonSerializerOptions Options)
 			{
-				return Reader.GetString().ToObjectId();
+				string? String = Reader.GetString();
+				if (String == null)
+				{
+					throw new InvalidDataException("Unable to parse object id");
+				}
+				return String.ToObjectId();
 			}
 
 			public override void Write(Utf8JsonWriter Writer, ObjectId ObjectId, JsonSerializerOptions Options)
@@ -238,7 +243,13 @@ namespace HordeServer
 			public override DateTime Read(ref Utf8JsonReader Reader, Type TypeToConvert, JsonSerializerOptions Options)
 			{
 				Debug.Assert(TypeToConvert == typeof(DateTime));
-				return DateTime.Parse(Reader.GetString(), CultureInfo.CurrentCulture);
+
+				string? String = Reader.GetString();
+				if (String == null)
+				{
+					throw new InvalidDataException("Unable to parse DateTime");
+				}
+				return DateTime.Parse(String, CultureInfo.CurrentCulture);
 			}
 
 			public override void Write(Utf8JsonWriter Writer, DateTime DateTime, JsonSerializerOptions Options)
@@ -458,7 +469,7 @@ namespace HordeServer
 				 {
 					 Options.Events.OnValidatePrincipal = Context =>
 					 {
-						 if (Context.Principal.FindFirst(HordeClaimTypes.UserId) == null)
+						 if (Context.Principal?.FindFirst(HordeClaimTypes.UserId) == null)
 						 {
 							 Context.RejectPrincipal();
 						 }
diff --git a/Engine/Source/Programs/Horde/HordeServer/Utilities/CompactBinaryFormatter.cs b/Engine/Source/Programs/Horde/HordeServer/Utilities/CompactBinaryFormatter.cs
--- a/Engine/Source/Programs/Horde/HordeServer/Utilities/CompactBinaryFormatter.cs
+++ b/Engine/Source/Programs/Horde/HordeServer/Utilities/CompactBinaryFormatter.cs
@@ -75,7 +75,7 @@ namespace HordeServer.Utilities
 				}
 				throw new Exception($"Unable to parse compact binary request: {FormatHexDump(Data, 256)}", Ex);
 			}
-			return await InputFormatterResult.SuccessAsync(CbSerializer.Deserialize(new CbField(Data), Context.ModelType));
+			return await InputFormatterResult.SuccessAsync(CbSerializer.Deserialize(new CbField(Data), Context.ModelType)!);
 		}
 
 		static string FormatHexDump(byte[] Data, int MaxLength)
@@ -123,7 +123,7 @@ namespace HordeServer.Utilities
 			}
 			else
 			{
-				Data = CbSerializer.Serialize(Context.ObjectType, Context.Object).GetView();
+				Data = CbSerializer.Serialize(Context.ObjectType!, Context.Object!).GetView();
 			}
 			await Context.HttpContext.Response.BodyWriter.WriteAsync(Data);
 		}
diff --git a/Engine/Source/Programs/Horde/HordeServer/Utilities/JsonElementExtensions.cs b/Engine/Source/Programs/Horde/HordeServer/Utilities/JsonElementExtensions.cs
--- a/Engine/Source/Programs/Horde/HordeServer/Utilities/JsonElementExtensions.cs
+++ b/Engine/Source/Programs/Horde/HordeServer/Utilities/JsonElementExtensions.cs
@@ -100,7 +100,7 @@ namespace HordeServer.Utilities
 			JsonElement Property;
 			if (Element.TryGetProperty(Name, out Property) && Property.ValueKind == JsonValueKind.String)
 			{
-				Value = Property.GetString();
+				Value = Property.GetString()!;
 				return true;
 			}
 			else
diff --git a/Engine/Source/Programs/Horde/HordeServer/Utilities/StringId.cs b/Engine/Source/Programs/Horde/HordeServer/Utilities/StringId.cs
--- a/Engine/Source/Programs/Horde/HordeServer/Utilities/StringId.cs
+++ b/Engine/Source/Programs/Horde/HordeServer/Utilities/StringId.cs
@@ -211,7 +211,7 @@ namespace HordeServer.Utilities
 		/// <inheritdoc/>
 		public override StringId<T> Read(ref Utf8JsonReader Reader, Type TypeToConvert, JsonSerializerOptions Options)
 		{
-			return new StringId<T>(Reader.GetString());
+			return new StringId<T>(Reader.GetString()!);
 		}
 
 		/// <inheritdoc/>
diff --git a/Engine/Source/Programs/Horde/HordeServer/Utilities/TransactionBuilder.cs b/Engine/Source/Programs/Horde/HordeServer/Utilities/TransactionBuilder.cs
--- a/Engine/Source/Programs/Horde/HordeServer/Utilities/TransactionBuilder.cs
+++ b/Engine/Source/Programs/Horde/HordeServer/Utilities/TransactionBuilder.cs
@@ -122,7 +122,7 @@ namespace HordeServer.Utilities
 			PropertyInfo IndexerProperty = Reflection<Dictionary<TKey, TValue>>.Indexer;
 			foreach (KeyValuePair<TKey, TValue?> Update in Updates)
 			{
-				MethodCallExpression Index = Expression.Call(Expr.Body, IndexerProperty.GetMethod, new[] { Expression.Constant(Update.Key) });
+				MethodCallExpression Index = Expression.Call(Expr.Body, IndexerProperty.GetMethod!, new[] { Expression.Constant(Update.Key) });
 				if (Update.Value == null)
 				{
 					Expression<Func<TDocument, object>> IndexExpression = Expression.Lambda<Func<TDocument, object>>(Index, Expr.Parameters[0]);
@@ -160,7 +160,7 @@ namespace HordeServer.Utilities
 			if (Body.NodeType == ExpressionType.MemberAccess)
 			{
 				MemberExpression MemberExpression = (MemberExpression)Body;
-				object? Object = Evaluate(MemberExpression.Expression, Target);
+				object? Object = Evaluate(MemberExpression.Expression!, Target);
 				PropertyInfo PropertyInfo = (PropertyInfo)MemberExpression.Member;
 				PropertyInfo.SetValue(Object, Value);
 			}
@@ -174,14 +174,14 @@ namespace HordeServer.Utilities
 			else if (Body.NodeType == ExpressionType.Call)
 			{
 				MethodCallExpression CallExpression = (MethodCallExpression)Body;
-				System.Collections.IDictionary Dictionary = (System.Collections.IDictionary)Evaluate(CallExpression.Object, Target)!;
+				System.Collections.IDictionary Dictionary = (System.Collections.IDictionary)Evaluate(CallExpression.Object!, Target)!;
 				object? Key = Evaluate(CallExpression.Arguments[0], null);
 				Dictionary[Key!] = Value;
 			}
 			else if (Body.NodeType == ExpressionType.Index)
 			{
 				IndexExpression IndexExpression = (IndexExpression)Body;
-				object? Object = Evaluate(IndexExpression.Object, Target);
+				object? Object = Evaluate(IndexExpression.Object!, Target);
 
 				object?[] Arguments = new object?[IndexExpression.Arguments.Count];
 				for (int Idx = 0; Idx < IndexExpression.Arguments.Count; Idx++)
@@ -189,7 +189,7 @@ namespace HordeServer.Utilities
 					Arguments[Idx] = Evaluate(IndexExpression.Arguments[Idx], null);
 				}
 
-				IndexExpression.Indexer.SetValue(Object, Value, Arguments);
+				IndexExpression.Indexer!.SetValue(Object, Value, Arguments);
 			}
 			else
 			{
@@ -222,14 +222,14 @@ namespace HordeServer.Utilities
 			else if (Body.NodeType == ExpressionType.Call)
 			{
 				MethodCallExpression CallExpression = (MethodCallExpression)Body;
-				System.Collections.IDictionary? Dictionary = (System.Collections.IDictionary?)Evaluate(CallExpression.Object, Target);
+				System.Collections.IDictionary? Dictionary = (System.Collections.IDictionary?)Evaluate(CallExpression.Object!, Target);
 				object? Key = Evaluate(CallExpression.Arguments[0], null);
 				Dictionary!.Remove(Key!);
 			}
 			else if (Body.NodeType == ExpressionType.Index)
 			{
 				IndexExpression IndexExpression = (IndexExpression)Body;
-				System.Collections.IDictionary? Dictionary = (System.Collections.IDictionary?)Evaluate(IndexExpression.Object, Target);
+				System.Collections.IDictionary? Dictionary = (System.Collections.IDictionary?)Evaluate(IndexExpression.Object!, Target);
 				object? Key = Evaluate(IndexExpression.Arguments[0], null);
 				Dictionary!.Remove(Key!);
 			}
@@ -255,7 +255,7 @@ namespace HordeServer.Utilities
 			if (Expression.NodeType == ExpressionType.Call)
 			{
 				MethodCallExpression CallExpression = (MethodCallExpression)Expression;
-				object? Object = Evaluate(CallExpression.Object, Parameter);
+				object? Object = Evaluate(CallExpression.Object!, Parameter);
 				object?[] Arguments = CallExpression.Arguments.Select(x => Evaluate(x, Parameter)).ToArray();
 				return CallExpression.Method.Invoke(Object, Arguments);
 			}
@@ -267,7 +267,7 @@ namespace HordeServer.Utilities
 			else if (Expression.NodeType == ExpressionType.MemberAccess)
 			{
 				MemberExpression MemberExpression = (MemberExpression)Expression;
-				object? Target = Evaluate(MemberExpression.Expression, Parameter);
+				object? Target = Evaluate(MemberExpression.Expression!, Parameter);
 
 				MemberInfo Member = MemberExpression.Member;
 				switch (Member.MemberType)
diff --git a/Engine/Source/Programs/Horde/HordeServerTests/AgentServiceTest.cs b/Engine/Source/Programs/Horde/HordeServerTests/AgentServiceTest.cs
--- a/Engine/Source/Programs/Horde/HordeServerTests/AgentServiceTest.cs
+++ b/Engine/Source/Programs/Horde/HordeServerTests/AgentServiceTest.cs
@@ -65,7 +65,7 @@ namespace HordeServerTests
         {
 	        return new ClaimsPrincipal(new ClaimsIdentity(new List<Claim>
 	        {
-		        new Claim(HordeClaimTypes.AgentSessionId, Agent.SessionId.ToString()),
+		        new Claim(HordeClaimTypes.AgentSessionId, Agent.SessionId.ToString()!),
 	        }, "TestAuthType"));
         }
     }
diff --git a/Engine/Source/Programs/Horde/HordeServerTests/DatabaseRunner.cs b/Engine/Source/Programs/Horde/HordeServerTests/DatabaseRunner.cs
--- a/Engine/Source/Programs/Horde/HordeServerTests/DatabaseRunner.cs
+++ b/Engine/Source/Programs/Horde/HordeServerTests/DatabaseRunner.cs
@@ -120,7 +120,7 @@ namespace HordeServerTests
 		
 		private string GetBinaryPath()
 		{
-			FileReference File = new FileReference(new Uri(Assembly.GetExecutingAssembly().CodeBase!).LocalPath);
+			FileReference File = new FileReference(new Uri(Assembly.GetExecutingAssembly().Location).LocalPath);
 			FileReference BinPath = FileReference.Combine(File.Directory, BinName);
 			return BinPath.FullName;
 		}
diff --git a/Engine/Source/Programs/Horde/HordeServerTests/GrpcTestHelpers.cs b/Engine/Source/Programs/Horde/HordeServerTests/GrpcTestHelpers.cs
--- a/Engine/Source/Programs/Horde/HordeServerTests/GrpcTestHelpers.cs
+++ b/Engine/Source/Programs/Horde/HordeServerTests/GrpcTestHelpers.cs
@@ -112,15 +112,15 @@ namespace HordeServerTests
 
 			if (await _channel.Reader.WaitToReadAsync())
 			{
-				_channel.Reader.TryRead(out var message);
-				Current = message;
-				return true;
-			}
-			else
-			{
-				Current = null!;
-				return false;
+				if (_channel.Reader.TryRead(out var message))
+				{
+					Current = message;
+					return true;
+				}
 			}
+
+			Current = null!;
+			return false;
 		}
 	}
 }
\ No newline at end of file
diff --git a/Engine/Source/Programs/Horde/HordeServerTests/RpcServiceTest.cs b/Engine/Source/Programs/Horde/HordeServerTests/RpcServiceTest.cs
--- a/Engine/Source/Programs/Horde/HordeServerTests/RpcServiceTest.cs
+++ b/Engine/Source/Programs/Horde/HordeServerTests/RpcServiceTest.cs
@@ -61,7 +61,7 @@ namespace HordeServerTests
 		{
 			public override ConnectionInfo Connection { get; } = null!;
 			public override IFeatureCollection Features { get; } = null!;
-			public override IDictionary<object, object> Items { get; set; } = null!;
+			public override IDictionary<object, object?> Items { get; set; } = null!;
 			public override HttpRequest Request { get; } = null!;
 			public override CancellationToken RequestAborted { get; set; }
 			public override IServiceProvider RequestServices { get; set; } = null!;
@@ -278,15 +278,15 @@ namespace HordeServerTests
 				
 				if (await Channel.Reader.WaitToReadAsync(CancellationToken))
 				{
-					Channel.Reader.TryRead(out var Message);
-					Current = Message;
-					return true;
-				}
-				else
-				{
-					Current = null!;
-					return false;
+					if (Channel.Reader.TryRead(out var Message))
+					{
+						Current = Message;
+						return true;
+					}
 				}
+
+				Current = null!;
+				return false;
 			}
 
 			public T Current { get; private set; } = null!;
diff --git a/Engine/Source/Programs/Horde/HordeServerTests/ServiceAccountAuthTest.cs b/Engine/Source/Programs/Horde/HordeServerTests/ServiceAccountAuthTest.cs
--- a/Engine/Source/Programs/Horde/HordeServerTests/ServiceAccountAuthTest.cs
+++ b/Engine/Source/Programs/Horde/HordeServerTests/ServiceAccountAuthTest.cs
@@ -97,10 +97,10 @@ namespace HordeServerTests
 			ServiceAccountAuthHandler Handler = await GetAuthHandlerAsync("ServiceAccount mytoken");
 			AuthenticateResult Result = await Handler.AuthenticateAsync();
 			Assert.IsTrue(Result.Succeeded);
-			Assert.AreEqual(3, Result.Ticket.Principal.Claims.Count());
-			Assert.AreEqual(ServiceAccountAuthHandler.AuthenticationScheme, Result.Ticket.Principal.FindFirst(ClaimTypes.Name).Value);
-			Assert.AreEqual("myvalue", Result.Ticket.Principal.FindFirst("myclaim").Value);
-			Assert.AreEqual("bar", Result.Ticket.Principal.FindFirst("foo").Value);
+			Assert.AreEqual(3, Result.Ticket!.Principal.Claims.Count());
+			Assert.AreEqual(ServiceAccountAuthHandler.AuthenticationScheme, Result.Ticket.Principal.FindFirst(ClaimTypes.Name)!.Value);
+			Assert.AreEqual("myvalue", Result.Ticket!.Principal.FindFirst("myclaim")!.Value);
+			Assert.AreEqual("bar", Result.Ticket!.Principal.FindFirst("foo")!.Value);
 		}
 	}
 }
\ No newline at end of file
