Horde: Additional logging for queue expiry and index management.

diff --git a/Engine/Source/Programs/Horde/HordeServer/Compute/Impl/RedisTaskScheduler.cs b/Engine/Source/Programs/Horde/HordeServer/Compute/Impl/RedisTaskScheduler.cs
--- a/Engine/Source/Programs/Horde/HordeServer/Compute/Impl/RedisTaskScheduler.cs
+++ b/Engine/Source/Programs/Horde/HordeServer/Compute/Impl/RedisTaskScheduler.cs
@@ -168,7 +168,11 @@ namespace HordeServer.Compute.Impl
 			ITransaction Transaction = Redis.CreateTransaction();
 			Transaction.AddCondition(Condition.KeyExists(GetQueueKey(QueueId)));
 			_ = Transaction.With(QueueIndex).AddAsync(QueueId);
-			await Transaction.ExecuteAsync();
+
+			if (await Transaction.ExecuteAsync())
+			{
+				Logger.LogInformation("Added queue {QueueId} to index", QueueId);
+			}
 
 			await Redis.PublishAsync(NewQueueChannel, QueueId);
 		}
@@ -182,7 +186,11 @@ namespace HordeServer.Compute.Impl
 			ITransaction Transaction = Redis.CreateTransaction();
 			Transaction.AddCondition(Condition.KeyNotExists(GetQueueKey(QueueId)));
 			_ = Transaction.With(QueueIndex).RemoveAsync(QueueId);
-			await Transaction.ExecuteAsync(CommandFlags.FireAndForget);
+
+			if (await Transaction.ExecuteAsync())
+			{
+				Logger.LogInformation("Removed queue {QueueId} from index", QueueId);
+			}
 		}
 
 		/// <summary>
@@ -195,6 +203,8 @@ namespace HordeServer.Compute.Impl
 			RedisList<TTask> List = GetQueue(QueueId);
 
 			long NewLength = await List.RightPushAsync(Task);
+			Logger.LogInformation("Length of queue {QueueId} is {Length}", QueueId, NewLength);
+
 			if (NewLength == 1)
 			{
 				await AddQueueToIndexAsync(QueueId);
@@ -211,6 +221,8 @@ namespace HordeServer.Compute.Impl
 			RedisList<TTask> List = GetQueue(QueueId);
 
 			long NewLength = await List.LeftPushAsync(Task);
+			Logger.LogInformation("Length of queue {QueueId} is {Length}", QueueId, NewLength);
+
 			if (NewLength == 1)
 			{
 				await AddQueueToIndexAsync(QueueId);
@@ -345,6 +357,7 @@ namespace HordeServer.Compute.Impl
 					}
 					if (Interlocked.CompareExchange(ref LocalActiveQueues, NewLocalActiveKeys, LocalActiveKeysCopy) == LocalActiveKeysCopy)
 					{
+						Logger.LogInformation("Refreshing active queue {QueueId}", QueueId);
 						await ActiveQueues.SetAsync(QueueId, DateTime.UtcNow);
 						break;
 					}
@@ -361,7 +374,7 @@ namespace HordeServer.Compute.Impl
 			HashSet<TQueueId> Keys = new HashSet<TQueueId>(await QueueIndex.MembersAsync());
 			HashSet<TQueueId> InvalidKeys = new HashSet<TQueueId>();
 
-			DateTime MinTime = DateTime.UtcNow - TimeSpan.FromMinutes(1.0);
+			DateTime MinTime = DateTime.UtcNow - TimeSpan.FromMinutes(10.0);
 
 			HashEntry<TQueueId, DateTime>[] Entries = await ActiveQueues.GetAllAsync();
 			foreach (HashEntry<TQueueId, DateTime> Entry in Entries)
