Horde: Only log queue addition removal when it actually actually changes state.

diff --git a/Engine/Source/Programs/Horde/HordeServer/Compute/Impl/RedisTaskScheduler.cs b/Engine/Source/Programs/Horde/HordeServer/Compute/Impl/RedisTaskScheduler.cs
--- a/Engine/Source/Programs/Horde/HordeServer/Compute/Impl/RedisTaskScheduler.cs
+++ b/Engine/Source/Programs/Horde/HordeServer/Compute/Impl/RedisTaskScheduler.cs
@@ -167,9 +167,9 @@ namespace HordeServer.Compute.Impl
 		{
 			ITransaction Transaction = Redis.CreateTransaction();
 			Transaction.AddCondition(Condition.KeyExists(GetQueueKey(QueueId)));
-			_ = Transaction.With(QueueIndex).AddAsync(QueueId);
+			Task<bool> WasAdded = Transaction.With(QueueIndex).AddAsync(QueueId);
 
-			if (await Transaction.ExecuteAsync())
+			if (await Transaction.ExecuteAsync() && await WasAdded)
 			{
 				Logger.LogInformation("Added queue {QueueId} to index", QueueId);
 			}
@@ -185,9 +185,9 @@ namespace HordeServer.Compute.Impl
 		{
 			ITransaction Transaction = Redis.CreateTransaction();
 			Transaction.AddCondition(Condition.KeyNotExists(GetQueueKey(QueueId)));
-			_ = Transaction.With(QueueIndex).RemoveAsync(QueueId);
+			Task<bool> WasRemoved = Transaction.With(QueueIndex).RemoveAsync(QueueId);
 
-			if (await Transaction.ExecuteAsync())
+			if (await Transaction.ExecuteAsync() && await WasRemoved)
 			{
 				Logger.LogInformation("Removed queue {QueueId} from index", QueueId);
 			}
