Horde: Fix race condition adding tasks to task queues, and adding task queues to the index.

diff --git a/Engine/Source/Programs/Horde/HordeServer/Compute/Impl/ComputeService.cs b/Engine/Source/Programs/Horde/HordeServer/Compute/Impl/ComputeService.cs
--- a/Engine/Source/Programs/Horde/HordeServer/Compute/Impl/ComputeService.cs
+++ b/Engine/Source/Programs/Horde/HordeServer/Compute/Impl/ComputeService.cs
@@ -206,7 +206,7 @@ namespace HordeServer.Compute.Impl
 			{
 				ComputeTaskInfo TaskInfo = new ComputeTaskInfo(TaskHash, ChannelId);
 				Logger.LogDebug("Adding task {TaskHash} to queue {RequirementsHash}", TaskInfo.TaskHash.Hash, RequirementsHash);
-				Tasks.Add(TaskScheduler.EnqueueAsync(RequirementsHash, TaskInfo));
+				Tasks.Add(TaskScheduler.EnqueueAsync(RequirementsHash, TaskInfo, false));
 			}
 			await Task.WhenAll(Tasks);
 		}
@@ -253,7 +253,7 @@ namespace HordeServer.Compute.Impl
 		{
 			ComputeTaskMessage Message = Payload.Unpack<ComputeTaskMessage>();
 			ComputeTaskInfo TaskInfo = new ComputeTaskInfo(Message.Task, new ChannelId(Message.ChannelId));
-			return TaskScheduler.EnqueueAsync(Message.Requirements.Hash, TaskInfo);
+			return TaskScheduler.EnqueueAsync(Message.Requirements.Hash, TaskInfo, true);
 		}
 
 		/// <inheritdoc/>
diff --git a/Engine/Source/Programs/Horde/HordeServer/Compute/Impl/RedisTaskScheduler.cs b/Engine/Source/Programs/Horde/HordeServer/Compute/Impl/RedisTaskScheduler.cs
--- a/Engine/Source/Programs/Horde/HordeServer/Compute/Impl/RedisTaskScheduler.cs
+++ b/Engine/Source/Programs/Horde/HordeServer/Compute/Impl/RedisTaskScheduler.cs
@@ -39,14 +39,8 @@ namespace HordeServer.Compute.Impl
 		/// </summary>
 		/// <param name="QueueId">The queue identifier</param>
 		/// <param name="TaskId">The task to add</param>
-		Task EnqueueAsync(TQueueId QueueId, TTask TaskId);
-
-		/// <summary>
-		/// Inserts a previously dequeued task back at the front of the queue
-		/// </summary>
-		/// <param name="QueueId">The queue identifier</param>
-		/// <param name="Task">The task to add</param>
-		Task RequeueAsync(TQueueId QueueId, TTask Task);
+		/// <param name="AtFront">Whether to insert at the front of the queue</param>
+		Task EnqueueAsync(TQueueId QueueId, TTask TaskId, bool AtFront);
 
 		/// <summary>
 		/// Dequeue any task from a particular queue
@@ -160,36 +154,17 @@ namespace HordeServer.Compute.Impl
 		}
 
 		/// <summary>
-		/// Adds a queue to the index, atomically checking that the queue exists
+		/// Pushes a task onto either end of a queue
 		/// </summary>
-		/// <param name="QueueId">The queue to add to the index</param>
-		async ValueTask AddQueueToIndexAsync(TQueueId QueueId)
+		static Task<long> PushTaskAsync(RedisList<TTask> List, TTask Task, When When, CommandFlags Flags, bool AtFront)
 		{
-			ITransaction Transaction = Redis.CreateTransaction();
-			Transaction.AddCondition(Condition.KeyExists(GetQueueKey(QueueId)));
-			Task<bool> WasAdded = Transaction.With(QueueIndex).AddAsync(QueueId);
-
-			if (await Transaction.ExecuteAsync() && await WasAdded)
+			if (AtFront)
 			{
-				Logger.LogInformation("Added queue {QueueId} to index", QueueId);
+				return List.LeftPushAsync(Task, When, Flags);
 			}
-
-			await Redis.PublishAsync(NewQueueChannel, QueueId);
-		}
-
-		/// <summary>
-		/// Removes a queue from the index, atomically checking that it is empty
-		/// </summary>
-		/// <param name="QueueId">The queue to remove</param>
-		async ValueTask RemoveQueueFromIndexAsync(TQueueId QueueId)
-		{
-			ITransaction Transaction = Redis.CreateTransaction();
-			Transaction.AddCondition(Condition.KeyNotExists(GetQueueKey(QueueId)));
-			Task<bool> WasRemoved = Transaction.With(QueueIndex).RemoveAsync(QueueId);
-
-			if (await Transaction.ExecuteAsync() && await WasRemoved)
+			else
 			{
-				Logger.LogInformation("Removed queue {QueueId} from index", QueueId);
+				return List.RightPushAsync(Task, When, Flags);
 			}
 		}
 
@@ -198,34 +173,31 @@ namespace HordeServer.Compute.Impl
 		/// </summary>
 		/// <param name="QueueId">The queue to add the task to</param>
 		/// <param name="Task">Task to be scheduled</param>
-		public async Task EnqueueAsync(TQueueId QueueId, TTask Task)
+		/// <param name="AtFront">Whether to add to the front of the queue</param>
+		public async Task EnqueueAsync(TQueueId QueueId, TTask Task, bool AtFront)
 		{
 			RedisList<TTask> List = GetQueue(QueueId);
-
-			long NewLength = await List.RightPushAsync(Task);
-			Logger.LogInformation("Length of queue {QueueId} is {Length}", QueueId, NewLength);
-
-			if (NewLength == 1)
+			for (; ; )
 			{
-				await AddQueueToIndexAsync(QueueId);
-			}
-		}
+				long NewLength = await PushTaskAsync(List, Task, When.Exists, CommandFlags.None, AtFront);
+				if (NewLength > 0)
+				{
+					Logger.LogInformation("Length of queue {QueueId} is {Length}", QueueId, NewLength);
+					break;
+				}
 
-		/// <summary>
-		/// Adds an item to the front of a particular queue, creating and adding that queue to the index if necessary
-		/// </summary>
-		/// <param name="QueueId">The queue to add the task to</param>
-		/// <param name="Task">Task to be scheduled</param>
-		public async Task RequeueAsync(TQueueId QueueId, TTask Task)
-		{
-			RedisList<TTask> List = GetQueue(QueueId);
+				ITransaction Transaction = Redis.CreateTransaction();
+				_ = Transaction.With(QueueIndex).AddAsync(QueueId, CommandFlags.FireAndForget);
+				_ = PushTaskAsync(Transaction.With(List), Task, When.Always, CommandFlags.FireAndForget, AtFront);
 
-			long NewLength = await List.LeftPushAsync(Task);
-			Logger.LogInformation("Length of queue {QueueId} is {Length}", QueueId, NewLength);
+				if (await Transaction.ExecuteAsync())
+				{
+					Logger.LogInformation("Created queue {QueueId}", QueueId);
+					await Redis.PublishAsync(NewQueueChannel, QueueId);
+					break;
+				}
 
-			if (NewLength == 1)
-			{
-				await AddQueueToIndexAsync(QueueId);
+				Logger.LogDebug("EnqueueAsync() retrying...");
 			}
 		}
 
@@ -316,8 +288,16 @@ namespace HordeServer.Compute.Impl
 			TTask? Item = await GetQueue(QueueId).LeftPopAsync();
 			if (Item == null)
 			{
-				await RemoveQueueFromIndexAsync(QueueId);
+				ITransaction Transaction = Redis.CreateTransaction();
+				Transaction.AddCondition(Condition.KeyNotExists(GetQueueKey(QueueId)));
+				Task<bool> WasRemoved = Transaction.With(QueueIndex).RemoveAsync(QueueId);
+
+				if (await Transaction.ExecuteAsync() && await WasRemoved)
+				{
+					Logger.LogInformation("Removed queue {QueueId} from index", QueueId);
+				}
 			}
+
 			return Item;
 		}
 
@@ -478,7 +458,7 @@ namespace HordeServer.Compute.Impl
 			{
 				if (Entry != null)
 				{
-					await RequeueAsync(Entry.Value.QueueId, Entry.Value.TaskId);
+					await EnqueueAsync(Entry.Value.QueueId, Entry.Value.TaskId, true);
 				}
 			}
 		}
