Horde: Add a caching layer for the "dumb" blob store implementation in Horde.

diff --git a/Engine/Source/Programs/Horde/HordeServer/Startup.cs b/Engine/Source/Programs/Horde/HordeServer/Startup.cs
--- a/Engine/Source/Programs/Horde/HordeServer/Startup.cs
+++ b/Engine/Source/Programs/Horde/HordeServer/Startup.cs
@@ -342,7 +342,7 @@ namespace HordeServer
 			Services.AddSingleton<IDeviceCollection, DeviceCollection>();
 
 			// Storage
-			Services.AddSingleton<IBlobCollection, BlobCollection>();
+			Services.AddSingleton<IBlobCollection>(SP => new CachingBlobCollection(new BlobCollection(SP.GetRequiredService<IStorageBackend<BlobCollection>>()), 256 * 1024 * 1024));
 			Services.AddSingleton<IBucketCollection, BucketCollection>();
 			Services.AddSingleton<INamespaceCollection, NamespaceCollection>();
 			Services.AddSingleton<IObjectCollection, ObjectCollection>();
@@ -665,9 +665,15 @@ namespace HordeServer
 			/// <inheritdoc/>
 			public Task<Stream?> ReadAsync(string Path) => Inner.ReadAsync(Path);
 
+			/// <inheritdoc/>
+			public Task<ReadOnlyMemory<byte>?> ReadBytesAsync(string Path) => Inner.ReadBytesAsync(Path);
+
 			/// <inheritdoc/>
 			public Task WriteAsync(string Path, Stream Stream) => Inner.WriteAsync(Path, Stream);
 
+			/// <inheritdoc/>
+			public Task WriteBytesAsync(string Path, ReadOnlyMemory<byte> Data) => Inner.WriteBytesAsync(Path, Data);
+
 			/// <inheritdoc/>
 			public Task DeleteAsync(string Path) => Inner.DeleteAsync(Path);
 
diff --git a/Engine/Source/Programs/Horde/HordeServer/Storage/Collections/BlobCollection.cs b/Engine/Source/Programs/Horde/HordeServer/Storage/Collections/BlobCollection.cs
--- a/Engine/Source/Programs/Horde/HordeServer/Storage/Collections/BlobCollection.cs
+++ b/Engine/Source/Programs/Horde/HordeServer/Storage/Collections/BlobCollection.cs
@@ -2,6 +2,7 @@
 
 using EpicGames.Core;
 using HordeServer.Utilities;
+using Microsoft.Extensions.Caching.Memory;
 using Microsoft.Extensions.Options;
 using System;
 using System.Collections.Generic;
@@ -30,19 +31,46 @@ namespace HordeServer.Storage.Collections
 		}
 
 		/// <inheritdoc/>
-		public Task<Stream?> ReadAsync(NamespaceId NamespaceId, IoHash Hash)
+		public Task<Stream?> TryReadStreamAsync(NamespaceId NamespaceId, IoHash Hash)
 		{
 			string Path = GetPath(NamespaceId, Hash);
 			return StorageBackend.ReadAsync(Path);
 		}
 
 		/// <inheritdoc/>
-		public Task WriteAsync(NamespaceId NamespaceId, IoHash Hash, Stream Stream)
+		public async Task<ReadOnlyMemory<byte>?> TryReadBytesAsync(NamespaceId NamespaceId, IoHash Hash)
+		{
+			using (MemoryStream OutputStream = new MemoryStream())
+			{
+				using (Stream? InputStream = await TryReadStreamAsync(NamespaceId, Hash))
+				{
+					if (InputStream == null)
+					{
+						return null;
+					}
+
+					await InputStream.CopyToAsync(OutputStream);
+					return OutputStream.ToArray();
+				}
+			}
+		}
+
+		/// <inheritdoc/>
+		public Task WriteStreamAsync(NamespaceId NamespaceId, IoHash Hash, Stream Stream)
 		{
 			string Path = GetPath(NamespaceId, Hash);
 			return StorageBackend.WriteAsync(Path, Stream);
 		}
 
+		/// <inheritdoc/>
+		public async Task WriteBytesAsync(NamespaceId NamespaceId, IoHash Hash, ReadOnlyMemory<byte> Data)
+		{
+			using (ReadOnlyMemoryStream Stream = new ReadOnlyMemoryStream(Data))
+			{
+				await WriteStreamAsync(NamespaceId, Hash, Stream);
+			}
+		}
+
 		/// <inheritdoc/>
 		public async Task<List<IoHash>> ExistsAsync(NamespaceId NamespaceId, IEnumerable<IoHash> Hashes)
 		{
diff --git a/Engine/Source/Programs/Horde/HordeServer/Storage/Collections/BlobRpc.cs b/Engine/Source/Programs/Horde/HordeServer/Storage/Collections/BlobRpc.cs
--- a/Engine/Source/Programs/Horde/HordeServer/Storage/Collections/BlobRpc.cs
+++ b/Engine/Source/Programs/Horde/HordeServer/Storage/Collections/BlobRpc.cs
@@ -43,7 +43,7 @@ namespace HordeServer.Storage.Collections
 			NamespaceId NamespaceId = new NamespaceId(Request.NamespaceId);
 			IoHash Hash = new IoHash(Request.Hash.ToByteArray());
 
-			Stream? Stream = await BlobCollection.ReadAsync(NamespaceId, Hash);
+			Stream? Stream = await BlobCollection.TryReadStreamAsync(NamespaceId, Hash);
 			if(Stream == null)
 			{
 				throw new RpcException(new Status(StatusCode.NotFound, "Blob not found"));
diff --git a/Engine/Source/Programs/Horde/HordeServer/Storage/Collections/BlobStoreRpc.cs b/Engine/Source/Programs/Horde/HordeServer/Storage/Collections/BlobStoreRpc.cs
--- a/Engine/Source/Programs/Horde/HordeServer/Storage/Collections/BlobStoreRpc.cs
+++ b/Engine/Source/Programs/Horde/HordeServer/Storage/Collections/BlobStoreRpc.cs
@@ -46,21 +46,21 @@ namespace HordeServer.Storage.Collections
 		{
 			NamespaceId NamespaceId = new NamespaceId(Request.NamespaceId);
 
-			List<Task<byte[]?>> Tasks = new List<Task<byte[]?>>();
+			List<Task<ReadOnlyMemory<byte>?>> Tasks = new List<Task<ReadOnlyMemory<byte>?>>();
 			foreach (IoHash Hash in Request.Blobs)
 			{
-				Tasks.Add(Task.Run(() => BlobCollection.ReadBytesAsync(NamespaceId, Hash)));
+				Tasks.Add(Task.Run(() => BlobCollection.TryReadBytesAsync(NamespaceId, Hash)));
 			}
 
 			GetBlobsResponse Response = new GetBlobsResponse();
 			for (int Idx = 0; Idx < Request.Blobs.Count; Idx++)
 			{
-				byte[]? Data = await Tasks[Idx];
+				ReadOnlyMemory<byte>? Data = await Tasks[Idx];
 
 				GetBlobResponse Blob = new GetBlobResponse();
 				if (Data != null)
 				{
-					Blob.Data = ByteString.CopyFrom(Data);
+					Blob.Data = ByteString.CopyFrom(Data.Value.Span);
 				}
 				Blob.Exists = Data != null;
 				Response.Blobs.Add(Blob);
diff --git a/Engine/Source/Programs/Horde/HordeServer/Storage/Collections/CachingBlobCollection.cs b/Engine/Source/Programs/Horde/HordeServer/Storage/Collections/CachingBlobCollection.cs
new file mode 100644
--- /dev/null
+++ b/Engine/Source/Programs/Horde/HordeServer/Storage/Collections/CachingBlobCollection.cs
@@ -0,0 +1,84 @@
+// Copyright Epic Games, Inc. All Rights Reserved.
+
+using EpicGames.Core;
+using HordeServer.Utilities;
+using Microsoft.Extensions.Caching.Memory;
+using System;
+using System.Collections.Generic;
+using System.IO;
+using System.Linq;
+using System.Threading.Tasks;
+
+namespace HordeServer.Storage.Collections
+{
+	using NamespaceId = StringId<INamespace>;
+
+	/// <summary>
+	/// Implements an in-memory cache over a blob collection
+	/// </summary>
+	class CachingBlobCollection : IBlobCollection, IDisposable
+	{
+		IBlobCollection Inner;
+		MemoryCache Cache;
+
+		/// <summary>
+		/// Constructor
+		/// </summary>
+		/// <param name="Inner">The inner blob collection instance</param>
+		/// <param name="CacheSize">The cache size</param>
+		public CachingBlobCollection(IBlobCollection Inner, long CacheSize)
+		{
+			this.Inner = Inner;
+			this.Cache = new MemoryCache(new MemoryCacheOptions { SizeLimit = CacheSize });
+		}
+
+		/// <inheritdoc/>
+		public void Dispose()
+		{
+			Cache.Dispose();
+		}
+
+		/// <inheritdoc/>
+		public Task<Stream?> TryReadStreamAsync(NamespaceId NamespaceId, IoHash Hash) => Inner.TryReadStreamAsync(NamespaceId, Hash);
+
+		/// <inheritdoc/>
+		public async Task<ReadOnlyMemory<byte>?> TryReadBytesAsync(NamespaceId NamespaceId, IoHash Hash)
+		{
+			ReadOnlyMemory<byte> Data;
+			if(Cache.TryGetValue(Hash, out Data))
+			{
+				return Data;
+			}
+
+			ReadOnlyMemory<byte>? Result = await Inner.TryReadBytesAsync(NamespaceId, Hash);
+			if (Result != null)
+			{
+				using (ICacheEntry Entry = Cache.CreateEntry(Hash))
+				{
+					Entry.SetSize(Result.Value.Length);
+					Entry.SetValue(Result.Value);
+				}
+			}
+
+			return Result;
+		}
+
+		/// <inheritdoc/>
+		public Task WriteStreamAsync(NamespaceId NamespaceId, IoHash Hash, Stream Stream) => Inner.WriteStreamAsync(NamespaceId, Hash, Stream);
+
+		/// <inheritdoc/>
+		public Task WriteBytesAsync(NamespaceId NamespaceId, IoHash Hash, ReadOnlyMemory<byte> Data)
+		{
+			ReadOnlyMemory<byte> BytesCopy = Data.ToArray();
+			using (ICacheEntry Entry = Cache.CreateEntry(Hash))
+			{
+				Entry.SetSize(BytesCopy.Length);
+				Entry.SetValue(BytesCopy);
+			}
+			return Inner.WriteBytesAsync(NamespaceId, Hash, BytesCopy);
+		}
+
+		/// <inheritdoc/>
+		public Task<List<IoHash>> ExistsAsync(NamespaceId NamespaceId, IEnumerable<IoHash> Hashes) => Inner.ExistsAsync(NamespaceId, Hashes);
+	}
+}
diff --git a/Engine/Source/Programs/Horde/HordeServer/Storage/Collections/ObjectCollection.cs b/Engine/Source/Programs/Horde/HordeServer/Storage/Collections/ObjectCollection.cs
--- a/Engine/Source/Programs/Horde/HordeServer/Storage/Collections/ObjectCollection.cs
+++ b/Engine/Source/Programs/Horde/HordeServer/Storage/Collections/ObjectCollection.cs
@@ -36,12 +36,12 @@ namespace HordeServer.Storage.Collections
 		/// <inheritdoc/>
 		public async Task<CbObject?> GetAsync(NamespaceId NsId, IoHash Hash)
 		{
-			byte[]? Data = await BlobCollection.ReadBytesAsync(NsId, Hash);
+			ReadOnlyMemory<byte>? Data = await BlobCollection.ReadBytesAsync(NsId, Hash);
 			if(Data == null)
 			{
 				return null;
 			}
-			return new CbObject(Data);
+			return new CbObject(Data.Value);
 		}
 
 		/// <inheritdoc/>
diff --git a/Engine/Source/Programs/Horde/HordeServer/Storage/Controllers/BlobsController.cs b/Engine/Source/Programs/Horde/HordeServer/Storage/Controllers/BlobsController.cs
--- a/Engine/Source/Programs/Horde/HordeServer/Storage/Controllers/BlobsController.cs
+++ b/Engine/Source/Programs/Horde/HordeServer/Storage/Controllers/BlobsController.cs
@@ -106,7 +106,7 @@ namespace HordeServer.Storage.Controllers
 
 			using (MemoryStream MemoryStream = new MemoryStream(Data))
 			{
-				await BlobCollection.WriteAsync(NamespaceId, Hash, MemoryStream);
+				await BlobCollection.WriteStreamAsync(NamespaceId, Hash, MemoryStream);
 			}
 
 			return new PutBlobResponse { Hash = Hash };
@@ -151,7 +151,7 @@ namespace HordeServer.Storage.Controllers
 				return Forbid();
 			}
 
-			Stream? Stream = await BlobCollection.ReadAsync(NamespaceId, Hash);
+			Stream? Stream = await BlobCollection.TryReadStreamAsync(NamespaceId, Hash);
 			if (Stream == null)
 			{
 				return NotFound();
diff --git a/Engine/Source/Programs/Horde/HordeServer/Storage/IBlobCollection.cs b/Engine/Source/Programs/Horde/HordeServer/Storage/IBlobCollection.cs
--- a/Engine/Source/Programs/Horde/HordeServer/Storage/IBlobCollection.cs
+++ b/Engine/Source/Programs/Horde/HordeServer/Storage/IBlobCollection.cs
@@ -32,7 +32,15 @@ namespace HordeServer.Storage
 		/// <param name="NamespaceId">Namespace to operate on</param>
 		/// <param name="Hash">Hash of the blob</param>
 		/// <returns>Stream for the blob, or null if it does not exist</returns>
-		Task<Stream?> ReadAsync(NamespaceId NamespaceId, IoHash Hash);
+		Task<Stream?> TryReadStreamAsync(NamespaceId NamespaceId, IoHash Hash);
+
+		/// <summary>
+		/// Reads a blob as an array of bytes. Throws an exception if the blob is too large.
+		/// </summary>
+		/// <param name="NamespaceId">Namespace to operate on</param>
+		/// <param name="Hash">Hash of the blob</param>
+		/// <returns>Stream for the blob, or null if it does not exist</returns>
+		Task<ReadOnlyMemory<byte>?> TryReadBytesAsync(NamespaceId NamespaceId, IoHash Hash);
 
 		/// <summary>
 		/// Adds an item to storage
@@ -40,7 +48,16 @@ namespace HordeServer.Storage
 		/// <param name="NamespaceId">Namespace to operate on</param>
 		/// <param name="Hash">Hash of the blob</param>
 		/// <param name="Stream">The stream to write</param>
-		Task WriteAsync(NamespaceId NamespaceId, IoHash Hash, Stream Stream);
+		Task WriteStreamAsync(NamespaceId NamespaceId, IoHash Hash, Stream Stream);
+
+		/// <summary>
+		/// Gets a blob as a byte array
+		/// </summary>
+		/// <param name="NamespaceId"></param>
+		/// <param name="Hash">Expected hash of the data</param>
+		/// <param name="Data">The data to be written</param>
+		/// <returns></returns>
+		Task WriteBytesAsync(NamespaceId NamespaceId, IoHash Hash, ReadOnlyMemory<byte> Data);
 
 		/// <summary>
 		/// Determines which of a set of blobs exist
@@ -51,6 +68,26 @@ namespace HordeServer.Storage
 		Task<List<IoHash>> ExistsAsync(NamespaceId NamespaceId, IEnumerable<IoHash> Hashes);
 	}
 
+	/// <summary>
+	/// Exception thrown for missing blobs
+	/// </summary>
+	public sealed class BlobNotFoundException : Exception
+	{
+		/// <summary>
+		/// Hash of the missing blob
+		/// </summary>
+		public IoHash Hash { get; }
+
+		/// <summary>
+		/// Constructor
+		/// </summary>
+		/// <param name="Hash">Hash of the missing blob</param>
+		public BlobNotFoundException(IoHash Hash) : base($"Unable to find blob {Hash}")
+		{
+			this.Hash = Hash;
+		}
+	}
+
 	/// <summary>
 	/// Extension methods for <see cref="IBlobCollection"/>
 	/// </summary>
@@ -63,21 +100,14 @@ namespace HordeServer.Storage
 		/// <param name="NamespaceId"></param>
 		/// <param name="Hash"></param>
 		/// <returns></returns>
-		public static async Task<byte[]?> ReadBytesAsync(this IBlobCollection Collection, NamespaceId NamespaceId, IoHash Hash)
+		public static async Task<ReadOnlyMemory<byte>> ReadBytesAsync(this IBlobCollection Collection, NamespaceId NamespaceId, IoHash Hash)
 		{
-			using (MemoryStream OutputStream = new MemoryStream())
+			ReadOnlyMemory<byte>? Result = await Collection.TryReadBytesAsync(NamespaceId, Hash);
+			if (Result == null)
 			{
-				using (Stream? InputStream = await Collection.ReadAsync(NamespaceId, Hash))
-				{
-					if(InputStream == null)
-					{
-						return null;
-					}
-
-					await InputStream.CopyToAsync(OutputStream);
-					return OutputStream.ToArray();
-				}
+				throw new BlobNotFoundException(Hash);
 			}
+			return Result.Value;
 		}
 
 		/// <summary>
@@ -90,26 +120,10 @@ namespace HordeServer.Storage
 		public static async Task<IoHash> WriteBytesAsync(this IBlobCollection Collection, NamespaceId NamespaceId, ReadOnlyMemory<byte> Data)
 		{
 			IoHash Hash = IoHash.Compute(Data.Span);
-			await WriteBytesAsync(Collection, NamespaceId, Hash, Data);
+			await Collection.WriteBytesAsync(NamespaceId, Hash, Data);
 			return Hash;
 		}
 
-		/// <summary>
-		/// Gets a blob as a byte array
-		/// </summary>
-		/// <param name="Collection"></param>
-		/// <param name="NamespaceId"></param>
-		/// <param name="Hash">Expected hash of the data</param>
-		/// <param name="Data">The data to be written</param>
-		/// <returns></returns>
-		public static async Task WriteBytesAsync(this IBlobCollection Collection, NamespaceId NamespaceId, IoHash Hash, ReadOnlyMemory<byte> Data)
-		{
-			using (ReadOnlyMemoryStream Stream = new ReadOnlyMemoryStream(Data))
-			{
-				await Collection.WriteAsync(NamespaceId, Hash, Stream);
-			}
-		}
-
 		/// <summary>
 		/// Test whether a single blob exists
 		/// </summary>
diff --git a/Engine/Source/Programs/Horde/HordeServer/Storage/IStorageBackend.cs b/Engine/Source/Programs/Horde/HordeServer/Storage/IStorageBackend.cs
--- a/Engine/Source/Programs/Horde/HordeServer/Storage/IStorageBackend.cs
+++ b/Engine/Source/Programs/Horde/HordeServer/Storage/IStorageBackend.cs
@@ -100,16 +100,18 @@ namespace HordeServer.Storage
 		/// <param name="StorageBackend"></param>
 		/// <param name="Path"></param>
 		/// <returns></returns>
-		public static async Task<ReadOnlyMemory<byte>> ReadBytesAsync(this IStorageBackend StorageBackend, string Path)
+		public static async Task<ReadOnlyMemory<byte>?> ReadBytesAsync(this IStorageBackend StorageBackend, string Path)
 		{
 			using (Stream? InputStream = await StorageBackend.ReadAsync(Path))
 			{
+				if (InputStream == null)
+				{
+					return null;
+				}
+
 				using (MemoryStream OutputStream = new MemoryStream())
 				{
-					if (InputStream != null)
-					{
-						await InputStream.CopyToAsync(OutputStream);
-					}
+					await InputStream.CopyToAsync(OutputStream);
 					return OutputStream.ToArray();
 				}
 			}
