Horde: Support for replicating stream hierarchies to the blob store. Currently disabled, pending implementation of code to pack index data into blobs.

diff --git a/Engine/Source/Programs/Horde/HordeServer/Commits/Impl/CommitService.cs b/Engine/Source/Programs/Horde/HordeServer/Commits/Impl/CommitService.cs
--- a/Engine/Source/Programs/Horde/HordeServer/Commits/Impl/CommitService.cs
+++ b/Engine/Source/Programs/Horde/HordeServer/Commits/Impl/CommitService.cs
@@ -22,12 +22,13 @@ using System.Collections.Concurrent;
 using System.Globalization;
 using Microsoft.Extensions.Hosting;
 using EpicGames.Perforce.Managed;
-using Microsoft.Extensions.Caching.Memory;
 using System.Threading.Channels;
 using StackExchange.Redis;
 using EpicGames.Perforce;
 using EpicGames.Redis;
 using System.Diagnostics;
+using System.IO.Compression;
+using System.IO;
 
 namespace HordeServer.Commits.Impl
 {
@@ -66,6 +67,7 @@ namespace HordeServer.Commits.Impl
 
 		// Collections
 		ICommitCollection CommitCollection;
+		IBlobCollection BlobCollection;
 		IObjectCollection ObjectCollection;
 		IRefCollection RefCollection;
 		IStreamCollection StreamCollection;
@@ -80,13 +82,14 @@ namespace HordeServer.Commits.Impl
 		/// <summary>
 		/// Constructor
 		/// </summary>
-		public CommitService(IDatabase Redis, ICommitCollection CommitCollection, IObjectCollection ObjectCollection, IRefCollection RefCollection, IStreamCollection StreamCollection, IPerforceService PerforceService, IUserCollection UserCollection, DatabaseService DatabaseService, ILogger<CommitService> Logger)
+		public CommitService(IDatabase Redis, ICommitCollection CommitCollection, IBlobCollection BlobCollection, IObjectCollection ObjectCollection, IRefCollection RefCollection, IStreamCollection StreamCollection, IPerforceService PerforceService, IUserCollection UserCollection, DatabaseService DatabaseService, ILogger<CommitService> Logger)
 		{
 			this.Redis = Redis;
 			this.RedisDirtyStreams = new RedisSet<StreamId>(Redis, RedisBaseKey.Append("streams"));
 			this.RedisReservations = new RedisSortedSet<StreamId>(Redis, RedisBaseKey.Append("reservations"));
 
 			this.CommitCollection = CommitCollection;
+			this.BlobCollection = BlobCollection;
 			this.ObjectCollection = ObjectCollection;
 			this.RefCollection = RefCollection;
 			this.StreamCollection = StreamCollection;
@@ -436,7 +439,7 @@ namespace HordeServer.Commits.Impl
 			{
 				try
 				{
-					await UpdateStreamsInternalAsync(BackgroundTasks);
+					await UpdateTreesInternalAsync(BackgroundTasks);
 				}
 				catch (Exception Ex)
 				{
@@ -445,7 +448,7 @@ namespace HordeServer.Commits.Impl
 			}
 		}
 
-		async Task UpdateStreamsInternalAsync(List<(StreamId, Task)> BackgroundTasks)
+		async Task UpdateTreesInternalAsync(List<(StreamId, Task)> BackgroundTasks)
 		{
 			// Remove any complete background tasks
 			for (int Idx = BackgroundTasks.Count - 1; Idx >= 0; Idx--)
@@ -483,7 +486,7 @@ namespace HordeServer.Commits.Impl
 						{
 							if (await RedisReservations.AddAsync(CheckStream, NewScore, When.NotExists))
 							{
-								Task NewTask = Task.Run(() => UpdateStreamAsync(CheckStream));
+								Task NewTask = Task.Run(() => UpdateStreamTreesAsync(CheckStream));
 								BackgroundTasks.Add((CheckStream, NewTask));
 
 								if (BackgroundTasks.Count == MaxBackgroundTasks)
@@ -521,7 +524,7 @@ namespace HordeServer.Commits.Impl
 			}
 		}
 
-		async Task UpdateStreamAsync(StreamId StreamId)
+		async Task UpdateStreamTreesAsync(StreamId StreamId)
 		{
 			Stopwatch Timer = Stopwatch.StartNew();
 
@@ -529,7 +532,7 @@ namespace HordeServer.Commits.Impl
 			for (; ; )
 			{
 				// Update the stream, updating the reservation every 30 seconds
-				Task InternalTask = Task.Run(() => UpdateStreamInternalAsync(StreamId, StreamChanges));
+				Task InternalTask = Task.Run(() => UpdateStreamTreesInternalAsync(StreamId, StreamChanges));
 				while (!InternalTask.IsCompleted)
 				{
 					Task DelayTask = Task.Delay(TimeSpan.FromSeconds(15.0));
@@ -554,7 +557,7 @@ namespace HordeServer.Commits.Impl
 			await RedisReservations.RemoveAsync(StreamId);
 		}
 
-		async Task UpdateStreamInternalAsync(StreamId StreamId, RedisList<int> Changes)
+		async Task UpdateStreamTreesInternalAsync(StreamId StreamId, RedisList<int> Changes)
 		{
 			IStream? Stream = await StreamCollection.GetAsync(StreamId);
 			if (Stream == null)
@@ -570,10 +573,13 @@ namespace HordeServer.Commits.Impl
 
 			InfoRecord ServerInfo = await Perforce.GetInfoAsync(InfoOptions.ShortOutput);
 
-			ViewMap View = await GetStreamViewAsync(Perforce, "//UE5/Dev-Main-Minimal");// Stream.Name);
+			ViewMap View = await GetStreamViewAsync(Perforce, Stream.Name);
 
-			ICommit? Commit = null;
-			StreamTreeRef? TreeRef = null;
+#if ENABLE_TREE_MIRRORING
+			ICommit? PrevCommit = null;
+			StreamTreeRef? PrevRoot = null;
+#endif
+			ObjectSet ObjectSet = new ObjectSet(BlobCollection, NamespaceId, 256 * 1024, DateTime.UtcNow);
 			for (; ; )
 			{
 				// Get the first two changelists to update
@@ -582,46 +588,66 @@ namespace HordeServer.Commits.Impl
 				{
 					break;
 				}
+#if ENABLE_TREE_MIRRORING
+				// Clear the previous commit if it no longer matches
+				if (PrevCommit != null && PrevCommit.Change != Values[0])
+				{
+					PrevCommit = null;
+				}
 
 				// Get the previous tree if we don't have it from the previous iteration
-				if (Commit == null || Commit.Change != Values[0])
+				if (PrevCommit == null)
 				{
-					Commit = await CommitCollection.GetCommitAsync(Stream.Id, Values[0]);
-					if (Commit == null)
+					PrevCommit = await CommitCollection.GetCommitAsync(Stream.Id, Values[0]);
+					if (PrevCommit == null)
 					{
-						TreeRef = null;
+						await Changes.LeftPopAsync();
+						continue;
 					}
-					else if (Commit.TreeRef == null)
+
+					bool HasTree = false;
+					if (PrevCommit.TreeRef != null)
 					{
-						TreeRef = await AddTreeRefAsync(Perforce, Stream, View, null, null, Commit, ServerInfo.Utf8PathComparer);
+						CommitTree? PrevCommitTree = await ReadCommitTreeAsync(PrevCommit);
+						if (PrevCommitTree != null)
+						{
+							try
+							{
+								throw new NotImplementedException();
+//								await Builder.ResetAsync(PrevCommitTree);
+//								HasTree = true;
+							}
+							catch (Exception Ex)
+							{
+								Logger.LogError(Ex, "Unable to read commit tree for {StreamId} CL {Change}", Stream.Id, PrevCommit.Change);
+							}
+						}
 					}
-					else
+
+					if (!HasTree)
 					{
-						TreeRef = await GetTreeRefAsync(Commit);
+						PrevRoot = await AddTreeRefAsync(Perforce, Stream, View, null, null, PrevCommit, ObjectSet, ServerInfo.Utf8PathComparer);
 					}
 				}
 
-				// Copy the previous commit and tree ref
-				ICommit? PrevCommit = Commit;
-				StreamTreeRef? PrevTreeRef = TreeRef;
-
 				// Move to the next commit
-				Commit = await CommitCollection.GetCommitAsync(StreamId, Values[1]);
-				if (Commit == null)
-				{
-					TreeRef = null;
-				}
-				else
+				ICommit? Commit = await CommitCollection.GetCommitAsync(StreamId, Values[1]);
+				StreamTreeRef? Root = null;
+				if (Commit != null)
 				{
-					TreeRef = await AddTreeRefAsync(Perforce, Stream, View, PrevCommit, PrevTreeRef, Commit, ServerInfo.Utf8PathComparer);
+					Root = await AddTreeRefAsync(Perforce, Stream, View, PrevCommit, PrevRoot, Commit, ObjectSet, ServerInfo.Utf8PathComparer);
 				}
 
+				// Move to the next commit
+				PrevCommit = Commit;
+				PrevRoot = Root;
+#endif
 				// Remove the first change number from this queue
 				await Changes.LeftPopAsync();
 			}
 		}
 
-		async Task<StreamTreeRef?> GetTreeRefAsync(ICommit Commit)
+		async Task<CommitTree?> GetCommitTreeAsync(ICommit Commit)
 		{
 			if (Commit.TreeRef == null)
 			{
@@ -634,53 +660,42 @@ namespace HordeServer.Commits.Impl
 				return null;
 			}
 
-			return CbSerializer.Deserialize<StreamTreeRef>(Ref.Value.AsField());
+			return CbSerializer.Deserialize<CommitTree>(Ref.Value.AsField());
 		}
 
-		async Task<StreamTreeRef?> AddTreeRefAsync(IPerforceConnection Perforce, IStream Stream, ViewMap View, ICommit? PrevCommit, StreamTreeRef? PrevTreeRef, ICommit Commit, Utf8StringComparer PathComparer)
+		async Task<StreamTreeRef> AddTreeRefAsync(IPerforceConnection Perforce, IStream Stream, ViewMap View, ICommit? PrevCommit, StreamTreeRef? PrevRoot, ICommit Commit, ObjectSet ObjectSet, Utf8StringComparer PathComparer)
 		{
-			StreamTree Tree;
-			if (PrevCommit == null || PrevTreeRef == null)
+			StreamTreeRef Root;
+			if (PrevCommit == null || PrevRoot == null)
 			{
 				Logger.LogInformation("Performing snapshot of {StreamId} at CL {Change}", Commit.StreamId, Commit.Change);
-				Tree = await CreateTreeAsync(Perforce, View, Commit.Change, PathComparer);
+				Root = await CreateTreeAsync(Perforce, View, Commit.Change, ObjectSet, PathComparer);
 			}
 			else
 			{
 				Logger.LogInformation("Performing incremental update of {StreamId} from CL {PrevChange} to CL {Change}", Commit.StreamId, PrevCommit.Change, Commit.Change);
-				Tree = await UpdateTreeAsync(PrevTreeRef, Perforce, View, Commit.Change, PathComparer);
+				Root = await UpdateTreeAsync(PrevRoot, Perforce, View, Commit.Change, ObjectSet, PathComparer);
 			}
 
+			CommitTree Tree = await FlushAsync(ObjectSet, Root);
 			string RefName = $"tree_{Commit.StreamId}_{Commit.Change}";
 
-			StreamTreeRef TreeRef = new StreamTreeRef(Tree.Path, Tree.ToCbObject().GetHash());
-
-			CbWriter Writer = new CbWriter();
-			Writer.BeginObject();
-			if (TreeRef.Path != Stream.Name)
-			{
-				Writer.WriteString("path", TreeRef.Path);
-			}
-			TreeRef.Write(Writer);
-			Writer.EndObject();
-
-			List<IoHash> MissingHashes = await RefCollection.SetAsync(NamespaceId, BucketId, RefName, Writer.ToObject());
+			List<IoHash> MissingHashes = await RefCollection.SetAsync(NamespaceId, BucketId, RefName, Tree.Serialize(Stream.Name));
 			if (MissingHashes.Count > 0)
 			{
-				Logger.LogWarning("Missing hashes when attempting to add ref: {MissingHashes}", String.Join(", ", MissingHashes.Select(x => x.ToString())));
-				return null;
+				throw new Exception($"Missing hashes when attempting to add ref: {String.Join(", ", MissingHashes.Select(x => x.ToString()))}");
 			}
 
 			NewCommit NewCommit = new NewCommit(Commit);
 			NewCommit.TreeRef = RefName;
 			await CommitCollection.AddOrReplaceAsync(NewCommit);
 
-			return TreeRef;
+			return Tree.Root;
 		}
 
-		#endregion
+#endregion
 
-		#region Snapshot generation
+#region Tree Snapshots
 
 		/// <summary>
 		/// Creates a snapshot of a stream at a particular changelist
@@ -688,9 +703,12 @@ namespace HordeServer.Commits.Impl
 		/// <param name="Connection">The Perforce repository to mirror</param>
 		/// <param name="View">View for the stream</param>
 		/// <param name="Change">Changelist to replicate</param>
+		/// <param name="ObjectSet"></param>
 		/// <param name="PathComparer">Comparison type for paths</param>
-		async Task<StreamTree> CreateTreeAsync(IPerforceConnection Connection, ViewMap View, int Change, Utf8StringComparer PathComparer)
+		async Task<StreamTreeRef> CreateTreeAsync(IPerforceConnection Connection, ViewMap View, int Change, ObjectSet ObjectSet, Utf8StringComparer PathComparer)
 		{
+			StreamTreeBuilder Root = new StreamTreeBuilder();
+
 			// Optimize the view by removing unnecessary inclusions of the same path
 			View = new ViewMap(View);
 			for (int Idx = 0; Idx < View.Entries.Count; Idx++)
@@ -705,13 +723,9 @@ namespace HordeServer.Commits.Impl
 				}
 			}
 
-			// Delegate to write an object
-			Func<StreamTree, IoHash> WriteObject = x => x.ToCbObject().GetHash();
-
 			// In order to optimize for the common case where we're adding multiple files to the same directory, we keep an open list of
 			// StreamTreeBuilder objects down to the last modified node. Whenever we need to move to a different node, we write the 
 			// previous path to the object store.
-			StreamTreeBuilder RootNode = new StreamTreeBuilder();
 			for(int Idx = 0; Idx < View.Entries.Count; Idx++)
 			{
 				ViewMapEntry Entry = View.Entries[Idx];
@@ -733,12 +747,11 @@ namespace HordeServer.Commits.Impl
 					}
 
 					Logger.LogInformation("Adding {Source}@{Change}", Entry.Source, Change);
-					await AddDepotFilesToSnapshotAsync(Connection, $"{Entry.Source}@{Change}", View, PathComparer, RootNode, WriteObject);
-
-					GC.Collect();
+					await AddDepotFilesToSnapshotAsync(Root, Connection, $"{Entry.Source}@{Change}", View, PathComparer, ObjectSet);
 				}
 			}
-			return RootNode.Encode(WriteObject);
+
+			return Root.EncodeRef(x => WriteTree(ObjectSet, x));
 		}
 
 		/// <summary>
@@ -768,13 +781,13 @@ namespace HordeServer.Commits.Impl
 		/// <summary>
 		/// Adds a set of files in the depot to a snapshot
 		/// </summary>
+		/// <param name="Root"></param>
 		/// <param name="Connection"></param>
 		/// <param name="FileSpec">The files to add</param>
 		/// <param name="View">View for the stream</param>
 		/// <param name="Comparer">Comparison type to use for names</param>
-		/// <param name="RootNode">The root node of the tree</param>
-		/// <param name="WriteObject">Delegate to write an object to the store</param>
-		async Task AddDepotFilesToSnapshotAsync(IPerforceConnection Connection, string FileSpec, ViewMap View, Utf8StringComparer Comparer, StreamTreeBuilder RootNode, Func<StreamTree, IoHash> WriteObject)
+		/// <param name="ObjectSet"></param>
+		static async Task AddDepotFilesToSnapshotAsync(StreamTreeBuilder Root, IPerforceConnection Connection, string FileSpec, ViewMap View, Utf8StringComparer Comparer, ObjectSet ObjectSet)
 		{
 			const string Filter = "^headAction=delete ^headAction=move/delete";
 			const string Fields = "depotFile,fileSize,digest,headType,headRev";
@@ -800,7 +813,7 @@ namespace HordeServer.Commits.Impl
 					Utf8String TargetFile;
 					if (File.DepotFile != Utf8String.Empty && View.TryMapFile(File.DepotFile, Comparer, out TargetFile))
 					{
-						await AddFileAsync(TargetFile, File, RootNode, Segments, WriteObject);
+						await AddFileAsync(Root, TargetFile, File, Segments, ObjectSet);
 					}
 				}
 			}
@@ -809,14 +822,14 @@ namespace HordeServer.Commits.Impl
 		/// <summary>
 		/// Adds a single depot file to a snapshot
 		/// </summary>
+		/// <param name="Root">The root tree builder</param>
 		/// <param name="Path">The stream path for the file</param>
 		/// <param name="MetaData">The file metadata</param>
-		/// <param name="RootNode">The last computed path through the tree</param>
 		/// <param name="Segments"></param>
-		/// <param name="WriteObject"></param>
-		async Task AddFileAsync(Utf8String Path, Utf8FStatRecord MetaData, StreamTreeBuilder RootNode, List<Utf8String> Segments, Func<StreamTree, IoHash> WriteObject)
+		/// <param name="ObjectSet"></param>
+		static async Task AddFileAsync(StreamTreeBuilder Root, Utf8String Path, Utf8FStatRecord MetaData, List<Utf8String> Segments, ObjectSet ObjectSet)
 		{
-			StreamTreeBuilder Node = RootNode;
+			StreamTreeBuilder Node = Root;
 
 			// Split the path into segments
 			Segments.Clear();
@@ -848,13 +861,13 @@ namespace HordeServer.Commits.Impl
 			}
 
 			// Collapse the rest of the current path
-			Node.EncodeChildren(WriteObject);
+			await CollapseChildrenAsync(ObjectSet, Node);
 
 			// Expand the rest of the path
 			for(; SegmentIdx < Segments.Count - 1; SegmentIdx++)
 			{
 				Utf8String Segment = Segments[SegmentIdx];
-				Node = await FindOrAddTreeAsync(Node, Segment);
+				Node = await ExpandChildAsync(ObjectSet, Node, Segment);
 			}
 
 			// Add the filename
@@ -862,36 +875,39 @@ namespace HordeServer.Commits.Impl
 			Node.NameToFile[FileName] = CreateStreamFile(MetaData);
 		}
 
-		#endregion
+#endregion
 
-		#region Incremental snapshot generation
+#region Tree Incremental Updates
 
 		/// <summary>
 		/// Creates a snapshot for the given stream, given an initial state and list of modified files
 		/// </summary>
-		/// <param name="PrevTreeRef"></param>
+		/// <param name="RootRef"></param>
 		/// <param name="Perforce"></param>
 		/// <param name="Change"></param>
 		/// <param name="View"></param>
+		/// <param name="ObjectSet"></param>
 		/// <param name="PathComparer"></param>
 		/// <returns></returns>
-		async Task<StreamTree> UpdateTreeAsync(StreamTreeRef PrevTreeRef, IPerforceConnection Perforce, ViewMap View, int Change, Utf8StringComparer PathComparer)
+		static async Task<StreamTreeRef> UpdateTreeAsync(StreamTreeRef RootRef, IPerforceConnection Perforce, ViewMap View, int Change, ObjectSet ObjectSet, Utf8StringComparer PathComparer)
 		{
-			List<FStatRecord> Records = await Perforce.FStatAsync(-1, Change, null, null, -1, FStatOptions.IncludeFileSizes, FileSpecList.Empty);
-			return await UpdateTreeAsync(PrevTreeRef, Records, View, PathComparer);
+			List<FStatRecord> Records = await Perforce.FStatAsync(-1, Change, null, null, -1, FStatOptions.IncludeFileSizes, FileSpecList.Any);
+			return await UpdateTreeAsync(RootRef, Records, View, ObjectSet, PathComparer);
 		}
 
 		/// <summary>
 		/// Creates a snapshot for the given stream, given an initial state and list of modified files
 		/// </summary>
-		/// <param name="PrevTreeRef"></param>
+		/// <param name="RootRef"></param>
 		/// <param name="Files"></param>
 		/// <param name="View"></param>
+		/// <param name="ObjectSet"></param>
 		/// <param name="PathComparer"></param>
 		/// <returns></returns>
-		async Task<StreamTree> UpdateTreeAsync(StreamTreeRef PrevTreeRef, IList<FStatRecord> Files, ViewMap View, Utf8StringComparer PathComparer)
+		static async Task<StreamTreeRef> UpdateTreeAsync(StreamTreeRef RootRef, IList<FStatRecord> Files, ViewMap View, ObjectSet ObjectSet, Utf8StringComparer PathComparer)
 		{
-			StreamTreeBuilder Tree = await ReadTreeAsync(PrevTreeRef);
+			// Read the root tree
+			StreamTreeBuilder Root = new StreamTreeBuilder(await ReadTreeAsync(ObjectSet, RootRef));
 
 			// Update the tree
 			foreach (FStatRecord File in Files)
@@ -901,7 +917,7 @@ namespace HordeServer.Commits.Impl
 				{
 					Utf8String LocalPath = LocalPathStr;
 
-					StreamTreeBuilder Node = Tree;
+					StreamTreeBuilder Node = Root;
 					for (; ; )
 					{
 						int NextIdx = LocalPath.IndexOf('/');
@@ -911,13 +927,23 @@ namespace HordeServer.Commits.Impl
 						}
 
 						Utf8String Name = LocalPath[0..NextIdx];
-						Node = await FindOrAddTreeAsync(Node, Name);
+						Node = await ExpandChildAsync(ObjectSet, Node, Name);
 
 						LocalPath = LocalPath[(NextIdx + 1)..];
 					}
 
+					FileAction Action = File.Action;
+					if (Action == FileAction.None)
+					{
+						Action = File.HeadAction;
+						if (Action == FileAction.None)
+						{
+							throw new NotImplementedException();
+						}
+					}
+
 					Utf8String FileName = LocalPath;
-					if (File.Action != FileAction.Delete && File.Action != FileAction.MoveDelete)
+					if (Action != FileAction.Delete && Action != FileAction.MoveDelete)
 					{
 						Node.NameToFile[FileName] = CreateStreamFile(File);
 					}
@@ -932,43 +958,21 @@ namespace HordeServer.Commits.Impl
 				}
 			}
 
-			return Tree.Encode(SubTree => SubTree.ToCbObject().GetHash());
+			return Root.EncodeRef(x => WriteTree(ObjectSet, x));
 		}
 
 #endregion
 
-		#region Misc tree manipulation
-
-		/// <summary>
-		/// Gets the view for a stream
-		/// </summary>
-		/// <param name="Connection">The Perforce connection</param>
-		/// <param name="StreamName">Name of the stream to be mirrored</param>
-		static async Task<ViewMap> GetStreamViewAsync(IPerforceConnection Connection, string StreamName)
-		{
-			StreamRecord Record = await Connection.GetStreamAsync(StreamName, true);
-
-			ViewMap View = new ViewMap();
-			foreach (string ViewLine in Record.View)
-			{
-				Match Match = Regex.Match(ViewLine, "^(-?)([^ ]+) *([^ ]+)$");
-				if (!Match.Success)
-				{
-					throw new PerforceException($"Unable to parse stream view: '{ViewLine}'");
-				}
-				View.Entries.Add(new ViewMapEntry(Match.Groups[1].Length == 0, Match.Groups[2].Value, Match.Groups[3].Value));
-			}
-
-			return View;
-		}
+#region Tree Builder
 
 		/// <summary>
-		/// Gets a <see cref="StreamTreeBuilder"/> for a subtree with the given name
+		/// Adds or expands a tree under the given node
 		/// </summary>
+		/// <param name="ObjectSet"></param>
 		/// <param name="Tree"></param>
 		/// <param name="Name"></param>
 		/// <returns></returns>
-		async Task<StreamTreeBuilder> FindOrAddTreeAsync(StreamTreeBuilder Tree, Utf8String Name)
+		static async Task<StreamTreeBuilder> ExpandChildAsync(ObjectSet ObjectSet, StreamTreeBuilder Tree, Utf8String Name)
 		{
 			StreamTreeBuilder? Result;
 			if (Tree.NameToTreeBuilder.TryGetValue(Name, out Result))
@@ -979,7 +983,7 @@ namespace HordeServer.Commits.Impl
 			StreamTreeRef? TreeRef;
 			if (Tree.NameToTree.TryGetValue(Name, out TreeRef))
 			{
-				Result = await ReadTreeAsync(TreeRef);
+				Result = new StreamTreeBuilder(await ReadTreeAsync(ObjectSet, TreeRef));
 
 				Tree.NameToTree.Remove(Name);
 				Tree.NameToTreeBuilder.Add(Name, Result);
@@ -994,69 +998,90 @@ namespace HordeServer.Commits.Impl
 		}
 
 		/// <summary>
-		/// Reads a tree builder for the given tree reference
+		/// Collapse children underneath the given tree builder, allowing data to be purged from memory and pushed
+		/// into persistent storage.
 		/// </summary>
-		/// <param name="TreeRef"></param>
+		/// <param name="ObjectSet"></param>
+		/// <param name="Tree"></param>
 		/// <returns></returns>
-		async Task<StreamTreeBuilder> ReadTreeAsync(StreamTreeRef TreeRef)
+		static async Task CollapseChildrenAsync(ObjectSet ObjectSet, StreamTreeBuilder Tree)
 		{
-			CbObject? Object = await ObjectCollection.GetAsync(NamespaceId, TreeRef.Hash);
-			if (Object == null)
+			foreach ((Utf8String ChildName, StreamTreeBuilder ChildBuilder) in Tree.NameToTreeBuilder)
 			{
-				throw new Exception($"Missing object {TreeRef.Hash} referenced by tree for {TreeRef.Path}");
+				await CollapseChildrenAsync(ObjectSet, ChildBuilder);
+				if (ChildBuilder.NameToFile.Count > 0 || ChildBuilder.NameToTree.Count > 0)
+				{
+					StreamTreeRef ChildTreeRef = ChildBuilder.EncodeRef(x => WriteTree(ObjectSet, x));
+					Tree.NameToTree[ChildName] = ChildTreeRef;
+				}
 			}
-			return new StreamTreeBuilder(new StreamTree(TreeRef.Path, Object));
+			Tree.NameToTreeBuilder.Clear();
 		}
 
 		/// <summary>
-		/// Creates a StreamFile from a FileMetaData object
+		/// Flush all the current data to storage, optimizing the underlying blobs if necessary
 		/// </summary>
-		/// <param name="MetaData"></param>
-		/// <returns></returns>
-		static StreamFile CreateStreamFile(Utf8FStatRecord MetaData)
+		/// <returns>The new tree object</returns>
+		static async Task<CommitTree> FlushAsync(ObjectSet ObjectSet, StreamTreeRef RootRef)
 		{
-			Utf8String FileType = Utf8String.Empty;
-			if (MetaData.Type != Utf8String.Empty)
-			{
-				FileType = MetaData.Type;
-			}
-			else if (MetaData.HeadType != Utf8String.Empty)
-			{
-				FileType = MetaData.HeadType;
-			}
-			return new StreamFile(MetaData.DepotFile!, MetaData.FileSize, new FileContentId(Md5Hash.Parse(MetaData.Digest!), FileType), MetaData.HeadRevision);
+			// Find the live set of objects
+			ObjectSet.RootSet.Clear();
+			ObjectSet.RootSet.Add(RootRef.Hash);
+
+			// Flush all the remaining objects
+			await ObjectSet.FlushAsync();
+
+			// Wait for all the blobs to finish writing and write the new tree
+			return new CommitTree(RootRef, new List<CbBinaryAttachment>(), ObjectSet.PackIndexes.ConvertAll(x => (CbBinaryAttachment)x.DataHash));
 		}
 
 		/// <summary>
-		/// Creates a StreamFile from a FileMetaData object
+		/// Reads a referenced tree from storage
 		/// </summary>
-		/// <param name="MetaData"></param>
+		/// <param name="ObjectSet">The object packer</param>
+		/// <param name="TreeRef">The tree reference</param>
 		/// <returns></returns>
-		static StreamFile CreateStreamFile(FStatRecord MetaData)
+		static async Task<StreamTree> ReadTreeAsync(ObjectSet ObjectSet, StreamTreeRef TreeRef)
 		{
-			string FileType = MetaData.Type ?? MetaData.HeadType ?? String.Empty;
-			return new StreamFile(MetaData.DepotFile!, MetaData.FileSize, new FileContentId(Md5Hash.Parse(MetaData.Digest!), FileType.ToString()), MetaData.HeadRevision);
+			ReadOnlyMemory<byte> Data = await ObjectSet.GetObjectDataAsync(TreeRef.Hash);
+			if (Data.Length == 0)
+			{
+				throw new Exception($"Unable to find tree with hash {TreeRef.Hash} for {TreeRef.Path}");
+			}
+			return new StreamTree(TreeRef.Path, new CbObject(Data));
 		}
 
 		/// <summary>
-		/// Print the contents of a snapshot
+		/// Serializes a tree to an object packer
 		/// </summary>
-		/// <param name="TreeRef"></param>
+		/// <param name="ObjectSet"></param>
+		/// <param name="Tree"></param>
 		/// <returns></returns>
-		Task PrintSnapshotAsync(StreamTreeRef TreeRef)
+		static IoHash WriteTree(ObjectSet ObjectSet, StreamTree Tree)
 		{
-			return PrintSnapshotInternalAsync(TreeRef, "/");
+			CbObject Object = Tree.ToCbObject();
+			ReadOnlyMemory<byte> Data = Object.GetView();
+
+			List<IoHash> Refs = new List<IoHash>();
+			Object.IterateAttachments(Field => Refs.Add(Field.AsHash()));
+
+			IoHash Hash = IoHash.Compute(Data.Span);
+			ObjectSet.Add(Hash, Data.Span, Refs.ToArray());
+			return Hash;
 		}
 
 		/// <summary>
-		/// Print the contents of a snapshot
+		/// Prints the state of the tree
 		/// </summary>
+		/// <param name="ObjectSet"></param>
 		/// <param name="TreeRef"></param>
-		/// <param name="Prefix"></param>
+		/// <param name="Logger"></param>
 		/// <returns></returns>
-		async Task PrintSnapshotInternalAsync(StreamTreeRef TreeRef, string Prefix)
+		public Task PrintAsync(ObjectSet ObjectSet, StreamTreeRef TreeRef, ILogger Logger) => PrintInternalAsync(ObjectSet, TreeRef, String.Empty, Logger);
+
+		async Task PrintInternalAsync(ObjectSet ObjectSet, StreamTreeRef TreeRef, string Prefix, ILogger Logger)
 		{
-			StreamTree Tree = new StreamTree(TreeRef.Path, await ObjectCollection.GetRequiredObjectAsync(NamespaceId, TreeRef.Hash));
+			StreamTree Tree = await ReadTreeAsync(ObjectSet, TreeRef);
 
 			foreach ((Utf8String Name, StreamFile File) in Tree.NameToFile)
 			{
@@ -1064,10 +1089,80 @@ namespace HordeServer.Commits.Impl
 			}
 			foreach ((Utf8String Name, StreamTreeRef ChildTreeRef) in Tree.NameToTree)
 			{
-				await PrintSnapshotInternalAsync(ChildTreeRef, $"{Prefix}/{Name}");
+				await PrintInternalAsync(ObjectSet, ChildTreeRef, $"{Prefix}/{Name}", Logger);
 			}
 		}
 
-		#endregion
+#endregion
+
+#region Tree Storage
+
+		/// <summary>
+		/// Reads the tree for a given commit
+		/// </summary>
+		/// <param name="Commit"></param>
+		async Task<CommitTree> ReadCommitTreeAsync(ICommit Commit)
+		{
+			IRef? Ref = await RefCollection.GetAsync(NamespaceId, BucketId, Commit.TreeRef!);
+			return CbSerializer.Deserialize<CommitTree>(Ref!.Value.AsField());
+		}
+
+#endregion
+
+#region Misc tree manipulation
+
+		/// <summary>
+		/// Gets the view for a stream
+		/// </summary>
+		/// <param name="Connection">The Perforce connection</param>
+		/// <param name="StreamName">Name of the stream to be mirrored</param>
+		static async Task<ViewMap> GetStreamViewAsync(IPerforceConnection Connection, string StreamName)
+		{
+			StreamRecord Record = await Connection.GetStreamAsync(StreamName, true);
+
+			ViewMap View = new ViewMap();
+			foreach (string ViewLine in Record.View)
+			{
+				Match Match = Regex.Match(ViewLine, "^(-?)([^ ]+) *([^ ]+)$");
+				if (!Match.Success)
+				{
+					throw new PerforceException($"Unable to parse stream view: '{ViewLine}'");
+				}
+				View.Entries.Add(new ViewMapEntry(Match.Groups[1].Length == 0, Match.Groups[2].Value, Match.Groups[3].Value));
+			}
+
+			return View;
+		}
+
+		/// <summary>
+		/// Creates a StreamFile from a FileMetaData object
+		/// </summary>
+		/// <param name="MetaData"></param>
+		/// <returns></returns>
+		static StreamFile CreateStreamFile(Utf8FStatRecord MetaData)
+		{
+			Utf8String FileType = Utf8String.Empty;
+			if (MetaData.Type != Utf8String.Empty)
+			{
+				FileType = MetaData.Type;
+			}
+			else if (MetaData.HeadType != Utf8String.Empty)
+			{
+				FileType = MetaData.HeadType;
+			}
+			return new StreamFile(MetaData.DepotFile!, MetaData.FileSize, new FileContentId(Md5Hash.Parse(MetaData.Digest!), FileType), MetaData.HeadRevision);
+		}
+
+		/// <summary>
+		/// Creates a StreamFile from a FileMetaData object
+		/// </summary>
+		/// <param name="MetaData"></param>
+		/// <returns></returns>
+		static StreamFile CreateStreamFile(FStatRecord MetaData)
+		{
+			string FileType = MetaData.Type ?? MetaData.HeadType ?? String.Empty;
+			return new StreamFile(MetaData.DepotFile!, MetaData.FileSize, new FileContentId(Md5Hash.Parse(MetaData.Digest!), FileType.ToString()), MetaData.HeadRevision);
+		}
+#endregion
 	}
 }
diff --git a/Engine/Source/Programs/Horde/HordeServer/Commits/Impl/CommitTree.cs b/Engine/Source/Programs/Horde/HordeServer/Commits/Impl/CommitTree.cs
new file mode 100644
--- /dev/null
+++ b/Engine/Source/Programs/Horde/HordeServer/Commits/Impl/CommitTree.cs
@@ -0,0 +1,70 @@
+// Copyright Epic Games, Inc. All Rights Reserved.
+
+using EpicGames.Core;
+using EpicGames.Perforce.Managed;
+using EpicGames.Serialization;
+using HordeServer.Storage;
+using HordeServer.Utilities;
+using Microsoft.Extensions.Logging;
+using MongoDB.Driver;
+using System;
+using System.Collections.Generic;
+using System.Diagnostics;
+using System.Linq;
+using System.Threading.Tasks;
+
+namespace HordeServer.Commits.Impl
+{
+	using NamespaceId = StringId<INamespace>;
+
+	class CommitTree
+	{
+		class ObjectFormat
+		{
+			[CbField("p")]
+			public Utf8String Path { get; set; }
+
+			[CbField("h")]
+			public IoHash Hash { get; set; }
+
+			[CbField("i")]
+			public List<CbBinaryAttachment>? IndexBlobs { get; set; }
+
+			[CbField("d")]
+			public List<CbBinaryAttachment>? DataBlobs { get; set; }
+		}
+
+		public StreamTreeRef Root { get; }
+		public List<CbBinaryAttachment> IndexBlobs { get; }
+		public List<CbBinaryAttachment> DataBlobs { get; }
+
+		public CommitTree(StreamTreeRef Root)
+			: this(Root, new List<CbBinaryAttachment>(), new List<CbBinaryAttachment>())
+		{
+		}
+
+		public CommitTree(StreamTreeRef Root, List<CbBinaryAttachment> IndexBlobs, List<CbBinaryAttachment> DataBlobs)
+		{
+			this.Root = Root;
+			this.IndexBlobs = IndexBlobs;
+			this.DataBlobs = DataBlobs;
+		}
+
+		public CbObject Serialize(Utf8String DefaultPath)
+		{
+			ObjectFormat Format = new ObjectFormat();
+			Format.Path = (Root.Path != DefaultPath) ? Root.Path : Utf8String.Empty;
+			Format.Hash = Root.Hash;
+			Format.IndexBlobs = IndexBlobs;
+			Format.DataBlobs = DataBlobs;
+			return CbSerializer.Serialize(Format);
+		}
+
+		public static CommitTree Deserialize(CbObject Object, Utf8String DefaultPath)
+		{
+			ObjectFormat Format = CbSerializer.Deserialize<ObjectFormat>(Object.AsField());
+			Utf8String Path = Format.Path.IsEmpty ? DefaultPath : Format.Path;
+			return new CommitTree(new StreamTreeRef(Path, Format.Hash), Format.IndexBlobs ?? new List<CbBinaryAttachment>(), Format.DataBlobs ?? new List<CbBinaryAttachment>());
+		}
+	}
+}
diff --git a/Engine/Source/Programs/Horde/HordeServer/Utilities/ObjectSet.cs b/Engine/Source/Programs/Horde/HordeServer/Utilities/ObjectSet.cs
new file mode 100644
--- /dev/null
+++ b/Engine/Source/Programs/Horde/HordeServer/Utilities/ObjectSet.cs
@@ -0,0 +1,512 @@
+// Copyright Epic Games, Inc. All Rights Reserved.
+
+using EpicGames.Core;
+using EpicGames.Serialization;
+using EpicGames.Serialization.Converters;
+using HordeServer.Storage;
+using System;
+using System.Collections.Generic;
+using System.Diagnostics.CodeAnalysis;
+using System.Linq;
+using System.Threading.Tasks;
+
+namespace HordeServer.Utilities
+{
+	using NamespaceId = StringId<INamespace>;
+
+	/// <summary>
+	/// Information about a blob stored in a blob pack file
+	/// </summary>
+	class ObjectPackEntry
+	{
+		public IoHash Hash { get; }
+		public int Offset { get; }
+		public int Length { get; }
+		public IoHash[] Refs { get; }
+
+		public ObjectPackEntry(IoHash Hash, int Offset, int Length, IoHash[] Refs)
+		{
+			this.Hash = Hash;
+			this.Offset = Offset;
+			this.Length = Length;
+			this.Refs = Refs;
+		}
+	}
+
+	/// <summary>
+	/// Index for a blob pack
+	/// </summary>
+	[CbConverter(typeof(ObjectPackIndexConverter))]
+	class ObjectPackIndex
+	{
+		public DateTime Time { get; }
+		public ObjectPackEntry[] Blobs { get; }
+		public IoHash DataHash { get; }
+		public int DataSize;
+
+		Dictionary<IoHash, ObjectPackEntry> HashToInfo;
+
+		public ObjectPackIndex(DateTime Time, ObjectPackEntry[] Blobs, IoHash DataHash, int DataSize)
+		{
+			this.Time = Time;
+			this.Blobs = Blobs;
+			this.DataHash = DataHash;
+			this.DataSize = DataSize;
+
+			HashToInfo = Blobs.ToDictionary(x => x.Hash, x => x);
+		}
+
+		public bool Contains(IoHash Hash) => HashToInfo.ContainsKey(Hash);
+
+		public bool TryGetEntry(IoHash Hash, [NotNullWhen(true)] out ObjectPackEntry? BlobInfo) => HashToInfo.TryGetValue(Hash, out BlobInfo);
+	}
+
+	/// <summary>
+	/// Converter for BlobPackIndex objects
+	/// </summary>
+	class ObjectPackIndexConverter : CbConverter<ObjectPackIndex>
+	{
+		class EncodeFormat
+		{
+			[CbField("time")]
+			public DateTime Time { get; set; }
+
+			[CbField("exports")]
+			public IoHash[]? Exports { get; set; }
+
+			[CbField("lengths")]
+			public int[]? Lengths { get; set; }
+
+			[CbField("refs")]
+			public IoHash[][]? Refs { get; set; }
+
+			[CbField("data")]
+			public CbBinaryAttachment DataHash { get; set; }
+
+			[CbField("size")]
+			public int DataSize { get; set; }
+		}
+
+		/// <inheritdoc/>
+		public override ObjectPackIndex Read(CbField Field)
+		{
+			EncodeFormat Format = CbSerializer.Deserialize<EncodeFormat>(Field);
+
+			ObjectPackEntry[] Objects = new ObjectPackEntry[Format.Exports!.Length];
+
+			int Offset = 0;
+			for (int Idx = 0; Idx < Format.Exports.Length; Idx++)
+			{
+				Objects[Idx] = new ObjectPackEntry(Format.Exports[Idx], Offset, Format.Lengths![Idx], Format.Refs![Idx]);
+				Offset += Format.Lengths[Idx];
+			}
+
+			return new ObjectPackIndex(Format.Time, Objects, Format.DataHash, Format.DataSize);
+		}
+
+		/// <inheritdoc/>
+		public override void Write(CbWriter Writer, ObjectPackIndex Index)
+		{
+			Writer.BeginObject();
+			WriteInternal(Writer, Index);
+			Writer.EndObject();
+		}
+
+		/// <inheritdoc/>
+		public override void WriteNamed(CbWriter Writer, Utf8String Name, ObjectPackIndex Index)
+		{
+			Writer.BeginObject(Name);
+			WriteInternal(Writer, Index);
+			Writer.EndObject();
+		}
+
+		static void WriteInternal(CbWriter Writer, ObjectPackIndex Index)
+		{
+			EncodeFormat Format = new EncodeFormat();
+			Format.Time = Index.Time;
+			Format.Exports = Index.Blobs.ConvertAll(x => x.Hash).ToArray();
+			Format.Lengths = Index.Blobs.ConvertAll(x => x.Length).ToArray();
+			Format.Refs = Index.Blobs.ConvertAll(x => x.Refs).ToArray();
+			Format.DataHash = Index.DataHash;
+			Format.DataSize = Index.DataSize;
+			CbSerializer.Serialize(Writer, Format);
+		}
+	}
+
+	/// <summary>
+	/// Helper class to maintain a set of small objects, re-packing blobs according to a heuristic to balance download performance with churn.
+	/// </summary>
+	class ObjectSet
+	{
+		readonly IBlobCollection BlobCollection;
+		readonly NamespaceId NamespaceId;
+		public int MaxPackSize { get; }
+
+		public HashSet<IoHash> RootSet { get; set; } = new HashSet<IoHash>();
+
+		DateTime Time;
+
+		int NextPackSize;
+		byte[] NextPackData;
+		List<ObjectPackEntry> NextPackEntries = new List<ObjectPackEntry>();
+		Dictionary<IoHash, ObjectPackEntry> NextPackHashToEntry = new Dictionary<IoHash, ObjectPackEntry>();
+
+		public List<ObjectPackIndex> PackIndexes { get; } = new List<ObjectPackIndex>();
+		List<Task> WriteTasks = new List<Task>();
+
+		/// <summary>
+		/// Constructor
+		/// </summary>
+		/// <param name="BlobCollection"></param>
+		/// <param name="NamespaceId"></param>
+		/// <param name="MaxPackSize"></param>
+		/// <param name="Time">The initial update time; used to determine the age of blobs</param>
+		public ObjectSet(IBlobCollection BlobCollection, NamespaceId NamespaceId, int MaxPackSize, DateTime Time)
+		{
+			this.BlobCollection = BlobCollection;
+			this.NamespaceId = NamespaceId;
+			this.MaxPackSize = MaxPackSize;
+
+			NextPackData = null!;
+
+			SetTime(Time);
+			Reset();
+		}
+
+		/// <summary>
+		/// Reset the current state of the next blob
+		/// </summary>
+		void Reset()
+		{
+			NextPackSize = 0;
+			NextPackData = new byte[MaxPackSize];
+			NextPackEntries.Clear();
+			NextPackHashToEntry.Clear();
+		}
+
+		/// <summary>
+		/// Reset the current timestamp
+		/// </summary>
+		/// <param name="Time">The new timestamp</param>
+		public void SetTime(DateTime Time)
+		{
+			this.Time = Time;
+		}
+
+		/// <summary>
+		/// Copies an existing entry into storage
+		/// </summary>
+		/// <param name="Hash">Hash of the data</param>
+		/// <param name="Data">The data buffer</param>
+		/// <param name="Refs">References to other objects</param>
+		public void Add(IoHash Hash, ReadOnlySpan<byte> Data, ReadOnlySpan<IoHash> Refs)
+		{
+			if (!NextPackHashToEntry.ContainsKey(Hash))
+			{
+				// Create enough space for the new data
+				CreateSpace(Data.Length);
+
+				// Copy the data into the buffer
+				Data.CopyTo(NextPackData.AsSpan(NextPackSize));
+
+				// Add the blob
+				ObjectPackEntry Entry = new ObjectPackEntry(Hash, NextPackSize, Data.Length, Refs.ToArray());
+				NextPackHashToEntry.Add(Hash, Entry);
+				NextPackEntries.Add(Entry);
+				NextPackSize += Data.Length;
+			}
+		}
+
+		/// <summary>
+		/// Adds an item to the packer
+		/// </summary>
+		/// <param name="Size">Size of the data</param>
+		/// <param name="ReadData">Delegate to copy the data into a span</param>
+		/// <param name="Refs">References to other objects</param>
+		public IoHash Add(int Size, Action<Memory<byte>> ReadData, IoHash[] Refs)
+		{
+			// Get the last blob and make sure there's enough space in it
+			CreateSpace(Size);
+
+			// Copy the data into the new blob
+			Memory<byte> Output = NextPackData.AsMemory(NextPackSize, Size);
+			ReadData(Output);
+
+			// Update the metadata for it
+			IoHash Hash = IoHash.Compute(Output.Span);
+			if (!NextPackHashToEntry.ContainsKey(Hash))
+			{
+				ObjectPackEntry Entry = new ObjectPackEntry(Hash, NextPackSize, Size, Refs);
+				NextPackHashToEntry.Add(Hash, Entry);
+				NextPackEntries.Add(Entry);
+				NextPackSize += Size;
+			}
+			return Hash;
+		}
+
+		/// <summary>
+		/// Finds data for an object with the given hash, from the current pack files
+		/// </summary>
+		/// <param name="Hash"></param>
+		/// <returns></returns>
+		public async Task<ReadOnlyMemory<byte>> GetObjectDataAsync(IoHash Hash)
+		{
+			await Task.WhenAll(WriteTasks);
+
+			ObjectPackEntry? Entry;
+			if (NextPackHashToEntry.TryGetValue(Hash, out Entry))
+			{
+				return NextPackData.AsMemory(Entry.Offset, Entry.Length);
+			}
+
+			foreach (ObjectPackIndex Pack in PackIndexes)
+			{
+				if (Pack.TryGetEntry(Hash, out Entry))
+				{
+					ReadOnlyMemory<byte> PackData = await BlobCollection.ReadBytesAsync(NamespaceId, Pack.DataHash);
+					return PackData.Slice(Entry.Offset, Entry.Length);
+				}
+			}
+
+			return default;
+		}
+
+		/// <summary>
+		/// Tries to find an entry for the given hash from the current set of pack files
+		/// </summary>
+		/// <param name="Hash">Hash of the </param>
+		/// <param name="Entry"></param>
+		/// <returns></returns>
+		public bool TryGetEntry(IoHash Hash, [NotNullWhen(true)] out ObjectPackEntry? Entry)
+		{
+			ObjectPackEntry? LocalEntry;
+			if (NextPackHashToEntry.TryGetValue(Hash, out LocalEntry))
+			{
+				Entry = LocalEntry;
+				return true;
+			}
+
+			foreach (ObjectPackIndex Pack in PackIndexes)
+			{
+				if (Pack.TryGetEntry(Hash, out LocalEntry))
+				{
+					Entry = LocalEntry;
+					return true;
+				}
+			}
+
+			Entry = null;
+			return false;
+		}
+
+		/// <summary>
+		/// Flush any pending blobs to disk
+		/// </summary>
+		public async Task FlushAsync()
+		{
+			// Find the live set of objects
+			HashSet<IoHash> LiveSet = new HashSet<IoHash>();
+			foreach(IoHash RootHash in RootSet)
+			{
+				FindLiveSet(RootHash, LiveSet);
+			}
+
+			// Find the total cost of all the current blobs, then loop through the blobs trying to find a more optimal arrangement
+			double TotalCost = PackIndexes.Sum(x => GetCostHeuristic(x)) + GetCostHeuristic(NextPackSize, TimeSpan.Zero);
+			for (; ; )
+			{
+				// Exclude any objects that are in the pending blobs, since we will always upload these
+				HashSet<IoHash> NewLiveSet = new HashSet<IoHash>(LiveSet);
+				NewLiveSet.ExceptWith(NextPackHashToEntry.Values.Select(x => x.Hash));
+
+				// Get the size and cost of the next blob
+				double NextBlobCost = GetCostHeuristic(NextPackSize, TimeSpan.Zero);
+
+				// Pass through all the blobs to find the best one to merge in
+				double MergeCost = TotalCost;
+				ObjectPackIndex? MergePack = null;
+				for (int Idx = PackIndexes.Count - 1; Idx >= 0; Idx--)
+				{
+					ObjectPackIndex PackIndex = PackIndexes[Idx];
+
+					// Try to merge any old blobs with the next blob
+					if (PackIndex.Time < Time)
+					{
+						// Calculate the cost of the last blob if we merge this one with it. We remove blobs as we iterate
+						// through the list, since subsequent blobs will not usefully contribute the same items.
+						double NewTotalCost = TotalCost - GetCostHeuristic(PackIndex) - NextBlobCost;
+
+						int NewNextPackSize = NextPackSize;
+						foreach (ObjectPackEntry Entry in PackIndex.Blobs)
+						{
+							if (LiveSet.Contains(Entry.Hash))
+							{
+								if (NewNextPackSize + Entry.Length > MaxPackSize)
+								{
+									NewTotalCost += GetCostHeuristic(NewNextPackSize, TimeSpan.Zero);
+									NewNextPackSize = 0;
+								}
+								NewNextPackSize += Entry.Length;
+							}
+						}
+
+						NewTotalCost += GetCostHeuristic(NewNextPackSize, TimeSpan.Zero);
+
+						// Compute the potential cost if we replace the partial blob with the useful parts of this blob
+						if (NewTotalCost < MergeCost)
+						{
+							MergePack = PackIndex;
+							MergeCost = NewTotalCost;
+						}
+					}
+
+					// Remove any items in this blob from the remaining live set. No other blobs need to include them.
+					NewLiveSet.ExceptWith(PackIndex.Blobs.Select(x => x.Hash));
+				}
+
+				// Bail out if we didn't find anything to merge
+				if (MergePack == null)
+				{
+					break;
+				}
+
+				// Get the data for this blob
+				ReadOnlyMemory<byte> MergeData = await BlobCollection.ReadBytesAsync(NamespaceId, MergePack.DataHash);
+
+				// Add anything that's still part of the live set into the new blobs
+				int Offset = 0;
+				foreach (ObjectPackEntry Blob in MergePack.Blobs)
+				{
+					if (LiveSet.Contains(Blob.Hash))
+					{
+						ReadOnlyMemory<byte> Data = MergeData.Slice(Offset, Blob.Length);
+						Add(Blob.Hash, Data.Span, Blob.Refs);
+					}
+					Offset += Blob.Length;
+				}
+
+				// Discard the old blob
+				PackIndexes.Remove(MergePack);
+				TotalCost = MergeCost;
+			}
+
+			// Write the current blob
+			FlushCurrentPack();
+
+			// Wait for all the writes to finish
+			await Task.WhenAll(WriteTasks);
+			WriteTasks.Clear();
+		}
+
+		/// <summary>
+		/// Finds the live set for a particular tree, and updates tree entries with the size of used items within them
+		/// </summary>
+		/// <param name="Hash"></param>
+		/// <param name="LiveSet"></param>
+		void FindLiveSet(IoHash Hash, HashSet<IoHash> LiveSet)
+		{
+			if (LiveSet.Add(Hash))
+			{
+				ObjectPackEntry? Entry;
+				if (!TryGetEntry(Hash, out Entry))
+				{
+					throw new Exception($"Missing blob {Hash} from working set");
+				}
+				foreach (IoHash Ref in Entry.Refs)
+				{
+					FindLiveSet(Ref, LiveSet);
+				}
+			}
+		}
+
+		/// <summary>
+		/// Creates enough space to store the given block of data
+		/// </summary>
+		/// <param name="Size"></param>
+		void CreateSpace(int Size)
+		{
+			// Get the last blob and make sure there's enough space in it
+			if (NextPackSize + Size > MaxPackSize)
+			{
+				FlushCurrentPack();
+			}
+
+			// Resize the next blob buffer if necessary
+			if (Size > NextPackData.Length)
+			{
+				Array.Resize(ref NextPackData, Size);
+			}
+		}
+
+		/// <summary>
+		/// Finalize the current blob and start writing it to storage
+		/// </summary>
+		void FlushCurrentPack()
+		{
+			if (NextPackSize > 0)
+			{
+				// Write the buffer to storage
+				Array.Resize(ref NextPackData, NextPackSize);
+				ReadOnlyMemory<byte> Data = NextPackData;
+				IoHash DataHash = IoHash.Compute(Data.Span);
+				WriteTasks.Add(Task.Run(() => BlobCollection.WriteBytesAsync(NamespaceId, DataHash, Data)));
+
+				// Create the new index
+				ObjectPackIndex Index = new ObjectPackIndex(Time, NextPackEntries.ToArray(), DataHash, NextPackSize);
+				PackIndexes.Add(Index);
+
+				// Clear the next pack buffer
+				Reset();
+			}
+		}
+
+		/// <inheritdoc cref="GetCostHeuristic(int, TimeSpan)"/>
+		/// <param name="Index">Index to calculate the heuristic for</param>
+		public double GetCostHeuristic(ObjectPackIndex Index) => GetCostHeuristic(Index.Blobs.Length, Time - Index.Time);
+
+		/// <summary>
+		/// Heuristic which estimates the cost of a particular blob. This is used to compare scenarios of merging blobs to reduce download
+		/// size against keeping older blobs which a lot of agents already have.
+		/// </summary>
+		/// <param name="Size">Size of the blob</param>
+		/// <param name="Age">Age of the blob</param>
+		/// <returns>Heuristic for the cost of a blob</returns>
+		public static double GetCostHeuristic(int Size, TimeSpan Age)
+		{
+			// Time overhead to starting a download
+			const double DownloadInit = 0.1;
+
+			// Download speed for agents, in bytes/sec
+			const double DownloadRate = 1024 * 1024;
+
+			// Probability of an agent having to download everything. Prevents bias against keeping a large number of files.
+			const double CleanSyncProbability = 0.2;
+
+			// Average length of time between agents having to update
+			TimeSpan AverageCoherence = TimeSpan.FromHours(4.0);
+
+			// Scale the age into a -1.0 -> 1.0 range around AverageCoherence
+			double ScaledAge = (AverageCoherence - Age).TotalSeconds / AverageCoherence.TotalSeconds;
+
+			// Get the probability of agents having to sync this blob based on its age. This is modeled as a logistic function (1 / (1 + e^-x))
+			// with value of 0.5 at AverageCoherence, and MaxInterval at zero.
+
+			// Find the scale factor for the 95% interval
+			//    1 / (1 + e^-x) = MaxInterval
+			//    e^-x = (1 / MaxInterval) - 1
+			//    x = -ln((1 / MaxInterval) - 1)
+			const double MaxInterval = 0.95;
+			double SigmoidScale = -Math.Log((1.0 / MaxInterval) - 1.0);
+
+			// Find the probability of having to sync this 
+			double Param = ScaledAge * SigmoidScale;
+			double Probability = 1.0 / (1.0 + Math.Exp(-Param));
+
+			// Scale the probability against having to do a full sync
+			Probability = CleanSyncProbability + (Probability * (1.0 - CleanSyncProbability));
+
+			// Compute the final cost estimate; the amount of time we expect agents to spend downloading the file
+			return Probability * (DownloadInit + (Size / DownloadRate));
+		}
+	}
+}
diff --git a/Engine/Source/Programs/Shared/EpicGames.Perforce/FileSpecList.cs b/Engine/Source/Programs/Shared/EpicGames.Perforce/FileSpecList.cs
--- a/Engine/Source/Programs/Shared/EpicGames.Perforce/FileSpecList.cs
+++ b/Engine/Source/Programs/Shared/EpicGames.Perforce/FileSpecList.cs
@@ -16,6 +16,11 @@ namespace EpicGames.Perforce
 		/// </summary>
 		public static FileSpecList Empty { get; } = new FileSpecList(new List<string>());
 
+		/// <summary>
+		/// Matches any files in the depot
+		/// </summary>
+		public static FileSpecList Any { get; } = new FileSpecList(new List<string> { "//..." });
+
 		/// <summary>
 		/// The list of filespecs
 		/// </summary>
