Horde: Fix unit tests.

diff --git a/Engine/Source/Programs/Horde/HordeServerTests/StorageTests.cs b/Engine/Source/Programs/Horde/HordeServerTests/StorageTests.cs
--- a/Engine/Source/Programs/Horde/HordeServerTests/StorageTests.cs
+++ b/Engine/Source/Programs/Horde/HordeServerTests/StorageTests.cs
@@ -35,20 +35,20 @@ namespace HordeServerTests
 			byte[] TestData = Encoding.UTF8.GetBytes("Hello world");
 			IoHash Hash = IoHash.Compute(TestData);
 
-			Assert.IsNull(await BlobCollection.ReadAsync(NamespaceId, Hash));
+			Assert.IsNull(await BlobCollection.TryReadStreamAsync(NamespaceId, Hash));
 			Assert.IsFalse(await BlobCollection.ExistsAsync(NamespaceId, Hash));
 
 			IoHash ReturnedHash = await BlobCollection.WriteBytesAsync(NamespaceId, TestData);
 			Assert.AreEqual(Hash, ReturnedHash);
 			Assert.IsTrue(await BlobCollection.ExistsAsync(NamespaceId, Hash));
 
-			byte[]? StoredData = await BlobCollection.ReadBytesAsync(NamespaceId, Hash);
+			ReadOnlyMemory<byte>? StoredData = await BlobCollection.TryReadBytesAsync(NamespaceId, Hash);
 			Assert.IsNotNull(StoredData);
-			Assert.IsTrue(TestData.AsSpan().SequenceEqual(StoredData));
+			Assert.IsTrue(TestData.AsSpan().SequenceEqual(StoredData.Value.Span));
 			Assert.IsTrue(await BlobCollection.ExistsAsync(NamespaceId, Hash));
 
 			NamespaceId OtherNamespaceId = new NamespaceId("other-ns");
-			Assert.IsNull(await BlobCollection.ReadAsync(OtherNamespaceId, Hash));
+			Assert.IsNull(await BlobCollection.TryReadStreamAsync(OtherNamespaceId, Hash));
 			Assert.IsFalse(await BlobCollection.ExistsAsync(OtherNamespaceId, Hash));
 		}
 
