Horde: Test fix.

diff --git a/Engine/Source/Programs/Horde/HordeServer/Storage/Collections/ObjectCollection.cs b/Engine/Source/Programs/Horde/HordeServer/Storage/Collections/ObjectCollection.cs
--- a/Engine/Source/Programs/Horde/HordeServer/Storage/Collections/ObjectCollection.cs
+++ b/Engine/Source/Programs/Horde/HordeServer/Storage/Collections/ObjectCollection.cs
@@ -36,7 +36,7 @@ namespace HordeServer.Storage.Collections
 		/// <inheritdoc/>
 		public async Task<CbObject?> GetAsync(NamespaceId NsId, IoHash Hash)
 		{
-			ReadOnlyMemory<byte>? Data = await BlobCollection.ReadBytesAsync(NsId, Hash);
+			ReadOnlyMemory<byte>? Data = await BlobCollection.TryReadBytesAsync(NsId, Hash);
 			if(Data == null)
 			{
 				return null;
