Horde: Fixed resource count not being stored on agent leases, allowing more jobs to start than the agent has capacity for.

diff --git a/Engine/Source/Programs/Horde/HordeServer/Compute/Impl/ComputeService.cs b/Engine/Source/Programs/Horde/HordeServer/Compute/Impl/ComputeService.cs
--- a/Engine/Source/Programs/Horde/HordeServer/Compute/Impl/ComputeService.cs
+++ b/Engine/Source/Programs/Horde/HordeServer/Compute/Impl/ComputeService.cs
@@ -240,7 +240,7 @@ namespace HordeServer.Compute.Impl
 					string LeaseName = $"Remote action ({TaskInfo.TaskHash})";
 					byte[] Payload = Any.Pack(ComputeTask).ToByteArray();
 
-					AgentLease Lease = new AgentLease(ObjectId.GenerateNewId(), LeaseName, null, null, null, LeaseState.Pending, null, Requirements.Exclusive, Payload);
+					AgentLease Lease = new AgentLease(ObjectId.GenerateNewId(), LeaseName, null, null, null, LeaseState.Pending, Requirements.Resources, Requirements.Exclusive, Payload);
 					Logger.LogDebug("Created lease {LeaseId} for channel {ChannelId} task {TaskHash} req {RequirementsHash}", Lease.Id, ComputeTask.ChannelId, ComputeTask.Task.Hash, ComputeTask.Requirements.Hash);
 					return Lease;
 				}
