Horde: Add 'id' as an agent property that can be used for remote execution conditions.

diff --git a/Engine/Source/Programs/Horde/HordeServer/Models/Agent.cs b/Engine/Source/Programs/Horde/HordeServer/Models/Agent.cs
--- a/Engine/Source/Programs/Horde/HordeServer/Models/Agent.cs
+++ b/Engine/Source/Programs/Horde/HordeServer/Models/Agent.cs
@@ -381,6 +381,11 @@ namespace HordeServer.Models
 	/// </summary>
 	static class KnownPropertyNames
 	{
+		/// <summary>
+		/// The agent id
+		/// </summary>
+		public const string Id = "Id";
+
 		/// <summary>
 		/// The operating system (Linux, MacOS, Windows)
 		/// </summary>
@@ -594,7 +599,11 @@ namespace HordeServer.Models
 		/// <returns>Property values</returns>
 		public static IEnumerable<string> GetPropertyValues(this IAgent Agent, string Name)
 		{
-			if (Name.Equals(KnownPropertyNames.Pool, StringComparison.OrdinalIgnoreCase))
+			if (Name.Equals(KnownPropertyNames.Id, StringComparison.OrdinalIgnoreCase))
+			{
+				yield return Agent.Id.ToString();
+			}
+			else if (Name.Equals(KnownPropertyNames.Pool, StringComparison.OrdinalIgnoreCase))
 			{
 				foreach (PoolId PoolId in Agent.GetPools())
 				{
