Log out what files we're getting ignore patterns from.

#rnx

diff --git a/Engine/Source/Programs/Horde/HordeAgent/Execution/BuildGraphExecutor.cs b/Engine/Source/Programs/Horde/HordeAgent/Execution/BuildGraphExecutor.cs
--- a/Engine/Source/Programs/Horde/HordeAgent/Execution/BuildGraphExecutor.cs
+++ b/Engine/Source/Programs/Horde/HordeAgent/Execution/BuildGraphExecutor.cs
@@ -656,6 +656,7 @@ namespace HordeAgent.Execution
 				FileReference IgnorePatternFile = FileReference.Combine(BaseDir, "Build", "Horde", "IgnorePatterns.txt");
 				if (FileReference.Exists(IgnorePatternFile))
 				{
+					Log.TraceInformation("Reading ignore patterns from {0}...", IgnorePatternFile);
 					IgnorePatternLines.AddRange(await FileReference.ReadAllLinesAsync(IgnorePatternFile));
 				}
 			}
