Horde: Fix race condition for update of issue span states, modified flag on spans not being respected, and issue summary text not being updated correctly. Also deprecate mutating existing document states in favor of returning document copies.

diff --git a/Engine/Source/Programs/Horde/HordeServer/Collections/Impl/IssueCollection.cs b/Engine/Source/Programs/Horde/HordeServer/Collections/Impl/IssueCollection.cs
--- a/Engine/Source/Programs/Horde/HordeServer/Collections/Impl/IssueCollection.cs
+++ b/Engine/Source/Programs/Horde/HordeServer/Collections/Impl/IssueCollection.cs
@@ -465,37 +465,26 @@ namespace HordeServer.Collections.Impl
 			}
 		}
 
-		async Task<IIssue?> ApplyTransactionAsync(IIssue Issue, TransactionBuilder<Issue> Transaction)
+		async Task<Issue?> TryUpdateIssueAsync(IIssue Issue, UpdateDefinition<Issue> Update)
 		{
 			Issue IssueDocument = (Issue)Issue;
 
 			int PrevUpdateIndex = IssueDocument.UpdateIndex;
-			Transaction.Set(x => x.UpdateIndex, PrevUpdateIndex + 1);
+			Update = Update.Set(x => x.UpdateIndex, PrevUpdateIndex + 1);
 
-			UpdateResult Result = await Issues.UpdateOneAsync(x => x.Id == IssueDocument.Id && x.UpdateIndex == PrevUpdateIndex, Transaction.ToUpdateDefinition());
-			if (Result.ModifiedCount == 0)
-			{
-				return null;
-			}
-
-			Transaction.ApplyTo(IssueDocument);
-			return IssueDocument;
+			FindOneAndUpdateOptions<Issue, Issue> Options = new FindOneAndUpdateOptions<Issue, Issue> { ReturnDocument = ReturnDocument.After };
+			return await Issues.FindOneAndUpdateAsync<Issue>(x => x.Id == IssueDocument.Id && x.UpdateIndex == PrevUpdateIndex, Update, Options);
 		}
 
-		async Task<IIssueSpan?> ApplyTransactionAsync(IIssueSpan IssueSpan, TransactionBuilder<IssueSpan> Transaction)
+		async Task<IssueSpan?> TryUpdateSpanAsync(IIssueSpan IssueSpan, UpdateDefinition<IssueSpan> Update)
 		{
 			IssueSpan IssueSpanDocument = (IssueSpan)IssueSpan;
 
 			int PrevUpdateIndex = IssueSpanDocument.UpdateIndex;
-			Transaction.Set(x => x.UpdateIndex, PrevUpdateIndex + 1);
+			Update = Update.Set(x => x.UpdateIndex, PrevUpdateIndex + 1);
 
-			UpdateResult Result = await IssueSpans.UpdateOneAsync(x => x.Id == IssueSpanDocument.Id && x.UpdateIndex == PrevUpdateIndex, Transaction.ToUpdateDefinition());
-			if (Result.ModifiedCount > 0)
-			{
-				Transaction.ApplyTo(IssueSpanDocument);
-				return IssueSpanDocument;
-			}
-			return null;
+			FindOneAndUpdateOptions<IssueSpan, IssueSpan> Options = new FindOneAndUpdateOptions<IssueSpan, IssueSpan> { ReturnDocument = ReturnDocument.After };
+			return await IssueSpans.FindOneAndUpdateAsync<IssueSpan>(x => x.Id == IssueSpanDocument.Id && x.UpdateIndex == PrevUpdateIndex, Update, Options);
 		}
 
 		/// <inheritdoc/>
@@ -685,7 +674,20 @@ namespace HordeServer.Collections.Impl
 
 				using (IAsyncCursor<ProjectedIssueId> Cursor = await IssueSpans.Aggregate().Match(Filter).Group(x => x.IssueId, x => new ProjectedIssueId { _id = x.Key }).SortByDescending(x => x._id).ToCursorAsync())
 				{
-					return await PaginatedJoinAsync(Cursor, (NextIds, NextIndex, NextCount) => FilterIssuesByOtherFieldsAsync(NextIds, null, null, null, NextIndex, NextCount), Index, Count);
+					List<Issue> Results = await PaginatedJoinAsync(Cursor, (NextIds, NextIndex, NextCount) => FilterIssuesByOtherFieldsAsync(NextIds, null, null, null, NextIndex, NextCount), Index, Count);
+					if (Resolved != null)
+					{
+						for (int Idx = Results.Count - 1; Idx >= 0; Idx--)
+						{
+							Issue Issue = Results[Idx];
+							if ((Issue.ResolvedAt != null) != Resolved.Value)
+							{
+								Logger.LogWarning("Issue {IssueId} has resolved state out of sync with spans", Issue.Id);
+								Results.RemoveAt(Idx);
+							}
+						}
+					}
+					return Results;
 				}
 			}
 		}
@@ -781,46 +783,46 @@ namespace HordeServer.Collections.Impl
 
 			DateTime UtcNow = DateTime.UtcNow;
 
-			TransactionBuilder<Issue> Transaction = new TransactionBuilder<Issue>();
+			List<UpdateDefinition<Issue>> Updates = new List<UpdateDefinition<Issue>>();
 			if (NewSeverity != null)
 			{
-				Transaction.Set(x => x.Severity, NewSeverity.Value);
+				Updates.Add(Builders<Issue>.Update.Set(x => x.Severity, NewSeverity.Value));
 			}
 			if (NewSummary != null)
 			{
-				Transaction.Set(x => x.Summary, NewSummary);
+				Updates.Add(Builders<Issue>.Update.Set(x => x.Summary, NewSummary));
 			}
 			if (NewUserSummary != null)
 			{
 				if (NewUserSummary.Length == 0)
 				{
-					Transaction.Unset(x => x.UserSummary!);
+					Updates.Add(Builders<Issue>.Update.Unset(x => x.UserSummary!));
 				}
 				else
 				{
-					Transaction.Set(x => x.UserSummary, NewUserSummary);
+					Updates.Add(Builders<Issue>.Update.Set(x => x.UserSummary, NewUserSummary));
 				}
 			}
 			if (NewOwnerId != null)
 			{
 				if (NewOwnerId.Value == ObjectId.Empty)
 				{
-					Transaction.Unset(x => x.OwnerId!);
-					Transaction.Unset(x => x.NominatedAt!);
-					Transaction.Unset(x => x.NominatedById!);
+					Updates.Add(Builders<Issue>.Update.Unset(x => x.OwnerId!));
+					Updates.Add(Builders<Issue>.Update.Unset(x => x.NominatedAt!));
+					Updates.Add(Builders<Issue>.Update.Unset(x => x.NominatedById!));
 				}
 				else
 				{
-					Transaction.Set(x => x.OwnerId!, NewOwnerId.Value);
+					Updates.Add(Builders<Issue>.Update.Set(x => x.OwnerId!, NewOwnerId.Value));
 
-					Transaction.Set(x => x.NominatedAt, DateTime.UtcNow);
+					Updates.Add(Builders<Issue>.Update.Set(x => x.NominatedAt, DateTime.UtcNow));
 					if (NewNominatedById == null)
 					{
-						Transaction.Unset(x => x.NominatedById!);
+						Updates.Add(Builders<Issue>.Update.Unset(x => x.NominatedById!));
 					}
 					else
 					{
-						Transaction.Set(x => x.NominatedById, NewNominatedById.Value);
+						Updates.Add(Builders<Issue>.Update.Set(x => x.NominatedById, NewNominatedById.Value));
 					}
 					NewAcknowledged = NewAcknowledged ?? false;
 				}
@@ -831,14 +833,14 @@ namespace HordeServer.Collections.Impl
 				{
 					if (IssueDocument.AcknowledgedAt == null)
 					{
-						Transaction.Set(x => x.AcknowledgedAt, UtcNow);
+						Updates.Add(Builders<Issue>.Update.Set(x => x.AcknowledgedAt, UtcNow));
 					}
 				}
 				else
 				{
 					if (IssueDocument.AcknowledgedAt != null)
 					{
-						Transaction.Unset(x => x.AcknowledgedAt!);
+						Updates.Add(Builders<Issue>.Update.Unset(x => x.AcknowledgedAt!));
 					}
 				}
 			}
@@ -850,7 +852,7 @@ namespace HordeServer.Collections.Impl
 					if (Suspect.AuthorId == NewDeclinedById.Value && Suspect.DeclinedAt == null)
 					{
 						int SuspectIdxCopy = SuspectIdx;
-						Transaction.Set(x => x.Suspects[SuspectIdxCopy].DeclinedAt, UtcNow);
+						Updates.Add(Builders<Issue>.Update.Set(x => x.Suspects[SuspectIdxCopy].DeclinedAt, UtcNow));
 					}
 				}
 			}
@@ -858,11 +860,11 @@ namespace HordeServer.Collections.Impl
 			{
 				if (NewFixChange == 0)
 				{
-					Transaction.Unset(x => x.FixChange!);
+					Updates.Add(Builders<Issue>.Update.Unset(x => x.FixChange!));
 				}
 				else
 				{
-					Transaction.Set(x => x.FixChange, NewFixChange);
+					Updates.Add(Builders<Issue>.Update.Set(x => x.FixChange, NewFixChange));
 				}
 			}
 			if (NewFixStreamIds != null)
@@ -880,7 +882,7 @@ namespace HordeServer.Collections.Impl
 						NewStream.ContainsFix = null;
 					}
 				}
-				Transaction.Set(x => x.Streams, NewStreams);
+				Updates.Add(Builders<Issue>.Update.Set(x => x.Streams, NewStreams));
 			}
 			if (NewResolvedById != null)
 			{
@@ -888,40 +890,40 @@ namespace HordeServer.Collections.Impl
 				{
 					if (IssueDocument.ResolvedAt == null || IssueDocument.ResolvedById != NewResolvedById)
 					{
-						Transaction.Set(x => x.ResolvedAt, UtcNow);
-						Transaction.Set(x => x.ResolvedById, NewResolvedById.Value);
+						Updates.Add(Builders<Issue>.Update.Set(x => x.ResolvedAt, UtcNow));
+						Updates.Add(Builders<Issue>.Update.Set(x => x.ResolvedById, NewResolvedById.Value));
 					}
 				}
 				else
 				{
 					if (IssueDocument.ResolvedAt != null)
 					{
-						Transaction.Unset(x => x.ResolvedAt!);
+						Updates.Add(Builders<Issue>.Update.Unset(x => x.ResolvedAt!));
 					}
 					if (IssueDocument.ResolvedById != null)
 					{
-						Transaction.Unset(x => x.ResolvedById!);
+						Updates.Add(Builders<Issue>.Update.Unset(x => x.ResolvedById!));
 					}
 				}
 			}
 			if (NewLastSeenAt != null)
 			{
-				Transaction.Set(x => x.LastSeenAt, NewLastSeenAt.Value);
+				Updates.Add(Builders<Issue>.Update.Set(x => x.LastSeenAt, NewLastSeenAt.Value));
 			}
 			if (NewNotifySuspects != null)
 			{
-				Transaction.Set(x => x.NotifySuspects, NewNotifySuspects.Value);
+				Updates.Add(Builders<Issue>.Update.Set(x => x.NotifySuspects, NewNotifySuspects.Value));
 			}
 
-			IIssue? NewIssue = await ApplyTransactionAsync(Issue, Transaction);
-			if (NewIssue != null)
+			if (Updates.Count == 0)
 			{
-				if (NewResolvedById != null)
-				{
-					await UpdateIssueDerivedDataAsync(Issue);
-				}
+				return IssueDocument;
 			}
-			return NewIssue;
+			else
+			{
+				return await TryUpdateIssueAsync(Issue, Builders<Issue>.Update.Combine(Updates));
+			}
+
 		}
 
 		public async Task<IIssue?> UpdateIssueDerivedDataAsync(IIssue Issue)
@@ -933,133 +935,168 @@ namespace HordeServer.Collections.Impl
 				// Find all the spans that are attached to the issue
 				List<IssueSpan> Spans = await IssueSpans.Find(x => x.IssueId == Issue.Id).ToListAsync();
 
-				// Find the list of suspects for this issue. If the issue has an assigned owner, the other suspects can be eliminated.
-				List<IssueSuspect> NewSuspects = new List<IssueSuspect>();
-				if (Spans.Count > 0)
+				// Update the issue to reflect information in the spans
+				Issue? Result = await TryUpdateIssueDerivedDataAsync(IssueDocument, Spans);
+				if (Result != null)
 				{
-					HashSet<int> SuspectChanges = new HashSet<int>(Spans[0].Suspects.Where(x => Issue.OwnerId == null || x.AuthorId == Issue.OwnerId).Select(x => x.OriginatingChange ?? x.Change));
-					for (int SpanIdx = 1; SpanIdx < Spans.Count; SpanIdx++)
-					{
-						SuspectChanges.IntersectWith(Spans[SpanIdx].Suspects.Select(x => x.OriginatingChange ?? x.Change));
-					}
-					NewSuspects = Spans[0].Suspects.Where(x => SuspectChanges.Contains(x.OriginatingChange ?? x.Change)).Select(x => new IssueSuspect(Issue.Id, x)).ToList();
+					return Result;
 				}
-				if (Issue.OwnerId != null)
+
+				// Update the issue
+				IIssue? NewIssue = await GetIssueAsync(Issue.Id);
+				if(NewIssue == null)
 				{
-					NewSuspects.RemoveAll(x => x.AuthorId != Issue.OwnerId);
-					if (NewSuspects.Count == 0)
-					{
-						NewSuspects.Add(new IssueSuspect(Issue.Id, Issue.OwnerId.Value, 0, null, null));
-					}
+					return null;
 				}
+				Issue = NewIssue;
+			}
+		}
 
-				// Find the current list of suspects
-				List<IssueSuspect> CurSuspects = await IssueSuspects.Find(x => x.IssueId == Issue.Id).ToListAsync();
-
-				HashSet<(ObjectId, int)> CurSuspectKeys = new HashSet<(ObjectId, int)>(CurSuspects.Select(x => (x.AuthorId, x.Change)));
-				List<IssueSuspect> CreateSuspects = NewSuspects.Where(x => !CurSuspectKeys.Contains((x.AuthorId, x.Change))).ToList();
-
-				HashSet<(ObjectId, int)> NewSuspectKeys = new HashSet<(ObjectId, int)>(NewSuspects.Select(x => (x.AuthorId, x.Change)));
-				List<IssueSuspect> DeleteSuspects = CurSuspects.Where(x => !NewSuspectKeys.Contains((x.AuthorId, x.Change))).ToList();
-
-				// Apply the suspect changes
-				if (CreateSuspects.Count > 0)
+		async Task<Issue?> TryUpdateIssueDerivedDataAsync(Issue Issue, List<IssueSpan> Spans)
+		{
+			// Find the list of suspects for this issue. If the issue has an assigned owner, the other suspects can be eliminated.
+			List<IssueSuspect> NewSuspects = new List<IssueSuspect>();
+			if (Spans.Count > 0)
+			{
+				HashSet<int> SuspectChanges = new HashSet<int>(Spans[0].Suspects.Where(x => Issue.OwnerId == null || x.AuthorId == Issue.OwnerId).Select(x => x.OriginatingChange ?? x.Change));
+				for (int SpanIdx = 1; SpanIdx < Spans.Count; SpanIdx++)
 				{
-					await IssueSuspects.InsertManyIgnoreDuplicatesAsync(CreateSuspects);
-					CurSuspects.AddRange(CreateSuspects);
+					SuspectChanges.IntersectWith(Spans[SpanIdx].Suspects.Select(x => x.OriginatingChange ?? x.Change));
 				}
-				if (DeleteSuspects.Count > 0)
+				NewSuspects = Spans[0].Suspects.Where(x => SuspectChanges.Contains(x.OriginatingChange ?? x.Change)).Select(x => new IssueSuspect(Issue.Id, x)).ToList();
+			}
+			if (Issue.OwnerId != null)
+			{
+				NewSuspects.RemoveAll(x => x.AuthorId != Issue.OwnerId);
+				if (NewSuspects.Count == 0)
 				{
-					await IssueSuspects.DeleteManyAsync(Builders<IssueSuspect>.Filter.In(x => x.Id, DeleteSuspects.Select(y => y.Id)));
-					CurSuspects.RemoveAll(x => !NewSuspectKeys.Contains((x.AuthorId, x.Change)));
+					NewSuspects.Add(new IssueSuspect(Issue.Id, Issue.OwnerId.Value, 0, null, null));
 				}
+			}
+
+			// Find the current list of suspects
+			List<IssueSuspect> CurSuspects = await IssueSuspects.Find(x => x.IssueId == Issue.Id).ToListAsync();
 
-				// Calculate the last time that this issue was seen
-				DateTime NewLastSeenAt = Issue.CreatedAt;
-				foreach (IIssueSpan Span in Spans)
+			HashSet<(ObjectId, int)> CurSuspectKeys = new HashSet<(ObjectId, int)>(CurSuspects.Select(x => (x.AuthorId, x.Change)));
+			List<IssueSuspect> CreateSuspects = NewSuspects.Where(x => !CurSuspectKeys.Contains((x.AuthorId, x.Change))).ToList();
+
+			HashSet<(ObjectId, int)> NewSuspectKeys = new HashSet<(ObjectId, int)>(NewSuspects.Select(x => (x.AuthorId, x.Change)));
+			List<IssueSuspect> DeleteSuspects = CurSuspects.Where(x => !NewSuspectKeys.Contains((x.AuthorId, x.Change))).ToList();
+
+			// Apply the suspect changes
+			if (CreateSuspects.Count > 0)
+			{
+				await IssueSuspects.InsertManyIgnoreDuplicatesAsync(CreateSuspects);
+				CurSuspects.AddRange(CreateSuspects);
+			}
+			if (DeleteSuspects.Count > 0)
+			{
+				await IssueSuspects.DeleteManyAsync(Builders<IssueSuspect>.Filter.In(x => x.Id, DeleteSuspects.Select(y => y.Id)));
+				CurSuspects.RemoveAll(x => !NewSuspectKeys.Contains((x.AuthorId, x.Change)));
+			}
+
+			// Calculate the last time that this issue was seen
+			DateTime NewLastSeenAt = Issue.CreatedAt;
+			foreach (IIssueSpan Span in Spans)
+			{
+				if (Span.LastFailure.StepTime > NewLastSeenAt)
 				{
-					if (Span.LastFailure.StepTime > NewLastSeenAt)
-					{
-						NewLastSeenAt = Span.LastFailure.StepTime;
-					}
+					NewLastSeenAt = Span.LastFailure.StepTime;
 				}
+			}
 
-				// Calculate the time at which the issue is verified fixed
-				DateTime? VerifiedAt = null;
-				foreach (IIssueSpan Span in Spans)
+			// Calculate the time at which the issue is verified fixed
+			DateTime? VerifiedAt = null;
+			foreach (IIssueSpan Span in Spans)
+			{
+				if (Span.NextSuccess == null)
 				{
-					if(Span.NextSuccess == null)
-					{
-						break;
-					}
-					else if(VerifiedAt == null || Span.NextSuccess.StepTime < VerifiedAt.Value)
-					{
-						VerifiedAt = Span.NextSuccess.StepTime;
-					}
+					break;
 				}
-				if (Spans.Count == 0)
+				else if (VerifiedAt == null || Span.NextSuccess.StepTime < VerifiedAt.Value)
 				{
-					VerifiedAt = Issue.VerifiedAt ?? DateTime.UtcNow;
+					VerifiedAt = Span.NextSuccess.StepTime;
 				}
+			}
+			if (Spans.Count == 0)
+			{
+				VerifiedAt = Issue.VerifiedAt ?? DateTime.UtcNow;
+			}
 
-				// Update the spans with the correct resolved time
-				DateTime? ResolvedAt = Issue.ResolvedAt ?? VerifiedAt;
+			// Update the spans with the correct resolved time
+			DateTime? ResolvedAt = Issue.ResolvedAt ?? VerifiedAt;
 
-				List<ObjectId> UpdateSpanIds = Spans.Where(x => x.ResolvedAt != ResolvedAt).Select(x => x.Id).ToList();
-				if (UpdateSpanIds.Count > 0)
+			for(int Idx = 0; Idx < Spans.Count; Idx++)
+			{
+				IssueSpan? Span = Spans[Idx];
+				if (Span.ResolvedAt != ResolvedAt)
 				{
-					await IssueSpans.UpdateManyAsync(Builders<IssueSpan>.Filter.In(x => x.Id, UpdateSpanIds), Builders<IssueSpan>.Update.Set(x => x.ResolvedAt, ResolvedAt));
+					Span = await TryUpdateSpanAsync(Span, Builders<IssueSpan>.Update.Set(x => x.ResolvedAt, ResolvedAt));
+					if (Span == null)
+					{
+						return null;
+					}
+					Spans[Idx] = Span;
 				}
+			}
 
-				List<ObjectId> UpdateSuspectIds = CurSuspects.Where(x => x.ResolvedAt != ResolvedAt).Select(x => x.Id).ToList();
-				if (UpdateSuspectIds.Count > 0)
-				{
-					await IssueSuspects.UpdateManyAsync(Builders<IssueSuspect>.Filter.In(x => x.Id, UpdateSuspectIds), Builders<IssueSuspect>.Update.Set(x => x.ResolvedAt, ResolvedAt));
-				}
+			List<ObjectId> UpdateSuspectIds = CurSuspects.Where(x => x.ResolvedAt != ResolvedAt).Select(x => x.Id).ToList();
+			if (UpdateSuspectIds.Count > 0)
+			{
+				await IssueSuspects.UpdateManyAsync(Builders<IssueSuspect>.Filter.In(x => x.Id, UpdateSuspectIds), Builders<IssueSuspect>.Update.Set(x => x.ResolvedAt, ResolvedAt));
+			}
 
-				// Build a list of streams
-				List<IssueStream> Streams = new List<IssueStream>();
-				foreach (IGrouping<StreamId, IIssueSpan> Group in Spans.GroupBy(x => x.StreamId))
+			// Build a list of streams
+			List<IssueStream> Streams = new List<IssueStream>();
+			foreach (IGrouping<StreamId, IIssueSpan> Group in Spans.GroupBy(x => x.StreamId))
+			{
+				IssueStream? Stream = Issue.Streams.FirstOrDefault(x => x.StreamId == Group.Key);
+				if (Stream == null)
 				{
-					IssueStream? Stream = IssueDocument.Streams.FirstOrDefault(x => x.StreamId == Group.Key);
-					if (Stream == null)
-					{
-						Stream = new IssueStream { StreamId = Group.Key };
-					}
-					else
-					{
-						Stream = new IssueStream(Stream);
-					}
-					Streams.Add(Stream);
+					Stream = new IssueStream { StreamId = Group.Key };
 				}
-
-				// Update the issue document itself
-				TransactionBuilder<Issue> Transaction = new TransactionBuilder<Issue>();
-				Transaction.Set(x => x.Severity, Spans.Any(x => x.Severity == IssueSeverity.Error) ? IssueSeverity.Error : IssueSeverity.Warning);
-				Transaction.Set(x => x.NotifySuspects, Spans.Any(x => x.NotifySuspects));
-				Transaction.Set(x => x.LastSeenAt, NewLastSeenAt);
-				Transaction.Set(x => x.ResolvedAt, ResolvedAt);
-				Transaction.Set(x => x.VerifiedAt, VerifiedAt);
-				Transaction.Set(x => x.Suspects, CurSuspects);
-				Transaction.Set(x => x.Streams, Streams);
-				Transaction.Set(x => x.MinSuspectChange, (CurSuspects.Count > 0) ? CurSuspects.Min(x => x.Change) : 0);
-				Transaction.Set(x => x.MaxSuspectChange, (CurSuspects.Count > 0) ? CurSuspects.Max(x => x.Change) : 0);
-
-				// Apply the transaction
-				IIssue? NewIssue = await ApplyTransactionAsync(Issue, Transaction);
-				if(NewIssue != null)
+				else
 				{
-					return NewIssue;
+					Stream = new IssueStream(Stream);
 				}
+				Streams.Add(Stream);
+			}
 
-				// Update the issue
-				NewIssue = await GetIssueAsync(Issue.Id);
-				if(NewIssue == null)
+			// Update the issue document itself
+			UpdateDefinition<Issue> Update = Builders<Issue>.Update
+				.Set(x => x.Severity, Spans.Any(x => x.Severity == IssueSeverity.Error) ? IssueSeverity.Error : IssueSeverity.Warning)
+				.Set(x => x.NotifySuspects, Spans.Any(x => x.NotifySuspects))
+				.Set(x => x.LastSeenAt, NewLastSeenAt)
+				.Set(x => x.ResolvedAt, ResolvedAt)
+				.Set(x => x.VerifiedAt, VerifiedAt)
+				.Set(x => x.Suspects, CurSuspects)
+				.Set(x => x.Streams, Streams)
+				.Set(x => x.MinSuspectChange, (CurSuspects.Count > 0) ? CurSuspects.Min(x => x.Change) : 0)
+				.Set(x => x.MaxSuspectChange, (CurSuspects.Count > 0) ? CurSuspects.Max(x => x.Change) : 0);
+
+			// Apply the update
+			Issue? NewIssue = await TryUpdateIssueAsync(Issue, Update);
+			if (NewIssue == null)
+			{
+				return null;
+			}
+
+			// Clear the modified flag on any spans
+			for (int Idx = 0; Idx < Spans.Count; Idx++)
+			{
+				IssueSpan? Span = Spans[Idx];
+				if (Span.Modified)
 				{
-					return null;
+					Span = await TryUpdateSpanAsync(Span, Builders<IssueSpan>.Update.Set(x => x.Modified, false));
+					if (Span == null)
+					{
+						return null;
+					}
+					Spans[Idx] = Span;
 				}
-				Issue = NewIssue;
 			}
+
+			return NewIssue;
 		}
 
 		#endregion
@@ -1094,46 +1131,51 @@ namespace HordeServer.Collections.Impl
 		/// <inheritdoc/>
 		public async Task<IIssueSpan?> UpdateSpanAsync(IIssueSpan Span, IssueSeverity? NewSeverity = null, NewIssueStepData? NewLastSuccess = null, NewIssueStepData? NewFailure = null, NewIssueStepData? NewNextSuccess = null, List<NewIssueSpanSuspectData>? NewSuspects = null, bool? NewModified = null)
 		{
-			TransactionBuilder<IssueSpan> Transaction = new TransactionBuilder<IssueSpan>();
+			List<UpdateDefinition<IssueSpan>> Updates = new List<UpdateDefinition<IssueSpan>>();
 			if (NewSeverity != null)
 			{
-				Transaction.Set(x => x.Severity, NewSeverity.Value);
+				Updates.Add(Builders<IssueSpan>.Update.Set(x => x.Severity, NewSeverity.Value));
 			}
 			if (NewLastSuccess != null)
 			{
-				Transaction.Set(x => x.MinChange, NewLastSuccess.Change);
-				Transaction.Set(x => x.LastSuccess, new IssueStep(Span.Id, NewLastSuccess));
+				Updates.Add(Builders<IssueSpan>.Update.Set(x => x.MinChange, NewLastSuccess.Change));
+				Updates.Add(Builders<IssueSpan>.Update.Set(x => x.LastSuccess, new IssueStep(Span.Id, NewLastSuccess)));
 			}
 			if (NewFailure != null)
 			{
 				if (NewFailure.Change < Span.FirstFailure.Change)
 				{
-					Transaction.Set(x => x.FirstFailure, new IssueStep(Span.Id, NewFailure));
+					Updates.Add(Builders<IssueSpan>.Update.Set(x => x.FirstFailure, new IssueStep(Span.Id, NewFailure)));
 				}
 				if (NewFailure.Change >= Span.LastFailure.Change)
 				{
-					Transaction.Set(x => x.LastFailure, new IssueStep(Span.Id, NewFailure));
+					Updates.Add(Builders<IssueSpan>.Update.Set(x => x.LastFailure, new IssueStep(Span.Id, NewFailure)));
 				}
 				if (NewFailure.NotifySuspects != Span.NotifySuspects && NewFailure.Change >= Span.LastFailure.Change)
 				{
-					Transaction.Set(x => x.NotifySuspects, NewFailure.NotifySuspects);
+					Updates.Add(Builders<IssueSpan>.Update.Set(x => x.NotifySuspects, NewFailure.NotifySuspects));
 				}
 			}
 			if (NewNextSuccess != null)
 			{
-				Transaction.Set(x => x.MaxChange, NewNextSuccess.Change);
-				Transaction.Set(x => x.NextSuccess, new IssueStep(Span.Id, NewNextSuccess));
+				Updates.Add(Builders<IssueSpan>.Update.Set(x => x.MaxChange, NewNextSuccess.Change));
+				Updates.Add(Builders<IssueSpan>.Update.Set(x => x.NextSuccess, new IssueStep(Span.Id, NewNextSuccess)));
 			}
 			if (NewSuspects != null)
 			{
-				Transaction.Set(x => x.Suspects, CreateSpanSuspects(NewSuspects));
+				Updates.Add(Builders<IssueSpan>.Update.Set(x => x.Suspects, CreateSpanSuspects(NewSuspects)));
 			}
 			if (NewModified != null)
 			{
-				Transaction.Set(x => x.Modified, NewModified.Value);
+				Updates.Add(Builders<IssueSpan>.Update.Set(x => x.Modified, NewModified.Value));
 			}
-			
-			IIssueSpan? NewSpan = await ApplyTransactionAsync(Span, Transaction);
+
+			if (Updates.Count == 0)
+			{
+				return Span;
+			}
+
+			IssueSpan? NewSpan = await TryUpdateSpanAsync(Span, Builders<IssueSpan>.Update.Combine(Updates));
 			if (NewSpan != null && NewFailure != null)
 			{
 				await AddStepAsync(NewSpan.Id, NewFailure);
diff --git a/Engine/Source/Programs/Horde/HordeServer/Services/Impl/IssueService.cs b/Engine/Source/Programs/Horde/HordeServer/Services/Impl/IssueService.cs
--- a/Engine/Source/Programs/Horde/HordeServer/Services/Impl/IssueService.cs
+++ b/Engine/Source/Programs/Horde/HordeServer/Services/Impl/IssueService.cs
@@ -739,7 +739,7 @@ namespace HordeServer.Services.Impl
 						}
 
 						// Update the span
-						IIssueSpan? NewSpan = await IssueCollection.UpdateSpanAsync(OpenSpan, NewSeverity: NewSeverity, NewFailure: NewFailure, NewModified: bMarkAsModified);
+						IIssueSpan? NewSpan = await IssueCollection.UpdateSpanAsync(OpenSpan, NewSeverity: NewSeverity, NewFailure: NewFailure, NewModified: (bMarkAsModified? (bool?)true : null));
 						if (NewSpan == null)
 						{
 							return false;
@@ -1135,8 +1135,8 @@ namespace HordeServer.Services.Impl
 			IIssue? NewIssue = await IssueCollection.UpdateIssueDerivedDataAsync(Issue);
 			if (NewIssue != null)
 			{
-				string NewSummary = GetSummary(Issue.Fingerprint, Issue.Severity);
-				if (!NewSummary.Equals(Issue.Summary, StringComparison.Ordinal))
+				string NewSummary = GetSummary(NewIssue.Fingerprint, NewIssue.Severity);
+				if (!NewSummary.Equals(NewIssue.Summary, StringComparison.Ordinal))
 				{
 					NewIssue = await IssueCollection.UpdateIssueAsync(NewIssue, NewSummary: NewSummary);
 				}
diff --git a/Engine/Source/Programs/Horde/HordeServerTests/IssueServiceTests.cs b/Engine/Source/Programs/Horde/HordeServerTests/IssueServiceTests.cs
--- a/Engine/Source/Programs/Horde/HordeServerTests/IssueServiceTests.cs
+++ b/Engine/Source/Programs/Horde/HordeServerTests/IssueServiceTests.cs
@@ -1196,6 +1196,43 @@ namespace HordeServerTests
 			}
 		}
 
+		[TestMethod]
+		public async Task AutoResolveTest()
+		{
+			int IssueId;
+
+			// #1
+			// Scenario: Warning in first step
+			// Expected: Default issue is created
+			{
+				IJob Job = CreateJob(MainStreamId, 105, "Test Build", Graph);
+				await AddEvent(Job, 0, 0, new { level = nameof(LogLevel.Warning) }, EventSeverity.Warning);
+				await UpdateCompleteStep(Job, 0, 0, JobStepOutcome.Warnings);
+
+				List<IIssue> Issues = await IssueService.FindIssuesAsync();
+				Assert.AreEqual(1, Issues.Count);
+				Assert.AreEqual(IssueSeverity.Warning, Issues[0].Severity);
+
+				Assert.AreEqual("Warnings in Update Version Files", Issues[0].Summary);
+
+				IssueId = Issues[0].Id;
+			}
+
+			// #2
+			// Scenario: Job succeeds
+			// Expected: Issue is marked as resolved
+			{
+				IJob Job = CreateJob(MainStreamId, 115, "Test Build", Graph);
+				await UpdateCompleteStep(Job, 0, 0, JobStepOutcome.Success);
+
+				IIssue? Issue = await IssueService.GetIssueAsync(IssueId);
+				Assert.IsNotNull(Issue!.ResolvedAt);
+
+				List<IIssueSpan> Spans = await IssueService.GetIssueSpansAsync(Issue);
+				Assert.AreEqual(Spans.Count, 1);
+			}
+		}
+
 		private async Task ParseAsync(ObjectId LogId, string[] Lines)
 		{
 			LogParserContext Context = new LogParserContext();
