Adding warnings for when no ignore patterns were added and add a verbose trace for what ignore patterns are being used.

#rnx

diff --git a/Engine/Source/Programs/Horde/HordeAgent/Execution/BuildGraphExecutor.cs b/Engine/Source/Programs/Horde/HordeAgent/Execution/BuildGraphExecutor.cs
--- a/Engine/Source/Programs/Horde/HordeAgent/Execution/BuildGraphExecutor.cs
+++ b/Engine/Source/Programs/Horde/HordeAgent/Execution/BuildGraphExecutor.cs
@@ -661,10 +661,20 @@ namespace HordeAgent.Execution
 				}
 			}
 
+			if (!IgnorePatternLines.Any())
+			{
+				Log.TraceWarning("No ignore pattern files were read. BaseDirs:");
+				foreach (DirectoryReference BaseDir in BaseDirs)
+				{
+					Log.TraceWarning(BaseDir.ToString());
+				}
+			}
+
 			HashSet<string> IgnorePatterns = new HashSet<string>(StringComparer.Ordinal);
 			foreach (string Line in IgnorePatternLines)
 			{
 				string TrimLine = Line.Trim();
+				Log.TraceVerbose("IgnorePattern: {0}", TrimLine);
 				if (TrimLine.Length > 0 && TrimLine[0] != '#')
 				{
 					IgnorePatterns.Add(TrimLine);
