Make the IgnorePattern output an info for now to make this issue easier to track down.

#rnx

diff --git a/Engine/Source/Programs/Horde/HordeAgent/Execution/BuildGraphExecutor.cs b/Engine/Source/Programs/Horde/HordeAgent/Execution/BuildGraphExecutor.cs
--- a/Engine/Source/Programs/Horde/HordeAgent/Execution/BuildGraphExecutor.cs
+++ b/Engine/Source/Programs/Horde/HordeAgent/Execution/BuildGraphExecutor.cs
@@ -674,7 +674,7 @@ namespace HordeAgent.Execution
 			foreach (string Line in IgnorePatternLines)
 			{
 				string TrimLine = Line.Trim();
-				Log.TraceVerbose("IgnorePattern: {0}", TrimLine);
+				Log.TraceInformation("IgnorePattern: {0}", TrimLine);
 				if (TrimLine.Length > 0 && TrimLine[0] != '#')
 				{
 					IgnorePatterns.Add(TrimLine);
