Horde: Fix false positive on error codes starting with a "3".

diff --git a/Engine/Source/Programs/Horde/HordeAgent/Parser/Matchers/CrashEventMatcher.cs b/Engine/Source/Programs/Horde/HordeAgent/Parser/Matchers/CrashEventMatcher.cs
--- a/Engine/Source/Programs/Horde/HordeAgent/Parser/Matchers/CrashEventMatcher.cs
+++ b/Engine/Source/Programs/Horde/HordeAgent/Parser/Matchers/CrashEventMatcher.cs
@@ -43,7 +43,7 @@ namespace HordeAgent.Parser.Matchers
 			}
 
 			Match? Match;
-			if (Cursor.TryMatch("ExitCode=(3|139|255)", out Match))
+			if (Cursor.TryMatch(@"ExitCode=(3|139|255)(?!\d)", out Match))
 			{
 				LogEventBuilder Builder = new LogEventBuilder(Cursor);
 				Builder.Lines[0].AddSpan(Match.Groups[1], "exitCode");
diff --git a/Engine/Source/Programs/Horde/HordeAgentTests/LogParserTests.cs b/Engine/Source/Programs/Horde/HordeAgentTests/LogParserTests.cs
--- a/Engine/Source/Programs/Horde/HordeAgentTests/LogParserTests.cs
+++ b/Engine/Source/Programs/Horde/HordeAgentTests/LogParserTests.cs
@@ -122,6 +122,24 @@ namespace HordeAgentTests
 			}
 		}
 
+		[TestMethod]
+		public void CrashEventMatcher2()
+		{
+			string[] Lines =
+			{
+				@"Took 620.820352s to run UE4Editor-Cmd.exe, ExitCode=3",
+				@"Took 620.820352s to run UE4Editor-Cmd.exe, ExitCode=30",
+			};
+
+			List<CapturedEvent> Events = Parse(Lines);
+			Assert.AreEqual(1, Events.Count);
+			Assert.AreEqual(LogLevel.Error, Events[0].Level);
+			Assert.AreEqual(KnownLogEvents.AutomationTool_CrashExitCode, Events[0].Id);
+
+			Assert.AreEqual(0, Events[0].LineIndex);
+			Assert.AreEqual(1, Events[0].LineCount);
+		}
+
 		[TestMethod]
 		public void CompileEventMatcher()
 		{
