Change how the ignore pattern stuff logs to try to get it actually working.

#rnx

diff --git a/Engine/Source/Programs/Horde/HordeAgent/Execution/BuildGraphExecutor.cs b/Engine/Source/Programs/Horde/HordeAgent/Execution/BuildGraphExecutor.cs
--- a/Engine/Source/Programs/Horde/HordeAgent/Execution/BuildGraphExecutor.cs
+++ b/Engine/Source/Programs/Horde/HordeAgent/Execution/BuildGraphExecutor.cs
@@ -643,7 +643,7 @@ namespace HordeAgent.Execution
 			}
 		}
 
-		static async Task<List<string>> ReadIgnorePatternsAsync(DirectoryReference WorkspaceDir)
+		static async Task<List<string>> ReadIgnorePatternsAsync(DirectoryReference WorkspaceDir, ILogger Logger)
 		{
 			List<DirectoryReference> BaseDirs = new List<DirectoryReference>();
 			BaseDirs.Add(DirectoryReference.Combine(WorkspaceDir, "Engine"));
@@ -656,17 +656,17 @@ namespace HordeAgent.Execution
 				FileReference IgnorePatternFile = FileReference.Combine(BaseDir, "Build", "Horde", "IgnorePatterns.txt");
 				if (FileReference.Exists(IgnorePatternFile))
 				{
-					Log.TraceInformation("Reading ignore patterns from {0}...", IgnorePatternFile);
+					Logger.LogInformation("Reading ignore patterns from {0}...", IgnorePatternFile);
 					IgnorePatternLines.AddRange(await FileReference.ReadAllLinesAsync(IgnorePatternFile));
 				}
 			}
 
 			if (!IgnorePatternLines.Any())
 			{
-				Log.TraceWarning("No ignore pattern files were read. BaseDirs:");
+				Logger.LogWarning("No ignore pattern files were read. BaseDirs:");
 				foreach (DirectoryReference BaseDir in BaseDirs)
 				{
-					Log.TraceWarning(BaseDir.ToString());
+					Logger.LogWarning(BaseDir.ToString());
 				}
 			}
 
@@ -674,7 +674,6 @@ namespace HordeAgent.Execution
 			foreach (string Line in IgnorePatternLines)
 			{
 				string TrimLine = Line.Trim();
-				Log.TraceInformation("IgnorePattern: {0}", TrimLine);
 				if (TrimLine.Length > 0 && TrimLine[0] != '#')
 				{
 					IgnorePatterns.Add(TrimLine);
@@ -746,7 +745,7 @@ namespace HordeAgent.Execution
 					Context.PerforceStream = Stream.Name;
 					Context.PerforceChange = Job.Change;
 
-					List<string> IgnorePatterns = await ReadIgnorePatternsAsync(WorkspaceDir);
+					List<string> IgnorePatterns = await ReadIgnorePatternsAsync(WorkspaceDir, Logger);
 					using (LogParser Filter = new LogParser(Logger, Context, IgnorePatterns))
 					{
 						await Process.CopyToAsync((Buffer, Offset, Length) => Filter.WriteData(Buffer.AsSpan(Offset, Length), false), 4096, CancellationToken);
