Horde: Speculative fixes for jobs being triggered incorrectly.

* Don't allow setting the last trigger CL to be behind the current CL (may happen if a pod hangs for a long time running a P4 command)
* Force exit of trigger update loop if returned change is outside the requested range (failsafe)

Also add additional logging for the last triggered job with each update.

diff --git a/Engine/Source/Programs/Horde/HordeServer/Services/ScheduleService.cs b/Engine/Source/Programs/Horde/HordeServer/Services/ScheduleService.cs
--- a/Engine/Source/Programs/Horde/HordeServer/Services/ScheduleService.cs
+++ b/Engine/Source/Programs/Horde/HordeServer/Services/ScheduleService.cs
@@ -195,7 +195,7 @@ namespace HordeServer.Services
 				Logger.LogDebug("Schedule {StreamId}/{TemplateRefId} will next trigger at {NextTriggerTime}", Stream.Id, TemplateRefId, NextTriggerTime.Value);
 				return NextTriggerTime;
 			}
-			Logger.LogInformation("Schedule {StreamId}/{TemplateRefId} is ready to trigger", Stream.Id, TemplateRefId);
+			Logger.LogInformation("Schedule {StreamId}/{TemplateRefId} is ready to trigger (last: {Change} at {Time})", Stream.Id, TemplateRefId, Schedule.LastTriggerChange, Schedule.LastTriggerTime);
 
 			// Get a list of jobs that we need to remove
 			List<ObjectId> RemoveJobIds = new List<ObjectId>();
@@ -291,7 +291,7 @@ namespace HordeServer.Services
 				}
 
 				// Quit if we didn't find anything
-				if (Change <= 0)
+				if (Change < MinChangeNumber)
 				{
 					break;
 				}
diff --git a/Engine/Source/Programs/Horde/HordeServer/Services/StreamService.cs b/Engine/Source/Programs/Horde/HordeServer/Services/StreamService.cs
--- a/Engine/Source/Programs/Horde/HordeServer/Services/StreamService.cs
+++ b/Engine/Source/Programs/Horde/HordeServer/Services/StreamService.cs
@@ -145,7 +145,17 @@ namespace HordeServer.Services
 				}
 
 				List<ObjectId> NewActiveJobs = TemplateRef.Schedule.ActiveJobs.Except(RemoveJobs).Union(AddJobs).ToList();
-				if (await Streams.TryUpdateScheduleTriggerAsync(NewStream, TemplateRefId, LastTriggerTime, LastTriggerChange, NewActiveJobs))
+
+				int? NewLastTriggerChange = LastTriggerChange;
+				DateTimeOffset? NewLastTriggerTime = LastTriggerTime;
+
+				if (NewLastTriggerChange != null && NewLastTriggerChange.Value < TemplateRef.Schedule.LastTriggerChange)
+				{
+					NewLastTriggerChange = null;
+					NewLastTriggerTime = null;
+				}
+
+				if (await Streams.TryUpdateScheduleTriggerAsync(NewStream, TemplateRefId, NewLastTriggerTime, NewLastTriggerChange, NewActiveJobs))
 				{
 					return NewStream;
 				}
