Make sure to strip out the empty string that is generated when splitting the empty IgnorePatterns.txt file.

#rnx

diff --git a/Engine/Source/Programs/Horde/HordeAgent/Execution/BuildGraphExecutor.cs b/Engine/Source/Programs/Horde/HordeAgent/Execution/BuildGraphExecutor.cs
--- a/Engine/Source/Programs/Horde/HordeAgent/Execution/BuildGraphExecutor.cs
+++ b/Engine/Source/Programs/Horde/HordeAgent/Execution/BuildGraphExecutor.cs
@@ -650,7 +650,7 @@ namespace HordeAgent.Execution
 			AddRestrictedDirs(BaseDirs, "Restricted");
 			AddRestrictedDirs(BaseDirs, "Platforms");
 
-			List<string> IgnorePatternLines = new List<string>(Properties.Resources.IgnorePatterns.Split('\n'));
+			List<string> IgnorePatternLines = new List<string>(Properties.Resources.IgnorePatterns.Split('\n', StringSplitOptions.RemoveEmptyEntries));
 			foreach (DirectoryReference BaseDir in BaseDirs)
 			{
 				FileReference IgnorePatternFile = FileReference.Combine(BaseDir, "Build", "Horde", "IgnorePatterns.txt");
@@ -748,6 +748,11 @@ namespace HordeAgent.Execution
 					List<string> IgnorePatterns = await ReadIgnorePatternsAsync(WorkspaceDir, Logger);
 					using (LogParser Filter = new LogParser(Logger, Context, IgnorePatterns))
 					{
+						Filter.WriteLine($"Using {IgnorePatterns.Count} IgnorePattern(s):");
+						foreach (string IgnorePattern in IgnorePatterns)
+						{
+							Filter.WriteLine(IgnorePattern);
+						}
 						await Process.CopyToAsync((Buffer, Offset, Length) => Filter.WriteData(Buffer.AsSpan(Offset, Length), false), 4096, CancellationToken);
 						Filter.WriteData(Array.Empty<byte>(), true);
 					}
