Horde: Fix existing schedule state not being copied over when updating streams.

diff --git a/Engine/Source/Programs/Horde/HordeServer/Collections/Impl/StreamCollection.cs b/Engine/Source/Programs/Horde/HordeServer/Collections/Impl/StreamCollection.cs
--- a/Engine/Source/Programs/Horde/HordeServer/Collections/Impl/StreamCollection.cs
+++ b/Engine/Source/Programs/Horde/HordeServer/Collections/Impl/StreamCollection.cs
@@ -117,7 +117,7 @@ namespace HordeServer.Collections.Impl
 		public async Task<IStream?> TryCreateOrReplaceAsync(StreamId Id, IStream? Stream, string ConfigPath, string Revision, ProjectId ProjectId, StreamConfig Config)
 		{
 			List<StreamTab> Tabs = Config.Tabs.ConvertAll(x => StreamTab.FromRequest(x));
-			Dictionary<TemplateRefId, TemplateRef> TemplateRefs = await CreateTemplateRefsAsync(Config.Templates, null, TemplateCollection);
+			Dictionary<TemplateRefId, TemplateRef> TemplateRefs = await CreateTemplateRefsAsync(Config.Templates, Stream, TemplateCollection);
 
 			Dictionary<string, AgentType>? AgentTypes = null;
 			if (Config.AgentTypes != null)
