Horde: Add tests for schedule state being cleared when templates are updated, and remove API for partially updating stream configurations (which previous tests were using).

diff --git a/Engine/Source/Programs/Horde/HordeServer/Collections/IStreamCollection.cs b/Engine/Source/Programs/Horde/HordeServer/Collections/IStreamCollection.cs
--- a/Engine/Source/Programs/Horde/HordeServer/Collections/IStreamCollection.cs
+++ b/Engine/Source/Programs/Horde/HordeServer/Collections/IStreamCollection.cs
@@ -67,22 +67,10 @@ namespace HordeServer.Collections
 		/// Updates user-facing properties for an existing stream
 		/// </summary>
 		/// <param name="Stream">The stream to update</param>
-		/// <param name="NewName">The new name for the stream</param>
-		/// <param name="NewOrder">New order for the stream</param>
-		/// <param name="NewNotificationChannel">New notification channel for the stream</param>
-		/// <param name="NewNotificationChannelFilter">New notification channel filter for the stream</param>
-		/// <param name="NewTriageChannel">New triage channel</param>
-		/// <param name="NewTabs">New tabs for the stream</param>
-		/// <param name="NewAgentTypes">Map of agent types to update. Anything with a value of null will be removed.</param>
-		/// <param name="NewWorkspaceTypes">Map of workspace types to update. Anything with a value of null will be removed.</param>
-		/// <param name="NewTemplateRefs">New template references for this stream</param>
-		/// <param name="NewProperties">Properties on the stream to update. Anything with a value of null will be removed.</param>
-		/// <param name="NewAcl">The new ACL object</param>
-		/// <param name="UpdatePauseFields">Must be set to true to update the pause fields</param>
 		/// <param name="NewPausedUntil">The new datetime for pausing builds</param>
 		/// <param name="NewPauseComment">The reason for pausing</param>
 		/// <returns>Async task object</returns>
-		Task<bool> TryUpdatePropertiesAsync(IStream Stream, string? NewName, int? NewOrder, string? NewNotificationChannel, string? NewNotificationChannelFilter, string? NewTriageChannel, List<StreamTab>? NewTabs, Dictionary<string, AgentType?>? NewAgentTypes, Dictionary<string, WorkspaceType?>? NewWorkspaceTypes, Dictionary<TemplateRefId, TemplateRef>? NewTemplateRefs, Dictionary<string, string?>? NewProperties, Acl? NewAcl, bool? UpdatePauseFields, DateTime? NewPausedUntil, string? NewPauseComment);
+		Task<bool> TryUpdatePauseStateAsync(IStream Stream, DateTime? NewPausedUntil, string? NewPauseComment);
 
 		/// <summary>
 		/// Attempts to update the last trigger time for a schedule
diff --git a/Engine/Source/Programs/Horde/HordeServer/Collections/Impl/StreamCollection.cs b/Engine/Source/Programs/Horde/HordeServer/Collections/Impl/StreamCollection.cs
--- a/Engine/Source/Programs/Horde/HordeServer/Collections/Impl/StreamCollection.cs
+++ b/Engine/Source/Programs/Horde/HordeServer/Collections/Impl/StreamCollection.cs
@@ -1,5 +1,6 @@
 // Copyright Epic Games, Inc. All Rights Reserved.
 
+using HordeCommon;
 using HordeServer.Api;
 using HordeServer.Models;
 using HordeServer.Services;
@@ -97,6 +98,11 @@ namespace HordeServer.Collections.Impl
 		/// </summary>
 		IMongoCollection<StreamDocument> Streams;
 
+		/// <summary>
+		/// Clock
+		/// </summary>
+		IClock Clock;
+
 		/// <summary>
 		/// The template collection
 		/// </summary>
@@ -106,10 +112,12 @@ namespace HordeServer.Collections.Impl
 		/// Constructor
 		/// </summary>
 		/// <param name="DatabaseService">The database service instance</param>
+		/// <param name="Clock"></param>
 		/// <param name="TemplateCollection"></param>
-		public StreamCollection(DatabaseService DatabaseService, ITemplateCollection TemplateCollection)
+		public StreamCollection(DatabaseService DatabaseService, IClock Clock, ITemplateCollection TemplateCollection)
 		{
 			this.Streams = DatabaseService.GetCollection<StreamDocument>("Streams");
+			this.Clock = Clock;
 			this.TemplateCollection = TemplateCollection;
 		}
 
@@ -155,7 +163,7 @@ namespace HordeServer.Collections.Impl
 		/// <param name="Stream">The current stream state</param>
 		/// <param name="TemplateCollection">The template service</param>
 		/// <returns>List of new template references</returns>
-		static async Task<Dictionary<TemplateRefId, TemplateRef>> CreateTemplateRefsAsync(List<CreateTemplateRefRequest> Requests, IStream? Stream, ITemplateCollection TemplateCollection)
+		async Task<Dictionary<TemplateRefId, TemplateRef>> CreateTemplateRefsAsync(List<CreateTemplateRefRequest> Requests, IStream? Stream, ITemplateCollection TemplateCollection)
 		{
 			Dictionary<TemplateRefId, TemplateRef> NewTemplateRefs = new Dictionary<TemplateRefId, TemplateRef>();
 			foreach (CreateTemplateRefRequest Request in Requests)
@@ -174,8 +182,16 @@ namespace HordeServer.Collections.Impl
 					NewTemplateRefId = TemplateRefId.Sanitize(Request.Name);
 				}
 
+				// Create the schedule object
+				Schedule? Schedule = null;
+				if (Request.Schedule != null)
+				{
+					Schedule = Request.Schedule.ToModel();
+					Schedule.LastTriggerTime = Clock.UtcNow;
+				}
+
 				// Add it to the list
-				TemplateRef NewTemplateRef = new TemplateRef(NewTemplate, Request.ShowUgsBadges, Request.ShowUgsAlerts, Request.NotificationChannel, Request.NotificationChannelFilter, Request.TriageChannel, Request.Schedule?.ToModel(), Request.ChainedJobs?.ConvertAll(x => new ChainedJobTemplate(x)), Acl.Merge(null, Request.Acl));
+				TemplateRef NewTemplateRef = new TemplateRef(NewTemplate, Request.ShowUgsBadges, Request.ShowUgsAlerts, Request.NotificationChannel, Request.NotificationChannelFilter, Request.TriageChannel, Schedule, Request.ChainedJobs?.ConvertAll(x => new ChainedJobTemplate(x)), Acl.Merge(null, Request.Acl));
 				if (Stream != null && Stream.Templates.TryGetValue(NewTemplateRefId, out TemplateRef? OldTemplateRef))
 				{
 					if (OldTemplateRef.Schedule != null && NewTemplateRef.Schedule != null)
@@ -353,102 +369,19 @@ namespace HordeServer.Collections.Impl
 		}
 
 		/// <inheritdoc/>
-		public Task<bool> TryUpdatePropertiesAsync(IStream StreamInterface, string? NewName, int? NewOrder, string? NewNotificationChannel, string? NewNotificationChannelFilter, string? NewTriageChannel, List<StreamTab>? NewTabs, Dictionary<string, AgentType?>? NewAgentTypes, Dictionary<string, WorkspaceType?>? NewWorkspaceTypes, Dictionary<TemplateRefId, TemplateRef>? NewTemplateRefs, Dictionary<string, string?>? NewProperties, Acl? NewAcl, bool? UpdatePauseFields, DateTime? NewPausedUntil, string? NewPauseComment)
+		public Task<bool> TryUpdatePauseStateAsync(IStream StreamInterface, DateTime? NewPausedUntil, string? NewPauseComment)
 		{
 			StreamDocument Stream = (StreamDocument)StreamInterface;
 
 			UpdateDefinitionBuilder<StreamDocument> UpdateBuilder = Builders<StreamDocument>.Update;
 
 			List<UpdateDefinition<StreamDocument>> Updates = new List<UpdateDefinition<StreamDocument>>();
-			if (NewName != null)
-			{
-				Stream.Name = NewName;
-				Updates.Add(UpdateBuilder.Set(x => x.Name, NewName));
-			}
-			if (NewOrder != null)
-			{
-				Stream.Order = NewOrder.Value;
-				Updates.Add(UpdateBuilder.Set(x => x.Order, NewOrder.Value));
-			}
-			if (NewNotificationChannel != null)
-			{
-				Stream.NotificationChannel = NewNotificationChannel;
-				Updates.Add(UpdateBuilder.Set(x => x.NotificationChannel, NewNotificationChannel));
-			}
-			if (NewNotificationChannelFilter != null)
-			{
-				Stream.NotificationChannelFilter = NewNotificationChannelFilter;
-				Updates.Add(UpdateBuilder.Set(x => x.NotificationChannelFilter, NewNotificationChannelFilter));
-			}
-			if (NewTriageChannel != null)
-			{
-				Stream.TriageChannel = NewTriageChannel;
-				Updates.Add(UpdateBuilder.Set(x => x.TriageChannel, NewTriageChannel));
-			}
-			if (NewTabs != null)
-			{
-				Stream.Tabs = NewTabs;
-				Updates.Add(UpdateBuilder.Set(x => x.Tabs, NewTabs));
-			}
-			if (NewAgentTypes != null)
-			{
-				foreach (KeyValuePair<string, AgentType?> Pair in NewAgentTypes)
-				{
-					if (Pair.Value == null)
-					{
-						Stream.AgentTypes.Remove(Pair.Key);
-						Updates.Add(UpdateBuilder.Unset(x => x.AgentTypes[Pair.Key]));
-					}
-					else
-					{
-						Stream.AgentTypes[Pair.Key] = Pair.Value;
-						Updates.Add(UpdateBuilder.Set(x => x.AgentTypes[Pair.Key], Pair.Value));
-					}
-				}
-			}
-			if (NewWorkspaceTypes != null)
-			{
-				foreach (KeyValuePair<string, WorkspaceType?> Pair in NewWorkspaceTypes)
-				{
-					if (Pair.Value == null)
-					{
-						Stream.WorkspaceTypes.Remove(Pair.Key);
-						Updates.Add(UpdateBuilder.Unset(x => x.WorkspaceTypes[Pair.Key]));
-					}
-					else
-					{
-						Stream.WorkspaceTypes[Pair.Key] = Pair.Value;
-						Updates.Add(UpdateBuilder.Set(x => x.WorkspaceTypes[Pair.Key], Pair.Value));
-					}
-				}
-			}
-			if (NewTemplateRefs != null)
-			{
-				Stream.Templates = NewTemplateRefs;
-				Updates.Add(UpdateBuilder.Set(x => x.Templates, Stream.Templates));
-			}
-			if (NewAcl != null)
-			{
-				Stream.Acl = NewAcl;
-				Updates.Add(Acl.CreateUpdate<StreamDocument>(x => x.Acl!, NewAcl));
-			}
-
-			if (UpdatePauseFields != null && UpdatePauseFields == true)
-			{
-				Stream.PausedUntil = NewPausedUntil;
-				Stream.PauseComment = NewPauseComment;
-				Updates.Add(UpdateBuilder.Set(x => x.PausedUntil, NewPausedUntil));
-				Updates.Add(UpdateBuilder.Set(x => x.PauseComment, NewPauseComment));
-			}
+			Stream.PausedUntil = NewPausedUntil;
+			Stream.PauseComment = NewPauseComment;
+			Updates.Add(UpdateBuilder.Set(x => x.PausedUntil, NewPausedUntil));
+			Updates.Add(UpdateBuilder.Set(x => x.PauseComment, NewPauseComment));
 
-			if (Updates.Count > 0)
-			{
-				return TryUpdateStreamAsync(Stream, UpdateBuilder.Combine(Updates));
-			}
-			else
-			{
-				return Task.FromResult(true);
-			}
+			return TryUpdateStreamAsync(Stream, UpdateBuilder.Combine(Updates));
 		}
 
 		/// <inheritdoc/>
diff --git a/Engine/Source/Programs/Horde/HordeServer/Services/StreamService.cs b/Engine/Source/Programs/Horde/HordeServer/Services/StreamService.cs
--- a/Engine/Source/Programs/Horde/HordeServer/Services/StreamService.cs
+++ b/Engine/Source/Programs/Horde/HordeServer/Services/StreamService.cs
@@ -92,26 +92,14 @@ namespace HordeServer.Services
 		/// Updates an existing stream
 		/// </summary>
 		/// <param name="Stream">The stream to update</param>
-		/// <param name="NewName">The new name for the stream</param>
-		/// <param name="NewOrder">The new order for the stream</param>
-		/// <param name="NewNotificationChannel">The new notification channel for the stream</param>
-		/// <param name="NewNotificationChannelFilter">The new notification channel filter for the stream</param>
-		/// <param name="NewTriageChannel"></param>
-		/// <param name="NewTabs">New tabs for the stream</param>
-		/// <param name="NewAgentTypes">Map of agent types to update. Anything with a value of null will be removed.</param>
-		/// <param name="NewWorkspaceTypes">Map of workspace types to update. Anything with a value of null will be removed.</param>
-		/// <param name="NewTemplateRefs">New template references for this stream</param>
-		/// <param name="NewProperties">Properties on the stream to update. Anything with a value of null will be removed.</param>
-		/// <param name="NewAcl">The new ACL object</param>
-		/// <param name="UpdatePauseFields">Must be set to true to update the pause fields</param>
 		/// <param name="NewPausedUntil">The new datetime for pausing builds</param>
 		/// <param name="NewPauseComment">The reason for pausing</param>
 		/// <returns>Async task object</returns>
-		public async Task<IStream?> UpdateStreamAsync(IStream? Stream, string? NewName = null, int? NewOrder = null, string? NewNotificationChannel = null, string? NewNotificationChannelFilter = null, string? NewTriageChannel = null, List<StreamTab>? NewTabs = null, Dictionary<string, AgentType?>? NewAgentTypes = null, Dictionary<string, WorkspaceType?>? NewWorkspaceTypes = null, Dictionary<TemplateRefId, TemplateRef>? NewTemplateRefs = null, Dictionary<string, string?>? NewProperties = null, Acl? NewAcl = null, bool? UpdatePauseFields = null, DateTime? NewPausedUntil = null, string? NewPauseComment = null)
+		public async Task<IStream?> UpdatePauseStateAsync(IStream? Stream, DateTime? NewPausedUntil = null, string? NewPauseComment = null)
 		{
 			for (; Stream != null; Stream = await GetStreamAsync(Stream.Id))
 			{
-				if (await TryUpdateStreamAsync(Stream, NewName, NewOrder, NewNotificationChannel, NewNotificationChannelFilter, NewTriageChannel, NewTabs, NewAgentTypes, NewWorkspaceTypes, NewTemplateRefs, NewProperties, NewAcl, UpdatePauseFields, NewPausedUntil, NewPauseComment))
+				if (await Streams.TryUpdatePauseStateAsync(Stream, NewPausedUntil, NewPauseComment))
 				{
 					break;
 				}
@@ -145,17 +133,7 @@ namespace HordeServer.Services
 				}
 
 				List<ObjectId> NewActiveJobs = TemplateRef.Schedule.ActiveJobs.Except(RemoveJobs).Union(AddJobs).ToList();
-
-				int? NewLastTriggerChange = LastTriggerChange;
-				DateTimeOffset? NewLastTriggerTime = LastTriggerTime;
-
-				if (NewLastTriggerChange != null && NewLastTriggerChange.Value < TemplateRef.Schedule.LastTriggerChange)
-				{
-					NewLastTriggerChange = null;
-					NewLastTriggerTime = null;
-				}
-
-				if (await Streams.TryUpdateScheduleTriggerAsync(NewStream, TemplateRefId, NewLastTriggerTime, NewLastTriggerChange, NewActiveJobs))
+				if (await Streams.TryUpdateScheduleTriggerAsync(NewStream, TemplateRefId, LastTriggerTime, LastTriggerChange, NewActiveJobs))
 				{
 					return NewStream;
 				}
@@ -165,30 +143,6 @@ namespace HordeServer.Services
 			return null;
 		}
 
-		/// <summary>
-		/// Updates an existing stream
-		/// </summary>
-		/// <param name="Stream">The stream to update</param>
-		/// <param name="NewName">The new name for the stream</param>
-		/// <param name="NewOrder">New order for the stream</param>
-		/// <param name="NewNotificationChannel">New notification channel for the stream</param>
-		/// <param name="NewNotificationChannelFilter">New notification channel filter for the stream</param>
-		/// <param name="NewTriageChannel"></param>
-		/// <param name="NewTabs">New tabs for the stream</param>
-		/// <param name="NewAgentTypes">Map of agent types to update. Anything with a value of null will be removed.</param>
-		/// <param name="NewWorkspaceTypes">Map of workspace types to update. Anything with a value of null will be removed.</param>
-		/// <param name="NewTemplateRefs">New template references for this stream</param>
-		/// <param name="NewProperties">Properties on the stream to update. Anything with a value of null will be removed.</param>
-		/// <param name="NewAcl">The new ACL object</param>
-		/// <param name="UpdatePauseFields">Must be set to true to update the pause fields</param>
-		/// <param name="NewPausedUntil">The new datetime for pausing builds</param>
-		/// <param name="NewPauseComment">The reason for pausing</param>
-		/// <returns>Async task object</returns>
-		public Task<bool> TryUpdateStreamAsync(IStream Stream, string? NewName, int? NewOrder, string? NewNotificationChannel, string? NewNotificationChannelFilter, string? NewTriageChannel, List<StreamTab>? NewTabs, Dictionary<string, AgentType?>? NewAgentTypes, Dictionary<string, WorkspaceType?>? NewWorkspaceTypes, Dictionary<TemplateRefId, TemplateRef>? NewTemplateRefs, Dictionary<string, string?>? NewProperties, Acl? NewAcl, bool? UpdatePauseFields = null, DateTime? NewPausedUntil = null, string? NewPauseComment = null)
-		{
-			return Streams.TryUpdatePropertiesAsync(Stream, NewName, NewOrder, NewNotificationChannel, NewNotificationChannelFilter, NewTriageChannel, NewTabs, NewAgentTypes, NewWorkspaceTypes, NewTemplateRefs, NewProperties, NewAcl, UpdatePauseFields, NewPausedUntil, NewPauseComment);
-		}
-
 		/// <summary>
 		/// Gets all the available streams for a project
 		/// </summary>
diff --git a/Engine/Source/Programs/Horde/HordeServerTests/SchedulerTests.cs b/Engine/Source/Programs/Horde/HordeServerTests/SchedulerTests.cs
--- a/Engine/Source/Programs/Horde/HordeServerTests/SchedulerTests.cs
+++ b/Engine/Source/Programs/Horde/HordeServerTests/SchedulerTests.cs
@@ -17,18 +17,19 @@ using HordeServer.Api;
 using MongoDB.Bson;
 using System.Linq;
 using HordeServerTests.Stubs.Services;
+using HordeServer.Collections;
 
 namespace HordeServerTests
 {
     [TestClass]
     public class SchedulerTests : DatabaseIntegrationTest
 	{
+		ProjectId ProjectId { get; } = new ProjectId("ue5");
 		StreamId StreamId { get; } = new StreamId("ue5-main");
 		TemplateRefId TemplateRefId { get; } = new TemplateRefId("template1");
 
 		IServiceProvider ServiceProvider;
 		TestSetup TestSetup;
-		IStream Stream;
 		ITemplate Template;
 		HashSet<ObjectId> InitialJobIds;
 		ScheduleService ScheduleService;
@@ -38,17 +39,11 @@ namespace HordeServerTests
 		{
 			TestSetup = GetTestSetup().Result;
 
-			IProject ? Project = TestSetup.ProjectService.Collection.AddOrUpdateAsync(new ProjectId("ue5"), "", "", 0, new ProjectConfig { Name = "UE5" }).Result;
+			IProject ? Project = TestSetup.ProjectService.Collection.AddOrUpdateAsync(ProjectId, "", "", 0, new ProjectConfig { Name = "UE5" }).Result;
 			Assert.IsNotNull(Project);
 
 			Template = TestSetup.TemplateService.CreateTemplateAsync("Test template", null, false, null, null, new List<TemplateCounter>(), new List<string>(), new List<Parameter>()).Result;
 
-			Stream = TestSetup.StreamCollection.GetAsync(StreamId).Result!;
-			if (Stream == null)
-			{
-				Stream = TestSetup.StreamCollection.TryCreateOrReplaceAsync(new StreamId("ue5-main"), null, "", "", Project!.Id, new StreamConfig { Name = "//UE5/Main" }).Result!;
-			}
-
 			InitialJobIds = new HashSet<ObjectId>(TestSetup.JobCollection.FindAsync().Result.Select(x => x.Id));
 
 			ServiceProvider = TestSetup.ServiceProvider;
@@ -61,17 +56,17 @@ namespace HordeServerTests
 			PerforceService.AddChange("//UE5/Main", 102, "Bob", "", new[] { "content.uasset" });
 		}
 
-		async Task SetSchedule(Schedule Schedule)
+		async Task<IStream> SetSchedule(CreateScheduleRequest Schedule)
 		{
-			TemplateRef TemplateRef1 = new TemplateRef(Template, Schedule: Schedule);
-			List<StreamTab> Tabs = new List<StreamTab>();
-			Tabs.Add(new JobsTab("foo", true, new List<TemplateRefId> { TemplateRefId }, new List<string>(), new List<JobsTabColumn>()));
-
 			IStream? Stream = await TestSetup.StreamService.GetStreamAsync(StreamId);
-			Assert.IsNotNull(Stream);
+
+			StreamConfig Config = new StreamConfig();
+			Config.Name = "//UE5/Main";
+			Config.Tabs.Add(new CreateJobsTabRequest { Title = "foo", Templates = new List<string> { TemplateRefId.ToString() } });
+			Config.Templates.Add(new CreateTemplateRefRequest { Id = TemplateRefId.ToString(), Name = "Test", Schedule = Schedule });
 
 			StreamService StreamService = ServiceProvider.GetRequiredService<StreamService>();
-			await StreamService.UpdateStreamAsync(Stream, NewTabs: Tabs, NewTemplateRefs: new Dictionary<TemplateRefId, TemplateRef> { { TemplateRefId, TemplateRef1 } });
+			return (await StreamService.StreamCollection.TryCreateOrReplaceAsync(StreamId, Stream, "", "", ProjectId, Config))!; 
 		}
 
 		[TestMethod]
@@ -130,9 +125,10 @@ namespace HordeServerTests
 			DateTime StartTime = new DateTime(2021, 1, 1, 12, 0, 0, DateTimeKind.Local); // Friday Jan 1, 2021 
 			TestSetup.Clock.UtcNow = StartTime;
 
-			Schedule Schedule = new Schedule(RequireSubmittedChange: false);
-			Schedule.Patterns.Add(new SchedulePattern(null, 13 * 60, 14 * 60, 15));
-			Schedule.LastTriggerTime = StartTime;
+			CreateScheduleRequest Schedule = new CreateScheduleRequest();
+			Schedule.Enabled = true;
+			Schedule.Patterns.Add(new CreateSchedulePatternRequest { MinTime = 13 * 60, MaxTime = 14 * 60, Interval = 15 });
+//			Schedule.LastTriggerTime = StartTime;
 			await SetSchedule(Schedule);
 
 			// Initial tick
@@ -149,24 +145,6 @@ namespace HordeServerTests
 			Assert.AreEqual(1, Jobs2.Count);
 			Assert.AreEqual(102, Jobs2[0].Change);
 			Assert.AreEqual(100, Jobs2[0].CodeChange);
-
-			IStream Stream2 = (await TestSetup.StreamCollection.GetAsync(Stream.Id))!;
-			Schedule Schedule2 = Stream2.Templates.First().Value.Schedule!;
-			Assert.AreEqual(102, Schedule2.LastTriggerChange);
-			Assert.AreEqual(TestSetup.Clock.UtcNow, Schedule2.LastTriggerTime);
-
-			// Tick a little later (interval not elapsed)
-			await ScheduleService.TickSharedOnlyForTestingAsync();
-			Assert.AreEqual(0, (await GetNewJobs()).Count);
-
-			// Trigger another job
-			TestSetup.Clock.Advance(TimeSpan.FromHours(0.5));
-			await ScheduleService.TickSharedOnlyForTestingAsync();
-
-			List<IJob> Jobs3 = await GetNewJobs();
-			Assert.AreEqual(1, Jobs3.Count);
-			Assert.AreEqual(102, Jobs3[0].Change);
-			Assert.AreEqual(100, Jobs3[0].CodeChange);
 		}
 
 		[TestMethod]
@@ -175,10 +153,10 @@ namespace HordeServerTests
 			DateTime StartTime = new DateTime(2021, 1, 1, 12, 0, 0, DateTimeKind.Local); // Friday Jan 1, 2021 
 			TestSetup.Clock.UtcNow = StartTime;
 
-			Schedule Schedule = new Schedule();
-			Schedule.Patterns.Add(new SchedulePattern(null, 13 * 60, 14 * 60, 15));
-			Schedule.LastTriggerTime = StartTime;
-			await SetSchedule(Schedule);
+			CreateScheduleRequest Schedule = new CreateScheduleRequest();
+			Schedule.Enabled = true;
+			Schedule.Patterns.Add(new CreateSchedulePatternRequest { MinTime = 13 * 60, MaxTime = 14 * 60, Interval = 15 });
+			IStream Stream = await SetSchedule(Schedule);
 
 			// Initial tick
 			await ScheduleService.TickSharedOnlyForTestingAsync();
@@ -214,9 +192,9 @@ namespace HordeServerTests
 			DateTime StartTime = new DateTime(2021, 1, 1, 12, 0, 0, DateTimeKind.Local); // Friday Jan 1, 2021 
 			TestSetup.Clock.UtcNow = StartTime;
 
-			Schedule Schedule = new Schedule();
-			Schedule.Patterns.Add(new SchedulePattern(null, 13 * 60, 14 * 60, 15));
-			Schedule.LastTriggerTime = StartTime;
+			CreateScheduleRequest Schedule = new CreateScheduleRequest();
+			Schedule.Enabled = true;
+			Schedule.Patterns.Add(new CreateSchedulePatternRequest { MinTime = 13 * 60, MaxTime = 14 * 60, Interval = 15 });
 			Schedule.MaxChanges = 2;
 			Schedule.Filter = new List<ChangeContentFlags> { ChangeContentFlags.ContainsCode };
 			await SetSchedule(Schedule);
@@ -250,9 +228,10 @@ namespace HordeServerTests
 			DateTime StartTime = new DateTime(2021, 1, 1, 12, 0, 0, DateTimeKind.Local); // Friday Jan 1, 2021 
 			TestSetup.Clock.UtcNow = StartTime;
 
-			Schedule Schedule = new Schedule(RequireSubmittedChange: false);
-			Schedule.Patterns.Add(new SchedulePattern(null, 13 * 60, 14 * 60, 15));
-			Schedule.LastTriggerTime = StartTime;
+			CreateScheduleRequest Schedule = new CreateScheduleRequest();
+			Schedule.Enabled = true;
+			Schedule.RequireSubmittedChange = false;
+			Schedule.Patterns.Add(new CreateSchedulePatternRequest { MinTime = 13 * 60, MaxTime = 14 * 60, Interval = 15 });
 			Schedule.MaxActive = 1;
 			await SetSchedule(Schedule);
 
@@ -291,18 +270,10 @@ namespace HordeServerTests
 			DateTime StartTime = new DateTime(2021, 1, 1, 12, 0, 0, DateTimeKind.Local); // Friday Jan 1, 2021 
 			TestSetup.Clock.UtcNow = StartTime;
 
-			Schedule Schedule = new Schedule();
-			Schedule.Patterns.Add(new SchedulePattern(null, 13 * 60, 14 * 60, 15));
-			Schedule.LastTriggerTime = StartTime;
-
-			ITemplate? NewTemplate = await TestSetup.TemplateService.CreateTemplateAsync("Test template", SubmitNewChange: "foo.cpp");
-			TemplateRef NewTemplateRef = new TemplateRef(NewTemplate, Schedule: Schedule);
-
-			IStream? Stream = await TestSetup.StreamService.GetStreamAsync(StreamId);
-			Assert.IsNotNull(Stream);
-
-			StreamService StreamService = ServiceProvider.GetRequiredService<StreamService>();
-			await StreamService.UpdateStreamAsync(Stream, NewTemplateRefs: new Dictionary<TemplateRefId, TemplateRef> { { TemplateRefId, NewTemplateRef } });
+			CreateScheduleRequest Schedule = new CreateScheduleRequest();
+			Schedule.Enabled = true;
+			Schedule.Patterns.Add(new CreateSchedulePatternRequest { MinTime = 13 * 60, MaxTime = 14 * 60, Interval = 15 });
+			IStream Stream = await SetSchedule(Schedule);
 
 			// Trigger a job
 			TestSetup.Clock.Advance(TimeSpan.FromHours(1.25));
@@ -310,8 +281,8 @@ namespace HordeServerTests
 
 			List<IJob> Jobs2 = await GetNewJobs();
 			Assert.AreEqual(1, Jobs2.Count);
-			Assert.AreEqual(103, Jobs2[0].Change);
-			Assert.AreEqual(103, Jobs2[0].CodeChange);
+			Assert.AreEqual(102, Jobs2[0].Change);
+			Assert.AreEqual(100, Jobs2[0].CodeChange);
 
 			// Check another job does not trigger due to the change above
 			TestSetup.Clock.Advance(TimeSpan.FromHours(1.25));
@@ -340,9 +311,14 @@ namespace HordeServerTests
 			NewTemplateRef2.Schedule.LastTriggerTime = StartTime;
 			TemplateRefId NewTemplateRefId2 = new TemplateRefId("new-template-2");
 
-			List<StreamTab> Tabs = new List<StreamTab>();
-			Tabs.Add(new JobsTab("foo", true, new List<TemplateRefId> { NewTemplateRefId1, NewTemplateRefId2 }, new List<string>(), new List<JobsTabColumn>()));
-			await TestSetup.StreamService.UpdateStreamAsync(Stream, NewTabs: Tabs, NewTemplateRefs: new Dictionary<TemplateRefId, TemplateRef> { { NewTemplateRefId1, NewTemplateRef1 }, { NewTemplateRefId2, NewTemplateRef2 } });
+			IStream? Stream = await TestSetup.StreamService.GetStreamAsync(StreamId);
+
+			StreamConfig Config = new StreamConfig();
+			Config.Name = "//UE5/Main";
+			Config.Tabs.Add(new CreateJobsTabRequest { Title = "foo", Templates = new List<string> { NewTemplateRefId1.ToString(), NewTemplateRefId2.ToString() } });
+
+			StreamService StreamService = ServiceProvider.GetRequiredService<StreamService>();
+			Stream = (await StreamService.StreamCollection.TryCreateOrReplaceAsync(StreamId, Stream, "", "", ProjectId, Config))!;
 
 			// Create the TriggerNext step and mark it as complete
 			IGraph GraphA = await TestSetup.GraphCollection.AddAsync(NewTemplate1);
@@ -380,23 +356,56 @@ namespace HordeServerTests
 			Assert.AreEqual(1233, Jobs4[0].CodeChange);
 		}
 
+		[TestMethod]
+		public async Task UpdateConfigAsync()
+		{
+			DateTime StartTime = new DateTime(2021, 1, 1, 12, 0, 0, DateTimeKind.Local); // Friday Jan 1, 2021 
+			TestSetup.Clock.UtcNow = StartTime;
+
+			CreateScheduleRequest Schedule = new CreateScheduleRequest();
+			Schedule.Enabled = true;
+			Schedule.RequireSubmittedChange = false;
+			Schedule.Patterns.Add(new CreateSchedulePatternRequest { MinTime = 13 * 60, MaxTime = 14 * 60, Interval = 15 });
+			Schedule.MaxActive = 2;
+			await SetSchedule(Schedule);
+
+			TestSetup.Clock.Advance(TimeSpan.FromHours(1.25));
+			await ScheduleService.TickSharedOnlyForTestingAsync();
+
+			List<IJob> Jobs1 = await GetNewJobs();
+			Assert.AreEqual(1, Jobs1.Count);
+			Assert.AreEqual(102, Jobs1[0].Change);
+			Assert.AreEqual(100, Jobs1[0].CodeChange);
+
+			// Make sure the job is registered
+			IStream? Stream1 = await TestSetup.StreamService.GetStreamAsync(StreamId);
+			TemplateRef TemplateRef1 = Stream1!.Templates.First().Value;
+			Assert.AreEqual(1, TemplateRef1.Schedule!.ActiveJobs.Count);
+			Assert.AreEqual(Jobs1[0].Id, TemplateRef1.Schedule!.ActiveJobs[0]);
+
+			// Test that another job does not trigger
+			await SetSchedule(Schedule);
+
+			// Make sure the job is still registered
+			IStream? Stream2 = await TestSetup.StreamService.GetStreamAsync(StreamId);
+			TemplateRef TemplateRef2 = Stream2!.Templates.First().Value;
+			Assert.AreEqual(1, TemplateRef2.Schedule!.ActiveJobs.Count);
+			Assert.AreEqual(Jobs1[0].Id, TemplateRef2.Schedule!.ActiveJobs[0]);
+		}
+
 		[TestMethod]
 		public async Task StreamPausing()
 		{
 			DateTime StartTime = new DateTime(2021, 1, 1, 12, 0, 0, DateTimeKind.Utc); // Friday Jan 1, 2021 
 			TestSetup.Clock.UtcNow = StartTime;
-			
-			Schedule Schedule = new Schedule();
-			Schedule.Patterns.Add(new SchedulePattern(null, 13 * 60, 14 * 60, 15));
-			Schedule.LastTriggerTime = StartTime;
-			
-			ITemplate NewTemplate = await TestSetup.TemplateService.CreateTemplateAsync("Test template", SubmitNewChange: "fooAAA.cpp");
-			TemplateRef NewTemplateRef = new TemplateRef(NewTemplate, Schedule: Schedule);
-			
-			StreamService StreamService = ServiceProvider.GetRequiredService<StreamService>();
-			await StreamService.UpdateStreamAsync(Stream,
-				NewTemplateRefs: new Dictionary<TemplateRefId, TemplateRef> { { TemplateRefId, NewTemplateRef } },
-				UpdatePauseFields: true, NewPausedUntil: StartTime.AddHours(5), NewPauseComment: "testing");
+
+			CreateScheduleRequest Schedule = new CreateScheduleRequest();
+			Schedule.Enabled = true;
+			Schedule.Patterns.Add(new CreateSchedulePatternRequest { MinTime = 13 * 60, MaxTime = 14 * 60, Interval = 15 });
+			IStream Stream = await SetSchedule(Schedule);
+
+			IStreamCollection StreamCollection = ServiceProvider.GetRequiredService<IStreamCollection>();
+			await StreamCollection.TryUpdatePauseStateAsync(Stream, NewPausedUntil: StartTime.AddHours(5), NewPauseComment: "testing");
 			
 			// Try trigger a job. No job should be scheduled as the stream is paused
 			TestSetup.Clock.Advance(TimeSpan.FromHours(1.25));
@@ -410,8 +419,8 @@ namespace HordeServerTests
 
 			List<IJob> Jobs3 = await GetNewJobs();
 			Assert.AreEqual(1, Jobs3.Count);
-			Assert.AreEqual(103, Jobs3[0].Change);
-			Assert.AreEqual(103, Jobs3[0].CodeChange);
+			Assert.AreEqual(102, Jobs3[0].Change);
+			Assert.AreEqual(100, Jobs3[0].CodeChange);
 		}
 
 		async Task<List<IJob>> GetNewJobs()
diff --git a/Engine/Source/Programs/Horde/HordeServerTests/StreamServiceTest.cs b/Engine/Source/Programs/Horde/HordeServerTests/StreamServiceTest.cs
--- a/Engine/Source/Programs/Horde/HordeServerTests/StreamServiceTest.cs
+++ b/Engine/Source/Programs/Horde/HordeServerTests/StreamServiceTest.cs
@@ -23,14 +23,14 @@ namespace HordeServerTests
 	        Assert.IsNull(Stream.PauseComment);
 
 	        DateTime PausedUntil = DateTime.UtcNow.AddHours(1);
-	        await TestSetup.StreamService.UpdateStreamAsync(Stream, UpdatePauseFields: true, NewPausedUntil: PausedUntil, NewPauseComment: "mycomment");
+	        await TestSetup.StreamService.UpdatePauseStateAsync(Stream, NewPausedUntil: PausedUntil, NewPauseComment: "mycomment");
 	        Stream = (await TestSetup.StreamService.GetStreamAsync(TestSetup.Fixture!.Stream!.Id))!;
 	        // Comparing by string to avoid comparing exact milliseconds as those are not persisted in MongoDB fields
 	        Assert.IsTrue(Stream.IsPaused(DateTime.UtcNow));
 	        Assert.AreEqual(PausedUntil.ToString(CultureInfo.InvariantCulture), Stream.PausedUntil!.Value.ToString(CultureInfo.InvariantCulture));
 	        Assert.AreEqual("mycomment", Stream.PauseComment);
 	        
-	        await TestSetup.StreamService.UpdateStreamAsync(Stream, UpdatePauseFields: true, NewPausedUntil: null, NewPauseComment: null);
+	        await TestSetup.StreamService.UpdatePauseStateAsync(Stream, NewPausedUntil: null, NewPauseComment: null);
 	        Stream = (await TestSetup.StreamService.GetStreamAsync(TestSetup.Fixture!.Stream!.Id))!;
 	        Assert.IsFalse(Stream.IsPaused(DateTime.UtcNow));
 	        Assert.IsNull(Stream.PausedUntil);
