Remove most of the debug logging for ignore patterns. Keep the code to write out the list of ignore pattern files that are loaded up.

#rnx

diff --git a/Engine/Source/Programs/Horde/HordeAgent/Execution/BuildGraphExecutor.cs b/Engine/Source/Programs/Horde/HordeAgent/Execution/BuildGraphExecutor.cs
--- a/Engine/Source/Programs/Horde/HordeAgent/Execution/BuildGraphExecutor.cs
+++ b/Engine/Source/Programs/Horde/HordeAgent/Execution/BuildGraphExecutor.cs
@@ -661,15 +661,6 @@ namespace HordeAgent.Execution
 				}
 			}
 
-			if (!IgnorePatternLines.Any())
-			{
-				Logger.LogWarning("No ignore pattern files were read. BaseDirs:");
-				foreach (DirectoryReference BaseDir in BaseDirs)
-				{
-					Logger.LogWarning(BaseDir.ToString());
-				}
-			}
-
 			HashSet<string> IgnorePatterns = new HashSet<string>(StringComparer.Ordinal);
 			foreach (string Line in IgnorePatternLines)
 			{
@@ -748,11 +739,6 @@ namespace HordeAgent.Execution
 					List<string> IgnorePatterns = await ReadIgnorePatternsAsync(WorkspaceDir, Logger);
 					using (LogParser Filter = new LogParser(Logger, Context, IgnorePatterns))
 					{
-						Filter.WriteLine($"Using {IgnorePatterns.Count} IgnorePattern(s):");
-						foreach (string IgnorePattern in IgnorePatterns)
-						{
-							Filter.WriteLine(IgnorePattern);
-						}
 						await Process.CopyToAsync((Buffer, Offset, Length) => Filter.WriteData(Buffer.AsSpan(Offset, Length), false), 4096, CancellationToken);
 						Filter.WriteData(Array.Empty<byte>(), true);
 					}
