Horde: Add separate assemblies for HordeServer.Common and HordeServer.Compute, preparing for Jupiter integration.

diff --git a/Engine/Source/Programs/Horde/HordeServer.Common/HordeServer.Common.csproj b/Engine/Source/Programs/Horde/HordeServer.Common/HordeServer.Common.csproj
new file mode 100644
--- /dev/null
+++ b/Engine/Source/Programs/Horde/HordeServer.Common/HordeServer.Common.csproj
@@ -0,0 +1,7 @@
+<Project Sdk="Microsoft.NET.Sdk">
+
+  <PropertyGroup>
+    <TargetFramework>netcoreapp3.1</TargetFramework>
+  </PropertyGroup>
+
+</Project>
diff --git a/Engine/Source/Programs/Horde/HordeServer.Compute/HordeServer.Compute.csproj b/Engine/Source/Programs/Horde/HordeServer.Compute/HordeServer.Compute.csproj
new file mode 100644
--- /dev/null
+++ b/Engine/Source/Programs/Horde/HordeServer.Compute/HordeServer.Compute.csproj
@@ -0,0 +1,12 @@
+<Project Sdk="Microsoft.NET.Sdk">
+
+  <PropertyGroup>
+    <TargetFramework>netcoreapp3.1</TargetFramework>
+  </PropertyGroup>
+
+  <ItemGroup>
+    <ProjectReference Include="..\..\Shared\EpicGames.Core\EpicGames.Core.csproj" />
+    <ProjectReference Include="..\HordeServer.Common\HordeServer.Common.csproj" />
+  </ItemGroup>
+
+</Project>
diff --git a/Engine/Source/Programs/Horde/HordeServer/HordeServer.csproj b/Engine/Source/Programs/Horde/HordeServer/HordeServer.csproj
--- a/Engine/Source/Programs/Horde/HordeServer/HordeServer.csproj
+++ b/Engine/Source/Programs/Horde/HordeServer/HordeServer.csproj
@@ -88,6 +88,8 @@
     <ProjectReference Include="..\..\Shared\EpicGames.Serialization\EpicGames.Serialization.csproj" />
     <ProjectReference Include="..\HordeCommon\HordeCommon.csproj" />
     <ProjectReference Include="..\..\..\ThirdParty\Perforce\P4Api.Net\p4api.net\p4api.net.csproj" />
+    <ProjectReference Include="..\HordeServer.Common\HordeServer.Common.csproj" />
+    <ProjectReference Include="..\HordeServer.Compute\HordeServer.Compute.csproj" />
   </ItemGroup>
 
   <ItemGroup>
