Horde: Updating server to NET Core 5.0.

diff --git a/Engine/Source/Programs/Horde/HordeServer.Common/HordeServer.Common.csproj b/Engine/Source/Programs/Horde/HordeServer.Common/HordeServer.Common.csproj
--- a/Engine/Source/Programs/Horde/HordeServer.Common/HordeServer.Common.csproj
+++ b/Engine/Source/Programs/Horde/HordeServer.Common/HordeServer.Common.csproj
@@ -1,7 +1,7 @@
 <Project Sdk="Microsoft.NET.Sdk">
 
   <PropertyGroup>
-    <TargetFramework>netcoreapp3.1</TargetFramework>
+    <TargetFramework>net5.0</TargetFramework>
   </PropertyGroup>
 
 </Project>
diff --git a/Engine/Source/Programs/Horde/HordeServer.Compute/HordeServer.Compute.csproj b/Engine/Source/Programs/Horde/HordeServer.Compute/HordeServer.Compute.csproj
--- a/Engine/Source/Programs/Horde/HordeServer.Compute/HordeServer.Compute.csproj
+++ b/Engine/Source/Programs/Horde/HordeServer.Compute/HordeServer.Compute.csproj
@@ -1,7 +1,7 @@
 <Project Sdk="Microsoft.NET.Sdk">
 
   <PropertyGroup>
-    <TargetFramework>netcoreapp3.1</TargetFramework>
+    <TargetFramework>net5.0</TargetFramework>
   </PropertyGroup>
 
   <ItemGroup>
diff --git a/Engine/Source/Programs/Horde/HordeServer/HordeServer.csproj b/Engine/Source/Programs/Horde/HordeServer/HordeServer.csproj
--- a/Engine/Source/Programs/Horde/HordeServer/HordeServer.csproj
+++ b/Engine/Source/Programs/Horde/HordeServer/HordeServer.csproj
@@ -1,7 +1,7 @@
 ﻿<Project Sdk="Microsoft.NET.Sdk.Web">
 
   <PropertyGroup>
-    <TargetFramework>netcoreapp3.1</TargetFramework>
+    <TargetFramework>net5.0</TargetFramework>
     <Nullable>enable</Nullable>
     <Deterministic>true</Deterministic>
     <AssemblyName>HordeServer</AssemblyName>
