Horde: Upgrade HordeServerTests to NET 5

diff --git a/Engine/Source/Programs/Horde/HordeServerTests/HordeServerTests.csproj b/Engine/Source/Programs/Horde/HordeServerTests/HordeServerTests.csproj
--- a/Engine/Source/Programs/Horde/HordeServerTests/HordeServerTests.csproj
+++ b/Engine/Source/Programs/Horde/HordeServerTests/HordeServerTests.csproj
@@ -1,7 +1,7 @@
 <Project Sdk="Microsoft.NET.Sdk">
 
   <PropertyGroup>
-    <TargetFramework>netcoreapp3.1</TargetFramework>
+    <TargetFramework>net5.0</TargetFramework>
     <Nullable>enable</Nullable>
     <IsPackable>false</IsPackable>
   </PropertyGroup>
