Horde: Add public Dockerfile and BuildGraph script for HordeServer.

diff --git a/Engine/Source/Programs/Horde/BuildHorde.xml b/Engine/Source/Programs/Horde/BuildHorde.xml
new file mode 100644
--- /dev/null
+++ b/Engine/Source/Programs/Horde/BuildHorde.xml
@@ -0,0 +1,31 @@
+<?xml version='1.0' ?>
+<BuildGraph xmlns="http://www.epicgames.com/BuildGraph" xmlns:xsi="http://www.w3.org/2001/XMLSchema-instance" xsi:schemaLocation="http://www.epicgames.com/BuildGraph ../../Engine/Build/Graph/Schema.xsd" >
+
+	<Option Name="PreflightChange" DefaultValue="" Description="Preflight changelist number"/>
+	<Option Name="Configuration" DefaultValue="release" Description="Configuration to build"/>
+	
+	<Property Name="EngineDir" Value="$(RootDir)/Engine"/>
+
+	<Property Name="Version" Value="$(EngineMajorVersion).$(EngineMinorVersion).$(EnginePatchVersion)"/>
+	<Property Name="InformationalVersion" Value="$(Version)-$(Change)"/>
+	<Property Name="InformationalVersion" Value="$(InformationalVersion)-PF-$(PreflightChange)" If="'$(PreflightChange)' != ''"/>
+	<Property Name="VersionArguments" Value="/p:Version=$(Version).0 /p:InformationalVersion=$(InformationalVersion)"/>
+
+	<Agent Name="HordeServer" Type="Win64_Docker">
+		<Property Name="StagingDir" Value="$(RootDir)/Engine/Saved/HordeServer"/>
+		
+		<Node Name="Build HordeServer">
+			<!-- Tag all the files that need to be staged to build -->
+			<Property Name="StagePaths">
+				Engine/Source/Programs/Shared/...
+				Engine/Source/Programs/Horde/...
+				Engine/Source/ThirdParty/Perforce/...
+			</Property>
+			<Tag Files="$(StagePaths)" Except=".../.vs/...;.../.git/...;.../bin/...;.../obj/..." With="#InputFiles"/>
+			<Tag Files="Engine/Source/Programs/Shared/EpicGames.Perforce.Native/bin/..." Filter="*.pdb;*.dll;*.so" With="#InputFiles"/>
+
+			<!-- Build the Docker image and publish it -->
+			<Docker-Build BaseDir="Engine/Source" Files="#InputFiles" Tag="hordeserver-public" DockerFile="Engine/Source/Programs/Horde/HordeServer/Dockerfile" Arguments="--build-arg msbuild_args=&quot;$(VersionArguments)&quot;"/>
+		</Node>
+	</Agent>
+</BuildGraph>
