Horde: Remove private constructors from response objects. These classes are no longer used for deserialization, so parameterless constructors are not required.

diff --git a/Engine/Source/Programs/Horde/HordeServer/Api/Acls.cs b/Engine/Source/Programs/Horde/HordeServer/Api/Acls.cs
--- a/Engine/Source/Programs/Horde/HordeServer/Api/Acls.cs
+++ b/Engine/Source/Programs/Horde/HordeServer/Api/Acls.cs
@@ -150,15 +150,6 @@ namespace HordeServer.Api
 		/// </summary>
 		public List<string>? Exceptions { get; set; }
 
-		/// <summary>
-		/// Parameterless constructor for serialization
-		/// </summary>
-		private GetAclResponse()
-		{
-			Entries = new List<GetAclEntryResponse>();
-			Exceptions = new List<string>();
-		}
-
 		/// <summary>
 		/// Constructor
 		/// </summary>
diff --git a/Engine/Source/Programs/Horde/HordeServer/Api/Agents.cs b/Engine/Source/Programs/Horde/HordeServer/Api/Agents.cs
--- a/Engine/Source/Programs/Horde/HordeServer/Api/Agents.cs
+++ b/Engine/Source/Programs/Horde/HordeServer/Api/Agents.cs
@@ -56,14 +56,6 @@ namespace HordeServer.Api
 		/// </summary>
 		public string Id { get; set; }
 
-		/// <summary>
-		/// Parameterless constructor for serialization
-		/// </summary>
-		private CreateAgentResponse()
-		{
-			Id = null!;
-		}
-
 		/// <summary>
 		/// Constructor
 		/// </summary>
@@ -180,15 +172,6 @@ namespace HordeServer.Api
 		/// </summary>
 		public LeaseState? State { get; set; }
 
-		/// <summary>
-		/// Parameterless constructor for serialization
-		/// </summary>
-		private GetAgentLeaseResponse()
-		{
-			this.Id = null!;
-			this.Name = String.Empty;
-		}
-
 		/// <summary>
 		/// Constructor
 		/// </summary>
@@ -255,14 +238,6 @@ namespace HordeServer.Api
 		/// </summary>
 		public string? Version { get; set; }
 
-		/// <summary>
-		/// Private constructor for serialization
-		/// </summary>
-		private GetAgentSessionResponse()
-		{
-			this.Id = null!;
-		}
-
 		/// <summary>
 		/// Constructor
 		/// </summary>
@@ -312,15 +287,6 @@ namespace HordeServer.Api
 		/// </summary>
 		public bool bIncremental { get; set; }
 
-		/// <summary>
-		/// Private constructor for serialization
-		/// </summary>
-		private GetAgentWorkspaceResponse()
-		{
-			Identifier = null!;
-			Stream = null!;
-		}
-
 		/// <summary>
 		/// Constructor
 		/// </summary>
@@ -461,18 +427,6 @@ namespace HordeServer.Api
 		/// </summary>
 		public string? Comment { get; set; }
 
-		/// <summary>
-		/// Parameterless constructor for serialization
-		/// </summary>
-		private GetAgentResponse()
-		{
-			this.Id = null!;
-			this.Name = null!;
-			Properties = new List<string>();
-			Resources = new Dictionary<string, int>();
-			this.Leases = new List<GetAgentLeaseResponse>();
-		}
-
 		/// <summary>
 		/// Constructor
 		/// </summary>
diff --git a/Engine/Source/Programs/Horde/HordeServer/Api/Artifacts.cs b/Engine/Source/Programs/Horde/HordeServer/Api/Artifacts.cs
--- a/Engine/Source/Programs/Horde/HordeServer/Api/Artifacts.cs
+++ b/Engine/Source/Programs/Horde/HordeServer/Api/Artifacts.cs
@@ -16,14 +16,6 @@ namespace HordeServer.Api
 		/// </summary>
 		public string Id { get; set; }
 
-		/// <summary>
-		/// Private constructor for serialization
-		/// </summary>
-		private CreateArtifactResponse()
-		{
-			Id = null!;
-		}
-
 		/// <summary>
 		/// Constructor
 		/// </summary>
diff --git a/Engine/Source/Programs/Horde/HordeServer/Api/Changes.cs b/Engine/Source/Programs/Horde/HordeServer/Api/Changes.cs
--- a/Engine/Source/Programs/Horde/HordeServer/Api/Changes.cs
+++ b/Engine/Source/Programs/Horde/HordeServer/Api/Changes.cs
@@ -28,15 +28,6 @@ namespace HordeServer.Api
 		/// </summary>
 		public string Description { get; set; }
 
-		/// <summary>
-		/// Private constructor for serialization
-		/// </summary>
-		private GetChangeSummaryResponse()
-		{
-			Author = String.Empty;
-			Description = String.Empty;
-		}
-
 		/// <summary>
 		/// Constructor
 		/// </summary>
@@ -74,16 +65,6 @@ namespace HordeServer.Api
 		/// </summary>
 		public List<string> Files { get; set; }
 
-		/// <summary>
-		/// Private constructor for serialization
-		/// </summary>
-		private GetChangeDetailsResponse()
-		{
-			Author = null!;
-			Description = null!;
-			Files = null!;
-		}
-
 		/// <summary>
 		/// Constructor
 		/// </summary>
diff --git a/Engine/Source/Programs/Horde/HordeServer/Api/Credentials.cs b/Engine/Source/Programs/Horde/HordeServer/Api/Credentials.cs
--- a/Engine/Source/Programs/Horde/HordeServer/Api/Credentials.cs
+++ b/Engine/Source/Programs/Horde/HordeServer/Api/Credentials.cs
@@ -36,14 +36,6 @@ namespace HordeServer.Api
 		/// </summary>
 		public string Id { get; set; }
 
-		/// <summary>
-		/// Private constructor for serialization
-		/// </summary>
-		private CreateCredentialResponse()
-		{
-			Id = null!;
-		}
-
 		/// <summary>
 		/// Constructor
 		/// </summary>
@@ -100,16 +92,6 @@ namespace HordeServer.Api
 		/// </summary>
 		public GetAclResponse? Acl { get; set; }
 
-		/// <summary>
-		/// Parameterless constructor for serialization
-		/// </summary>
-		private GetCredentialResponse()
-		{
-			Id = null!;
-			Name = null!;
-			Properties = null!;
-		}
-
 		/// <summary>
 		/// Constructor
 		/// </summary>
diff --git a/Engine/Source/Programs/Horde/HordeServer/Api/Jobs.cs b/Engine/Source/Programs/Horde/HordeServer/Api/Jobs.cs
--- a/Engine/Source/Programs/Horde/HordeServer/Api/Jobs.cs
+++ b/Engine/Source/Programs/Horde/HordeServer/Api/Jobs.cs
@@ -127,14 +127,6 @@ namespace HordeServer.Api
 		/// </summary>
 		public string Id { get; set; }
 
-		/// <summary>
-		/// Private constructor for serialization
-		/// </summary>
-		private CreateJobResponse()
-		{
-			Id = null!;
-		}
-
 		/// <summary>
 		/// Constructor
 		/// </summary>
@@ -535,14 +527,6 @@ namespace HordeServer.Api
 		/// </summary>
 		public Dictionary<string, string>? Properties { get; set; }
 
-		/// <summary>
-		/// Parameterless constructor for serialization
-		/// </summary>
-		private GetStepResponse()
-		{
-			Id = null!;
-		}
-
 		/// <summary>
 		/// Constructor
 		/// </summary>
@@ -744,15 +728,6 @@ namespace HordeServer.Api
 		/// </summary>
 		public DateTimeOffset? FinishTime { get; set; }
 
-		/// <summary>
-		/// Parameterless constructor for serialization
-		/// </summary>
-		private GetBatchResponse()
-		{
-			Id = null!;
-			Steps = new List<GetStepResponse>();
-		}
-
 		/// <summary>
 		/// Converts this batch into a public response object
 		/// </summary>
diff --git a/Engine/Source/Programs/Horde/HordeServer/Api/LogEvents.cs b/Engine/Source/Programs/Horde/HordeServer/Api/LogEvents.cs
--- a/Engine/Source/Programs/Horde/HordeServer/Api/LogEvents.cs
+++ b/Engine/Source/Programs/Horde/HordeServer/Api/LogEvents.cs
@@ -47,15 +47,6 @@ namespace HordeServer.Api
 		/// </summary>
 		public List<JsonElement> Lines { get; set; }
 
-		/// <summary>
-		/// Private constructor for serialization
-		/// </summary>
-		private GetLogEventResponse()
-		{
-			LogId = null!;
-			Lines = new List<JsonElement>();
-		}
-
 		/// <summary>
 		/// Constructor
 		/// </summary>
diff --git a/Engine/Source/Programs/Horde/HordeServer/Api/LogFiles.cs b/Engine/Source/Programs/Horde/HordeServer/Api/LogFiles.cs
--- a/Engine/Source/Programs/Horde/HordeServer/Api/LogFiles.cs
+++ b/Engine/Source/Programs/Horde/HordeServer/Api/LogFiles.cs
@@ -74,14 +74,6 @@ namespace HordeServer.Api
 		/// </summary>
 		public string Id { get; set; }
 
-		/// <summary>
-		/// Private constructor for serialization
-		/// </summary>
-		private CreateLogFileResponse()
-		{
-			Id = null!;
-		}
-
 		/// <summary>
 		/// Constructor
 		/// </summary>
diff --git a/Engine/Source/Programs/Horde/HordeServer/Api/Nodes.cs b/Engine/Source/Programs/Horde/HordeServer/Api/Nodes.cs
--- a/Engine/Source/Programs/Horde/HordeServer/Api/Nodes.cs
+++ b/Engine/Source/Programs/Horde/HordeServer/Api/Nodes.cs
@@ -271,18 +271,6 @@ namespace HordeServer.Api
 		/// </summary>
 		public IReadOnlyDictionary<string, string>? Properties { get; set; }
 
-		/// <summary>
-		/// Parameterless constructor for serialization
-		/// </summary>
-		private GetNodeResponse()
-		{
-			this.Name = null!;
-			this.InputDependencies = new List<string>();
-			this.OrderDependencies = new List<string>();
-			this.Credentials = new Dictionary<string, string>();
-			this.Properties = new Dictionary<string, string>();
-		}
-
 		/// <summary>
 		/// Constructor
 		/// </summary>
diff --git a/Engine/Source/Programs/Horde/HordeServer/Api/Pools.cs b/Engine/Source/Programs/Horde/HordeServer/Api/Pools.cs
--- a/Engine/Source/Programs/Horde/HordeServer/Api/Pools.cs
+++ b/Engine/Source/Programs/Horde/HordeServer/Api/Pools.cs
@@ -57,14 +57,6 @@ namespace HordeServer.Api
 		/// </summary>
 		public string Id { get; set; }
 
-		/// <summary>
-		/// Private constructor for serialization
-		/// </summary>
-		private CreatePoolResponse()
-		{
-			Id = null!;
-		}
-
 		/// <summary>
 		/// Constructor
 		/// </summary>
@@ -178,17 +170,6 @@ namespace HordeServer.Api
 		/// </summary>
 		public IReadOnlyDictionary<string, string> Properties { get; set; }
 
-		/// <summary>
-		/// Parameterless constructor for serialization
-		/// </summary>
-		private GetPoolResponse()
-		{
-			Id = null!;
-			Name = null!;
-			Workspaces = null!;
-			Properties = null!;
-		}
-
 		/// <summary>
 		/// Constructor
 		/// </summary>
diff --git a/Engine/Source/Programs/Horde/HordeServer/Api/Projects.cs b/Engine/Source/Programs/Horde/HordeServer/Api/Projects.cs
--- a/Engine/Source/Programs/Horde/HordeServer/Api/Projects.cs
+++ b/Engine/Source/Programs/Horde/HordeServer/Api/Projects.cs
@@ -49,92 +49,6 @@ namespace HordeServer.Api
 		}
 	}
 
-	/// <summary>
-	/// Parameters to create a new project
-	/// </summary>
-	public class CreateProjectRequest
-	{
-		/// <summary>
-		/// Name for the new project
-		/// </summary>
-		[Required]
-		public string Name { get; set; } = null!;
-
-		/// <summary>
-		/// Order to display this project on the dashboard
-		/// </summary>
-		public int? Order { get; set; }
-
-		/// <summary>
-		/// Categories to include in this project
-		/// </summary>
-		public List<CreateProjectCategoryRequest>? Categories { get; set; }
-
-		/// <summary>
-		/// Properties for the new project
-		/// </summary>
-		public Dictionary<string, string>? Properties { get; set; }
-	}
-
-	/// <summary>
-	/// Response from creating a new project
-	/// </summary>
-	public class CreateProjectResponse
-	{
-		/// <summary>
-		/// Unique id for the new project
-		/// </summary>
-		public string Id { get; set; }
-
-		/// <summary>
-		/// Private constructor for serialization
-		/// </summary>
-		private CreateProjectResponse()
-		{
-			Id = null!;
-		}
-
-		/// <summary>
-		/// Constructor
-		/// </summary>
-		/// <param name="Id">Unique id for the new project</param>
-		public CreateProjectResponse(string Id)
-		{
-			this.Id = Id;
-		}
-	}
-
-	/// <summary>
-	/// Parameters to update a project
-	/// </summary>
-	public class UpdateProjectRequest
-	{
-		/// <summary>
-		/// Optional new name for the project
-		/// </summary>
-		public string? Name { get; set; }
-
-		/// <summary>
-		/// Order to display this project on the dashboard
-		/// </summary>
-		public int? Order { get; set; }
-
-		/// <summary>
-		/// New set of categories for the request
-		/// </summary>
-		public List<CreateProjectCategoryRequest>? Categories { get; set; }
-
-		/// <summary>
-		/// Properties to update for the project. Properties set to null will be removed.
-		/// </summary>
-		public Dictionary<string, string>? Properties { get; set; }
-
-		/// <summary>
-		/// Custom permissions for this object
-		/// </summary>
-		public UpdateAclRequest? Acl { get; set; }
-	}
-
 	/// <summary>
 	/// Information about a stream within a project
 	/// </summary>
@@ -246,17 +160,6 @@ namespace HordeServer.Api
 		/// </summary>
 		public GetAclResponse? Acl { get; set; }
 
-		/// <summary>
-		/// Parameterless constructor for serialization
-		/// </summary>
-		private GetProjectResponse()
-		{
-			Id = null!;
-			Name = null!;
-			Streams = null!;
-			Categories = null!;
-		}
-
 		/// <summary>
 		/// Constructor
 		/// </summary>
diff --git a/Engine/Source/Programs/Horde/HordeServer/Api/Schedules.cs b/Engine/Source/Programs/Horde/HordeServer/Api/Schedules.cs
--- a/Engine/Source/Programs/Horde/HordeServer/Api/Schedules.cs
+++ b/Engine/Source/Programs/Horde/HordeServer/Api/Schedules.cs
@@ -163,13 +163,6 @@ namespace HordeServer.Api
 		/// </summary>
 		public int? Interval { get; set; }
 
-		/// <summary>
-		/// Private constructor for serialization
-		/// </summary>
-		private GetSchedulePatternResponse()
-		{
-		}
-
 		/// <summary>
 		/// Constructor
 		/// </summary>
@@ -269,16 +262,6 @@ namespace HordeServer.Api
 		/// </summary>
 		public List<string> ActiveJobs { get; set; }
 
-		/// <summary>
-		/// Parameterless constructor for serialization
-		/// </summary>
-		private GetScheduleResponse()
-		{
-			TemplateParameters = null!;
-			Patterns = null!;
-			ActiveJobs = new List<string>();
-		}
-
 		/// <summary>
 		/// Constructor
 		/// </summary>
@@ -312,14 +295,6 @@ namespace HordeServer.Api
 		/// </summary>
 		public List<DateTime> Times { get; set; }
 
-		/// <summary>
-		/// Private constructor for serialization
-		/// </summary>
-		private GetScheduleForecastResponse()
-		{
-			this.Times = null!;
-		}
-
 		/// <summary>
 		/// Constructor
 		/// </summary>
diff --git a/Engine/Source/Programs/Horde/HordeServer/Api/Software.cs b/Engine/Source/Programs/Horde/HordeServer/Api/Software.cs
--- a/Engine/Source/Programs/Horde/HordeServer/Api/Software.cs
+++ b/Engine/Source/Programs/Horde/HordeServer/Api/Software.cs
@@ -80,14 +80,6 @@ namespace HordeServer.Api
 		/// </summary>
 		public DateTime? MadeDefaultAtTime { get; set; }
 
-		/// <summary>
-		/// Default constructor for serialization
-		/// </summary>
-		private GetSoftwareResponse()
-		{
-			Id = null!;
-		}
-
 		/// <summary>
 		/// Constructor
 		/// </summary>
diff --git a/Engine/Source/Programs/Horde/HordeServer/Api/Streams.cs b/Engine/Source/Programs/Horde/HordeServer/Api/Streams.cs
--- a/Engine/Source/Programs/Horde/HordeServer/Api/Streams.cs
+++ b/Engine/Source/Programs/Horde/HordeServer/Api/Streams.cs
@@ -357,14 +357,6 @@ namespace HordeServer.Api
 		/// </summary>
 		public Dictionary<string, string>? Environment { get; set; }
 
-		/// <summary>
-		/// Default constructor
-		/// </summary>
-		private GetAgentTypeResponse()
-		{
-			Pool = null!;
-		}
-
 		/// <summary>
 		/// Constructor
 		/// </summary>
@@ -648,19 +640,6 @@ namespace HordeServer.Api
 		/// </summary>
 		public string? PauseComment { get; set; }
 
-		/// <summary>
-		/// Parameterless constructor for serialization
-		/// </summary>
-		private GetStreamResponse()
-		{
-			Id = null!;
-			ProjectId = null!;
-			Name = null!;
-			Tabs = null!;
-			AgentTypes = null!;
-			Templates = null!;
-		}
-
 		/// <summary>
 		/// Constructor
 		/// </summary>
diff --git a/Engine/Source/Programs/Horde/HordeServer/Controllers/CredentialsController.cs b/Engine/Source/Programs/Horde/HordeServer/Controllers/CredentialsController.cs
--- a/Engine/Source/Programs/Horde/HordeServer/Controllers/CredentialsController.cs
+++ b/Engine/Source/Programs/Horde/HordeServer/Controllers/CredentialsController.cs
@@ -51,7 +51,7 @@ namespace HordeServer.Controllers
 		/// <returns>Http result code</returns>
 		[HttpPost]
 		[Route("/api/v1/credentials")]
-		public async Task<ActionResult<CreateCredentialResponse>> CreateCredentialAsync([FromBody] CreateProjectRequest Create)
+		public async Task<ActionResult<CreateCredentialResponse>> CreateCredentialAsync([FromBody] CreateCredentialRequest Create)
 		{
 			if(!await AclService.AuthorizeAsync(AclAction.CreateCredential, User))
 			{
