Horde: Remove private parameterless constructors. NET 5 no longer supports deserializing JSON using these.

diff --git a/Engine/Source/Programs/Horde/HordeServer/Api/Aggregates.cs b/Engine/Source/Programs/Horde/HordeServer/Api/Aggregates.cs
--- a/Engine/Source/Programs/Horde/HordeServer/Api/Aggregates.cs
+++ b/Engine/Source/Programs/Horde/HordeServer/Api/Aggregates.cs
@@ -25,15 +25,6 @@ namespace HordeServer.Api
 		[Required]
 		public List<string> Nodes { get; set; }
 
-		/// <summary>
-		/// Private constructor for serialization
-		/// </summary>
-		private CreateAggregateRequest()
-		{
-			Name = null!;
-			Nodes = null!;
-		}
-
 		/// <summary>
 		/// Constructor
 		/// </summary>
diff --git a/Engine/Source/Programs/Horde/HordeServer/Api/Artifacts.cs b/Engine/Source/Programs/Horde/HordeServer/Api/Artifacts.cs
--- a/Engine/Source/Programs/Horde/HordeServer/Api/Artifacts.cs
+++ b/Engine/Source/Programs/Horde/HordeServer/Api/Artifacts.cs
@@ -46,16 +46,6 @@ namespace HordeServer.Api
 		/// </summary>
 		public List<string>? ArtifactIds { get; set; }
 
-		/// <summary>
-		/// Private constructor for serialization
-		/// </summary>
-		private GetArtifactZipRequest()
-		{
-			JobId = null!;
-			StepId = null!;
-			ArtifactIds = null!;
-		}
-
 		/// <summary>
 		/// Constructor
 		/// </summary>
diff --git a/Engine/Source/Programs/Horde/HordeServer/Api/Jobs.cs b/Engine/Source/Programs/Horde/HordeServer/Api/Jobs.cs
--- a/Engine/Source/Programs/Horde/HordeServer/Api/Jobs.cs
+++ b/Engine/Source/Programs/Horde/HordeServer/Api/Jobs.cs
@@ -110,10 +110,10 @@ namespace HordeServer.Api
 		/// <summary>
 		/// Private constructor for serialization
 		/// </summary>
-		private CreateJobRequest()
+		public CreateJobRequest(string StreamId, string TemplateId)
 		{
-			StreamId = null!;
-			TemplateId = null!;
+			this.StreamId = StreamId;
+			this.TemplateId = TemplateId;
 		}
 	}
 
diff --git a/Engine/Source/Programs/Horde/HordeServer/Api/LogFiles.cs b/Engine/Source/Programs/Horde/HordeServer/Api/LogFiles.cs
--- a/Engine/Source/Programs/Horde/HordeServer/Api/LogFiles.cs
+++ b/Engine/Source/Programs/Horde/HordeServer/Api/LogFiles.cs
@@ -28,62 +28,6 @@ namespace HordeServer.Api
 		Json
 	}
 
-	/// <summary>
-	/// Parameters required to create log file
-	/// </summary>
-	public class CreateLogFileRequest
-	{
-		/// <summary>
-		/// Job Id this log file belongs to
-		/// </summary>
-		[Required]
-		public string JobId { get; set; }
-
-		/// <summary>
-		/// The type of events stored in this log
-		/// </summary>
-		public LogType Type { get; set; }
-
-		/// <summary>
-		/// Private constructor for serialization
-		/// </summary>
-		private CreateLogFileRequest()
-		{
-			JobId = null!;
-		}
-
-		/// <summary>
-		/// Constructor
-		/// </summary>
-		/// <param name="JobId">Unique job id that this log file belongs to</param>
-		/// <param name="Type">Type of events to be stored in this file</param>
-		public CreateLogFileRequest(string JobId, LogType Type)
-		{
-			this.JobId = JobId;
-			this.Type = Type;
-		}
-	}
-
-	/// <summary>
-	/// Response from creating a log file
-	/// </summary>
-	public class CreateLogFileResponse
-	{
-		/// <summary>
-		/// Unique id for this log file
-		/// </summary>
-		public string Id { get; set; }
-
-		/// <summary>
-		/// Constructor
-		/// </summary>
-		/// <param name="Id">The log file id</param>
-		public CreateLogFileResponse(string Id)
-		{
-			this.Id = Id;
-		}
-	}
-
 	/// <summary>
 	/// Response describing a log file
 	/// </summary>
diff --git a/Engine/Source/Programs/Horde/HordeServer/Api/Nodes.cs b/Engine/Source/Programs/Horde/HordeServer/Api/Nodes.cs
--- a/Engine/Source/Programs/Horde/HordeServer/Api/Nodes.cs
+++ b/Engine/Source/Programs/Horde/HordeServer/Api/Nodes.cs
@@ -63,14 +63,6 @@ namespace HordeServer.Api
 		/// </summary>
 		public Dictionary<string, string>? Properties { get; set; }
 
-		/// <summary>
-		/// Private constructor for serialization
-		/// </summary>
-		private CreateNodeRequest()
-		{
-			this.Name = null!;
-		}
-
 		/// <summary>
 		/// Constructor
 		/// </summary>
@@ -114,15 +106,6 @@ namespace HordeServer.Api
 		[Required]
 		public List<CreateNodeRequest> Nodes { get; set; }
 
-		/// <summary>
-		/// Private constructor for serialization
-		/// </summary>
-		private CreateGroupRequest()
-		{
-			AgentType = null!;
-			Nodes = new List<CreateNodeRequest>();
-		}
-
 		/// <summary>
 		/// Constructor
 		/// </summary>
@@ -186,34 +169,6 @@ namespace HordeServer.Api
 		/// Nodes which must be part of the job for the aggregate to be valid
 		/// </summary>
 		public List<string> IncludedNodes { get; set; } = new List<string>();
-
-		/// <summary>
-		/// Private constructor for serialization
-		/// </summary>
-		private CreateLabelRequest()
-		{
-		}
-
-		/// <summary>
-		/// Constructor
-		/// </summary>
-		/// <param name="DashboardName">Name of this label</param>
-		/// <param name="DashboardCategory">Category of this label</param>
-		/// <param name="UgsName">Name of this label for UGS</param>
-		/// <param name="UgsProject">Name of the project to show this label for in UGS</param>
-		/// <param name="Change">Which change the label applies to</param>
-		/// <param name="RequiredNodes">Nodes which must be part of the job for the aggregate to be shown</param>
-		/// <param name="IncludedNodes">Nodes to include in the status of this aggregate, if present in the job</param>
-		public CreateLabelRequest(string? DashboardName, string? DashboardCategory, string? UgsName, string? UgsProject, LabelChange Change, List<string> RequiredNodes, List<string> IncludedNodes)
-		{
-			this.DashboardName = String.IsNullOrEmpty(DashboardName)? null : DashboardName;
-			this.DashboardCategory = String.IsNullOrEmpty(DashboardCategory)? null : DashboardCategory;
-			this.UgsName = String.IsNullOrEmpty(UgsName)? null : UgsName;
-			this.UgsProject = String.IsNullOrEmpty(UgsProject)? null : UgsProject;
-			this.Change = Change;
-			this.RequiredNodes = RequiredNodes;
-			this.IncludedNodes = IncludedNodes;
-		}
 	}
 
 	/// <summary>
diff --git a/Engine/Source/Programs/Horde/HordeServer/Api/Projects.cs b/Engine/Source/Programs/Horde/HordeServer/Api/Projects.cs
--- a/Engine/Source/Programs/Horde/HordeServer/Api/Projects.cs
+++ b/Engine/Source/Programs/Horde/HordeServer/Api/Projects.cs
@@ -41,11 +41,11 @@ namespace HordeServer.Api
 		public List<string> ExcludePatterns { get; set; } = new List<string>();
 
 		/// <summary>
-		/// Private constructor for serialization
+		/// Constructor
 		/// </summary>
-		private CreateProjectCategoryRequest()
+		public CreateProjectCategoryRequest(string Name)
 		{
-			Name = null!;
+			this.Name = Name;
 		}
 	}
 
diff --git a/Engine/Source/Programs/Horde/HordeServer/Api/Streams.cs b/Engine/Source/Programs/Horde/HordeServer/Api/Streams.cs
--- a/Engine/Source/Programs/Horde/HordeServer/Api/Streams.cs
+++ b/Engine/Source/Programs/Horde/HordeServer/Api/Streams.cs
@@ -230,108 +230,6 @@ namespace HordeServer.Api
 		public UpdateAclRequest? Acl { get; set; }
 	}
 
-	/// <summary>
-	/// Parameters to create a new stream
-	/// </summary>
-	public class SetStreamRequest
-	{
-		/// <summary>
-		/// Unique id for the project
-		/// </summary>
-		[Required]
-		public string ProjectId { get; set; }
-
-		/// <summary>
-		/// Name for the new stream
-		/// </summary>
-		[Required]
-		public string Name { get; set; }
-
-		/// <summary>
-		/// Path to the config file in Perforce
-		/// </summary>
-		public string? ConfigPath { get; set; }
-
-		/// <summary>
-		/// Order for this stream
-		/// </summary>
-		public int? Order { get; set; }
-
-		/// <summary>
-		/// Notification channel for all jobs in this stream
-		/// </summary>
-		public string? NotificationChannel { get; set; }
-
-		/// <summary>
-		/// Notification channel filter for this template. Can be Success|Failure|Warnings
-		/// </summary>
-		public string? NotificationChannelFilter { get; set; }
-
-		/// <summary>
-		/// Channel to post issue triage notifications
-		/// </summary>
-		public string? TriageChannel { get; set; }
-
-		/// <summary>
-		/// Default template for running preflights
-		/// </summary>
-		public string? DefaultPreflightTemplate { get; set; }
-
-		/// <summary>
-		/// Default template for running preflights
-		/// </summary>
-		public DefaultPreflightRequest? DefaultPreflight { get; set; }
-
-		/// <summary>
-		/// List of tabs to show for the new stream
-		/// </summary>
-		public List<CreateStreamTabRequest>? Tabs { get; set; }
-
-		/// <summary>
-		/// Map of agent name to type
-		/// </summary>
-		public Dictionary<string, CreateAgentTypeRequest>? AgentTypes { get; set; }
-
-		/// <summary>
-		/// Map of workspace name to type
-		/// </summary>
-		public Dictionary<string, CreateWorkspaceTypeRequest>? WorkspaceTypes { get; set; }
-
-		/// <summary>
-		/// List of templates to create
-		/// </summary>
-		public List<CreateTemplateRefRequest>? Templates { get; set; }
-
-		/// <summary>
-		/// Properties for the new stream
-		/// </summary>
-		public Dictionary<string, string>? Properties { get; set; }
-
-		/// <summary>
-		/// Custom permissions for this object
-		/// </summary>
-		public UpdateAclRequest? Acl { get; set; }
-		
-		/// <summary>
-		/// Pause stream builds until specified date
-		/// </summary>
-		public DateTime? PausedUntil { get; set; }
-		
-		/// <summary>
-		/// Reason for pausing builds of the stream
-		/// </summary>
-		public string? PauseComment { get; set; }
-
-		/// <summary>
-		/// Private constructor for serialization
-		/// </summary>
-		private SetStreamRequest()
-		{
-			this.ProjectId = null!;
-			this.Name = null!;
-		}
-	}
-
 	/// <summary>
 	/// Mapping from a BuildGraph agent type to a set of machines on the farm
 	/// </summary>
diff --git a/Engine/Source/Programs/Horde/HordeServer/Controllers/LogsController.cs b/Engine/Source/Programs/Horde/HordeServer/Controllers/LogsController.cs
--- a/Engine/Source/Programs/Horde/HordeServer/Controllers/LogsController.cs
+++ b/Engine/Source/Programs/Horde/HordeServer/Controllers/LogsController.cs
@@ -83,25 +83,6 @@ namespace HordeServer.Controllers
 			this.JobService = JobService;
 		}
 
-		/// <summary>
-		/// Creates a new logfile
-		/// </summary>
-		/// <param name="Create">Parameters for the new LogFile</param>
-		/// <returns>Http result code</returns>
-		[HttpPost]
-		[Route("/api/v1/logs")]
-		public async Task<ActionResult<CreateLogFileResponse>> CreateLogFile([FromBody] CreateLogFileRequest Create)
-		{
-			ObjectId JobId = Create.JobId.ToObjectId();
-			if (!await JobService.AuthorizeAsync(JobId, AclAction.CreateLog, User, null))
-			{
-				return Forbid();
-			}
-
-			ILogFile NewLogFile = await LogFileService.CreateLogFileAsync(JobId, null, Create.Type);
-			return new CreateLogFileResponse(NewLogFile.Id.ToString());
-		}
-
 		/// <summary>
 		/// Retrieve metadata about a specific log file
 		/// </summary>
diff --git a/Engine/Source/Programs/Horde/HordeServer/Services/RpcService.cs b/Engine/Source/Programs/Horde/HordeServer/Services/RpcService.cs
--- a/Engine/Source/Programs/Horde/HordeServer/Services/RpcService.cs
+++ b/Engine/Source/Programs/Horde/HordeServer/Services/RpcService.cs
@@ -839,8 +839,15 @@ namespace HordeServer.Services
 			List<Api.CreateLabelRequest> ApiLabels = new List<Api.CreateLabelRequest>();
 			foreach (CreateLabelRequest Label in Request.Labels)
 			{
-				Api.CreateLabelRequest ApiLabel = new Api.CreateLabelRequest(Label.DashboardName, Label.DashboardCategory, Label.UgsName, Label.UgsProject, Label.Change, Label.RequiredNodes.ToList(), Label.IncludedNodes.ToList());
-				ApiLabels.Add(ApiLabel);
+				Api.CreateLabelRequest NewLabel = new Api.CreateLabelRequest();
+				NewLabel.DashboardName = String.IsNullOrEmpty(Label.DashboardName) ? null : Label.DashboardName;
+				NewLabel.DashboardCategory = String.IsNullOrEmpty(Label.DashboardCategory) ? null : Label.DashboardCategory;
+				NewLabel.UgsName = String.IsNullOrEmpty(Label.UgsName) ? null : Label.UgsName;
+				NewLabel.UgsProject = String.IsNullOrEmpty(Label.UgsProject) ? null : Label.UgsProject;
+				NewLabel.Change = Label.Change;
+				NewLabel.RequiredNodes = Label.RequiredNodes.ToList();
+				NewLabel.IncludedNodes = Label.IncludedNodes.ToList();
+				ApiLabels.Add(NewLabel);
 			}
 
 			ObjectId JobIdValue = Request.JobId.ToObjectId();
