Horde: Change classes used to send REST API requests to define graphs into internal, programmatic versions.

diff --git a/Engine/Source/Programs/Horde/HordeServer/Api/Aggregates.cs b/Engine/Source/Programs/Horde/HordeServer/Api/Aggregates.cs
deleted file mode 100644
diff --git a/Engine/Source/Programs/Horde/HordeServer/Api/Jobs.cs b/Engine/Source/Programs/Horde/HordeServer/Api/Jobs.cs
--- a/Engine/Source/Programs/Horde/HordeServer/Api/Jobs.cs
+++ b/Engine/Source/Programs/Horde/HordeServer/Api/Jobs.cs
@@ -87,21 +87,6 @@ namespace HordeServer.Api
 		/// </summary>
 		public bool? UpdateIssues { get; set; }
 
-		/// <summary>
-		/// Nodes for the new job
-		/// </summary>
-		public List<CreateGroupRequest>? Groups { get; set; }
-
-		/// <summary>
-		/// Aggregates for the new job
-		/// </summary>
-		public List<CreateAggregateRequest>? Aggregates { get; set; }
-
-		/// <summary>
-		/// Labels for the new job
-		/// </summary>
-		public List<CreateLabelRequest>? Labels { get; set; }
-
 		/// <summary>
 		/// Arguments for the job
 		/// </summary>
diff --git a/Engine/Source/Programs/Horde/HordeServer/Api/Nodes.cs b/Engine/Source/Programs/Horde/HordeServer/Api/Nodes.cs
--- a/Engine/Source/Programs/Horde/HordeServer/Api/Nodes.cs
+++ b/Engine/Source/Programs/Horde/HordeServer/Api/Nodes.cs
@@ -12,165 +12,6 @@ using System.Threading.Tasks;
 
 namespace HordeServer.Api
 {
-	/// <summary>
-	/// Information required to create a node
-	/// </summary>
-	public class CreateNodeRequest
-	{
-		/// <summary>
-		/// The name of this node 
-		/// </summary>
-		[Required]
-		public string Name { get; set; } = null!;
-
-		/// <summary>
-		/// List of nodes which must succeed for this node to run
-		/// </summary>
-		public List<string>? InputDependencies { get; set; }
-
-		/// <summary>
-		/// List of nodes which must have completed for this node to run
-		/// </summary>
-		public List<string>? OrderDependencies { get; set; }
-
-		/// <summary>
-		/// The priority of this node
-		/// </summary>
-		public Priority? Priority { get; set; }
-
-		/// <summary>
-		/// This node can be run multiple times
-		/// </summary>
-		public bool? AllowRetry { get; set; }
-
-		/// <summary>
-		/// This node can start running early, before dependencies of other nodes in the same group are complete
-		/// </summary>
-		public bool? RunEarly { get; set; }
-
-		/// <summary>
-		/// Whether to include warnings in the diagnostic output
-		/// </summary>
-		public bool? Warnings { get; set; }
-
-		/// <summary>
-		/// Credentials required for this node to run. This dictionary maps from environment variable names to a credential property in the format 'CredentialName.PropertyName'.
-		/// </summary>
-		public Dictionary<string, string>? Credentials { get; set; }
-
-		/// <summary>
-		/// Properties for this node
-		/// </summary>
-		public Dictionary<string, string>? Properties { get; set; }
-
-		/// <summary>
-		/// Constructor
-		/// </summary>
-		/// <param name="Name">Name of the node</param>
-		/// <param name="InputDependencies">List of nodes which must have completed succesfully for this node to run</param>
-		/// <param name="OrderDependencies">List of nodes which must have completed for this node to run</param>
-		/// <param name="Priority">Priority of this node</param>
-		/// <param name="AllowRetry">Whether the node can be run multiple times</param>
-		/// <param name="RunEarly">Whether the node can run early, before dependencies of other nodes in the same group complete</param>
-		/// <param name="Warnings">Whether to include warnings in the diagnostic output (defaults to true)</param>
-		/// <param name="Credentials">Credentials required for this node to run</param>
-		/// <param name="Properties">Properties for the node</param>
-		public CreateNodeRequest(string Name, List<string>? InputDependencies = null, List<string>? OrderDependencies = null, Priority? Priority = null, bool? AllowRetry = null, bool? RunEarly = null, bool? Warnings = null, Dictionary<string, string>? Credentials = null, Dictionary<string, string>? Properties = null)
-		{
-			this.Name = Name;
-			this.InputDependencies = InputDependencies;
-			this.OrderDependencies = OrderDependencies;
-			this.Priority = Priority;
-			this.AllowRetry = AllowRetry;
-			this.RunEarly = RunEarly;
-			this.Warnings = Warnings;
-			this.Credentials = Credentials;
-			this.Properties = Properties;
-		}
-	}
-
-	/// <summary>
-	/// Information about a group of nodes
-	/// </summary>
-	public class CreateGroupRequest
-	{
-		/// <summary>
-		/// The type of agent to execute this group
-		/// </summary>
-		[Required]
-		public string AgentType { get; set; }
-
-		/// <summary>
-		/// Nodes in the group
-		/// </summary>
-		[Required]
-		public List<CreateNodeRequest> Nodes { get; set; }
-
-		/// <summary>
-		/// Constructor
-		/// </summary>
-		/// <param name="AgentType">The type of agent to execute this group</param>
-		/// <param name="Nodes">Nodes in this group</param>
-		public CreateGroupRequest(string AgentType, List<CreateNodeRequest> Nodes)
-		{
-			this.AgentType = AgentType;
-			this.Nodes = Nodes;
-		}
-	}
-
-	/// <summary>
-	/// Information about a group of nodes
-	/// </summary>
-	public class CreateLabelRequest
-	{
-		/// <summary>
-		/// Category for this label
-		/// </summary>
-		[Obsolete("Use DashboardCategory instead")]
-		public string? Category => DashboardCategory;
-
-		/// <summary>
-		/// Name of the aggregate
-		/// </summary>
-		[Obsolete("Use DashboardName instead")]
-		public string? Name => DashboardName;
-
-		/// <summary>
-		/// Name of the aggregate
-		/// </summary>
-		public string? DashboardName { get; set; }
-
-		/// <summary>
-		/// Category for this label
-		/// </summary>
-		public string? DashboardCategory { get; set; }
-
-		/// <summary>
-		/// Name of the badge in UGS
-		/// </summary>
-		public string? UgsName { get; set; }
-
-		/// <summary>
-		/// Project to show this label for in UGS
-		/// </summary>
-		public string? UgsProject { get; set; }
-
-		/// <summary>
-		/// Which change the label applies to
-		/// </summary>
-		public LabelChange Change { get; set; }
-
-		/// <summary>
-		/// Nodes which must be part of the job for the aggregate to be valid
-		/// </summary>
-		public List<string> RequiredNodes { get; set; } = new List<string>();
-
-		/// <summary>
-		/// Nodes which must be part of the job for the aggregate to be valid
-		/// </summary>
-		public List<string> IncludedNodes { get; set; } = new List<string>();
-	}
-
 	/// <summary>
 	/// Information required to create a node
 	/// </summary>
diff --git a/Engine/Source/Programs/Horde/HordeServer/Collections/IGraphCollection.cs b/Engine/Source/Programs/Horde/HordeServer/Collections/IGraphCollection.cs
--- a/Engine/Source/Programs/Horde/HordeServer/Collections/IGraphCollection.cs
+++ b/Engine/Source/Programs/Horde/HordeServer/Collections/IGraphCollection.cs
@@ -38,7 +38,7 @@ namespace HordeServer.Collections
 		/// <param name="NewAggregateRequests">List of aggregate requests</param>
 		/// <param name="NewLabelRequests">List of label requests</param>
 		/// <returns>The new graph definition</returns>
-		Task<IGraph> AppendAsync(IGraph? BaseGraph, List<CreateGroupRequest>? NewGroupRequests = null, List<CreateAggregateRequest>? NewAggregateRequests = null, List<CreateLabelRequest>? NewLabelRequests = null);
+		Task<IGraph> AppendAsync(IGraph? BaseGraph, List<NewGroup>? NewGroupRequests = null, List<NewAggregate>? NewAggregateRequests = null, List<NewLabel>? NewLabelRequests = null);
 
 		/// <summary>
 		/// Gets the graph for a job
diff --git a/Engine/Source/Programs/Horde/HordeServer/Collections/Impl/GraphCollection.cs b/Engine/Source/Programs/Horde/HordeServer/Collections/Impl/GraphCollection.cs
--- a/Engine/Source/Programs/Horde/HordeServer/Collections/Impl/GraphCollection.cs
+++ b/Engine/Source/Programs/Horde/HordeServer/Collections/Impl/GraphCollection.cs
@@ -202,7 +202,7 @@ namespace HordeServer.Collections.Impl
 				this.Id = ContentHash.SHA1(BsonExtensionMethods.ToBson(this));
 			}
 
-			public GraphDocument(GraphDocument BaseGraph, List<CreateGroupRequest>? NewGroupRequests, List<CreateAggregateRequest>? NewAggregateRequests, List<CreateLabelRequest>? NewLabelRequests)
+			public GraphDocument(GraphDocument BaseGraph, List<NewGroup>? NewGroupRequests, List<NewAggregate>? NewAggregateRequests, List<NewLabel>? NewLabelRequests)
 			{
 				Dictionary<string, NodeRef> NodeNameToRef = new Dictionary<string, NodeRef>(BaseGraph.GetNodeNameToRef(), StringComparer.OrdinalIgnoreCase);
 
@@ -210,10 +210,10 @@ namespace HordeServer.Collections.Impl
 				List<NodeGroup> NewGroups = new List<NodeGroup>(BaseGraph.Groups);
 				if (NewGroupRequests != null)
 				{
-					foreach (CreateGroupRequest NewGroupRequest in NewGroupRequests)
+					foreach (NewGroup NewGroupRequest in NewGroupRequests)
 					{
 						List<Node> Nodes = new List<Node>();
-						foreach (CreateNodeRequest NewNodeRequest in NewGroupRequest.Nodes)
+						foreach (NewNode NewNodeRequest in NewGroupRequest.Nodes)
 						{
 							int NodeIdx = Nodes.Count;
 
@@ -237,7 +237,7 @@ namespace HordeServer.Collections.Impl
 				List<Aggregate> NewAggregates = new List<Aggregate>(BaseGraph.Aggregates);
 				if (NewAggregateRequests != null)
 				{
-					foreach (CreateAggregateRequest NewAggregateRequest in NewAggregateRequests)
+					foreach (NewAggregate NewAggregateRequest in NewAggregateRequests)
 					{
 						List<NodeRef> Nodes = NewAggregateRequest.Nodes.ConvertAll(x => NodeNameToRef[x]);
 						NewAggregates.Add(new Aggregate(NewAggregateRequest.Name, Nodes));
@@ -248,7 +248,7 @@ namespace HordeServer.Collections.Impl
 				List<Label> NewLabels = new List<Label>(BaseGraph.Labels);
 				if (NewLabelRequests != null)
 				{
-					foreach (CreateLabelRequest NewLabelRequest in NewLabelRequests)
+					foreach (NewLabel NewLabelRequest in NewLabelRequests)
 					{
 						List<NodeRef> RequiredNodes = NewLabelRequest.RequiredNodes.ConvertAll(x => NodeNameToRef[x]);
 						List<NodeRef> IncludedNodes = NewLabelRequest.IncludedNodes.ConvertAll(x => NodeNameToRef[x]);
@@ -397,7 +397,7 @@ namespace HordeServer.Collections.Impl
 		}
 
 		/// <inheritdoc/>
-		public async Task<IGraph> AppendAsync(IGraph? BaseGraph, List<CreateGroupRequest>? NewGroupRequests, List<CreateAggregateRequest>? NewAggregateRequests, List<CreateLabelRequest>? NewLabelRequests)
+		public async Task<IGraph> AppendAsync(IGraph? BaseGraph, List<NewGroup>? NewGroupRequests, List<NewAggregate>? NewAggregateRequests, List<NewLabel>? NewLabelRequests)
 		{
 			GraphDocument Graph = new GraphDocument((GraphDocument?)BaseGraph ?? GraphDocument.Empty, NewGroupRequests, NewAggregateRequests, NewLabelRequests);
 			await AddAsync(Graph);
diff --git a/Engine/Source/Programs/Horde/HordeServer/Controllers/JobsController.cs b/Engine/Source/Programs/Horde/HordeServer/Controllers/JobsController.cs
--- a/Engine/Source/Programs/Horde/HordeServer/Controllers/JobsController.cs
+++ b/Engine/Source/Programs/Horde/HordeServer/Controllers/JobsController.cs
@@ -170,18 +170,6 @@ namespace HordeServer.Controllers
 
 			// New groups for the job
 			IGraph Graph = await Graphs.AddAsync(Template);
-			if (Create.Groups != null)
-			{
-				Graph = await Graphs.AppendAsync(Graph, Create.Groups, null, null);
-			}
-			if (Create.Aggregates != null)
-			{
-				Graph = await Graphs.AppendAsync(Graph, null, Create.Aggregates, null);
-			}
-			if (Create.Labels != null)
-			{
-				Graph = await Graphs.AppendAsync(Graph, null, null, Create.Labels);
-			}
 
 			// Get the change to build
 			int Change;
@@ -629,7 +617,7 @@ namespace HordeServer.Controllers
 		/// <returns>Id of the new job</returns>
 		[HttpPost]
 		[Route("/api/v1/jobs/{JobId}/groups")]
-		public async Task<ActionResult> CreateGroupsAsync(string JobId, [FromBody] List<CreateGroupRequest> Requests)
+		public async Task<ActionResult> CreateGroupsAsync(string JobId, [FromBody] List<NewGroup> Requests)
 		{
 			Dictionary<string, int> ExpectedDurationCache = new Dictionary<string, int>();
 
@@ -1254,7 +1242,7 @@ namespace HordeServer.Controllers
 		/// <returns>Response object with the first index of the new aggregates</returns>
 		[HttpPost]
 		[Route("/api/v1/jobs/{JobId}/aggregates")]
-		public async Task<ActionResult> CreateAggregatesAsync(string JobId, [FromBody] List<CreateAggregateRequest> Requests)
+		public async Task<ActionResult> CreateAggregatesAsync(string JobId, [FromBody] List<NewAggregate> Requests)
 		{
 			ObjectId JobIdValue = JobId.ToObjectId();
 			for (; ; )
diff --git a/Engine/Source/Programs/Horde/HordeServer/Models/Graph.cs b/Engine/Source/Programs/Horde/HordeServer/Models/Graph.cs
--- a/Engine/Source/Programs/Horde/HordeServer/Models/Graph.cs
+++ b/Engine/Source/Programs/Horde/HordeServer/Models/Graph.cs
@@ -280,10 +280,10 @@ namespace HordeServer.Models
 		/// <returns>True if the node was found, false otherwise</returns>
 		public static bool TryFindNode(this IGraph Graph, string NodeName, out NodeRef Ref)
 		{
-			for(int GroupIdx = 0; GroupIdx < Graph.Groups.Count; GroupIdx++)
+			for (int GroupIdx = 0; GroupIdx < Graph.Groups.Count; GroupIdx++)
 			{
 				INodeGroup Group = Graph.Groups[GroupIdx];
-				for(int NodeIdx = 0; NodeIdx < Group.Nodes.Count; NodeIdx++)
+				for (int NodeIdx = 0; NodeIdx < Group.Nodes.Count; NodeIdx++)
 				{
 					INode Node = Group.Nodes[NodeIdx];
 					if (String.Equals(Node.Name, NodeName, StringComparison.OrdinalIgnoreCase))
@@ -366,5 +366,189 @@ namespace HordeServer.Models
 			return Enumerable.Concat(Node.InputDependencies, Node.OrderDependencies).Select(x => Graph.GetNode(x));
 		}
 	}
+
+	/// <summary>
+	/// Information required to create a node
+	/// </summary>
+	public class NewNode
+	{
+		/// <summary>
+		/// The name of this node 
+		/// </summary>
+		public string Name { get; set; } = null!;
+
+		/// <summary>
+		/// List of nodes which must succeed for this node to run
+		/// </summary>
+		public List<string>? InputDependencies { get; set; }
+
+		/// <summary>
+		/// List of nodes which must have completed for this node to run
+		/// </summary>
+		public List<string>? OrderDependencies { get; set; }
+
+		/// <summary>
+		/// The priority of this node
+		/// </summary>
+		public Priority? Priority { get; set; }
+
+		/// <summary>
+		/// This node can be run multiple times
+		/// </summary>
+		public bool? AllowRetry { get; set; }
+
+		/// <summary>
+		/// This node can start running early, before dependencies of other nodes in the same group are complete
+		/// </summary>
+		public bool? RunEarly { get; set; }
+
+		/// <summary>
+		/// Whether to include warnings in the diagnostic output
+		/// </summary>
+		public bool? Warnings { get; set; }
+
+		/// <summary>
+		/// Credentials required for this node to run. This dictionary maps from environment variable names to a credential property in the format 'CredentialName.PropertyName'.
+		/// </summary>
+		public Dictionary<string, string>? Credentials { get; set; }
+
+		/// <summary>
+		/// Properties for this node
+		/// </summary>
+		public Dictionary<string, string>? Properties { get; set; }
+
+		/// <summary>
+		/// Constructor
+		/// </summary>
+		/// <param name="Name">Name of the node</param>
+		/// <param name="InputDependencies">List of nodes which must have completed succesfully for this node to run</param>
+		/// <param name="OrderDependencies">List of nodes which must have completed for this node to run</param>
+		/// <param name="Priority">Priority of this node</param>
+		/// <param name="AllowRetry">Whether the node can be run multiple times</param>
+		/// <param name="RunEarly">Whether the node can run early, before dependencies of other nodes in the same group complete</param>
+		/// <param name="Warnings">Whether to include warnings in the diagnostic output (defaults to true)</param>
+		/// <param name="Credentials">Credentials required for this node to run</param>
+		/// <param name="Properties">Properties for the node</param>
+		public NewNode(string Name, List<string>? InputDependencies = null, List<string>? OrderDependencies = null, Priority? Priority = null, bool? AllowRetry = null, bool? RunEarly = null, bool? Warnings = null, Dictionary<string, string>? Credentials = null, Dictionary<string, string>? Properties = null)
+		{
+			this.Name = Name;
+			this.InputDependencies = InputDependencies;
+			this.OrderDependencies = OrderDependencies;
+			this.Priority = Priority;
+			this.AllowRetry = AllowRetry;
+			this.RunEarly = RunEarly;
+			this.Warnings = Warnings;
+			this.Credentials = Credentials;
+			this.Properties = Properties;
+		}
+	}
+
+	/// <summary>
+	/// Information about a group of nodes
+	/// </summary>
+	public class NewGroup
+	{
+		/// <summary>
+		/// The type of agent to execute this group
+		/// </summary>
+		public string AgentType { get; set; }
+
+		/// <summary>
+		/// Nodes in the group
+		/// </summary>
+		public List<NewNode> Nodes { get; set; }
+
+		/// <summary>
+		/// Constructor
+		/// </summary>
+		/// <param name="AgentType">The type of agent to execute this group</param>
+		/// <param name="Nodes">Nodes in this group</param>
+		public NewGroup(string AgentType, List<NewNode> Nodes)
+		{
+			this.AgentType = AgentType;
+			this.Nodes = Nodes;
+		}
+	}
+
+	/// <summary>
+	/// Information about a group of nodes
+	/// </summary>
+	public class NewLabel
+	{
+		/// <summary>
+		/// Category for this label
+		/// </summary>
+		[Obsolete("Use DashboardCategory instead")]
+		public string? Category => DashboardCategory;
+
+		/// <summary>
+		/// Name of the aggregate
+		/// </summary>
+		[Obsolete("Use DashboardName instead")]
+		public string? Name => DashboardName;
+
+		/// <summary>
+		/// Name of the aggregate
+		/// </summary>
+		public string? DashboardName { get; set; }
+
+		/// <summary>
+		/// Category for this label
+		/// </summary>
+		public string? DashboardCategory { get; set; }
+
+		/// <summary>
+		/// Name of the badge in UGS
+		/// </summary>
+		public string? UgsName { get; set; }
+
+		/// <summary>
+		/// Project to show this label for in UGS
+		/// </summary>
+		public string? UgsProject { get; set; }
+
+		/// <summary>
+		/// Which change the label applies to
+		/// </summary>
+		public LabelChange Change { get; set; }
+
+		/// <summary>
+		/// Nodes which must be part of the job for the aggregate to be valid
+		/// </summary>
+		public List<string> RequiredNodes { get; set; } = new List<string>();
+
+		/// <summary>
+		/// Nodes which must be part of the job for the aggregate to be valid
+		/// </summary>
+		public List<string> IncludedNodes { get; set; } = new List<string>();
+	}
+
+	/// <summary>
+	/// Information about a group of nodes
+	/// </summary>
+	public class NewAggregate
+	{
+		/// <summary>
+		/// Name of the aggregate
+		/// </summary>
+		public string Name { get; set; }
+
+		/// <summary>
+		/// Nodes which must be part of the job for the aggregate to be valid
+		/// </summary>
+		public List<string> Nodes { get; set; }
+
+		/// <summary>
+		/// Constructor
+		/// </summary>
+		/// <param name="Name">Name of this aggregate</param>
+		/// <param name="Nodes">Nodes which must be part of the job for the aggregate to be shown</param>
+		public NewAggregate(string Name, List<string> Nodes)
+		{
+			this.Name = Name;
+			this.Nodes = Nodes;
+		}
+	}
 }
 
+
diff --git a/Engine/Source/Programs/Horde/HordeServer/Services/RpcService.cs b/Engine/Source/Programs/Horde/HordeServer/Services/RpcService.cs
--- a/Engine/Source/Programs/Horde/HordeServer/Services/RpcService.cs
+++ b/Engine/Source/Programs/Horde/HordeServer/Services/RpcService.cs
@@ -817,29 +817,29 @@ namespace HordeServer.Services
 		/// <returns>Information about the new agent</returns>
 		public async override Task<UpdateGraphResponse> UpdateGraph(UpdateGraphRequest Request, ServerCallContext Context)
 		{
-			List<Api.CreateGroupRequest> ApiGroups = new List<Api.CreateGroupRequest>();
+			List<NewGroup> NewGroups = new List<NewGroup>();
 			foreach (CreateGroupRequest Group in Request.Groups)
 			{
-				List<Api.CreateNodeRequest> ApiNodes = new List<Api.CreateNodeRequest>();
+				List<NewNode> NewNodes = new List<NewNode>();
 				foreach (CreateNodeRequest Node in Group.Nodes)
 				{
-					Api.CreateNodeRequest ApiNode = new Api.CreateNodeRequest(Node.Name, Node.InputDependencies.ToList(), Node.OrderDependencies.ToList(), Node.Priority, Node.AllowRetry, Node.RunEarly, Node.Warnings, new Dictionary<string, string>(Node.Credentials), new Dictionary<string, string>(Node.Properties));
-					ApiNodes.Add(ApiNode);
+					NewNode NewNode = new NewNode(Node.Name, Node.InputDependencies.ToList(), Node.OrderDependencies.ToList(), Node.Priority, Node.AllowRetry, Node.RunEarly, Node.Warnings, new Dictionary<string, string>(Node.Credentials), new Dictionary<string, string>(Node.Properties));
+					NewNodes.Add(NewNode);
 				}
-				ApiGroups.Add(new Api.CreateGroupRequest(Group.AgentType, ApiNodes));
+				NewGroups.Add(new NewGroup(Group.AgentType, NewNodes));
 			}
 
-			List<Api.CreateAggregateRequest> ApiAggregates = new List<Api.CreateAggregateRequest>();
+			List<NewAggregate> NewAggregates = new List<NewAggregate>();
 			foreach (CreateAggregateRequest Aggregate in Request.Aggregates)
 			{
-				Api.CreateAggregateRequest ApiAggregate = new Api.CreateAggregateRequest(Aggregate.Name, Aggregate.Nodes.ToList());
-				ApiAggregates.Add(ApiAggregate);
+				NewAggregate NewAggregate = new NewAggregate(Aggregate.Name, Aggregate.Nodes.ToList());
+				NewAggregates.Add(NewAggregate);
 			}
 
-			List<Api.CreateLabelRequest> ApiLabels = new List<Api.CreateLabelRequest>();
+			List<NewLabel> NewLabels = new List<NewLabel>();
 			foreach (CreateLabelRequest Label in Request.Labels)
 			{
-				Api.CreateLabelRequest NewLabel = new Api.CreateLabelRequest();
+				NewLabel NewLabel = new NewLabel();
 				NewLabel.DashboardName = String.IsNullOrEmpty(Label.DashboardName) ? null : Label.DashboardName;
 				NewLabel.DashboardCategory = String.IsNullOrEmpty(Label.DashboardCategory) ? null : Label.DashboardCategory;
 				NewLabel.UgsName = String.IsNullOrEmpty(Label.UgsName) ? null : Label.UgsName;
@@ -847,7 +847,7 @@ namespace HordeServer.Services
 				NewLabel.Change = Label.Change;
 				NewLabel.RequiredNodes = Label.RequiredNodes.ToList();
 				NewLabel.IncludedNodes = Label.IncludedNodes.ToList();
-				ApiLabels.Add(NewLabel);
+				NewLabels.Add(NewLabel);
 			}
 
 			ObjectId JobIdValue = Request.JobId.ToObjectId();
@@ -864,7 +864,7 @@ namespace HordeServer.Services
 				}
 
 				IGraph Graph = await JobService.GetGraphAsync(Job);
-				Graph = await Graphs.AppendAsync(Graph, ApiGroups, ApiAggregates, ApiLabels);
+				Graph = await Graphs.AppendAsync(Graph, NewGroups, NewAggregates, NewLabels);
 
 				if (await JobService.TryUpdateGraphAsync(Job, Graph))
 				{
diff --git a/Engine/Source/Programs/Horde/HordeServerTests/JobCollectionTests.cs b/Engine/Source/Programs/Horde/HordeServerTests/JobCollectionTests.cs
--- a/Engine/Source/Programs/Horde/HordeServerTests/JobCollectionTests.cs
+++ b/Engine/Source/Programs/Horde/HordeServerTests/JobCollectionTests.cs
@@ -38,16 +38,16 @@ namespace HordeServerTests
 			this.JobCollection = new JobCollection(DatabaseService, NullLogger<JobCollection>.Instance);
 		}
 
-		CreateGroupRequest AddGroup(List<CreateGroupRequest> Groups)
+		NewGroup AddGroup(List<NewGroup> Groups)
 		{
-			CreateGroupRequest Group = new CreateGroupRequest("win64", new List<CreateNodeRequest>());
+			NewGroup Group = new NewGroup("win64", new List<NewNode>());
 			Groups.Add(Group);
 			return Group;
 		}
 
-		CreateNodeRequest AddNode(CreateGroupRequest Group, string Name, string[]? InputDependencies, Action<CreateNodeRequest>? Action = null)
+		NewNode AddNode(NewGroup Group, string Name, string[]? InputDependencies, Action<NewNode>? Action = null)
 		{
-			CreateNodeRequest Node = new CreateNodeRequest(Name, InputDependencies?.ToList(), InputDependencies?.ToList(), null, null, null, null, null, null);
+			NewNode Node = new NewNode(Name, InputDependencies?.ToList(), InputDependencies?.ToList(), null, null, null, null, null, null);
 			if (Action != null)
 			{
 				Action.Invoke(Node);
@@ -90,16 +90,16 @@ namespace HordeServerTests
 			await StartBatch(Job, BaseGraph, 0);
 			await RunStep(Job, BaseGraph, 0, 0, JobStepOutcome.Success); // Setup Build
 
-			List<CreateGroupRequest> NewGroups = new List<CreateGroupRequest>();
+			List<NewGroup> NewGroups = new List<NewGroup>();
 
-			CreateGroupRequest InitialGroup = AddGroup(NewGroups);
+			NewGroup InitialGroup = AddGroup(NewGroups);
 			AddNode(InitialGroup, "Update Version Files", null);
 			AddNode(InitialGroup, "Compile Editor", new[] { "Update Version Files" });
 
-			CreateGroupRequest CompileGroup = AddGroup(NewGroups);
+			NewGroup CompileGroup = AddGroup(NewGroups);
 			AddNode(CompileGroup, "Compile Client", new[] { "Update Version Files" });
 
-			CreateGroupRequest PublishGroup = AddGroup(NewGroups);
+			NewGroup PublishGroup = AddGroup(NewGroups);
 			AddNode(PublishGroup, "Cook Client", new[] { "Compile Editor" }, x => x.RunEarly = true);
 			AddNode(PublishGroup, "Publish Client", new[] { "Compile Client", "Cook Client" });
 			AddNode(PublishGroup, "Post-Publish Client", null, x => x.OrderDependencies = new List<string> { "Publish Client" });
@@ -162,9 +162,9 @@ namespace HordeServerTests
 			await StartBatch(Job, BaseGraph, 0);
 			await RunStep(Job, BaseGraph, 0, 0, JobStepOutcome.Success); // Setup Build
 
-			List<CreateGroupRequest> NewGroups = new List<CreateGroupRequest>();
+			List<NewGroup> NewGroups = new List<NewGroup>();
 
-			CreateGroupRequest InitialGroup = AddGroup(NewGroups);
+			NewGroup InitialGroup = AddGroup(NewGroups);
 			AddNode(InitialGroup, "Step 1", null);
 			AddNode(InitialGroup, "Step 2", HasDependency? new[] { "Step 1" } : null);
 			AddNode(InitialGroup, "Step 3", new[] { "Step 2" });
diff --git a/Engine/Source/Programs/Horde/HordeServerTests/JobServiceTests.cs b/Engine/Source/Programs/Horde/HordeServerTests/JobServiceTests.cs
--- a/Engine/Source/Programs/Horde/HordeServerTests/JobServiceTests.cs
+++ b/Engine/Source/Programs/Horde/HordeServerTests/JobServiceTests.cs
@@ -158,15 +158,15 @@ namespace HordeServerTests
 			ITemplate Template = await TestSetup.TemplateService.CreateTemplateAsync("Test template", null, false, null, null, new List<TemplateCounter>(), new List<string>(), new List<Parameter>());
 			IGraph Graph = await TestSetup.GraphCollection.AddAsync(Template);
 
-			CreateGroupRequest GroupA = new CreateGroupRequest("win", new List<CreateNodeRequest>());
-			GroupA.Nodes.Add(new CreateNodeRequest("Compile"));
+			NewGroup GroupA = new NewGroup("win", new List<NewNode>());
+			GroupA.Nodes.Add(new NewNode("Compile"));
 
-			CreateGroupRequest GroupB = new CreateGroupRequest("win", new List<CreateNodeRequest>());
-			GroupB.Nodes.Add(new CreateNodeRequest("Cook", RunEarly: true));
-			GroupB.Nodes.Add(new CreateNodeRequest("Middle"));
-			GroupB.Nodes.Add(new CreateNodeRequest("Pak", InputDependencies: new List<string> { "Compile", "Cook" }));
+			NewGroup GroupB = new NewGroup("win", new List<NewNode>());
+			GroupB.Nodes.Add(new NewNode("Cook", RunEarly: true));
+			GroupB.Nodes.Add(new NewNode("Middle"));
+			GroupB.Nodes.Add(new NewNode("Pak", InputDependencies: new List<string> { "Compile", "Cook" }));
 
-			Graph = await TestSetup.GraphCollection.AppendAsync(Graph, new List<CreateGroupRequest> { GroupA, GroupB });
+			Graph = await TestSetup.GraphCollection.AppendAsync(Graph, new List<NewGroup> { GroupA, GroupB });
 
 			IJob Job = await TestSetup.JobService.CreateJobAsync(null, Stream!, new TemplateRefId("temp"), Template.Id, Graph, "Hello", 1234, 1233, 999, null, null, "joe", null, null, null, null, true, true, null, null, null, Template.Counters, new List<string> { "-Target=Pak" });
 
diff --git a/Engine/Source/Programs/Horde/HordeServerTests/SchedulerTests.cs b/Engine/Source/Programs/Horde/HordeServerTests/SchedulerTests.cs
--- a/Engine/Source/Programs/Horde/HordeServerTests/SchedulerTests.cs
+++ b/Engine/Source/Programs/Horde/HordeServerTests/SchedulerTests.cs
@@ -322,8 +322,8 @@ namespace HordeServerTests
 
 			// Create the TriggerNext step and mark it as complete
 			IGraph GraphA = await TestSetup.GraphCollection.AddAsync(NewTemplate1);
-			CreateGroupRequest GroupA = new CreateGroupRequest("win", new List<CreateNodeRequest> { new CreateNodeRequest("TriggerNext") });
-			GraphA = await TestSetup.GraphCollection.AppendAsync(GraphA, new List<CreateGroupRequest> { GroupA });
+			NewGroup GroupA = new NewGroup("win", new List<NewNode> { new NewNode("TriggerNext") });
+			GraphA = await TestSetup.GraphCollection.AppendAsync(GraphA, new List<NewGroup> { GroupA });
 
 			// Tick the schedule and make sure it doesn't trigger
 			await ScheduleService.TickSharedOnlyForTestingAsync();
