Horde: Remove REAPI implementation.

diff --git a/Engine/Source/Programs/Horde/HordeAgent/Commands/Execution/ExecuteCommand.cs b/Engine/Source/Programs/Horde/HordeAgent/Commands/Execution/ExecuteCommand.cs
deleted file mode 100644
diff --git a/Engine/Source/Programs/Horde/HordeAgent/Commands/Execution/ExecuteCommandV2.cs b/Engine/Source/Programs/Horde/HordeAgent/Commands/Execution/ExecuteCommandV2.cs
--- a/Engine/Source/Programs/Horde/HordeAgent/Commands/Execution/ExecuteCommandV2.cs
+++ b/Engine/Source/Programs/Horde/HordeAgent/Commands/Execution/ExecuteCommandV2.cs
@@ -15,7 +15,6 @@ using HordeAgent.Services;
 using Grpc.Net.Client;
 using System.Net.Http.Headers;
 using Grpc.Core;
-using Google.LongRunning;
 using Microsoft.Extensions.Hosting;
 using Microsoft.Extensions.DependencyInjection;
 using Microsoft.Extensions.Configuration;
diff --git a/Engine/Source/Programs/Horde/HordeAgent/Execution/ActionExecutor.cs b/Engine/Source/Programs/Horde/HordeAgent/Execution/ActionExecutor.cs
deleted file mode 100644
diff --git a/Engine/Source/Programs/Horde/HordeAgent/Execution/ComputeTaskExecutor.cs b/Engine/Source/Programs/Horde/HordeAgent/Execution/ComputeTaskExecutor.cs
--- a/Engine/Source/Programs/Horde/HordeAgent/Execution/ComputeTaskExecutor.cs
+++ b/Engine/Source/Programs/Horde/HordeAgent/Execution/ComputeTaskExecutor.cs
@@ -3,7 +3,6 @@
 using EpicGames.Core;
 using EpicGames.Horde;
 using EpicGames.Horde.Storage;
-using Google.LongRunning;
 using Google.Protobuf;
 using Grpc.Core;
 using Microsoft.Extensions.Options;
@@ -23,7 +22,6 @@ using HordeAgent.Utility;
 using HordeCommon.Rpc.Tasks;
 using Microsoft.Extensions.Logging;
 using HordeCommon.Rpc;
-using EpicGames.Horde.Common.RemoteExecution;
 using EpicGames.Serialization;
 using EpicGames.Horde.Compute;
 using EpicGames.Horde.Common;
diff --git a/Engine/Source/Programs/Horde/HordeAgent/HordeAgent.csproj b/Engine/Source/Programs/Horde/HordeAgent/HordeAgent.csproj
--- a/Engine/Source/Programs/Horde/HordeAgent/HordeAgent.csproj
+++ b/Engine/Source/Programs/Horde/HordeAgent/HordeAgent.csproj
@@ -84,21 +84,15 @@
 
   <ItemGroup>
     <ProjectReference Include="..\..\Shared\EpicGames.Core\EpicGames.Core.csproj" />
-    <ProjectReference Include="..\..\Shared\EpicGames.Horde.Client\EpicGames.Horde.Client.csproj" />
-    <ProjectReference Include="..\..\Shared\EpicGames.Horde.Common\EpicGames.Horde.Common.csproj" />
     <ProjectReference Include="..\..\Shared\EpicGames.Horde\EpicGames.Horde.csproj" />
+    <ProjectReference Include="..\..\Shared\EpicGames.Horde.Client\EpicGames.Horde.Client.csproj" />
     <ProjectReference Include="..\..\Shared\EpicGames.Perforce.Managed\EpicGames.Perforce.Managed.csproj" />
     <ProjectReference Include="..\..\Shared\EpicGames.Perforce\EpicGames.Perforce.csproj" />
     <ProjectReference Include="..\HordeCommon\HordeCommon.csproj" />
   </ItemGroup>
 
   <ItemGroup>
-    <Protobuf Include="..\..\Shared\EpicGames.Horde.Common\Protos\google\bytestream\bytestream.proto" ProtoRoot="..\..\Shared\EpicGames.Horde.Common\Protos" GrpcServices="Client" />
-    <Protobuf Include="..\..\Shared\EpicGames.Horde.Common\Protos\google\devtools\remoteworkers\v1test2\bots.proto" ProtoRoot="..\..\Shared\EpicGames.Horde.Common\Protos" GrpcServices="Client" />
-    <Protobuf Include="..\..\Shared\EpicGames.Horde.Common\Protos\google\devtools\remoteworkers\v1test2\worker.proto" ProtoRoot="..\..\Shared\EpicGames.Horde.Common\Protos" GrpcServices="Client" />
-    <Protobuf Include="..\..\Shared\EpicGames.Horde.Common\Protos\horde\action_rpc.proto" ProtoRoot="..\..\Shared\EpicGames.Horde.Common\Protos" GrpcServices="Client" />
     <Protobuf Include="..\HordeCommon\Protos\horde\horde_rpc.proto" ProtoRoot="..\HordeCommon\Protos" GrpcServices="Client" />
-    <Protobuf Include="..\..\Shared\EpicGames.Horde.Common\Protos\horde\tasks\action_task.proto" ProtoRoot="..\..\Shared\EpicGames.Horde.Common\Protos" GrpcServices="Client" />
   </ItemGroup>
 
   <ItemGroup>
diff --git a/Engine/Source/Programs/Horde/HordeAgent/Services/GrpcService.cs b/Engine/Source/Programs/Horde/HordeAgent/Services/GrpcService.cs
--- a/Engine/Source/Programs/Horde/HordeAgent/Services/GrpcService.cs
+++ b/Engine/Source/Programs/Horde/HordeAgent/Services/GrpcService.cs
@@ -12,7 +12,6 @@ using System.Collections.Generic;
 using System.Net.Http;
 using System.Net.Http.Headers;
 using System.Text;
-using Build.Bazel.Remote.Execution.V2;
 using Grpc.Core;
 using HordeCommon.Rpc;
 
@@ -38,10 +37,6 @@ namespace HordeAgent.Services
 		/// </summary>
 		ILogger Logger;
 
-		internal Func<ChannelBase, ContentAddressableStorage.ContentAddressableStorageClient> CasClientFactory;
-		internal Func<ChannelBase, ActionCache.ActionCacheClient> ActionCacheClientFactory;
-		internal Func<ChannelBase, ActionRpc.ActionRpcClient> ActionRpcClientFactory;
-
 		/// <summary>
 		/// Constructor
 		/// </summary>
@@ -52,10 +47,6 @@ namespace HordeAgent.Services
 			this.Settings = Settings;
 			this.ServerProfile = Settings.Value.GetCurrentServerProfile();
 			this.Logger = Logger;
-			
-			CasClientFactory = Channel => new ContentAddressableStorage.ContentAddressableStorageClient(Channel);
-			ActionCacheClientFactory = Channel => new ActionCache.ActionCacheClient(Channel);
-			ActionRpcClientFactory = Channel => new ActionRpc.ActionRpcClient(Channel);
 		}
 
 		/// <summary>
@@ -117,20 +108,5 @@ namespace HordeAgent.Services
 				DisposeHttpClient = true
 			});
 		}
-
-		public ContentAddressableStorage.ContentAddressableStorageClient CreateCasClient(ChannelBase Channel)
-		{
-			return CasClientFactory(Channel);
-		}
-
-		public ActionCache.ActionCacheClient CreateActionCacheClient(ChannelBase Channel)
-		{
-			return ActionCacheClientFactory(Channel);
-		}
-
-		public ActionRpc.ActionRpcClient CreateActionRpcClient(ChannelBase Channel)
-		{
-			return ActionRpcClientFactory(Channel);
-		}
 	}
 }
diff --git a/Engine/Source/Programs/Horde/HordeAgent/Services/WorkerService.cs b/Engine/Source/Programs/Horde/HordeAgent/Services/WorkerService.cs
--- a/Engine/Source/Programs/Horde/HordeAgent/Services/WorkerService.cs
+++ b/Engine/Source/Programs/Horde/HordeAgent/Services/WorkerService.cs
@@ -674,14 +674,6 @@ namespace HordeAgent.Services
 		{
 			Any Payload = LeaseInfo.Lease.Payload;
 
-			ActionTask ActionTask;
-			if (LeaseInfo.Lease.Payload.TryUnpack(out ActionTask))
-			{
-				GlobalTracer.Instance.ActiveSpan?.SetTag("task", "Action");
-				Func<ILogger, Task<LeaseResult>> Handler = NewLogger => ActionAsync(RpcConnection, LeaseInfo.Lease.Id, ActionTask, NewLogger, LeaseInfo.CancellationTokenSource.Token);
-				return await HandleLeasePayloadWithLogAsync(RpcConnection, ActionTask.LogId, AgentId, null, null, Handler, LeaseInfo.CancellationTokenSource.Token);
-			}
-
 			ComputeTaskMessage ComputeTask;
 			if (LeaseInfo.Lease.Payload.TryUnpack(out ComputeTask))
 			{
@@ -822,81 +814,6 @@ namespace HordeAgent.Services
 			}
 		}
 
-		/// <summary>
-		/// Execute a remote execution aciton
-		/// </summary>
-		/// <param name="RpcConnection">RPC client for communicating with the server</param>
-		/// <param name="LeaseId">The lease id</param>
-		/// <param name="ActionTask">The action task parameters</param>
-		/// <param name="Logger">Logger for the task</param>
-		/// <param name="CancellationToken">Token used to cancel the operation</param>
-		/// <returns></returns>
-		internal async Task<LeaseResult> ActionAsync(IRpcConnection RpcConnection, string LeaseId, ActionTask ActionTask, ILogger Logger, CancellationToken CancellationToken)
-		{
-			DateTimeOffset ActionTaskStartTime = DateTimeOffset.UtcNow;
-			using (IRpcClientRef Client = await RpcConnection.GetClientRef(new RpcContext(), CancellationToken))
-			{
-				DirectoryReference LeaseDir = DirectoryReference.Combine(WorkingDir, "Remote", LeaseId);
-				DirectoryReference.CreateDirectory(LeaseDir);
-
-				GrpcChannel GetChannel(string? Url, string? Token)
-				{
-					if (string.IsNullOrEmpty(Url)) return Client.Channel;
-					if (!string.IsNullOrEmpty(Token))
-						return GrpcService.CreateGrpcChannel(Url, new AuthenticationHeaderValue("ServiceAccount", Token));
-					
-					return GrpcChannel.ForAddress(Url, new GrpcChannelOptions
-					{
-						// Required payloads coming from CAS service can be large
-						MaxReceiveMessageSize = 1024 * 1024 * 1024, // 1 GB
-						MaxSendMessageSize = 1024 * 1024 * 1024 // 1 GB
-					});
-				}
-
-				GrpcChannel CasChannel = GetChannel(ActionTask.CasUrl, ActionTask.ServiceAccountToken);// == null ? Client.Channel : GrpcChannel.ForAddress(ActionTask.CasUrl);
-				GrpcChannel ActionCacheChannel = GetChannel(ActionTask.ActionCacheUrl, ActionTask.ServiceAccountToken);//ActionTask.ActionCacheUrl == null ? Client.Channel : GrpcChannel.ForAddress(ActionTask.ActionCacheUrl);
-				GrpcChannel ActionRpcChannel = Client.Channel;
-
-				Logger.LogInformation("CasChannel={CasChannel}", CasChannel.Target);
-				Logger.LogInformation("ActionCacheChannel={ActionCacheChannel}", ActionCacheChannel.Target);
-				Logger.LogInformation("ActionRpcChannel={ActionRpcChannel}", ActionRpcChannel.Target);
-
-				ActionExecutor Executor = new ActionExecutor(
-					Settings.GetAgentName(), ActionTask.InstanceName,
-					GrpcService.CreateCasClient(CasChannel),
-					GrpcService.CreateActionCacheClient(ActionCacheChannel),
-					GrpcService.CreateActionRpcClient(ActionRpcChannel), Logger);
-
-				try
-				{
-					await Executor.ExecuteActionAsync(LeaseId, ActionTask, LeaseDir, ActionTaskStartTime, CancellationToken);
-				}
-				catch (Exception re)
-				{
-					Logger.LogError(re, "Unhandled exception during remote execution");
-					PostActionResultRequest PostResultRequest = new PostActionResultRequest();
-					PostResultRequest.LeaseId = LeaseId;
-					PostResultRequest.ActionDigest = ActionTask.Digest;
-					PostResultRequest.Error = new RpcException(new Status(StatusCode.Unknown, "Unhandled exception during remote execution: " + re.Message), re.Message).Status;
-					await Executor.ActionRpc.PostActionResultAsync(PostResultRequest);
-					
-					throw re;
-				}
-				finally
-				{
-					try
-					{
-						DirectoryReference.Delete(LeaseDir, true);
-					}
-					catch
-					{
-					}
-				}
-
-				return LeaseResult.Success;
-			}
-		}
-
 		/// <summary>
 		/// Execute a remote execution aciton
 		/// </summary>
diff --git a/Engine/Source/Programs/Horde/HordeAgentTests/ActionExecutorTest.cs b/Engine/Source/Programs/Horde/HordeAgentTests/ActionExecutorTest.cs
deleted file mode 100644
diff --git a/Engine/Source/Programs/Horde/HordeAgentTests/GrpcClients.cs b/Engine/Source/Programs/Horde/HordeAgentTests/GrpcClients.cs
deleted file mode 100644
diff --git a/Engine/Source/Programs/Horde/HordeAgentTests/TaskTest.cs b/Engine/Source/Programs/Horde/HordeAgentTests/TaskTest.cs
deleted file mode 100644
diff --git a/Engine/Source/Programs/Horde/HordeServer/Compute/Impl/ComputeService.cs b/Engine/Source/Programs/Horde/HordeServer/Compute/Impl/ComputeService.cs
--- a/Engine/Source/Programs/Horde/HordeServer/Compute/Impl/ComputeService.cs
+++ b/Engine/Source/Programs/Horde/HordeServer/Compute/Impl/ComputeService.cs
@@ -1,7 +1,6 @@
 // Copyright Epic Games, Inc. All Rights Reserved.
 
 using EpicGames.Core;
-using Google.LongRunning;
 using Google.Protobuf;
 using Google.Protobuf.Reflection;
 using Google.Protobuf.WellKnownTypes;
@@ -10,7 +9,6 @@ using HordeCommon;
 using HordeCommon.Rpc.Tasks;
 using HordeServer.Api;
 using HordeServer.Models;
-using HordeServer.Rpc;
 using HordeServer.Services;
 using Microsoft.Extensions.Logging;
 using MongoDB.Bson;
diff --git a/Engine/Source/Programs/Horde/HordeServer/Controllers/LeasesController.cs b/Engine/Source/Programs/Horde/HordeServer/Controllers/LeasesController.cs
--- a/Engine/Source/Programs/Horde/HordeServer/Controllers/LeasesController.cs
+++ b/Engine/Source/Programs/Horde/HordeServer/Controllers/LeasesController.cs
@@ -176,9 +176,9 @@ namespace HordeServer.Controllers
 				return NotFound();
 			}
 
-			if (!AgentLease.IsActionLease() && !AgentLease.IsConformLease())
+			if (!AgentLease.IsConformLease())
 			{
-				Logger.LogError("Lease abort only supported on action and conform leases for now, {LeaseId}", LeaseId);
+				Logger.LogError("Lease abort only supported on conform leases for now, {LeaseId}", LeaseId);
 				return Ok();
 			}
 
diff --git a/Engine/Source/Programs/Horde/HordeServer/HordeServer.csproj b/Engine/Source/Programs/Horde/HordeServer/HordeServer.csproj
--- a/Engine/Source/Programs/Horde/HordeServer/HordeServer.csproj
+++ b/Engine/Source/Programs/Horde/HordeServer/HordeServer.csproj
@@ -1,4 +1,4 @@
-﻿<Project Sdk="Microsoft.NET.Sdk.Web">
+<Project Sdk="Microsoft.NET.Sdk.Web">
 
   <PropertyGroup>
     <TargetFramework>net5.0</TargetFramework>
@@ -68,18 +68,12 @@
     <Protobuf Include="..\..\Shared\EpicGames.Horde\Protos\horde\storage\ref_table.proto" ProtoRoot="..\..\Shared\EpicGames.Horde\Protos" GrpcServices="Server" />
     <Protobuf Include="..\..\Shared\EpicGames.Horde\Protos\horde\storage\blob_store.proto" ProtoRoot="..\..\Shared\EpicGames.Horde\Protos" GrpcServices="Server" />
     <Protobuf Include="..\..\Shared\EpicGames.Horde\Protos\horde\compute\compute_rpc.proto" ProtoRoot="..\..\Shared\EpicGames.Horde\Protos" GrpcServices="Server" />
-    <Protobuf Include="..\..\Shared\EpicGames.Horde.Common\Protos\google\bytestream\bytestream.proto" ProtoRoot="..\..\Shared\EpicGames.Horde.Common\Protos" GrpcServices="Server" />
-    <Protobuf Include="..\..\Shared\EpicGames.Horde.Common\Protos\horde\action_rpc.proto" ProtoRoot="..\..\Shared\EpicGames.Horde.Common\Protos" GrpcServices="Server" />
-    <Protobuf Include="..\..\Shared\EpicGames.Horde.Common\Protos\horde\tasks\action_task.proto" ProtoRoot="..\..\Shared\EpicGames.Horde.Common\Protos" GrpcServices="Server" />
-    <Protobuf Include="..\..\Shared\EpicGames.Horde.Common\Protos\horde\storage\ref_service.proto" ProtoRoot="..\..\Shared\EpicGames.Horde.Common\Protos" GrpcServices="Server" />
-    <Protobuf Include="..\..\Shared\EpicGames.Horde.Common\Protos\horde\storage\blob_service.proto" ProtoRoot="..\..\Shared\EpicGames.Horde.Common\Protos" GrpcServices="Server" />
     <Protobuf Include="..\HordeCommon\Protos\horde\horde_rpc.proto" ProtoRoot="..\HordeCommon\Protos" GrpcServices="Server" />
     <Protobuf Include="..\HordeCommon\Protos\grpc\health\v1\health.proto" ProtoRoot="..\HordeCommon\Protos" GrpcServices="Server" />
   </ItemGroup>
 
   <ItemGroup>
     <ProjectReference Include="..\..\Shared\EpicGames.Core\EpicGames.Core.csproj" />
-    <ProjectReference Include="..\..\Shared\EpicGames.Horde.Common\EpicGames.Horde.Common.csproj" />
     <ProjectReference Include="..\..\Shared\EpicGames.Horde\EpicGames.Horde.csproj" />
     <ProjectReference Include="..\..\Shared\EpicGames.IoHash\EpicGames.IoHash.csproj" />
     <ProjectReference Include="..\..\Shared\EpicGames.Perforce.Managed\EpicGames.Perforce.Managed.csproj" />
diff --git a/Engine/Source/Programs/Horde/HordeServer/Models/Agent.cs b/Engine/Source/Programs/Horde/HordeServer/Models/Agent.cs
--- a/Engine/Source/Programs/Horde/HordeServer/Models/Agent.cs
+++ b/Engine/Source/Programs/Horde/HordeServer/Models/Agent.cs
@@ -295,23 +295,6 @@ namespace HordeServer.Models
 			}
 			return false;
 		}
-		/// <summary>
-		/// Determines if this is a remote action lease
-		/// </summary>
-		/// <returns>True if this is nn action lease</returns>
-		public bool IsActionLease()
-		{
-			if (Payload != null)
-			{
-				Any BasePayload = Any.Parser.ParseFrom(Payload);
-				if (BasePayload.Is(ActionTask.Descriptor))
-				{
-					return true;
-				}
-			}
-			return false;
-		}
-
 
 		/// <summary>
 		/// Gets user-readable payload information
@@ -350,14 +333,6 @@ namespace HordeServer.Models
 					Details["SoftwareId"] = UpgradeTask.SoftwareId;
 					Details["LogId"] = UpgradeTask.LogId;
 				}
-
-				ActionTask ActionTask;
-				if (BasePayload.TryUnpack(out ActionTask))
-				{
-					Details["Type"] = "Action";
-					Details["LogId"] = ActionTask.LogId;
-				}
-
 			}
 			return Details;
 		}
diff --git a/Engine/Source/Programs/Horde/HordeServer/Rpc/ActionCacheService.cs b/Engine/Source/Programs/Horde/HordeServer/Rpc/ActionCacheService.cs
deleted file mode 100644
diff --git a/Engine/Source/Programs/Horde/HordeServer/Rpc/ActionRpcService.cs b/Engine/Source/Programs/Horde/HordeServer/Rpc/ActionRpcService.cs
deleted file mode 100644
diff --git a/Engine/Source/Programs/Horde/HordeServer/Rpc/ByteStreamService.cs b/Engine/Source/Programs/Horde/HordeServer/Rpc/ByteStreamService.cs
deleted file mode 100644
diff --git a/Engine/Source/Programs/Horde/HordeServer/Rpc/CapabilitiesService.cs b/Engine/Source/Programs/Horde/HordeServer/Rpc/CapabilitiesService.cs
deleted file mode 100644
diff --git a/Engine/Source/Programs/Horde/HordeServer/Rpc/ContentStorageService.cs b/Engine/Source/Programs/Horde/HordeServer/Rpc/ContentStorageService.cs
deleted file mode 100644
diff --git a/Engine/Source/Programs/Horde/HordeServer/Rpc/ExecutionService.cs b/Engine/Source/Programs/Horde/HordeServer/Rpc/ExecutionService.cs
deleted file mode 100644
diff --git a/Engine/Source/Programs/Horde/HordeServer/Services/RpcService.cs b/Engine/Source/Programs/Horde/HordeServer/Services/RpcService.cs
--- a/Engine/Source/Programs/Horde/HordeServer/Services/RpcService.cs
+++ b/Engine/Source/Programs/Horde/HordeServer/Services/RpcService.cs
@@ -110,11 +110,6 @@ namespace HordeServer.Services
 		/// </summary>
 		ITestDataCollection TestData;
 
-		/// <summary>
-		/// The RPC action source
-		/// </summary>
-		ActionTaskSource ActionTaskSource;
-
 		/// <summary>
 		/// The conform task source
 		/// </summary>
@@ -141,10 +136,9 @@ namespace HordeServer.Services
 		/// <param name="LifetimeService">The application lifetime</param>
 		/// <param name="Graphs">Collection of graph documents</param>
 		/// <param name="TestData">Collection of testdata</param>
-		/// <param name="ActionTaskSource"></param>
 		/// <param name="ConformTaskSource"></param>
 		/// <param name="Logger">Log writer</param>
-		public RpcService(DatabaseService DatabaseService, AclService AclService, AgentService AgentService, StreamService StreamService, JobService JobService, AgentSoftwareService AgentSoftwareService, ArtifactService ArtifactService, ILogFileService LogFileService, CredentialService CredentialService, PoolService PoolService, LifetimeService LifetimeService, IGraphCollection Graphs, ITestDataCollection TestData, ActionTaskSource ActionTaskSource, ConformTaskSource ConformTaskSource, ILogger<RpcService> Logger)
+		public RpcService(DatabaseService DatabaseService, AclService AclService, AgentService AgentService, StreamService StreamService, JobService JobService, AgentSoftwareService AgentSoftwareService, ArtifactService ArtifactService, ILogFileService LogFileService, CredentialService CredentialService, PoolService PoolService, LifetimeService LifetimeService, IGraphCollection Graphs, ITestDataCollection TestData, ConformTaskSource ConformTaskSource, ILogger<RpcService> Logger)
 		{
 			this.DatabaseService = DatabaseService;
 			this.AclService = AclService;
@@ -159,7 +153,6 @@ namespace HordeServer.Services
 			this.LifetimeService = LifetimeService;
 			this.Graphs = Graphs;
 			this.TestData = TestData;
-			this.ActionTaskSource = ActionTaskSource;
 			this.ConformTaskSource = ConformTaskSource;
 			this.Logger = Logger;
 		}
diff --git a/Engine/Source/Programs/Horde/HordeServer/Startup.cs b/Engine/Source/Programs/Horde/HordeServer/Startup.cs
--- a/Engine/Source/Programs/Horde/HordeServer/Startup.cs
+++ b/Engine/Source/Programs/Horde/HordeServer/Startup.cs
@@ -59,7 +59,6 @@ using HordeServer.Collections.Impl;
 using HordeServer.Services.Impl;
 using Microsoft.Extensions.DependencyInjection.Extensions;
 using HordeServer.Authentication;
-using HordeServer.Rpc;
 using HordeServer.Tasks.Impl;
 using HordeServer.Tasks;
 using StatsdClient;
@@ -410,12 +409,6 @@ namespace HordeServer
 			Services.AddSingleton<TemplateService>();
 			Services.AddSingleton<UpgradeService>();
 
-			Services.AddSingleton<ActionCacheService>();
-			Services.AddSingleton<ByteStreamService>();
-			Services.AddSingleton<CapabilitiesService>();
-			Services.AddSingleton<ContentStorageService>();
-			Services.AddSingleton<ExecutionService>();
-
 			Services.AddSingleton<DeviceService>();
 
 			AWSOptions AwsOptions = Configuration.GetAWSOptions();
@@ -553,7 +546,6 @@ namespace HordeServer
 
 			// Task sources. Order of registration is important here; it dictates the order in which sources are served.
 			Services.AddSingleton<JobTaskSource>();
-			Services.AddSingleton<ActionTaskSource>();
 			Services.AddSingleton<ConformTaskSource>();
 			Services.AddSingleton<IComputeService, ComputeService>();
 
@@ -562,7 +554,6 @@ namespace HordeServer
 			Services.AddSingleton<ITaskSource, RestartTaskSource>();
 			Services.AddSingleton<ITaskSource, ConformTaskSource>(Provider => Provider.GetRequiredService<ConformTaskSource>());
 			Services.AddSingleton<ITaskSource, JobTaskSource>(Provider => Provider.GetRequiredService<JobTaskSource>());
-			Services.AddSingleton<ITaskSource, ActionTaskSource>(Provider => Provider.GetRequiredService<ActionTaskSource>());
 			Services.AddSingleton<ITaskSource>(Provider => new NewTaskSourceWrapper(Provider.GetRequiredService<IComputeService>()));
 
 			Services.AddHostedService(Provider => Provider.GetRequiredService<ConformTaskSource>());
@@ -814,17 +805,8 @@ namespace HordeServer
 				Endpoints.MapGrpcService<RpcService>();
 				
 				// Google Remote Execution API
-				Endpoints.MapGrpcService<ActionRpcService>();
-				Endpoints.MapGrpcService<ActionCacheService>();
-				Endpoints.MapGrpcService<ByteStreamService>();
-				Endpoints.MapGrpcService<CapabilitiesService>();
-				Endpoints.MapGrpcService<ContentStorageService>();
-				Endpoints.MapGrpcService<ExecutionService>();
-
 				Endpoints.MapGrpcService<ComputeRpcServer>();
-				Endpoints.MapGrpcService<BlobRpc>();
 				Endpoints.MapGrpcService<BlobStoreRpc>();
-				Endpoints.MapGrpcService<RefRpc>();
 				Endpoints.MapGrpcService<RefTableRpc>();
 
 				Endpoints.MapGrpcReflectionService();
diff --git a/Engine/Source/Programs/Horde/HordeServer/Storage/Collections/BlobRpc.cs b/Engine/Source/Programs/Horde/HordeServer/Storage/Collections/BlobRpc.cs
deleted file mode 100644
diff --git a/Engine/Source/Programs/Horde/HordeServer/Storage/Collections/RefRpc.cs b/Engine/Source/Programs/Horde/HordeServer/Storage/Collections/RefRpc.cs
deleted file mode 100644
diff --git a/Engine/Source/Programs/Horde/HordeServer/Tasks/Impl/ActionTaskSource.cs b/Engine/Source/Programs/Horde/HordeServer/Tasks/Impl/ActionTaskSource.cs
deleted file mode 100644
diff --git a/Engine/Source/Programs/Horde/HordeServer/Tasks/Impl/ClusteredActionTaskSource.cs b/Engine/Source/Programs/Horde/HordeServer/Tasks/Impl/ClusteredActionTaskSource.cs
deleted file mode 100644
diff --git a/Engine/Source/Programs/Horde/HordeServer/Utilities/RemoteExtensions.cs b/Engine/Source/Programs/Horde/HordeServer/Utilities/RemoteExtensions.cs
deleted file mode 100644
diff --git a/Engine/Source/Programs/Horde/HordeServerTests/RemoteExec/ActionTaskSourceTest.cs b/Engine/Source/Programs/Horde/HordeServerTests/RemoteExec/ActionTaskSourceTest.cs
deleted file mode 100644
diff --git a/Engine/Source/Programs/Horde/HordeServerTests/RemoteExec/ClusteredActionTaskSourceTest.cs b/Engine/Source/Programs/Horde/HordeServerTests/RemoteExec/ClusteredActionTaskSourceTest.cs
deleted file mode 100644
diff --git a/Engine/Source/Programs/Horde/HordeServerTests/TestSetup.cs b/Engine/Source/Programs/Horde/HordeServerTests/TestSetup.cs
--- a/Engine/Source/Programs/Horde/HordeServerTests/TestSetup.cs
+++ b/Engine/Source/Programs/Horde/HordeServerTests/TestSetup.cs
@@ -19,7 +19,6 @@ using HordeServer.Logs.Storage.Impl;
 using HordeServer.Models;
 using HordeServer.Notifications;
 using HordeServer.Notifications.Impl;
-using HordeServer.Rpc;
 using HordeServer.Services;
 using HordeServer.Services.Impl;
 using HordeServer.Storage;
@@ -198,15 +197,8 @@ namespace HordeServerTests
 			Services.AddSingleton<TemplateService>();
 			Services.AddSingleton<UpgradeService>();
 
-			Services.AddSingleton<ActionTaskSource>();
 			Services.AddSingleton<ConformTaskSource>();
 
-			Services.AddSingleton<ActionCacheService>();
-			Services.AddSingleton<ByteStreamService>();
-			Services.AddSingleton<CapabilitiesService>();
-			Services.AddSingleton<ContentStorageService>();
-			Services.AddSingleton<ExecutionService>();
-
 			Services.AddSingleton<IBlobCollection, BlobCollection>();
 			Services.AddSingleton<IObjectCollection, ObjectCollection>();
 			Services.AddSingleton<IRefCollection, RefCollection>();
diff --git a/Engine/Source/Programs/Shared/EpicGames.Horde.Common/EpicGames.Horde.Common.csproj b/Engine/Source/Programs/Shared/EpicGames.Horde.Common/EpicGames.Horde.Common.csproj
deleted file mode 100644
diff --git a/Engine/Source/Programs/Shared/EpicGames.Horde.Common/RemoteExecution/Json.cs b/Engine/Source/Programs/Shared/EpicGames.Horde.Common/RemoteExecution/Json.cs
deleted file mode 100644
diff --git a/Engine/Source/Programs/Shared/EpicGames.Horde.Common/RemoteExecution/StorageExtensions.cs b/Engine/Source/Programs/Shared/EpicGames.Horde.Common/RemoteExecution/StorageExtensions.cs
deleted file mode 100644
diff --git a/Engine/Source/Programs/Shared/EpicGames.Horde.Common/RemoteExecution/UploadList.cs b/Engine/Source/Programs/Shared/EpicGames.Horde.Common/RemoteExecution/UploadList.cs
deleted file mode 100644
diff --git a/Engine/Source/Programs/Shared/EpicGames.Horde.Common/Status.cs b/Engine/Source/Programs/Shared/EpicGames.Horde.Common/Status.cs
deleted file mode 100644
