Horde: Fix error fetching resource post-rename of tests project.

diff --git a/Engine/Source/Programs/Horde/HordeServer.Compute.Tests/LogIndexingTests.cs b/Engine/Source/Programs/Horde/HordeServer.Compute.Tests/LogIndexingTests.cs
--- a/Engine/Source/Programs/Horde/HordeServer.Compute.Tests/LogIndexingTests.cs
+++ b/Engine/Source/Programs/Horde/HordeServer.Compute.Tests/LogIndexingTests.cs
@@ -9,7 +9,7 @@ using HordeServer.Logs.Storage.Impl;
 using HordeServer.Models;
 using HordeServer.Services;
 using HordeServer.Utilities;
-using HordeServerTests.Properties;
+using HordeServer.Compute.Tests.Properties;
 using Microsoft.Extensions.Configuration;
 using Microsoft.Extensions.DependencyInjection;
 using Microsoft.Extensions.Logging;
