Horde: Map aborted color state to failure state
#rnx

diff --git a/Engine/Source/Programs/Horde/HordeDashboard/src/backend/Dashboard.ts b/Engine/Source/Programs/Horde/HordeDashboard/src/backend/Dashboard.ts
--- a/Engine/Source/Programs/Horde/HordeDashboard/src/backend/Dashboard.ts
+++ b/Engine/Source/Programs/Horde/HordeDashboard/src/backend/Dashboard.ts
@@ -293,23 +293,26 @@ export class Dashboard {
 
         const dark = this.darktheme;
 
-        return new Map<StatusColor, string>([
+        const colors = new Map<StatusColor, string>([
             [StatusColor.Success, dark ? "#3b7b0a" : "#52C705"],
             [StatusColor.Warnings, dark ? "#9a7b18" : "#EDC74A"],
             [StatusColor.Failure, dark ? "#882f19" : "#DE4522"],
             [StatusColor.Running, dark ? "#146579" : theme.palette.blueLight],
             [StatusColor.Waiting, dark ? "#474542" : "#A19F9D"],
             [StatusColor.Ready, dark ? "#474542" : "#A19F9D"],
-            [StatusColor.Skipped, dark ? "#63625c" : "#F3F2F1"],
-            [StatusColor.Aborted, dark ? "#63625c" : "#F3F2F1"],
+            [StatusColor.Skipped, dark ? "#63625c" : "#F3F2F1"],            
             [StatusColor.Unspecified, "#637087"]
         ]);
 
+        colors.set(StatusColor.Aborted, colors.get(StatusColor.Failure)!);
+
+        return colors;
+
     }
 
 
 
-    getStatusColors = () => {
+    getStatusColors = (): Map<StatusColor, string> => {
 
         const defaultStatusColors = this.getDefaultStatusColors();
 
@@ -318,7 +321,7 @@ export class Dashboard {
         const error = this.getPreference(DashboardPreference.ColorError);
         const running = this.getPreference(DashboardPreference.ColorRunning);
 
-        return new Map<StatusColor, string>([
+        const colors = new Map<StatusColor, string>([
             [StatusColor.Success, success ? success : defaultStatusColors.get(StatusColor.Success)!],
             [StatusColor.Warnings, warning ? warning : defaultStatusColors.get(StatusColor.Warnings)!],
             [StatusColor.Failure, error ? error : defaultStatusColors.get(StatusColor.Failure)!],
@@ -326,9 +329,12 @@ export class Dashboard {
             [StatusColor.Waiting, defaultStatusColors.get(StatusColor.Waiting)!],
             [StatusColor.Ready, defaultStatusColors.get(StatusColor.Ready)!],
             [StatusColor.Skipped, defaultStatusColors.get(StatusColor.Skipped)!],
-            [StatusColor.Aborted, defaultStatusColors.get(StatusColor.Aborted)!],
             [StatusColor.Unspecified, defaultStatusColors.get(StatusColor.Unspecified)!]
         ]);
+
+        colors.set(StatusColor.Aborted, colors.get(StatusColor.Failure)!);
+
+        return colors;
     }
 
 
