Horde: Update NuGet packages to match Jupiter versions.

diff --git a/Engine/Source/Programs/Horde/HordeServer/HordeServer.csproj b/Engine/Source/Programs/Horde/HordeServer/HordeServer.csproj
--- a/Engine/Source/Programs/Horde/HordeServer/HordeServer.csproj
+++ b/Engine/Source/Programs/Horde/HordeServer/HordeServer.csproj
@@ -1,4 +1,4 @@
-<Project Sdk="Microsoft.NET.Sdk.Web">
+﻿<Project Sdk="Microsoft.NET.Sdk.Web">
 
   <PropertyGroup>
     <TargetFramework>net5.0</TargetFramework>
@@ -10,18 +10,16 @@
   </PropertyGroup>
 
   <ItemGroup>
-    <PackageReference Include="AWSSDK.Core" Version="3.5.1.49" />
-    <PackageReference Include="AWSSDK.EC2" Version="3.5.20" />
-    <PackageReference Include="AWSSDK.Extensions.NETCore.Setup" Version="3.3.101" />
-    <PackageReference Include="AWSSDK.S3" Version="3.5.6.5" />
-    <PackageReference Include="AWSSDK.SecurityToken" Version="3.5.1.29" />
+    <PackageReference Include="AWSSDK.EC2" Version="3.7.0" />
+    <PackageReference Include="AWSSDK.Extensions.NETCore.Setup" Version="3.7.0.1" />
+    <PackageReference Include="AWSSDK.S3" Version="3.7.0.17" />
+    <PackageReference Include="AWSSDK.SecurityToken" Version="3.7.1.69" />
     <PackageReference Include="Datadog.Trace.OpenTracing" Version="1.18.0" />
     <PackageReference Include="DogStatsD-CSharp-Client" Version="6.0.0" />
-    <PackageReference Include="Google.Protobuf" Version="3.12.3" />
-    <PackageReference Include="Grpc.AspNetCore" Version="2.29.0" />
-    <PackageReference Include="Grpc.AspNetCore.Server.Reflection" Version="2.29.0" />
-    <PackageReference Include="Grpc.Net.Client" Version="2.29.0" />
-    <PackageReference Include="Grpc.Tools" Version="2.30.0">
+    <PackageReference Include="Grpc.AspNetCore" Version="2.37.0" />
+    <PackageReference Include="Grpc.AspNetCore.Server.Reflection" Version="2.37.0" />
+    <PackageReference Include="Grpc.Net.Client" Version="2.37.0" />
+    <PackageReference Include="Grpc.Tools" Version="2.37.0">
       <PrivateAssets>all</PrivateAssets>
       <IncludeAssets>runtime; build; native; contentfiles; analyzers; buildtransitive</IncludeAssets>
     </PackageReference>
@@ -35,8 +33,8 @@
       <PrivateAssets>all</PrivateAssets>
       <IncludeAssets>runtime; build; native; contentfiles; analyzers; buildtransitive</IncludeAssets>
     </PackageReference>
-    <PackageReference Include="Microsoft.Extensions.Caching.Memory" Version="3.1.6" />
-    <PackageReference Include="MongoDB.Driver" Version="2.10.4" />
+    <PackageReference Include="Microsoft.Extensions.Caching.Memory" Version="5.0.0" />
+    <PackageReference Include="MongoDB.Driver" Version="2.11.6" />
     <PackageReference Include="OpenTracing" Version="0.12.1" />
     <PackageReference Include="OpenTracing.Contrib.Grpc" Version="0.4.0" />
     <PackageReference Include="Polly" Version="7.2.1" />
diff --git a/Engine/Source/Programs/Horde/HordeServer/Utilities/MongoTracingCollection.cs b/Engine/Source/Programs/Horde/HordeServer/Utilities/MongoTracingCollection.cs
--- a/Engine/Source/Programs/Horde/HordeServer/Utilities/MongoTracingCollection.cs
+++ b/Engine/Source/Programs/Horde/HordeServer/Utilities/MongoTracingCollection.cs
@@ -56,11 +56,39 @@ namespace HordeServer.Utiltiies
 		}
 
 		/// <inheritdoc />
-		public async Task<IAsyncCursor<TResult>> AggregateAsync<TResult>(IClientSessionHandle Session, PipelineDefinition<T, TResult> pipeline, AggregateOptions Options = null!,
+		public async Task<IAsyncCursor<TResult>> AggregateAsync<TResult>(IClientSessionHandle Session, PipelineDefinition<T, TResult> Pipeline, AggregateOptions Options = null!,
 			CancellationToken  CancellationToken = new CancellationToken())
 		{
 			using IScope _ = GlobalTracer.Instance.BuildSpan(Prefix + "AggregateAsync").StartActive();
-			return await Collection.AggregateAsync(Session, pipeline, Options, CancellationToken);
+			return await Collection.AggregateAsync(Session, Pipeline, Options, CancellationToken);
+		}
+
+		/// <inheritdoc />
+		public void AggregateToCollection<TResult>(PipelineDefinition<T, TResult> Pipeline, AggregateOptions? Options = null, CancellationToken CancellationToken = default(CancellationToken))
+		{
+			using IScope _ = GlobalTracer.Instance.BuildSpan(Prefix + "AggregateToCollection").StartActive();
+			Collection.AggregateToCollection(Pipeline, Options, CancellationToken);
+		}
+
+		/// <inheritdoc />
+		public void AggregateToCollection<TResult>(IClientSessionHandle Session, PipelineDefinition<T, TResult> pipeline, AggregateOptions? options = null, CancellationToken cancellationToken = default(CancellationToken))
+		{
+			using IScope _ = GlobalTracer.Instance.BuildSpan(Prefix + "AggregateToCollection").StartActive();
+			Collection.AggregateToCollection(Session, pipeline, options, cancellationToken);
+		}
+
+		/// <inheritdoc />
+		public async Task AggregateToCollectionAsync<TResult>(PipelineDefinition<T, TResult> Pipeline, AggregateOptions? Options = null, CancellationToken CancellationToken = default(CancellationToken))
+		{
+			using IScope _ = GlobalTracer.Instance.BuildSpan(Prefix + "AggregateToCollectionAsync").StartActive();
+			await Collection.AggregateToCollectionAsync(Pipeline, Options, CancellationToken);
+		}
+
+		/// <inheritdoc />
+		public async Task AggregateToCollectionAsync<TResult>(IClientSessionHandle Session, PipelineDefinition<T, TResult> Pipeline, AggregateOptions? Options = null, CancellationToken CancellationToken = default(CancellationToken))
+		{
+			using IScope _ = GlobalTracer.Instance.BuildSpan(Prefix + "AggregateToCollectionAsync").StartActive();
+			await Collection.AggregateToCollectionAsync(Session, Pipeline, Options, CancellationToken);
 		}
 
 		/// <inheritdoc />
