Horde: Remove device platform references
#rnx

diff --git a/Engine/Source/Programs/Horde/HordeServer/Api/Devices.cs b/Engine/Source/Programs/Horde/HordeServer/Api/Devices.cs
--- a/Engine/Source/Programs/Horde/HordeServer/Api/Devices.cs
+++ b/Engine/Source/Programs/Horde/HordeServer/Api/Devices.cs
@@ -689,4 +689,30 @@ namespace HordeServer.Api
 		public string DeviceData { get; set; } = null!;
 	}
 
+	/// <summary>
+	/// Platform mapping update for V1 devices
+	/// </summary>
+	public class UpdatePlatformMapRequest
+	{
+		/// <summary>
+		/// Platform name => platform id
+		/// </summary>
+		[Required]
+		public Dictionary<string, string> PlatformMap { get; set; } = null!;
+
+		/// <summary>
+		/// Platform id => name
+		/// </summary>
+		[Required]
+		public Dictionary<string, string> PlatformReverseMap { get; set; } = null!;
+
+		/// <summary>
+		/// Perfspec v1 => model names
+		/// </summary>
+		[Required]
+		public Dictionary<string, string> PerfSpecHighMap { get; set; } = null!;
+
+	}
+
+
 }
diff --git a/Engine/Source/Programs/Horde/HordeServer/Controllers/DevicesController.cs b/Engine/Source/Programs/Horde/HordeServer/Controllers/DevicesController.cs
--- a/Engine/Source/Programs/Horde/HordeServer/Controllers/DevicesController.cs
+++ b/Engine/Source/Programs/Horde/HordeServer/Controllers/DevicesController.cs
@@ -28,6 +28,11 @@ namespace HordeServer.Controllers
 	public class DevicesController : ControllerBase
 	{
 
+		/// <summary>
+		/// The acl service singleton
+		/// </summary>
+		AclService AclService;
+
 		/// <summary>
 		/// The user collection instance
 		/// </summary>
@@ -47,11 +52,12 @@ namespace HordeServer.Controllers
 		/// <summary>
 		/// Constructor
 		/// </summary>
-		public DevicesController(DeviceService DeviceService, IUserCollection UserCollection, ILogger<DevicesController> Logger)
+		public DevicesController(DeviceService DeviceService, AclService AclService, IUserCollection UserCollection, ILogger<DevicesController> Logger)
 		{
-            this.UserCollection = UserCollection;
-            this.DeviceService = DeviceService;
-            this.Logger = Logger;
+			this.UserCollection = UserCollection;
+			this.DeviceService = DeviceService;
+			this.Logger = Logger;
+			this.AclService = AclService;
 		}
 
 		// DEVICES
@@ -634,34 +640,6 @@ namespace HordeServer.Controllers
 			High
 		};
 
-		static Dictionary<string, DevicePlatformId> PlatformMap = new Dictionary<string, DevicePlatformId>()
-			{
-				{ "PS4", new DevicePlatformId("ps4") },
-				{ "PS4-DevKit", new DevicePlatformId("ps4") },
-				{ "XboxOne", new DevicePlatformId("xboxone") },
-				{ "XboxOneGDK", new DevicePlatformId("xboxone") },
-				{ "XboxOne-DevKit", new DevicePlatformId("xboxone") },
-				{ "Switch", new DevicePlatformId("switch") },
-				{ "XSX", new DevicePlatformId("xsx") },
-				{ "PS5", new DevicePlatformId("ps5") }
-			};
-
-		static Dictionary<DevicePlatformId, string> PlatformReverseMap = new Dictionary<DevicePlatformId, string>()
-			{
-				{ new DevicePlatformId("ps4"),  "PS4-DevKit"},
-				{ new DevicePlatformId("ps5"),  "PS5"},
-				{ new DevicePlatformId("xboxone"),  "XboxOneGDK"},
-				{ new DevicePlatformId("xsx"),  "XSX"},
-				{ new DevicePlatformId("switch"),  "Switch"},
-			};
-
-		static Dictionary<DevicePlatformId, string> PerfSpecHighMap = new Dictionary<DevicePlatformId, string>()
-			{
-				{ new DevicePlatformId("ps4"),  "PS4Pro"},
-				{ new DevicePlatformId("xboxone"),  "XboxOneX"}
-			};
-
-
 		/// <summary>
 		/// Create a device reservation
 		/// </summary>
@@ -674,40 +652,40 @@ namespace HordeServer.Controllers
 			List<IDevicePool> Pools = await DeviceService.GetPoolsAsync();
 			List<IDevicePlatform> Platforms = await DeviceService.GetPlatformsAsync();
 
-            string Message;
-            string? PoolId = Request.PoolId;
+			string Message;
+			string? PoolId = Request.PoolId;
 
 			// @todo: Remove this once all streams are updated to provide jobid
-            string Details = "";
-            if ((String.IsNullOrEmpty(Request.JobId) || String.IsNullOrEmpty(Request.StepId)))
+			string Details = "";
+			if ((String.IsNullOrEmpty(Request.JobId) || String.IsNullOrEmpty(Request.StepId)))
 			{
- 				if (!string.IsNullOrEmpty(Request.ReservationDetails))
-				 {
+				if (!string.IsNullOrEmpty(Request.ReservationDetails))
+				{
 					Details = $" - {Request.ReservationDetails}";
-				 }
-				 else
-				 {
+				}
+				else
+				{
 					Details = $" - Host {Request.Hostname}";
-				 }							 
-            }
+				}
+			}
 
 			if (string.IsNullOrEmpty(PoolId))
 			{
-                Message = $"No pool specified, defaulting to UE4" + Details;
-                Logger.LogError(Message + $" JobId: {Request.JobId}, StepId: {Request.StepId}");
+				Message = $"No pool specified, defaulting to UE4" + Details;
+				Logger.LogError(Message + $" JobId: {Request.JobId}, StepId: {Request.StepId}");
 				await DeviceService.NotifyDeviceServiceAsync(Message, null, Request.JobId, Request.StepId);
 				PoolId = "ue4";
-                //return BadRequest(Message);
-            }			
+				//return BadRequest(Message);
+			}
 
 			DevicePoolId PoolIdValue = DevicePoolId.Sanitize(PoolId);
 			IDevicePool? Pool = Pools.FirstOrDefault(x => x.Id == PoolIdValue);
 			if (Pool == null)
 			{
-                Message = $"Unknown pool {PoolId} " + Details;
+				Message = $"Unknown pool {PoolId} " + Details;
 				Logger.LogError(Message);
 				await DeviceService.NotifyDeviceServiceAsync(Message, null, Request.JobId, Request.StepId);
-                return BadRequest(Message);
+				return BadRequest(Message);
 			}
 
 			List<DeviceRequestData> RequestedDevices = new List<DeviceRequestData>();
@@ -731,22 +709,24 @@ namespace HordeServer.Controllers
 
 				DevicePlatformId PlatformId;
 
-				if (!PlatformMap.TryGetValue(PlatformName, out PlatformId))
+				DevicePlatformMapV1 MapV1 = await DeviceService.GetPlatformMapV1();
+
+				if (!MapV1.PlatformMap.TryGetValue(PlatformName, out PlatformId))
 				{
-                    Message = $"Unknown platform {PlatformName}" + Details;
+					Message = $"Unknown platform {PlatformName}" + Details;
 					Logger.LogError(Message);
 					await DeviceService.NotifyDeviceServiceAsync(Message, null, Request.JobId, Request.StepId);
-                    return BadRequest(Message);
+					return BadRequest(Message);
 				}
 
 				IDevicePlatform? Platform = Platforms.FirstOrDefault(x => x.Id == PlatformId);
 
 				if (Platform == null)
 				{
-                    Message = $"Unknown platform {PlatformId}" + Details;
+					Message = $"Unknown platform {PlatformId}" + Details;
 					Logger.LogError(Message);
 					await DeviceService.NotifyDeviceServiceAsync(Message, null, Request.JobId, Request.StepId);
-                    return BadRequest(Message);
+					return BadRequest(Message);
 				}
 
 				List<string> IncludeModels = new List<string>();
@@ -755,7 +735,7 @@ namespace HordeServer.Controllers
 				if (PerfSpecName == "High")
 				{
 					string? Model = null;
-					if (PerfSpecHighMap.TryGetValue(PlatformId, out Model))
+					if (MapV1.PerfSpecHighMap.TryGetValue(PlatformId, out Model))
 					{
 						IncludeModels.Add(Model);
 					}
@@ -764,7 +744,7 @@ namespace HordeServer.Controllers
 				if (PerfSpecName == "Minimum" || PerfSpecName == "Recommended")
 				{
 					string? Model = null;
-					if (PerfSpecHighMap.TryGetValue(PlatformId, out Model))
+					if (MapV1.PerfSpecHighMap.TryGetValue(PlatformId, out Model))
 					{
 						ExcludeModels.Add(Model);
 					}
@@ -809,8 +789,8 @@ namespace HordeServer.Controllers
 
 			if (Reservation == null)
 			{
-                Logger.LogError($"Unable to find reservation for legacy guid {ReservationGuid}");
-                return BadRequest();
+				Logger.LogError($"Unable to find reservation for legacy guid {ReservationGuid}");
+				return BadRequest();
 			}
 
 			bool Updated = await DeviceService.TryUpdateReservationAsync(Reservation.Id);
@@ -821,9 +801,9 @@ namespace HordeServer.Controllers
 				return BadRequest();
 			}
 
-            List<IDevice> Devices = await DeviceService.GetDevicesAsync(Reservation.Devices);
+			List<IDevice> Devices = await DeviceService.GetDevicesAsync(Reservation.Devices);
 
-            GetLegacyReservationResponse Response = new GetLegacyReservationResponse();
+			GetLegacyReservationResponse Response = new GetLegacyReservationResponse();
 
 			Response.Guid = Reservation.LegacyGuid;
 			Response.DeviceNames = Reservation.Devices.Select(DeviceId => Devices.First(D => D.Id == DeviceId).Name).ToArray();
@@ -874,8 +854,10 @@ namespace HordeServer.Controllers
 				return BadRequest($"Unknown device {DeviceName}");
 			}
 
+			DevicePlatformMapV1 MapV1 = await DeviceService.GetPlatformMapV1();
+
 			string? Platform;
-			if (!PlatformReverseMap.TryGetValue(Device.PlatformId, out Platform))
+			if (!MapV1.PlatformReverseMap.TryGetValue(Device.PlatformId, out Platform))
 			{
 				return BadRequest($"Unable to map platform for {DeviceName} : {Device.PlatformId}");
 			}
@@ -889,19 +871,19 @@ namespace HordeServer.Controllers
 			Response.AvailableEndTime = "00:00:00";
 			Response.Enabled = true;
 			Response.DeviceData = "";
-						
+
 			Response.PerfSpec = "Minimum";
 
 			if (Device.ModelId != null)
 			{
 				string? Model = null;
-				if (PerfSpecHighMap.TryGetValue(Device.PlatformId, out Model))
+				if (MapV1.PerfSpecHighMap.TryGetValue(Device.PlatformId, out Model))
 				{
 					if (Model == Device.ModelId)
 					{
 						Response.PerfSpec = "High";
 					}
-					
+
 				}
 			}
 
@@ -923,20 +905,20 @@ namespace HordeServer.Controllers
 			{
 				Logger.LogError($"Device error reported for unknown device {DeviceName}");
 				return BadRequest($"Unknown device {DeviceName}");
-			}			
+			}
 
 			await DeviceService.UpdateDeviceAsync(Device.Id, NewProblem: true);
 
 			string Message = $"Device problem, {Device.Name} : {Device.PoolId.ToString().ToUpperInvariant()}";
 
-            IDeviceReservation? Reservation = await DeviceService.TryGetDeviceReservation(Device.Id);
+			IDeviceReservation? Reservation = await DeviceService.TryGetDeviceReservation(Device.Id);
 
-            string? JobId = null;
-            string? StepId = null;
-            if (Reservation != null) 
+			string? JobId = null;
+			string? StepId = null;
+			if (Reservation != null)
 			{
 				JobId = !string.IsNullOrEmpty(Reservation.JobId) ? Reservation.JobId : null;
-                StepId = !string.IsNullOrEmpty(Reservation.StepId) ? Reservation.StepId : null;
+				StepId = !string.IsNullOrEmpty(Reservation.StepId) ? Reservation.StepId : null;
 
 				if ((JobId == null || StepId == null))
 				{
@@ -947,16 +929,49 @@ namespace HordeServer.Controllers
 					else
 					{
 						Message += $" - Host {Reservation.Hostname}";
-					}                    
-                }
-            }
+					}
+				}
+			}
 
-            Logger.LogError(Message);
+			Logger.LogError(Message);
 
-            await DeviceService.NotifyDeviceServiceAsync(Message, Device.Id, JobId, StepId);
+			await DeviceService.NotifyDeviceServiceAsync(Message, Device.Id, JobId, StepId);
 
 			return Ok();
 
 		}
+
+		/// <summary>
+		/// Updates the platform map for v1 requests
+		/// </summary>
+		[HttpPut]
+		[Route("/api/v1/devices/platformmap")]
+		public async Task<ActionResult> UpdatePlatformMapV1([FromBody] UpdatePlatformMapRequest Request)
+		{
+			if (!await AclService.AuthorizeAsync(AclAction.AdminWrite, User))
+			{
+				return Forbid();
+			}
+
+			bool Result = false;
+
+			try
+			{
+				Result = await DeviceService.UpdatePlatformMapAsync(Request);
+			}
+			catch (Exception Ex)
+			{
+				Logger.LogError(Ex, "Error updating device platform map {Message}", Ex.Message);
+				throw;
+			}
+
+			if (!Result)
+			{
+				Logger.LogError("Unable to update device platform mapping");
+				return BadRequest();
+			}
+
+			return Ok();
+		}
 	}
 }
\ No newline at end of file
diff --git a/Engine/Source/Programs/Horde/HordeServer/Services/DeviceService.cs b/Engine/Source/Programs/Horde/HordeServer/Services/DeviceService.cs
--- a/Engine/Source/Programs/Horde/HordeServer/Services/DeviceService.cs
+++ b/Engine/Source/Programs/Horde/HordeServer/Services/DeviceService.cs
@@ -16,6 +16,7 @@ using System.Security.Claims;
 using System.Diagnostics.CodeAnalysis;
 using HordeServer.Notifications;
 using System.Linq;
+using HordeServer.Api;
 
 namespace HordeServer.Services
 {
@@ -23,6 +24,30 @@ namespace HordeServer.Services
 	using DevicePlatformId = StringId<IDevicePlatform>;
 	using DevicePoolId = StringId<IDevicePool>;
 
+	/// <summary>
+	/// Platform map required by V1 API
+	/// </summary>
+	[SingletonDocument("6165a2e26fd5f104e31e6862")]
+	public class DevicePlatformMapV1 : SingletonBase
+	{
+		/// <summary>
+		/// Platform V1 => Platform Id
+		/// </summary>
+		public Dictionary<string, DevicePlatformId> PlatformMap { get; set; } = new Dictionary<string, DevicePlatformId>();
+
+		/// <summary>
+		/// Platform Id => Platform V1
+		/// </summary>
+		public Dictionary<DevicePlatformId, string> PlatformReverseMap { get; set; } = new Dictionary<DevicePlatformId, string>();
+
+		/// <summary>
+		/// Perfspec V1 => Model
+		/// </summary>
+		public Dictionary<DevicePlatformId, string> PerfSpecHighMap { get; set; } = new Dictionary<DevicePlatformId, string>();
+
+	}
+
+
 	/// <summary>
 	/// Device management service
 	/// </summary>
@@ -59,10 +84,15 @@ namespace HordeServer.Services
 		/// </summary>
 		IDeviceCollection Devices;
 
+		/// <summary>
+		/// Platform map V1 singleton
+		/// </summary>
+		ISingletonDocument<DevicePlatformMapV1> PlatformMapSingleton;
+
 		/// <summary>
 		/// Device service constructor
 		/// </summary>
-		public DeviceService(IDeviceCollection Devices, JobService JobService, StreamService StreamService, AclService AclService, INotificationService NotificationService, ILogger<DeviceService> Logger)
+		public DeviceService(IDeviceCollection Devices, ISingletonDocument<DevicePlatformMapV1> PlatformMapSingleton, JobService JobService, StreamService StreamService, AclService AclService, INotificationService NotificationService, ILogger<DeviceService> Logger)
 			: base(TimeSpan.FromMinutes(1.0), Logger)
 		{
 			this.Devices = Devices;
@@ -71,6 +101,9 @@ namespace HordeServer.Services
             this.AclService = AclService;
             this.NotificationService = NotificationService;
             this.Logger = Logger;
+			
+			this.PlatformMapSingleton = PlatformMapSingleton;
+
 		}
 
 		/// <summary>
@@ -333,5 +366,76 @@ namespace HordeServer.Services
 			return Task.FromResult(User.IsInRole("Internal-Employees"));			
 		}
 
+		/// <summary>
+		/// Get Platform mappings for V1 API
+		/// </summary>		
+		public async Task<DevicePlatformMapV1> GetPlatformMapV1()
+		{
+			return await PlatformMapSingleton.GetAsync();
+		}
+
+		/// <summary>
+		/// Updates the platform mapping information required for the V1 api
+		/// </summary>				
+		public async Task<bool> UpdatePlatformMapAsync(UpdatePlatformMapRequest Request)
+		{
+
+			List<IDevicePlatform> Platforms = await GetPlatformsAsync();
+
+			// Update the platform map
+			for (int i = 0; i < 10; i++)
+			{
+				DevicePlatformMapV1 Instance = await PlatformMapSingleton.GetAsync();
+
+				Instance.PlatformMap = new Dictionary<string, DevicePlatformId>();
+				Instance.PlatformReverseMap = new Dictionary<DevicePlatformId, string>();
+				Instance.PerfSpecHighMap = new Dictionary<DevicePlatformId, string>();
+
+				foreach(KeyValuePair<string, string> Entry in Request.PlatformMap)
+				{
+					IDevicePlatform? Platform = Platforms.FirstOrDefault(P => P.Id == new DevicePlatformId(Entry.Value));
+					if (Platform == null)
+					{
+						throw new Exception($"Unknowm platform in map {Entry.Key} : {Entry.Value}");
+					}
+
+					Instance.PlatformMap.Add(Entry.Key, Platform.Id);
+				}
+
+				foreach (KeyValuePair<string, string> Entry in Request.PlatformReverseMap)
+				{
+					IDevicePlatform? Platform = Platforms.FirstOrDefault(P => P.Id == new DevicePlatformId(Entry.Key));
+					if (Platform == null)
+					{
+						throw new Exception($"Unknowm platform in reverse map {Entry.Key} : {Entry.Value}");
+					}
+
+					Instance.PlatformReverseMap.Add(Platform.Id, Entry.Value);
+				}
+
+				foreach (KeyValuePair<string, string> Entry in Request.PerfSpecHighMap)
+				{
+					IDevicePlatform? Platform = Platforms.FirstOrDefault(P => P.Id == new DevicePlatformId(Entry.Key));
+					if (Platform == null)
+					{
+						throw new Exception($"Unknowm platform in spec map {Entry.Key} : {Entry.Value}");
+					}
+
+					Instance.PerfSpecHighMap.Add(Platform.Id, Entry.Value);
+				}
+
+
+				if (await PlatformMapSingleton.TryUpdateAsync(Instance))
+				{
+					return true;
+				}
+
+				Thread.Sleep(1000);
+			}
+
+			return false;
+
+		}
+
 	}
 }
\ No newline at end of file
