Horde: Return copies of stream documents after updating them. Fixes issue with schedules triggering more times than expected, due to documents in memory being modified in place during a schedule update.

diff --git a/Engine/Source/Programs/Horde/HordeServer/Collections/IStreamCollection.cs b/Engine/Source/Programs/Horde/HordeServer/Collections/IStreamCollection.cs
--- a/Engine/Source/Programs/Horde/HordeServer/Collections/IStreamCollection.cs
+++ b/Engine/Source/Programs/Horde/HordeServer/Collections/IStreamCollection.cs
@@ -69,8 +69,8 @@ namespace HordeServer.Collections
 		/// <param name="Stream">The stream to update</param>
 		/// <param name="NewPausedUntil">The new datetime for pausing builds</param>
 		/// <param name="NewPauseComment">The reason for pausing</param>
-		/// <returns>Async task object</returns>
-		Task<bool> TryUpdatePauseStateAsync(IStream Stream, DateTime? NewPausedUntil, string? NewPauseComment);
+		/// <returns>The updated stream if successful, null otherwise</returns>
+		Task<IStream?> TryUpdatePauseStateAsync(IStream Stream, DateTime? NewPausedUntil, string? NewPauseComment);
 
 		/// <summary>
 		/// Attempts to update the last trigger time for a schedule
@@ -80,8 +80,8 @@ namespace HordeServer.Collections
 		/// <param name="LastTriggerTime">New last trigger time for the schedule</param>
 		/// <param name="LastTriggerChange">New last trigger changelist for the schedule</param>
 		/// <param name="NewActiveJobs">New list of active jobs</param>
-		/// <returns>True if the stream was updated</returns>
-		Task<bool> TryUpdateScheduleTriggerAsync(IStream Stream, TemplateRefId TemplateRefId, DateTimeOffset? LastTriggerTime, int? LastTriggerChange, List<ObjectId> NewActiveJobs);
+		/// <returns>The updated stream if successful, null otherwise</returns>
+		Task<IStream?> TryUpdateScheduleTriggerAsync(IStream Stream, TemplateRefId TemplateRefId, DateTimeOffset? LastTriggerTime, int? LastTriggerChange, List<ObjectId> NewActiveJobs);
 
 		/// <summary>
 		/// Delete a stream
diff --git a/Engine/Source/Programs/Horde/HordeServer/Collections/Impl/StreamCollection.cs b/Engine/Source/Programs/Horde/HordeServer/Collections/Impl/StreamCollection.cs
--- a/Engine/Source/Programs/Horde/HordeServer/Collections/Impl/StreamCollection.cs
+++ b/Engine/Source/Programs/Horde/HordeServer/Collections/Impl/StreamCollection.cs
@@ -127,13 +127,13 @@ namespace HordeServer.Collections.Impl
 			List<StreamTab> Tabs = Config.Tabs.ConvertAll(x => StreamTab.FromRequest(x));
 			Dictionary<TemplateRefId, TemplateRef> TemplateRefs = await CreateTemplateRefsAsync(Config.Templates, Stream, TemplateCollection);
 
-			Dictionary<string, AgentType>? AgentTypes = null;
+			Dictionary<string, AgentType> AgentTypes = new Dictionary<string, AgentType>();
 			if (Config.AgentTypes != null)
 			{
 				AgentTypes = Config.AgentTypes.Where(x => x.Value != null).ToDictionary(x => x.Key, x => new AgentType(x.Value!));
 			}
 
-			Dictionary<string, WorkspaceType>? WorkspaceTypes = null;
+			Dictionary<string, WorkspaceType> WorkspaceTypes = new Dictionary<string, WorkspaceType>();
 			if (Config.WorkspaceTypes != null)
 			{
 				WorkspaceTypes = Config.WorkspaceTypes.Where(x => x.Value != null).ToDictionary(x => x.Key, x => new WorkspaceType(x.Value!));
@@ -145,6 +145,8 @@ namespace HordeServer.Collections.Impl
 				DefaultPreflight = new DefaultPreflight(new TemplateRefId(Config.DefaultPreflightTemplate), null);
 			}
 
+			Validate(Id, DefaultPreflight, TemplateRefs, Tabs, AgentTypes, WorkspaceTypes);
+
 			Acl? Acl = Acl.Merge(new Acl(), Config.Acl);
 			if (Stream == null)
 			{
@@ -218,7 +220,7 @@ namespace HordeServer.Collections.Impl
 		}
 
 		/// <inheritdoc/>
-		public async Task<IStream?> TryCreateAsync(StreamId Id, ProjectId ProjectId, string ConfigPath, string ConfigRevision, StreamConfig Config, DefaultPreflight? DefaultPreflight, List<StreamTab>? Tabs, Dictionary<string, AgentType>? AgentTypes, Dictionary<string, WorkspaceType>? WorkspaceTypes, Dictionary<TemplateRefId, TemplateRef>? TemplateRefs, Acl? Acl)
+		async Task<IStream?> TryCreateAsync(StreamId Id, ProjectId ProjectId, string ConfigPath, string ConfigRevision, StreamConfig Config, DefaultPreflight? DefaultPreflight, List<StreamTab> Tabs, Dictionary<string, AgentType> AgentTypes, Dictionary<string, WorkspaceType> WorkspaceTypes, Dictionary<TemplateRefId, TemplateRef> TemplateRefs, Acl? Acl)
 		{
 			StreamDocument NewStream = new StreamDocument(Id, Config.Name, ProjectId);
 			NewStream.ClusterName = Config.ClusterName;
@@ -228,32 +230,12 @@ namespace HordeServer.Collections.Impl
 			NewStream.NotificationChannel = Config.NotificationChannel;
 			NewStream.NotificationChannelFilter = Config.NotificationChannelFilter;
 			NewStream.TriageChannel = Config.TriageChannel;
-
-			if (DefaultPreflight != null)
-			{
-				NewStream.DefaultPreflight = DefaultPreflight;
-			}
-			if (Tabs != null)
-			{
-				NewStream.Tabs = Tabs;
-			}
-			if (AgentTypes != null)
-			{
-				NewStream.AgentTypes = new Dictionary<string, AgentType>(AgentTypes, NewStream.AgentTypes.Comparer);
-			}
-			if (WorkspaceTypes != null)
-			{
-				NewStream.WorkspaceTypes = new Dictionary<string, WorkspaceType>(WorkspaceTypes, NewStream.WorkspaceTypes.Comparer);
-			}
-			if (TemplateRefs != null)
-			{
-				NewStream.Templates = TemplateRefs;
-			}
-			if (Acl != null)
-			{
-				NewStream.Acl = Acl;
-			}
-			NewStream.Validate();
+			NewStream.DefaultPreflight = DefaultPreflight;
+			NewStream.Tabs = Tabs;
+			NewStream.AgentTypes = AgentTypes;
+			NewStream.WorkspaceTypes = WorkspaceTypes;
+			NewStream.Templates = TemplateRefs;
+			NewStream.Acl = Acl;
 
 			try
 			{
@@ -274,71 +256,33 @@ namespace HordeServer.Collections.Impl
 		}
 
 		/// <inheritdoc/>
-		public async Task<IStream?> TryReplaceAsync(IStream StreamInterface, ProjectId ProjectId, string ConfigPath, string ConfigRevision, StreamConfig Config, DefaultPreflight? DefaultPreflight, List<StreamTab>? Tabs, Dictionary<string, AgentType>? AgentTypes, Dictionary<string, WorkspaceType>? WorkspaceTypes, Dictionary<TemplateRefId, TemplateRef>? TemplateRefs, Acl? Acl)
+		async Task<IStream?> TryReplaceAsync(IStream StreamInterface, ProjectId ProjectId, string ConfigPath, string ConfigRevision, StreamConfig Config, DefaultPreflight? DefaultPreflight, List<StreamTab> Tabs, Dictionary<string, AgentType>? AgentTypes, Dictionary<string, WorkspaceType>? WorkspaceTypes, Dictionary<TemplateRefId, TemplateRef>? TemplateRefs, Acl? Acl)
 		{
+			int Order = Config.Order ?? StreamDocument.DefaultOrder;
+
 			StreamDocument Stream = (StreamDocument)StreamInterface;
 
 			UpdateDefinitionBuilder<StreamDocument> UpdateBuilder = Builders<StreamDocument>.Update;
 
 			List<UpdateDefinition<StreamDocument>> Updates = new List<UpdateDefinition<StreamDocument>>();
-
-			Stream.Name = Config.Name;
-			Updates.Add(UpdateBuilder.Set(x => x.Name, Stream.Name));
-
-			Stream.ProjectId = ProjectId;
-			Updates.Add(UpdateBuilder.Set(x => x.ProjectId, Stream.ProjectId));
-
-			Stream.ClusterName = Config.ClusterName;
-			Updates.Add(UpdateBuilder.Set(x => x.ClusterName, Stream.ClusterName));
-
-			Stream.ConfigPath = ConfigPath;
-			Updates.Add(UpdateBuilder.Set(x => x.ConfigPath, Stream.ConfigPath));
-
-			Stream.ConfigRevision = ConfigRevision;
-			Updates.Add(UpdateBuilder.Set(x => x.ConfigRevision, Stream.ConfigRevision));
-
-			Stream.Order = Config.Order ?? StreamDocument.DefaultOrder;
-			Updates.Add(UpdateBuilder.Set(x => x.Order, Stream.Order));
-
-			Stream.NotificationChannel = Config.NotificationChannel;
-			Updates.Add(UpdateBuilder.Set(x => x.NotificationChannel, Stream.NotificationChannel));
-
-			Stream.NotificationChannelFilter = Config.NotificationChannelFilter;
-			Updates.Add(UpdateBuilder.Set(x => x.NotificationChannelFilter, Stream.NotificationChannelFilter));
-
-			Stream.TriageChannel = Config.TriageChannel;
-			Updates.Add(UpdateBuilder.Set(x => x.TriageChannel, Stream.TriageChannel));
-
-			Stream.DefaultPreflight = DefaultPreflight;
-			Updates.Add(UpdateBuilder.Set(x => x.DefaultPreflight, Stream.DefaultPreflight));
-			
-			Stream.Tabs = Tabs ?? new List<StreamTab>();
-			Updates.Add(UpdateBuilder.Set(x => x.Tabs, Stream.Tabs));
-
-			Stream.AgentTypes = AgentTypes ?? new Dictionary<string, AgentType>();
-			Updates.Add(UpdateBuilder.Set(x => x.AgentTypes, Stream.AgentTypes));
-
-			Stream.WorkspaceTypes = WorkspaceTypes ?? new Dictionary<string, WorkspaceType>();
-			Updates.Add(UpdateBuilder.Set(x => x.WorkspaceTypes, Stream.WorkspaceTypes));
-
-			Stream.Templates = TemplateRefs ?? new Dictionary<TemplateRefId, TemplateRef>();
-			Updates.Add(UpdateBuilder.Set(x => x.Templates, Stream.Templates));
-
-			Stream.Acl = Acl;
+			Updates.Add(UpdateBuilder.Set(x => x.Name, Config.Name));
+			Updates.Add(UpdateBuilder.Set(x => x.ProjectId, ProjectId));
+			Updates.Add(UpdateBuilder.Set(x => x.ClusterName, Config.ClusterName));
+			Updates.Add(UpdateBuilder.Set(x => x.ConfigPath, ConfigPath));
+			Updates.Add(UpdateBuilder.Set(x => x.ConfigRevision, ConfigRevision));
+			Updates.Add(UpdateBuilder.Set(x => x.Order, Order));
+			Updates.Add(UpdateBuilder.Set(x => x.NotificationChannel, Config.NotificationChannel));
+			Updates.Add(UpdateBuilder.Set(x => x.NotificationChannelFilter, Config.NotificationChannelFilter));
+			Updates.Add(UpdateBuilder.Set(x => x.TriageChannel, Config.TriageChannel));
+			Updates.Add(UpdateBuilder.Set(x => x.DefaultPreflight, DefaultPreflight));
+			Updates.Add(UpdateBuilder.Set(x => x.Tabs, Tabs ?? new List<StreamTab>()));
+			Updates.Add(UpdateBuilder.Set(x => x.AgentTypes, AgentTypes ?? new Dictionary<string, AgentType>()));
+			Updates.Add(UpdateBuilder.Set(x => x.WorkspaceTypes, WorkspaceTypes ?? new Dictionary<string, WorkspaceType>()));
+			Updates.Add(UpdateBuilder.Set(x => x.Templates, TemplateRefs ?? new Dictionary<TemplateRefId, TemplateRef>()));
 			Updates.Add(UpdateBuilder.SetOrUnsetNullRef(x => x.Acl, Acl));
-
-			Stream.Deleted = false;
 			Updates.Add(UpdateBuilder.Unset(x => x.Deleted));
 
-			Stream.Validate();
-			if(await TryUpdateStreamAsync(Stream, UpdateBuilder.Combine(Updates)))
-			{
-				return Stream;
-			}
-			else
-			{
-				return null;
-			}
+			return await TryUpdateStreamAsync(Stream, UpdateBuilder.Combine(Updates));
 		}
 
 		/// <inheritdoc/>
@@ -369,7 +313,7 @@ namespace HordeServer.Collections.Impl
 		}
 
 		/// <inheritdoc/>
-		public Task<bool> TryUpdatePauseStateAsync(IStream StreamInterface, DateTime? NewPausedUntil, string? NewPauseComment)
+		public async Task<IStream?> TryUpdatePauseStateAsync(IStream StreamInterface, DateTime? NewPausedUntil, string? NewPauseComment)
 		{
 			StreamDocument Stream = (StreamDocument)StreamInterface;
 
@@ -381,11 +325,11 @@ namespace HordeServer.Collections.Impl
 			Updates.Add(UpdateBuilder.Set(x => x.PausedUntil, NewPausedUntil));
 			Updates.Add(UpdateBuilder.Set(x => x.PauseComment, NewPauseComment));
 
-			return TryUpdateStreamAsync(Stream, UpdateBuilder.Combine(Updates));
+			return await TryUpdateStreamAsync(Stream, UpdateBuilder.Combine(Updates));
 		}
 
 		/// <inheritdoc/>
-		public async Task<bool> TryUpdateScheduleTriggerAsync(IStream StreamInterface, TemplateRefId TemplateRefId, DateTimeOffset? LastTriggerTime, int? LastTriggerChange, List<ObjectId> NewActiveJobs)
+		public async Task<IStream?> TryUpdateScheduleTriggerAsync(IStream StreamInterface, TemplateRefId TemplateRefId, DateTimeOffset? LastTriggerTime, int? LastTriggerChange, List<ObjectId> NewActiveJobs)
 		{
 			StreamDocument Stream = (StreamDocument)StreamInterface;
 			Schedule Schedule = Stream.Templates[TemplateRefId].Schedule!;
@@ -411,7 +355,7 @@ namespace HordeServer.Collections.Impl
 				Schedule.ActiveJobs = NewActiveJobs;
 			}
 
-			return Updates.Count == 0 || await TryUpdateStreamAsync(Stream, Builders<StreamDocument>.Update.Combine(Updates));
+			return (Updates.Count == 0)? StreamInterface : await TryUpdateStreamAsync(Stream, Builders<StreamDocument>.Update.Combine(Updates));
 		}
 
 		/// <summary>
@@ -419,15 +363,14 @@ namespace HordeServer.Collections.Impl
 		/// </summary>
 		/// <param name="Stream">The stream to update</param>
 		/// <param name="Update">The update definition</param>
-		/// <returns>True if the stream was updated, false otherwise</returns>
-		private async Task<bool> TryUpdateStreamAsync(StreamDocument Stream, UpdateDefinition<StreamDocument> Update)
+		/// <returns>The updated document, or null the update failed</returns>
+		private async Task<StreamDocument?> TryUpdateStreamAsync(StreamDocument Stream, UpdateDefinition<StreamDocument> Update)
 		{
-			Stream.Validate();
+			FilterDefinition<StreamDocument> Filter = Builders<StreamDocument>.Filter.Expr(x => x.Id == Stream.Id && x.UpdateIndex == Stream.UpdateIndex);
+			Update = Update.Set(x => x.UpdateIndex, Stream.UpdateIndex + 1);
 
-			int InitialUpdateIndex = Stream.UpdateIndex++;
-
-			UpdateResult Result = await Streams.UpdateOneAsync(x => x.Id == Stream.Id && x.UpdateIndex == InitialUpdateIndex, Update.Set(x => x.UpdateIndex, Stream.UpdateIndex));
-			return Result.ModifiedCount > 0;
+			FindOneAndUpdateOptions<StreamDocument> Options = new FindOneAndUpdateOptions<StreamDocument> { ReturnDocument = ReturnDocument.After };
+			return await Streams.FindOneAndUpdateAsync(Filter, Update, Options);
 		}
 
 		/// <inheritdoc/>
@@ -435,5 +378,45 @@ namespace HordeServer.Collections.Impl
 		{
 			await Streams.UpdateOneAsync<StreamDocument>(x => x.Id == StreamId, Builders<StreamDocument>.Update.Set(x => x.Deleted, true).Inc(x => x.UpdateIndex, 1));
 		}
+
+
+		/// <summary>
+		/// Checks the stream definition for consistency
+		/// </summary>
+		public static void Validate(StreamId StreamId, DefaultPreflight? DefaultPreflight, IReadOnlyDictionary<TemplateRefId, TemplateRef> Templates, IReadOnlyList<StreamTab> Tabs, IReadOnlyDictionary<string, AgentType> AgentTypes, IReadOnlyDictionary<string, WorkspaceType> WorkspaceTypes)
+		{
+			// Check the default preflight template is valid
+			if (DefaultPreflight != null)
+			{
+				if (DefaultPreflight.TemplateRefId != null && !Templates.ContainsKey(DefaultPreflight.TemplateRefId.Value))
+				{
+					throw new InvalidStreamException($"Default preflight template was listed as '{DefaultPreflight.TemplateRefId.Value}', but no template was found by that name");
+				}
+			}
+
+			// Check that all the templates are referenced by a tab
+			HashSet<TemplateRefId> RemainingTemplates = new HashSet<TemplateRefId>(Templates.Keys);
+			foreach (JobsTab JobsTab in Tabs.OfType<JobsTab>())
+			{
+				if (JobsTab.Templates != null)
+				{
+					RemainingTemplates.ExceptWith(JobsTab.Templates);
+				}
+			}
+			if (RemainingTemplates.Count > 0)
+			{
+				throw new InvalidStreamException(String.Join("\n", RemainingTemplates.Select(x => $"Template '{x}' is not listed on any tab for {StreamId}")));
+			}
+
+			// Check that all the agent types reference valid workspace names
+			foreach (KeyValuePair<string, AgentType> Pair in AgentTypes)
+			{
+				string? WorkspaceTypeName = Pair.Value.Workspace;
+				if (WorkspaceTypeName != null && !WorkspaceTypes.ContainsKey(WorkspaceTypeName))
+				{
+					throw new InvalidStreamException($"Agent type '{Pair.Key}' references undefined workspace type '{Pair.Value.Workspace}' in {StreamId}");
+				}
+			}
+		}
 	}
 }
diff --git a/Engine/Source/Programs/Horde/HordeServer/Models/Stream.cs b/Engine/Source/Programs/Horde/HordeServer/Models/Stream.cs
--- a/Engine/Source/Programs/Horde/HordeServer/Models/Stream.cs
+++ b/Engine/Source/Programs/Horde/HordeServer/Models/Stream.cs
@@ -645,46 +645,6 @@ namespace HordeServer.Models
 			return Regex.Replace(Stream.GetEscapedName(), @"^(\+\+[^+]*).*$", "$1");
 		}
 
-		/// <summary>
-		/// Checks the stream definition for consistency
-		/// </summary>
-		/// <param name="Stream">The stream object</param>
-		public static void Validate(this IStream Stream)
-		{
-			// Check the default preflight template is valid
-			if (Stream.DefaultPreflight != null)
-			{
-				if (Stream.DefaultPreflight.TemplateRefId != null && !Stream.Templates.ContainsKey(Stream.DefaultPreflight.TemplateRefId.Value))
-				{
-					throw new InvalidStreamException($"Default preflight template was listed as '{Stream.DefaultPreflight.TemplateRefId.Value}', but no template was found by that name");
-				}
-			}
-
-			// Check that all the templates are referenced by a tab
-			HashSet<TemplateRefId> RemainingTemplates = new HashSet<TemplateRefId>(Stream.Templates.Keys);
-			foreach(JobsTab JobsTab in Stream.Tabs.OfType<JobsTab>())
-			{
-				if (JobsTab.Templates != null)
-				{
-					RemainingTemplates.ExceptWith(JobsTab.Templates);
-				}
-			}
-			if(RemainingTemplates.Count > 0)
-			{
-				throw new InvalidStreamException(String.Join("\n", RemainingTemplates.Select(x => $"Template '{x}' is not listed on any tab for {Stream.Id}")));
-			}
-
-			// Check that all the agent types reference valid workspace names
-			foreach (KeyValuePair<string, AgentType> Pair in Stream.AgentTypes)
-			{
-				string? WorkspaceTypeName = Pair.Value.Workspace;
-				if (WorkspaceTypeName != null && !Stream.WorkspaceTypes.ContainsKey(WorkspaceTypeName))
-				{
-					throw new InvalidStreamException($"Agent type '{Pair.Key}' references undefined workspace type '{Pair.Value.Workspace}' in {Stream.Id}");
-				}
-			}
-		}
-
 		/// <summary>
 		/// Converts to a public response object
 		/// </summary>
diff --git a/Engine/Source/Programs/Horde/HordeServer/Services/StreamService.cs b/Engine/Source/Programs/Horde/HordeServer/Services/StreamService.cs
--- a/Engine/Source/Programs/Horde/HordeServer/Services/StreamService.cs
+++ b/Engine/Source/Programs/Horde/HordeServer/Services/StreamService.cs
@@ -99,12 +99,13 @@ namespace HordeServer.Services
 		{
 			for (; Stream != null; Stream = await GetStreamAsync(Stream.Id))
 			{
-				if (await Streams.TryUpdatePauseStateAsync(Stream, NewPausedUntil, NewPauseComment))
+				IStream? NewStream = await Streams.TryUpdatePauseStateAsync(Stream, NewPausedUntil, NewPauseComment);
+				if (NewStream != null)
 				{
-					break;
+					return NewStream;
 				}
 			}
-			return Stream;
+			return null;
 		}
 
 		/// <summary>
@@ -133,12 +134,15 @@ namespace HordeServer.Services
 				}
 
 				List<ObjectId> NewActiveJobs = TemplateRef.Schedule.ActiveJobs.Except(RemoveJobs).Union(AddJobs).ToList();
-				if (await Streams.TryUpdateScheduleTriggerAsync(NewStream, TemplateRefId, LastTriggerTime, LastTriggerChange, NewActiveJobs))
+
+				NewStream = await Streams.TryUpdateScheduleTriggerAsync(NewStream, TemplateRefId, LastTriggerTime, LastTriggerChange, NewActiveJobs);
+
+				if (NewStream != null)
 				{
 					return NewStream;
 				}
 
-				NewStream = await Streams.GetAsync(NewStream.Id);
+				NewStream = await Streams.GetAsync(Stream.Id);
 			}
 			return null;
 		}
