Horde: Wrap the existing Datadog tracer rather than creating a new one, and add a dummy trace using the native API to make sure it works.

diff --git a/Engine/Source/Programs/Horde/HordeServer/Program.cs b/Engine/Source/Programs/Horde/HordeServer/Program.cs
--- a/Engine/Source/Programs/Horde/HordeServer/Program.cs
+++ b/Engine/Source/Programs/Horde/HordeServer/Program.cs
@@ -151,8 +151,9 @@ namespace HordeServer
 
 			if (HordeSettings.WithDatadog)
 			{
+				using var _ = Datadog.Trace.Tracer.Instance.StartActive("Trace Test");
 				Serilog.Log.Logger.Information("Enabling datadog tracing");
-				GlobalTracer.Register(Datadog.Trace.OpenTracing.OpenTracingTracerFactory.CreateTracer());
+				GlobalTracer.Register(Datadog.Trace.OpenTracing.OpenTracingTracerFactory.WrapTracer(Datadog.Trace.Tracer.Instance));
 			}
 
 			if (Arguments.HasOption("-UpdateSchemas"))
