Horde: Test code for Datadog span tracing.

diff --git a/Engine/Source/Programs/Horde/HordeServer/Program.cs b/Engine/Source/Programs/Horde/HordeServer/Program.cs
--- a/Engine/Source/Programs/Horde/HordeServer/Program.cs
+++ b/Engine/Source/Programs/Horde/HordeServer/Program.cs
@@ -154,6 +154,8 @@ namespace HordeServer
 				using var _ = Datadog.Trace.Tracer.Instance.StartActive("Trace Test");
 				Serilog.Log.Logger.Information("Enabling datadog tracing");
 				GlobalTracer.Register(Datadog.Trace.OpenTracing.OpenTracingTracerFactory.WrapTracer(Datadog.Trace.Tracer.Instance));
+				using IScope Scope = GlobalTracer.Instance.BuildSpan("OpenTrace Test").StartActive();
+				Scope.Span.SetTag("TestProp", "hello");
 			}
 
 			if (Arguments.HasOption("-UpdateSchemas"))
diff --git a/Engine/Source/Programs/Horde/HordeServer/Services/DatabaseService.cs b/Engine/Source/Programs/Horde/HordeServer/Services/DatabaseService.cs
--- a/Engine/Source/Programs/Horde/HordeServer/Services/DatabaseService.cs
+++ b/Engine/Source/Programs/Horde/HordeServer/Services/DatabaseService.cs
@@ -137,6 +137,7 @@ namespace HordeServer.Services
 			
 			public void Handle(CommandStartedEvent e)
 			{
+				Serilog.Log.Debug("Mongo command start: {Command}", e.CommandName);
 				IScope Scope = GlobalTracer.Instance.BuildSpan("mongodb.command").StartActive();
 				string? DbName = null;
 				string? CollectionName = null;
@@ -176,8 +177,9 @@ namespace HordeServer.Services
 					Scope.Dispose();
 					Scopes.Remove(e.RequestId, out _);
 				}
+				Serilog.Log.Debug("Mongo command succeded: {Command} ({HaveScope})", e.CommandName, Scope != null);
 			}
-			
+
 			public void Handle(CommandFailedEvent e)
 			{
 				if (Scopes.TryGetValue(e.RequestId, out IScope? Scope))
@@ -186,6 +188,7 @@ namespace HordeServer.Services
 					Scope.Dispose();
 					Scopes.Remove(e.RequestId, out _);
 				}
+				Serilog.Log.Debug("Mongo command failed: {Command} ({HaveScope})", e.CommandName, Scope != null);
 			}
 		}
 
