Horde: More debug logging for tracing.

diff --git a/Engine/Source/Programs/Horde/HordeServer/Program.cs b/Engine/Source/Programs/Horde/HordeServer/Program.cs
--- a/Engine/Source/Programs/Horde/HordeServer/Program.cs
+++ b/Engine/Source/Programs/Horde/HordeServer/Program.cs
@@ -42,6 +42,7 @@ using Serilog.Events;
 using Serilog.Filters;
 using Serilog.Formatting.Json;
 using Serilog.Sinks.SystemConsole.Themes;
+using OpenTracing.Propagation;
 
 namespace HordeServer
 {
@@ -91,6 +92,36 @@ namespace HordeServer
 		}
 	}
 
+	class TestTracer : ITracer
+	{
+		ITracer Inner;
+
+		public TestTracer(ITracer Inner)
+		{
+			this.Inner = Inner;
+		}
+
+		public IScopeManager ScopeManager => Inner.ScopeManager;
+
+		public ISpan ActiveSpan => Inner.ActiveSpan;
+
+		public ISpanBuilder BuildSpan(string operationName)
+		{
+			Serilog.Log.Debug("Creating span {Name}", operationName);
+			return Inner.BuildSpan(operationName);
+		}
+
+		public ISpanContext Extract<TCarrier>(IFormat<TCarrier> format, TCarrier carrier)
+		{
+			return Inner.Extract<TCarrier>(format, carrier);
+		}
+
+		public void Inject<TCarrier>(ISpanContext spanContext, IFormat<TCarrier> format, TCarrier carrier)
+		{
+			Inner.Inject<TCarrier>(spanContext, format, carrier);
+		}
+	}
+
 	class Program
 	{
 		public static DirectoryReference AppDir { get; } = GetAppDir();
@@ -153,9 +184,10 @@ namespace HordeServer
 			{
 				using var _ = Datadog.Trace.Tracer.Instance.StartActive("Trace Test");
 				Serilog.Log.Logger.Information("Enabling datadog tracing");
-				GlobalTracer.Register(Datadog.Trace.OpenTracing.OpenTracingTracerFactory.WrapTracer(Datadog.Trace.Tracer.Instance));
+				GlobalTracer.Register(new TestTracer(Datadog.Trace.OpenTracing.OpenTracingTracerFactory.WrapTracer(Datadog.Trace.Tracer.Instance)));
 				using IScope Scope = GlobalTracer.Instance.BuildSpan("OpenTrace Test").StartActive();
 				Scope.Span.SetTag("TestProp", "hello");
+				Serilog.Log.Logger.Information("Enabling datadog tracing (OpenTrace)");
 			}
 
 			if (Arguments.HasOption("-UpdateSchemas"))
