Horde: More logging for DD tracing.

diff --git a/Engine/Source/Programs/Horde/HordeServer/Services/DatabaseService.cs b/Engine/Source/Programs/Horde/HordeServer/Services/DatabaseService.cs
--- a/Engine/Source/Programs/Horde/HordeServer/Services/DatabaseService.cs
+++ b/Engine/Source/Programs/Horde/HordeServer/Services/DatabaseService.cs
@@ -137,58 +137,66 @@ namespace HordeServer.Services
 			
 			public void Handle(CommandStartedEvent e)
 			{
-				Serilog.Log.Debug("Mongo command start: {Command}", e.CommandName);
-				IScope Scope = GlobalTracer.Instance.BuildSpan("mongodb.command").StartActive();
-				string? DbName = null;
-				string? CollectionName = null;
-				string? Query = null;
-
-				string? GetCollectionName(string CommandName, BsonDocument Command)
+				try
 				{
-					return Command.TryGetValue(CommandName, out BsonValue Value) ? Value.AsString : null;
+					IScope Scope = GlobalTracer.Instance.BuildSpan("mongodb.command").StartActive();
+					string? DbName = null;
+					string? CollectionName = null;
+					string? Query = null;
+
+					string? GetCollectionName(string CommandName, BsonDocument Command)
+					{
+						return Command.TryGetValue(CommandName, out BsonValue Value) ? Value.AsString : null;
+					}
+
+					switch (e.CommandName)
+					{
+						case "find":
+							CollectionName = GetCollectionName("find", e.Command);
+							Query = e.Command.TryGetValue("filter", out BsonValue Value) ? Value.ToJson() : null;
+							break;
+						case "insert":
+							CollectionName = GetCollectionName("insert", e.Command);
+							break;
+						case "update":
+							CollectionName = GetCollectionName("update", e.Command);
+							break;
+					}
+
+					Scope.Span.SetTag("Type", "db");
+					Scope.Span.SetTag("Command", e.CommandName);
+					Scope.Span.SetTag("DbName", DbName);
+					Scope.Span.SetTag("Collection", CollectionName);
+					Scope.Span.SetTag("Query", Query);
+					Scopes[e.RequestId] = Scope;
+
+					Serilog.Log.Debug("Mongo command start: {Command} (SpanId: {SpanId}, TraceId: {TraceId})", e.CommandName, Scope.Span.Context.SpanId, Scope.Span.Context.TraceId);
 				}
-				
-				switch (e.CommandName)
+				catch (Exception Ex)
 				{
-					case "find":
-						CollectionName = GetCollectionName("find", e.Command);
-						Query = e.Command.TryGetValue("filter", out BsonValue Value) ? Value.ToJson() : null;
-						break;
-					case "insert":
-						CollectionName = GetCollectionName("insert", e.Command);
-						break;
-					case "update":
-						CollectionName = GetCollectionName("update", e.Command);
-						break;
+					Serilog.Log.Error(Ex, "Mongo trace exception: {Message}, {Command}", Ex.Message, e.CommandName);
 				}
-
-				Scope.Span.SetTag("Type", "db");
-				Scope.Span.SetTag("Command", e.CommandName);
-				Scope.Span.SetTag("DbName", DbName);
-				Scope.Span.SetTag("Collection", CollectionName);
-				Scope.Span.SetTag("Query", Query);
-				Scopes[e.RequestId] = Scope;
 			}
 
 			public void Handle(CommandSucceededEvent e)
 			{
 				if (Scopes.TryGetValue(e.RequestId, out IScope? Scope))
 				{
+					Serilog.Log.Debug("Mongo command succeded: {Command}", e.CommandName);
 					Scope.Dispose();
 					Scopes.Remove(e.RequestId, out _);
 				}
-				Serilog.Log.Debug("Mongo command succeded: {Command} ({HaveScope})", e.CommandName, Scope != null);
 			}
 
 			public void Handle(CommandFailedEvent e)
 			{
 				if (Scopes.TryGetValue(e.RequestId, out IScope? Scope))
 				{
+					Serilog.Log.Debug("Mongo command failed: {Command}", e.CommandName);
 					Scope.Span.SetTag("Error", true);
 					Scope.Dispose();
 					Scopes.Remove(e.RequestId, out _);
 				}
-				Serilog.Log.Debug("Mongo command failed: {Command} ({HaveScope})", e.CommandName, Scope != null);
 			}
 		}
 
