Horde: More logging.

diff --git a/Engine/Source/Programs/Horde/HordeServer/Services/DatabaseService.cs b/Engine/Source/Programs/Horde/HordeServer/Services/DatabaseService.cs
--- a/Engine/Source/Programs/Horde/HordeServer/Services/DatabaseService.cs
+++ b/Engine/Source/Programs/Horde/HordeServer/Services/DatabaseService.cs
@@ -139,6 +139,8 @@ namespace HordeServer.Services
 			{
 				try
 				{
+					Serilog.Log.Information("Mongo command start: {Command}", e.CommandName);
+
 					IScope Scope = GlobalTracer.Instance.BuildSpan("mongodb.command").StartActive();
 					string? DbName = null;
 					string? CollectionName = null;
@@ -170,11 +172,11 @@ namespace HordeServer.Services
 					Scope.Span.SetTag("Query", Query);
 					Scopes[e.RequestId] = Scope;
 
-					Serilog.Log.Debug("Mongo command start: {Command} (SpanId: {SpanId}, TraceId: {TraceId})", e.CommandName, Scope.Span.Context.SpanId, Scope.Span.Context.TraceId);
+					Serilog.Log.Information("Mongo command start: {Command} (SpanId: {SpanId}, TraceId: {TraceId})", e.CommandName, Scope.Span.Context.SpanId, Scope.Span.Context.TraceId);
 				}
 				catch (Exception Ex)
 				{
-					Serilog.Log.Error(Ex, "Mongo trace exception: {Message}, {Command}", Ex.Message, e.CommandName);
+					Serilog.Log.Error(Ex, "Mongo trace exception: {Message}", Ex.Message);
 				}
 			}
 
@@ -182,7 +184,7 @@ namespace HordeServer.Services
 			{
 				if (Scopes.TryGetValue(e.RequestId, out IScope? Scope))
 				{
-					Serilog.Log.Debug("Mongo command succeded: {Command}", e.CommandName);
+					Serilog.Log.Information("Mongo command succeded: {Command}", e.CommandName);
 					Scope.Dispose();
 					Scopes.Remove(e.RequestId, out _);
 				}
@@ -192,7 +194,7 @@ namespace HordeServer.Services
 			{
 				if (Scopes.TryGetValue(e.RequestId, out IScope? Scope))
 				{
-					Serilog.Log.Debug("Mongo command failed: {Command}", e.CommandName);
+					Serilog.Log.Information("Mongo command failed: {Command}", e.CommandName);
 					Scope.Span.SetTag("Error", true);
 					Scope.Dispose();
 					Scopes.Remove(e.RequestId, out _);
