Horde: Remove MongoTracer.

I believe we're not getting trace output from this because completion callbacks do not happen in the same async context that they were initiated in, so the chain of parent spans gets corrupted.

Assuming we can use MongoTracingCollection for any profiling data we need for now.

diff --git a/Engine/Source/Programs/Horde/HordeServer/Services/DatabaseService.cs b/Engine/Source/Programs/Horde/HordeServer/Services/DatabaseService.cs
--- a/Engine/Source/Programs/Horde/HordeServer/Services/DatabaseService.cs
+++ b/Engine/Source/Programs/Horde/HordeServer/Services/DatabaseService.cs
@@ -116,92 +116,6 @@ namespace HordeServer.Services
 		/// </summary>
 		const int DefaultMongoPort = 27017;
 
-		/// <summary>
-		/// Trace MongoDB commands for observability
-		/// The built-in version from Datadog tracing library doesn't work for us (missing queries).
-		/// </summary>
-		class MongoTracer : IEventSubscriber
-		{
-			readonly ConcurrentDictionary<int, IScope> Scopes = new ConcurrentDictionary<int, IScope>();
-			private readonly ReflectionEventSubscriber _subscriber;
-
-			public MongoTracer()
-			{
-				_subscriber = new ReflectionEventSubscriber(this);
-			}
-
-			public bool TryGetEventHandler<TEvent>(out Action<TEvent> handler)
-			{
-				return _subscriber.TryGetEventHandler(out handler);
-			}
-			
-			public void Handle(CommandStartedEvent e)
-			{
-				try
-				{
-					Serilog.Log.Information("Mongo command start: {Command}", e.CommandName);
-
-					IScope Scope = GlobalTracer.Instance.BuildSpan("mongodb.command").StartActive();
-					string? DbName = null;
-					string? CollectionName = null;
-					string? Query = null;
-
-					string? GetCollectionName(string CommandName, BsonDocument Command)
-					{
-						return Command.TryGetValue(CommandName, out BsonValue Value) ? Value.AsString : null;
-					}
-
-					switch (e.CommandName)
-					{
-						case "find":
-							CollectionName = GetCollectionName("find", e.Command);
-							Query = e.Command.TryGetValue("filter", out BsonValue Value) ? Value.ToJson() : null;
-							break;
-						case "insert":
-							CollectionName = GetCollectionName("insert", e.Command);
-							break;
-						case "update":
-							CollectionName = GetCollectionName("update", e.Command);
-							break;
-					}
-
-					Scope.Span.SetTag("Type", "db");
-					Scope.Span.SetTag("Command", e.CommandName);
-					Scope.Span.SetTag("DbName", DbName);
-					Scope.Span.SetTag("Collection", CollectionName);
-					Scope.Span.SetTag("Query", Query);
-					Scopes[e.RequestId] = Scope;
-
-					Serilog.Log.Information("Mongo command start: {Command} (SpanId: {SpanId}, TraceId: {TraceId})", e.CommandName, Scope.Span.Context.SpanId, Scope.Span.Context.TraceId);
-				}
-				catch (Exception Ex)
-				{
-					Serilog.Log.Error(Ex, "Mongo trace exception: {Message}", Ex.Message);
-				}
-			}
-
-			public void Handle(CommandSucceededEvent e)
-			{
-				if (Scopes.TryGetValue(e.RequestId, out IScope? Scope))
-				{
-					Serilog.Log.Information("Mongo command succeded: {Command}", e.CommandName);
-					Scope.Dispose();
-					Scopes.Remove(e.RequestId, out _);
-				}
-			}
-
-			public void Handle(CommandFailedEvent e)
-			{
-				if (Scopes.TryGetValue(e.RequestId, out IScope? Scope))
-				{
-					Serilog.Log.Information("Mongo command failed: {Command}", e.CommandName);
-					Scope.Span.SetTag("Error", true);
-					Scope.Dispose();
-					Scopes.Remove(e.RequestId, out _);
-				}
-			}
-		}
-
 		/// <summary>
 		/// Constructor
 		/// </summary>
@@ -261,7 +175,6 @@ namespace HordeServer.Services
 					{
 						ClusterBuilder.Subscribe<CommandStartedEvent>(Event => TraceMongoCommand(Event.Command));
 					}
-					ClusterBuilder.Subscribe(new MongoTracer());
 				};
 				
 				MongoSettings.SslSettings = new SslSettings();
