Horde: Only consider an agent valid if it's enabled.

diff --git a/Engine/Source/Programs/Horde/HordeServer/Tasks/Impl/JobTaskSource.cs b/Engine/Source/Programs/Horde/HordeServer/Tasks/Impl/JobTaskSource.cs
--- a/Engine/Source/Programs/Horde/HordeServer/Tasks/Impl/JobTaskSource.cs
+++ b/Engine/Source/Programs/Horde/HordeServer/Tasks/Impl/JobTaskSource.cs
@@ -425,7 +425,7 @@ namespace HordeServer.Tasks.Impl
 			HashSet<PoolId> OnlinePools = new HashSet<PoolId>(Agents.Where(x => x.IsSessionValid(UtcNow)).SelectMany(x => x.ExplicitPools));
 			foreach (IPool Pool in Pools)
 			{
-				if (Pool.Condition != null && !OnlinePools.Contains(Pool.Id) && Agents.Any(x => x.IsSessionValid(UtcNow) && x.SatisfiesCondition(Pool.Condition)))
+				if (Pool.Condition != null && !OnlinePools.Contains(Pool.Id) && Agents.Any(x => x.IsSessionValid(UtcNow) && x.SatisfiesCondition(Pool.Condition) && x.Enabled))
 				{
 					OnlinePools.Add(Pool.Id);
 				}
@@ -435,7 +435,7 @@ namespace HordeServer.Tasks.Impl
 			HashSet<PoolId> ValidPools = new HashSet<PoolId>(OnlinePools.Union(Agents.Where(x => !x.IsSessionValid(UtcNow)).SelectMany(x => x.ExplicitPools)));
 			foreach (IPool Pool in Pools)
 			{
-				if (Pool.Condition != null && !ValidPools.Contains(Pool.Id) && Agents.Any(x => !x.IsSessionValid(UtcNow) && x.SatisfiesCondition(Pool.Condition)))
+				if (Pool.Condition != null && !ValidPools.Contains(Pool.Id) && Agents.Any(x => !x.IsSessionValid(UtcNow) && x.SatisfiesCondition(Pool.Condition) && x.Enabled))
 				{
 					ValidPools.Add(Pool.Id);
 				}
