Horde: Another swing at fixing handling of disabled agents in job scheduler.

diff --git a/Engine/Source/Programs/Horde/HordeServer/Tasks/Impl/JobTaskSource.cs b/Engine/Source/Programs/Horde/HordeServer/Tasks/Impl/JobTaskSource.cs
--- a/Engine/Source/Programs/Horde/HordeServer/Tasks/Impl/JobTaskSource.cs
+++ b/Engine/Source/Programs/Horde/HordeServer/Tasks/Impl/JobTaskSource.cs
@@ -422,7 +422,7 @@ namespace HordeServer.Tasks.Impl
 			List<IPool> Pools = await PoolCollection.GetAsync();
 
 			// Find all the pools which are currently online
-			HashSet<PoolId> OnlinePools = new HashSet<PoolId>(Agents.Where(x => x.IsSessionValid(UtcNow)).SelectMany(x => x.ExplicitPools));
+			HashSet<PoolId> OnlinePools = new HashSet<PoolId>(Agents.Where(x => x.IsSessionValid(UtcNow) && x.Enabled).SelectMany(x => x.ExplicitPools));
 			foreach (IPool Pool in Pools)
 			{
 				if (Pool.Condition != null && !OnlinePools.Contains(Pool.Id) && Agents.Any(x => x.IsSessionValid(UtcNow) && x.SatisfiesCondition(Pool.Condition) && x.Enabled))
@@ -432,7 +432,7 @@ namespace HordeServer.Tasks.Impl
 			}
 
 			// Find lists of valid pools and online pools
-			HashSet<PoolId> ValidPools = new HashSet<PoolId>(OnlinePools.Union(Agents.Where(x => !x.IsSessionValid(UtcNow)).SelectMany(x => x.ExplicitPools)));
+			HashSet<PoolId> ValidPools = new HashSet<PoolId>(OnlinePools.Union(Agents.Where(x => !x.IsSessionValid(UtcNow) || !x.Enabled).SelectMany(x => x.ExplicitPools)));
 			foreach (IPool Pool in Pools)
 			{
 				if (Pool.Condition != null && !ValidPools.Contains(Pool.Id) && Agents.Any(x => !x.IsSessionValid(UtcNow) && x.SatisfiesCondition(Pool.Condition) && x.Enabled))
