Horde: Include IP addresses of agent machine in properties list.

diff --git a/Engine/Source/Programs/Horde/HordeAgent/Services/WorkerService.cs b/Engine/Source/Programs/Horde/HordeAgent/Services/WorkerService.cs
--- a/Engine/Source/Programs/Horde/HordeAgent/Services/WorkerService.cs
+++ b/Engine/Source/Programs/Horde/HordeAgent/Services/WorkerService.cs
@@ -40,6 +40,7 @@ using System.Xml;
 using Status = Grpc.Core.Status;
 using EpicGames.Horde.Storage;
 using EpicGames.Horde.Compute;
+using System.Net;
 
 [assembly: InternalsVisibleTo("HordeAgentTests")]
 
@@ -1578,6 +1579,27 @@ namespace HordeAgent.Services
 				}
 			}
 
+			// Get the IP addresses
+			try
+			{
+				IPHostEntry Entry = Dns.GetHostEntry(Dns.GetHostName());
+				foreach (IPAddress Address in Entry.AddressList)
+				{
+					if (Address.AddressFamily == System.Net.Sockets.AddressFamily.InterNetwork)
+					{
+						PrimaryDevice.Properties.Add($"Ipv4={Address}");
+					}
+					else if (Address.AddressFamily == System.Net.Sockets.AddressFamily.InterNetworkV6)
+					{
+						PrimaryDevice.Properties.Add($"Ipv6={Address}");
+					}
+				}
+			}
+			catch (Exception Ex)
+			{
+				Logger.LogDebug(Ex, "Unable to get local IP address");
+			}
+
 			// Get the time that the machine booted
 			PrimaryDevice.Properties.Add($"BootTime={BootTime}");
 			PrimaryDevice.Properties.Add($"StartTime={StartTime}");
