Horde: Adjust darkmode text
#rnx

diff --git a/Engine/Source/Programs/Horde/HordeDashboard/src/index.css b/Engine/Source/Programs/Horde/HordeDashboard/src/index.css
--- a/Engine/Source/Programs/Horde/HordeDashboard/src/index.css
+++ b/Engine/Source/Programs/Horde/HordeDashboard/src/index.css
@@ -5,7 +5,7 @@ body {
     sans-serif;
   -webkit-font-smoothing: antialiased;
   -moz-osx-font-smoothing: grayscale; 
-  color:#2D3F5F;
+  color:#0F0F0FFF;
   overflow: hidden;
 }
 
