Changing Linux shutdown and restart process to use sudo to elevate permissions allowing us to run the agent process as local user horde-agent and only elevate permissions when a shutdown command has been sent from the server.

diff --git a/Engine/Source/Programs/Horde/HordeAgent/Utility/Shutdown.cs b/Engine/Source/Programs/Horde/HordeAgent/Utility/Shutdown.cs
--- a/Engine/Source/Programs/Horde/HordeAgent/Utility/Shutdown.cs
+++ b/Engine/Source/Programs/Horde/HordeAgent/Utility/Shutdown.cs
@@ -64,17 +64,6 @@ namespace HordeAgent.Utility
 
 		const int SE_PRIVILEGE_ENABLED = 0x00000002;
 
-		// Linux syscalls and reboot flags
-		[DllImport("libc.so.6", SetLastError = true)]
-		private static extern void sync();
-
-		[DllImport("libc.so.6", SetLastError = true)]
-		private static extern int reboot(int flag);
-
-		const int RB_AUTOBOOT = 0x1234567;
-
-		const int RB_POWER_OFF = 0x4321fedc;
-
 		/// <summary>
 		/// Initiate a shutdown operation
 		/// </summary>
@@ -130,24 +119,34 @@ namespace HordeAgent.Utility
 			}
 			else if (RuntimeInformation.IsOSPlatform(OSPlatform.Linux))
 			{
-				sync(); // Commit filesystem caches to disk.
+				String ShutdownArgs;
 				if (bRestartAfterShutdown)
-				{
-					Logger.LogInformation("Triggering restart");
-					if (reboot(RB_AUTOBOOT) != 0)
-					{
-						Logger.LogError("Restart failed ({0})", Marshal.GetLastWin32Error());
-						return false;
-					}
+                {
+					ShutdownArgs = "sudo shutdown -r +0 \"Horde Agent is restarting\"";
 				}
 				else
+                {
+					ShutdownArgs = "sudo shutdown +0 \"Horde Agent is shutting down\"";
+				}
+ 
+				Process ShutdownProcess = new Process()
 				{
-					Logger.LogInformation("Triggering shutdown");
-					if (reboot(RB_POWER_OFF) != 0)
+					StartInfo = new ProcessStartInfo
 					{
-						Logger.LogError("Shutdown failed ({0})", Marshal.GetLastWin32Error());
-						return false;
+						FileName = "/bin/sh",
+						Arguments = String.Format("-c \"{0}\"", ShutdownArgs),
+						UseShellExecute = false,
+						CreateNoWindow = true
 					}
+				};
+ 
+				ShutdownProcess.Start();
+				ShutdownProcess.WaitForExit();
+				int ExitCode = ShutdownProcess.ExitCode;
+				if (ExitCode != 0)
+				{
+					Logger.LogError("Shutdown failed ({0})", ExitCode);
+					return false;
 				}
 				return true;
 			}
