Horde: New implementation of ITaskSource. This interface makes it much simpler to handle cancellation and start/finish conditions in implementors.

diff --git a/Engine/Source/Programs/Horde/HordeServer.Compute.Tests/AgentServiceTest.cs b/Engine/Source/Programs/Horde/HordeServer.Compute.Tests/AgentServiceTest.cs
--- a/Engine/Source/Programs/Horde/HordeServer.Compute.Tests/AgentServiceTest.cs
+++ b/Engine/Source/Programs/Horde/HordeServer.Compute.Tests/AgentServiceTest.cs
@@ -51,16 +51,6 @@ namespace HordeServerTests
 	        Assert.IsFalse(AgentService.AuthorizeSession(Agent, GetUser(Agent)));
         }
 
-        [TestMethod]
-        public async Task AgentListenerDispose()
-        {
-	        // Test created to verify dispose bug with AgentEventListener
-	        TestSetup TestSetup = await GetTestSetup();
-	        AgentService AgentService = TestSetup.AgentService;
-	        AgentService.AgentEventListener Listener = new AgentService.AgentEventListener(AgentService, new AgentId("test"), ObjectId.GenerateNewId(), () => { });
-	        Listener.Dispose();
-        }
-
         private ClaimsPrincipal GetUser(IAgent Agent)
         {
 	        return new ClaimsPrincipal(new ClaimsIdentity(new List<Claim>
diff --git a/Engine/Source/Programs/Horde/HordeServer/Compute/IComputeService.cs b/Engine/Source/Programs/Horde/HordeServer/Compute/IComputeService.cs
--- a/Engine/Source/Programs/Horde/HordeServer/Compute/IComputeService.cs
+++ b/Engine/Source/Programs/Horde/HordeServer/Compute/IComputeService.cs
@@ -78,7 +78,7 @@ namespace HordeServer.Compute
 	/// <summary>
 	/// Interface for the compute service
 	/// </summary>
-	public interface IComputeService : INewTaskSource
+	public interface IComputeService : ITaskSource
 	{
 		/// <summary>
 		/// Post tasks to be executed to a channel
diff --git a/Engine/Source/Programs/Horde/HordeServer/Compute/Impl/ComputeService.cs b/Engine/Source/Programs/Horde/HordeServer/Compute/Impl/ComputeService.cs
--- a/Engine/Source/Programs/Horde/HordeServer/Compute/Impl/ComputeService.cs
+++ b/Engine/Source/Programs/Horde/HordeServer/Compute/Impl/ComputeService.cs
@@ -210,7 +210,7 @@ namespace HordeServer.Compute.Impl
 		}
 
 		/// <inheritdoc/>
-		public async Task<AgentLease?> TryAssignLeaseAsync(IAgent Agent, CancellationToken CancellationToken)
+		public async Task<AgentLease?> AssignLeaseAsync(IAgent Agent, CancellationToken CancellationToken)
 		{
 			// If the agent is disabled, just block until we cancel
 			if (!Agent.Enabled)
diff --git a/Engine/Source/Programs/Horde/HordeServer/Models/Agent.cs b/Engine/Source/Programs/Horde/HordeServer/Models/Agent.cs
--- a/Engine/Source/Programs/Horde/HordeServer/Models/Agent.cs
+++ b/Engine/Source/Programs/Horde/HordeServer/Models/Agent.cs
@@ -182,6 +182,11 @@ namespace HordeServer.Models
 	/// </summary>
 	public class AgentLease
 	{
+		/// <summary>
+		/// Lease returned when a task source wants to run, but needs to wait for other tasks to finish first.
+		/// </summary>
+		public static AgentLease Drain { get; } = new AgentLease(ObjectId.Empty, String.Empty, null, null, null, LeaseState.Pending, null, true, null);
+
 		/// <summary>
 		/// Name of this lease
 		/// </summary>
diff --git a/Engine/Source/Programs/Horde/HordeServer/Services/AgentService.cs b/Engine/Source/Programs/Horde/HordeServer/Services/AgentService.cs
--- a/Engine/Source/Programs/Horde/HordeServer/Services/AgentService.cs
+++ b/Engine/Source/Programs/Horde/HordeServer/Services/AgentService.cs
@@ -1,42 +1,29 @@
 // Copyright Epic Games, Inc. All Rights Reserved.
 
 using EpicGames.Core;
-using Google.Protobuf;
-using Google.Protobuf.Reflection;
 using Google.Protobuf.WellKnownTypes;
-using HordeServer.Api;
-using HordeServer.Collections;
 using HordeCommon;
+using HordeServer.Collections;
 using HordeServer.Models;
-using HordeCommon.Rpc.Tasks;
 using HordeServer.Tasks;
 using HordeServer.Utilities;
-using Microsoft.AspNetCore.Identity;
 using Microsoft.Extensions.Hosting;
 using Microsoft.Extensions.Logging;
-using Microsoft.Extensions.Options;
-using Microsoft.IdentityModel.Tokens;
 using MongoDB.Bson;
 using MongoDB.Driver;
+using StatsdClient;
 using System;
 using System.Collections.Generic;
 using System.Diagnostics;
-using System.IdentityModel.Tokens.Jwt;
 using System.Linq;
-using System.Linq.Expressions;
 using System.Security.Claims;
-using System.Text;
 using System.Threading;
 using System.Threading.Tasks;
-using StatsdClient;
-using System.Collections.Concurrent;
-using Microsoft.Extensions.Caching.Memory;
 
 namespace HordeServer.Services
 {
-	using PoolId = StringId<IPool>;
-	using AgentSoftwareVersion = StringId<IAgentSoftwareCollection>;
 	using AgentSoftwareChannelName = StringId<AgentSoftwareChannels>;
+	using PoolId = StringId<IPool>;
 
 	/// <summary>
 	/// Wraps funtionality for manipulating agents
@@ -108,11 +95,6 @@ namespace HordeServer.Services
 		/// </summary>
 		readonly IClock Clock;
 
-		/// <summary>
-		/// List of event listeners for agents connected to this server. Used to track leases that cannot be migrated to other machines (ie. REAPI leases)
-		/// </summary>
-		Dictionary<AgentId, AgentEventListener> Listeners = new Dictionary<AgentId, AgentEventListener>();
-
 		/// <summary>
 		/// Lazily updated list of current pools
 		/// </summary>
@@ -343,17 +325,49 @@ namespace HordeServer.Services
 			return Agent;
 		}
 
+		static AgentLease? SafeGetLeaseResultInternal(Task<AgentLease?> Task, ILogger Logger)
+		{
+			if (!Task.IsCompleted)
+			{
+				Logger.LogError("Task is not complete");
+				return null;
+			}
+			if (Task.IsFaulted)
+			{
+				Logger.LogError(Task.Exception, "Task has faulted");
+				return null;
+			}
+			return Task.Result;
+		}
+
+		async Task<(ITaskSource, AgentLease)?> GuardedWaitForLeaseAsync(ITaskSource Source, IAgent Agent, CancellationToken CancellationToken)
+		{
+			try
+			{
+				AgentLease? Lease = await Source.AssignLeaseAsync(Agent, CancellationToken);
+				return (Lease != null) ? (Source, Lease) : null;
+			}
+			catch (TaskCanceledException)
+			{
+				return null;
+			}
+			catch (Exception Ex)
+			{
+				Logger.LogError(Ex, "Exception while trying to assign lease"); 
+				return null;
+			}
+		}
+
 		/// <summary>
 		/// Waits for a lease to be assigned to an agent
 		/// </summary>
 		/// <param name="Agent">The agent to assign a lease to</param>
-		/// <param name="AbortTask">Task which can be used to stop polling</param>
+		/// <param name="CancellationToken"></param>
 		/// <returns>True if a lease was assigned, false otherwise</returns>
-		public async Task<IAgent?> WaitForLeaseAsync(IAgent? Agent, Task AbortTask)
+		public async Task<IAgent?> WaitForLeaseAsync(IAgent? Agent, CancellationToken CancellationToken)
 		{
 			while (Agent != null)
 			{
-				// Check there's an active session
 				if (!Agent.SessionExpiresAt.HasValue)
 				{
 					break;
@@ -362,105 +376,89 @@ namespace HordeServer.Services
 				// Check we have some time to wait
 				DateTime UtcNow = Clock.UtcNow;
 				TimeSpan MaxWaitTime = (Agent.SessionExpiresAt.Value - SessionExpiryTime + SessionLongPollTime) - UtcNow;
-				if(MaxWaitTime < TimeSpan.Zero)
+				if (MaxWaitTime <= TimeSpan.Zero)
 				{
 					break;
 				}
 
+				// Create a cancellation token that will expire with the session
+				using CancellationTokenSource CancellationSource = CancellationTokenSource.CreateLinkedTokenSource(ApplicationLifetime.ApplicationStopping, CancellationToken);
+				CancellationSource.CancelAfter(MaxWaitTime);
+
 				// If we're in a maintenance window, just wait for the time to expire
 				if (DowntimeService.IsDowntimeActive)
 				{
-					await Task.WhenAny(AbortTask, Task.Delay(MaxWaitTime));
+					await AsyncUtils.DelayNoThrow(MaxWaitTime, CancellationToken);
 					break;
 				}
 
-				// Subscribe for any new tasks becoming available
-				List<ITaskListener> Subscriptions = new List<ITaskListener>();
-				try
+				// Create all the tasks to wait for
+				List<Task<(ITaskSource, AgentLease)?>> Tasks = new List<Task<(ITaskSource, AgentLease)?>>();
+				foreach (ITaskSource TaskSource in TaskSources)
 				{
-					Task WaitTask = Task.Delay(MaxWaitTime, ApplicationLifetime.ApplicationStopping);
+					Tasks.Add(GuardedWaitForLeaseAsync(TaskSource, Agent, CancellationSource.Token));
+				}
 
-					// Create a list of tasks to wait for
-					List<Task> Tasks = new List<Task>();
-					Tasks.Add(WaitTask);
-					Tasks.Add(AbortTask);
+				// Wait for a lease to be available
+				(ITaskSource, AgentLease)? Result = null;
+				while (Tasks.Count > 0)
+				{
+					await Task.WhenAny(Tasks);
 
-					// Try to create a waiter for each task source
-					foreach (ITaskSource TaskSource in TaskSources)
+					for (int Idx = 0; Idx < Tasks.Count; Idx++)
 					{
-						ITaskListener? NewSubscription = await TaskSource.SubscribeAsync(Agent);
-						if (NewSubscription != null)
+						Task<(ITaskSource, AgentLease)?> Task = Tasks[Idx];
+						if (!Task.IsCompleted)
 						{
-							Subscriptions.Add(NewSubscription);
-							if (!NewSubscription.LeaseTask.IsCompleted || NewSubscription.LeaseTask.Result != null)
-							{
-								Tasks.Add(NewSubscription.LeaseTask);
-							}
-							if (NewSubscription.LeaseTask.IsCompleted)
-							{
-								break;
-							}
+							continue;
 						}
-					}
 
-					// Wait for any lease to complete
-					Task? Result = Tasks.FirstOrDefault(x => x.IsCompleted);
-					if (Result == null)
-					{
-						try
+						Tasks.RemoveAt(Idx--);
+
+						(ITaskSource, AgentLease)? TaskResult = Task.Result;
+						if (!TaskResult.HasValue)
 						{
-							Logger.LogDebug("Start long poll for agent {AgentId}", Agent.Id);
-							Result = await Task.WhenAny(Tasks);
-							Logger.LogDebug("Long poll for agent {AgentId} stopped by {Idx}", Agent.Id, Tasks.IndexOf(Result));
+							continue;
 						}
-						catch (OperationCanceledException)
+
+						if (Result == null)
 						{
-							Logger.LogDebug("Stopped long poll (application stopping)");
-							break;
+							Result = TaskResult;
+							CancellationSource.Cancel();
 						}
-					}
-
-					// Log the outcome of the log poll
-					if (Result == WaitTask)
-					{
-						Logger.LogDebug("Stopped long poll (timeout)");
-						break;
-					}
-					else if (Result == AbortTask)
-					{
-						Logger.LogDebug("Stopped long poll (client abort)");
-						break;
-					}
-
-					// If we got some new work, try to add a lease for it
-					Logger.LogDebug("Stopped long poll (found work)");
-
-					ITaskListener? Subscription = Subscriptions.FirstOrDefault(x => x.LeaseTask.IsCompleted);
-					if (Subscription != null)
-					{
-						NewLeaseInfo? NewAgentLease = await Subscription.LeaseTask;
-						if (NewAgentLease != null)
+						else
 						{
-							IAgent? NewAgent = await Agents.TryAddLeaseAsync(Agent, NewAgentLease.Lease);
-							if(NewAgent != null)
-							{
-								Subscription.Accepted = true;
-								if (NewAgentLease.OnConnectionLost != null)
-								{
-									AddListener(Agent.Id, NewAgentLease.Lease.Id, NewAgentLease.OnConnectionLost);
-								}
-								await CreateLeaseAsync(Agent, NewAgentLease.Lease);
-								return NewAgent;
-							}
+							(ITaskSource TaskSource, AgentLease TaskLease) = TaskResult.Value;
+							await TaskSource.CancelLeaseAsync(Agent, TaskLease.Id, Any.Parser.ParseFrom(TaskLease.Payload));
 						}
 					}
 				}
-				finally
+
+				// Exit if we didn't find any work to do. It may be that all the task sources returned null, in which case wait for the time period to expire.
+				if (Result == null)
 				{
-					foreach (ITaskListener Waiter in Subscriptions)
+					if (!CancellationSource.IsCancellationRequested)
 					{
-						await Waiter.DisposeAsync();
+						await CancellationSource.Token.AsTask();
 					}
+					break;
+				}
+
+				// Get the resulting lease
+				(ITaskSource Source, AgentLease Lease) = Result.Value;
+				if (Lease == AgentLease.Drain)
+				{
+					await AsyncUtils.DelayNoThrow(MaxWaitTime, CancellationToken);
+					break;
+				}
+
+				// Add the new lease to the agent
+				IAgent? NewAgent = await Agents.TryAddLeaseAsync(Agent, Lease);
+				if (NewAgent != null)
+				{
+					await Source.OnLeaseStartedAsync(NewAgent, Lease.Id, Any.Parser.ParseFrom(Lease.Payload));
+					await CreateLeaseAsync(Agent, Lease);
+					return NewAgent;
 				}
 
 				// Update the agent
@@ -498,6 +496,25 @@ namespace HordeServer.Services
 			return true;
 		}
 
+		/// <summary>
+		/// 
+		/// </summary>
+		public async Task<IAgent?> UpdateSessionWithWaitAsync(IAgent InAgent, ObjectId SessionId, AgentStatus Status, IReadOnlyList<string>? Properties, IReadOnlyDictionary<string, int>? Resources, IList<HordeCommon.Rpc.Messages.Lease> NewLeases, CancellationToken CancellationToken)
+		{
+			IAgent? Agent = InAgent;
+
+			// Capture the current agent update index. This allows us to detect if anything has changed.
+			uint UpdateIndex = Agent.UpdateIndex;
+
+			// Update the agent session and return to the caller if anything changes
+			Agent = await UpdateSessionAsync(Agent, SessionId, Status, Properties, Resources, NewLeases);
+			if (Agent != null && Agent.UpdateIndex == UpdateIndex && (Agent.Leases.Count > 0 || Agent.Status != AgentStatus.Stopping))
+			{
+				Agent = await WaitForLeaseAsync(Agent, CancellationToken);
+			}
+			return Agent;
+		}
+
 		/// <summary>
 		/// Updates the state of the current agent session
 		/// </summary>
@@ -510,255 +527,121 @@ namespace HordeServer.Services
 		/// <returns>Updated agent state</returns>
 		public async Task<IAgent?> UpdateSessionAsync(IAgent InAgent, ObjectId SessionId, AgentStatus Status, IReadOnlyList<string>? Properties, IReadOnlyDictionary<string, int>? Resources, IList<HordeCommon.Rpc.Messages.Lease> NewLeases)
 		{
-			AgentEventListener? Listener = null;
-			try
-			{
-				lock (Listeners)
-				{
-					if (Listeners.TryGetValue(InAgent.Id, out Listener))
-					{
-						Listener.BeginUpdate(InAgent.Leases.Select(x => x.Id));
-					}
-				}
+			DateTime UtcNow = Clock.UtcNow;
 
-				DateTime UtcNow = Clock.UtcNow;
-
-				Stopwatch Timer = Stopwatch.StartNew();
+			Stopwatch Timer = Stopwatch.StartNew();
 
-				IAgent? Agent = InAgent;
-				while (Agent != null)
+			IAgent? Agent = InAgent;
+			while (Agent != null)
+			{
+				// If the agent is stopping and doesn't have any leases, we can terminate the current session.
+				if (Status == AgentStatus.Stopping && NewLeases.Count == 0)
 				{
-					// If the agent is stopping and doesn't have any leases, we can terminate the current session.
-					if (Status == AgentStatus.Stopping && NewLeases.Count == 0)
-					{
-						// If we've already decided to terminate the session, this update is redundant but harmless
-						if (Agent.SessionId != SessionId)
-						{
-							break;
-						}
-
-						// If the session is valid, we can terminate once the agent leases are also empty
-						if (Agent.Leases.Count == 0)
-						{
-							if (!await TryTerminateSessionAsync(Agent))
-							{
-								Agent = await GetAgentAsync(Agent.Id);
-								continue;
-							}
-							break;
-						}
-					}
-
-					// Check the session id is correct.
+					// If we've already decided to terminate the session, this update is redundant but harmless
 					if (Agent.SessionId != SessionId)
 					{
-						throw new InvalidOperationException($"Invalid agent session {SessionId}");
-					}
-
-					// Check the session hasn't expired
-					if (!Agent.IsSessionValid(UtcNow))
-					{
-						throw new InvalidOperationException("Session has already expired");
-					}
-
-					// Extend the current session time if we're within a time period of the current time expiring. This reduces
-					// unnecessary DB writes a little, but it also allows us to skip the update and jump into a long poll state
-					// if there's still time on the session left.
-					DateTime? SessionExpiresAt = null;
-					if (!Agent.SessionExpiresAt.HasValue || UtcNow > (Agent.SessionExpiresAt - SessionExpiryTime) + SessionRenewTime)
-					{
-						SessionExpiresAt = UtcNow + SessionExpiryTime;
-					}
-
-					// Flag for whether the leases array should be updated
-					bool bUpdateLeases = false;
-					List<AgentLease> Leases = new List<AgentLease>(Agent.Leases);
-
-					// Remove any completed leases from the agent
-					Dictionary<ObjectId, HordeCommon.Rpc.Messages.Lease> LeaseIdToNewState = NewLeases.ToDictionary(x => x.Id.ToObjectId(), x => x);
-					for (int Idx = 0; Idx < Leases.Count; Idx++)
-					{
-						AgentLease Lease = Leases[Idx];
-						if (Lease.State == LeaseState.Cancelled)
-						{
-							HordeCommon.Rpc.Messages.Lease? NewLease;
-							if (!LeaseIdToNewState.TryGetValue(Lease.Id, out NewLease) || NewLease.State == LeaseState.Cancelled || NewLease.State == LeaseState.Completed)
-							{
-								await RemoveLeaseAsync(Agent, Lease, UtcNow, LeaseOutcome.Cancelled, null);
-								Leases.RemoveAt(Idx--);
-								bUpdateLeases = true;
-							}
-						}
-						else
-						{
-							HordeCommon.Rpc.Messages.Lease? NewLease;
-							if (LeaseIdToNewState.TryGetValue(Lease.Id, out NewLease) && NewLease.State != Lease.State)
-							{
-								if (NewLease.State == LeaseState.Cancelled || NewLease.State == LeaseState.Completed)
-								{
-									await RemoveLeaseAsync(Agent, Lease, UtcNow, NewLease.Outcome, NewLease.Output.ToByteArray());
-									Leases.RemoveAt(Idx--);
-								}
-								else if (NewLease.State == LeaseState.Active && Lease.State == LeaseState.Pending)
-								{
-									Lease.State = LeaseState.Active;
-								}
-								bUpdateLeases = true;
-							}
-						}
+						break;
 					}
 
-					// If the agent is stopping, cancel all the leases. Clear out the current session once it's complete.
-					if (Status == AgentStatus.Stopping)
+					// If the session is valid, we can terminate once the agent leases are also empty
+					if (Agent.Leases.Count == 0)
 					{
-						foreach (AgentLease Lease in Leases)
+						if (!await TryTerminateSessionAsync(Agent))
 						{
-							if (Lease.State != LeaseState.Cancelled)
-							{
-								Lease.State = LeaseState.Cancelled;
-								bUpdateLeases = true;
-							}
+							Agent = await GetAgentAsync(Agent.Id);
+							continue;
 						}
-					}
-
-					// Get the new dynamic pools for the agent
-					List<PoolId> DynamicPools = await GetDynamicPoolsAsync(Agent);
-
-					// Update the agent, and try to create new lease documents if we succeed
-					IAgent? NewAgent = await Agents.TryUpdateSessionAsync(Agent, Status, SessionExpiresAt, Properties, Resources, DynamicPools, bUpdateLeases ? Leases : null);
-					if (NewAgent != null)
-					{
 						break;
 					}
-
-					// Fetch the agent again
-					Agent = await GetAgentAsync(Agent.Id);
 				}
-				return Agent;
-			}
-			finally
-			{
-				Listener?.EndUpdate();
-			}
-		}
 
-		void AddListener(AgentId AgentId, ObjectId LeaseId, Action Callback)
-		{
-			lock (Listeners)
-			{
-				AgentEventListener? Listener;
-				if (!Listeners.TryGetValue(AgentId, out Listener) || !Listener.TryAdd(LeaseId, Callback))
+				// Check the session id is correct.
+				if (Agent.SessionId != SessionId)
 				{
-					Listener = new AgentEventListener(this, AgentId, LeaseId, Callback);
-					Listeners[AgentId] = Listener;
+					throw new InvalidOperationException($"Invalid agent session {SessionId}");
 				}
-			}
-		}
 
-		void RemoveListener(AgentEventListener Listener)
-		{
-			lock (Listeners)
-			{
-				AgentEventListener? CurrentListener;
-				if (Listeners.TryGetValue(Listener.AgentId, out CurrentListener) && CurrentListener == Listener)
+				// Check the session hasn't expired
+				if (!Agent.IsSessionValid(UtcNow))
 				{
-					Listeners.Remove(Listener.AgentId);
+					throw new InvalidOperationException("Session has already expired");
 				}
-			}
-			Listener.Dispose();
-		}
 
-		internal class AgentEventListener : IDisposable
-		{
-			object LockObject = new object();
-			AgentService AgentService;
-			public AgentId AgentId;
-			Timer? Timer;
-			List<(ObjectId, Action)> LeaseCallbacks = new List<(ObjectId, Action)>();
-
-			public AgentEventListener(AgentService AgentService, AgentId AgentId, ObjectId LeaseId, Action Callback)
-			{
-				this.AgentService = AgentService;
-				this.AgentId = AgentId;
-
-				Timer = new Timer(x => AgentService.RemoveListener(this));
-				LeaseCallbacks.Add((LeaseId, Callback));
-			}
-
-			public void Dispose()
-			{
-				if (Timer != null)
+				// Extend the current session time if we're within a time period of the current time expiring. This reduces
+				// unnecessary DB writes a little, but it also allows us to skip the update and jump into a long poll state
+				// if there's still time on the session left.
+				DateTime? SessionExpiresAt = null;
+				if (!Agent.SessionExpiresAt.HasValue || UtcNow > (Agent.SessionExpiresAt - SessionExpiryTime) + SessionRenewTime)
 				{
-					Timer.Dispose();
-					Timer = null;
-					BeginUpdate(Array.Empty<ObjectId>());
+					SessionExpiresAt = UtcNow + SessionExpiryTime;
 				}
-			}
 
-			public bool TryAdd(ObjectId LeaseId, Action Callback)
-			{
-				lock (LockObject)
+				// Flag for whether the leases array should be updated
+				bool bUpdateLeases = false;
+				List<AgentLease> Leases = new List<AgentLease>(Agent.Leases);
+
+				// Remove any completed leases from the agent
+				Dictionary<ObjectId, HordeCommon.Rpc.Messages.Lease> LeaseIdToNewState = NewLeases.ToDictionary(x => x.Id.ToObjectId(), x => x);
+				for (int Idx = 0; Idx < Leases.Count; Idx++)
 				{
-					if (LeaseCallbacks.Count > 0)
+					AgentLease Lease = Leases[Idx];
+					if (Lease.State == LeaseState.Cancelled)
 					{
-						LeaseCallbacks.Add((LeaseId, Callback));
-						return true;
+						HordeCommon.Rpc.Messages.Lease? NewLease;
+						if (!LeaseIdToNewState.TryGetValue(Lease.Id, out NewLease) || NewLease.State == LeaseState.Cancelled || NewLease.State == LeaseState.Completed)
+						{
+							await RemoveLeaseAsync(Agent, Lease, UtcNow, LeaseOutcome.Cancelled, null);
+							Leases.RemoveAt(Idx--);
+							bUpdateLeases = true;
+						}
 					}
-				}
-				return false;
-			}
-
-			public void BeginUpdate(IEnumerable<ObjectId> LeaseIds)
-			{
-				lock (LockObject)
-				{
-					if (LeaseCallbacks.Count > 0)
+					else
 					{
-						// Execute any callbacks for leases that have completed
-						for (int Idx = LeaseCallbacks.Count - 1; Idx >= 0; Idx--)
+						HordeCommon.Rpc.Messages.Lease? NewLease;
+						if (LeaseIdToNewState.TryGetValue(Lease.Id, out NewLease) && NewLease.State != Lease.State)
 						{
-							if (!LeaseIds.Contains(LeaseCallbacks[Idx].Item1))
+							if (NewLease.State == LeaseState.Cancelled || NewLease.State == LeaseState.Completed)
 							{
-								LeaseCallbacks[Idx].Item2();
-								LeaseCallbacks.RemoveAt(Idx);
+								await RemoveLeaseAsync(Agent, Lease, UtcNow, NewLease.Outcome, NewLease.Output.ToByteArray());
+								Leases.RemoveAt(Idx--);
 							}
+							else if (NewLease.State == LeaseState.Active && Lease.State == LeaseState.Pending)
+							{
+								Lease.State = LeaseState.Active;
+							}
+							bUpdateLeases = true;
 						}
+					}
+				}
 
-						// Return false when we don't have any remaining callbacks
-						if (LeaseCallbacks.Count == 0)
-						{
-							AgentService.RemoveListener(this);
-						}
-						else
+				// If the agent is stopping, cancel all the leases. Clear out the current session once it's complete.
+				if (Status == AgentStatus.Stopping)
+				{
+					foreach (AgentLease Lease in Leases)
+					{
+						if (Lease.State != LeaseState.Cancelled)
 						{
-							Timer?.Change(Timeout.Infinite, Timeout.Infinite);
+							Lease.State = LeaseState.Cancelled;
+							bUpdateLeases = true;
 						}
 					}
 				}
-			}
 
-			public void EndUpdate()
-			{
-				lock (LockObject)
+				// Get the new dynamic pools for the agent
+				List<PoolId> DynamicPools = await GetDynamicPoolsAsync(Agent);
+
+				// Update the agent, and try to create new lease documents if we succeed
+				IAgent? NewAgent = await Agents.TryUpdateSessionAsync(Agent, Status, SessionExpiresAt, Properties, Resources, DynamicPools, bUpdateLeases ? Leases : null);
+				if (NewAgent != null)
 				{
-					Timer?.Change(20 * 1000, Timeout.Infinite);
+					Agent = NewAgent;
+					break;
 				}
-			}
-		}
 
-		class AgentUpdateScope : IDisposable
-		{
-			AgentEventListener Listener;
-
-			public AgentUpdateScope(AgentEventListener Listener)
-			{
-				this.Listener = Listener;
-			}
-
-			public void Dispose()
-			{
-				Listener.EndUpdate();
+				// Fetch the agent again
+				Agent = await GetAgentAsync(Agent.Id);
 			}
+			return Agent;
 		}
 
 		/// <summary>
@@ -874,7 +757,7 @@ namespace HordeServer.Services
 				{
 					if (Any.Is(TaskSource.Descriptor))
 					{
-						await TaskSource.OnLeaseCompletedAsync(Agent, Lease.Id, Any, Outcome, Output);
+						await TaskSource.OnLeaseFinishedAsync(Agent, Lease.Id, Any, Outcome, Output);
 						break;
 					}
 				}
diff --git a/Engine/Source/Programs/Horde/HordeServer/Tasks/ITaskSource.cs b/Engine/Source/Programs/Horde/HordeServer/Tasks/ITaskSource.cs
--- a/Engine/Source/Programs/Horde/HordeServer/Tasks/ITaskSource.cs
+++ b/Engine/Source/Programs/Horde/HordeServer/Tasks/ITaskSource.cs
@@ -15,78 +15,6 @@ using System.Threading.Tasks;
 
 namespace HordeServer.Tasks
 {
-	/// <summary>
-	/// Information about a new task
-	/// </summary>
-	public class NewLeaseInfo
-	{
-		/// <summary>
-		/// Lease to execute the task
-		/// </summary>
-		public AgentLease Lease { get; }
-
-		/// <summary>
-		/// Callback for connection to the agent being lost
-		/// </summary>
-		public Action? OnConnectionLost { get; }
-
-		/// <summary>
-		/// Constructor
-		/// </summary>
-		/// <param name="Lease"></param>
-		/// <param name="OnConnectionLost"></param>
-		public NewLeaseInfo(AgentLease Lease, Action? OnConnectionLost = null)
-		{
-			this.Lease = Lease;
-			this.OnConnectionLost = OnConnectionLost;
-		}
-	}
-
-	/// <summary>
-	/// Interface for a an object that waits for a task
-	/// </summary>
-	public interface ITaskListener : IAsyncDisposable
-	{
-		/// <summary>
-		/// Task returning a lease for the agent. If a null lease is returned, subsequent task sources will not be checked for leases.
-		/// </summary>
-		Task<NewLeaseInfo?> LeaseTask { get; }
-
-		/// <summary>
-		/// Whether this lease has been accepted. An lease which has not been accepted will be returned to the pool when the subscription is disposed.
-		/// </summary>
-		bool Accepted { get; set; }
-	}
-
-	class TaskSubscription : ITaskListener
-	{
-		public Task<NewLeaseInfo?> LeaseTask { get; }
-
-		public bool Accepted { get; set; }
-
-		public TaskSubscription(Task<NewLeaseInfo?> LeaseTask)
-		{
-			this.LeaseTask = LeaseTask;
-		}
-
-		public static TaskSubscription FromResult(AgentLease? Lease)
-		{
-			if (Lease == null)
-			{
-				return new TaskSubscription(Task.FromResult<NewLeaseInfo?>(null));
-			}
-			else
-			{
-				return new TaskSubscription(Task.FromResult<NewLeaseInfo?>(new NewLeaseInfo(Lease)));
-			}
-		}
-
-		public ValueTask DisposeAsync()
-		{
-			return default;
-		}
-	}
-
 	/// <summary>
 	/// Handler for a certain lease type
 	/// </summary>
@@ -98,66 +26,12 @@ namespace HordeServer.Tasks
 		MessageDescriptor Descriptor { get; }
 
 		/// <summary>
-		/// Subscribes an agent to execute leases for this source. The lifetime of the returned subscription object indicates the time that the agent is volunteering for work.
-		/// </summary>
-		/// <param name="Agent">The agent state</param>
-		/// <returns>The subscription object</returns>
-		Task<ITaskListener?> SubscribeAsync(IAgent Agent);
-
-		/// <summary>
-		/// Notification that a task has completed
-		/// </summary>
-		/// <param name="Agent">The agent that was allocated to the lease</param>
-		/// <param name="LeaseId">The lease id</param>
-		/// <param name="Payload">The lease payload</param>
-		/// <param name="Outcome">Outcome of the lease</param>
-		/// <param name="Output">Output from the task</param>
-		/// <returns>Async task</returns>
-		Task OnLeaseCompletedAsync(IAgent Agent, ObjectId LeaseId, Any Payload, LeaseOutcome Outcome, ReadOnlyMemory<byte> Output);
-	}
-
-	/// <summary>
-	/// Strongly typed base class for implementation of ITaskHandler
-	/// </summary>
-	/// <typeparam name="T">Type of task message</typeparam>
-	public abstract class TaskSourceBase<T> : ITaskSource where T : IMessage<T>, new()
-	{
-		/// <inheritdoc/>
-		public MessageDescriptor Descriptor { get; } = new T().Descriptor;
-
-		/// <inheritdoc/>
-		public abstract Task<ITaskListener?> SubscribeAsync(IAgent Agent);
-
-		/// <inheritdoc/>
-		public Task OnLeaseCompletedAsync(IAgent Agent, ObjectId LeaseId, Any Task, LeaseOutcome Outcome, ReadOnlyMemory<byte> Output)
-		{
-			return OnLeaseCompletedAsync(Agent, LeaseId, Task.Unpack<T>(), Outcome, Output);
-		}
-
-		/// <inheritdoc cref="OnLeaseCompletedAsync(IAgent, ObjectId, Any, LeaseOutcome, ReadOnlyMemory{byte})"/>
-		public virtual Task OnLeaseCompletedAsync(IAgent Agent, ObjectId LeaseId, T TaskInput, LeaseOutcome Outcome, ReadOnlyMemory<byte> Output)
-		{
-			return Task.CompletedTask;
-		}
-	}
-
-	/// <summary>
-	/// Handler for a certain lease type
-	/// </summary>
-	public interface INewTaskSource
-	{
-		/// <summary>
-		/// Descriptor for this lease type
-		/// </summary>
-		MessageDescriptor Descriptor { get; }
-
-		/// <summary>
-		/// Waits for a lease to be availabe
+		/// Assigns a lease or waits for one to be available
 		/// </summary>
 		/// <param name="Agent">The agent to assign a lease to</param>
 		/// <param name="CancellationToken">Cancellation token for the wait</param>
 		/// <returns>New lease object</returns>
-		Task<AgentLease?> TryAssignLeaseAsync(IAgent Agent, CancellationToken CancellationToken);
+		Task<AgentLease?> AssignLeaseAsync(IAgent Agent, CancellationToken CancellationToken);
 
 		/// <summary>
 		/// Cancel a lease that was previously assigned to an agent, allowing it to be assigned out again
@@ -186,72 +60,28 @@ namespace HordeServer.Tasks
 		/// <param name="Output">Output from the task</param>
 		Task OnLeaseFinishedAsync(IAgent Agent, ObjectId LeaseId, Any Payload, LeaseOutcome Outcome, ReadOnlyMemory<byte> Output);
 	}
-
-	class NewTaskSourceWrapper : ITaskSource
+	
+	/// <summary>
+	/// Base implementation of <see cref="ITaskSource"/>
+	/// </summary>
+	/// <typeparam name="TMessage"></typeparam>
+	public abstract class TaskSourceBase<TMessage> : ITaskSource where TMessage : IMessage, new()
 	{
-		class TaskListener : ITaskListener
-		{
-			NewTaskSourceWrapper Outer;
-			IAgent Agent;
-			CancellationTokenSource CancellationTokenSource;
-
-			public Task<NewLeaseInfo?> LeaseTask { get; }
-
-			public bool Accepted { get; set; }
-
-			public TaskListener(NewTaskSourceWrapper Outer, IAgent Agent)
-			{
-				this.Outer = Outer;
-				this.Agent = Agent;
-				this.CancellationTokenSource = new CancellationTokenSource();
-				this.LeaseTask = GetLeaseAsync();
-			}
-
-			async Task<NewLeaseInfo?> GetLeaseAsync()
-			{
-				AgentLease? Lease = await Outer.TaskSource.TryAssignLeaseAsync(Agent, CancellationTokenSource.Token);
-				return (Lease == null)? null : new NewLeaseInfo(Lease);
-			}
+		static TMessage Message = new TMessage();
 
-			public async ValueTask DisposeAsync()
-			{
-				CancellationTokenSource.Cancel();
-				NewLeaseInfo? NewLease = await LeaseTask;
-
-				if (NewLease != null)
-				{
-					if (Accepted)
-					{
-						await Outer.TaskSource.OnLeaseStartedAsync(Agent, NewLease.Lease.Id, Any.Parser.ParseFrom(NewLease.Lease.Payload));
-					}
-					else 
-					{
-						await Outer.TaskSource.CancelLeaseAsync(Agent, NewLease.Lease.Id, Any.Parser.ParseFrom(NewLease.Lease.Payload));
-					}
-				}
-
-				CancellationTokenSource.Dispose();
-			}
-		}
-
-		INewTaskSource TaskSource;
+		/// <inheritdoc/>
+		public MessageDescriptor Descriptor => Message.Descriptor;
 
-		public MessageDescriptor Descriptor => TaskSource.Descriptor;
+		/// <inheritdoc/>
+		public abstract Task<AgentLease?> AssignLeaseAsync(IAgent Agent, CancellationToken CancellationToken);
 
-		public NewTaskSourceWrapper(INewTaskSource TaskSource)
-		{
-			this.TaskSource = TaskSource;
-		}
+		/// <inheritdoc/>
+		public virtual Task CancelLeaseAsync(IAgent Agent, ObjectId LeaseId, Any Payload) => Task.CompletedTask;
 
-		public Task<ITaskListener?> SubscribeAsync(IAgent Agent)
-		{
-			return Task.FromResult<ITaskListener?>(new TaskListener(this, Agent));
-		}
+		/// <inheritdoc/>
+		public virtual Task OnLeaseStartedAsync(IAgent Agent, ObjectId LeaseId, Any Payload) => Task.CompletedTask;
 
-		public Task OnLeaseCompletedAsync(IAgent Agent, ObjectId LeaseId, Any Payload, LeaseOutcome Outcome, ReadOnlyMemory<byte> Output)
-		{
-			return TaskSource.OnLeaseFinishedAsync(Agent, LeaseId, Payload, Outcome, Output);
-		}
+		/// <inheritdoc/>
+		public virtual Task OnLeaseFinishedAsync(IAgent Agent, ObjectId LeaseId, Any Payload, LeaseOutcome Outcome, ReadOnlyMemory<byte> Output) => Task.CompletedTask;
 	}
 }
-
diff --git a/Engine/Source/Programs/Horde/HordeServer/Tasks/Impl/ConformTaskSource.cs b/Engine/Source/Programs/Horde/HordeServer/Tasks/Impl/ConformTaskSource.cs
--- a/Engine/Source/Programs/Horde/HordeServer/Tasks/Impl/ConformTaskSource.cs
+++ b/Engine/Source/Programs/Horde/HordeServer/Tasks/Impl/ConformTaskSource.cs
@@ -155,7 +155,7 @@ namespace HordeServer.Tasks.Impl
 		}
 
 		/// <inheritdoc/>
-		public override async Task<ITaskListener?> SubscribeAsync(IAgent Agent)
+		public override async Task<AgentLease?> AssignLeaseAsync(IAgent Agent, CancellationToken CancellationToken)
 		{
 			DateTime UtcNow = DateTime.UtcNow;
 			if (!await IsConformPendingAsync(Agent, UtcNow))
@@ -179,14 +179,13 @@ namespace HordeServer.Tasks.Impl
 
 					byte[] Payload = Any.Pack(Task).ToByteArray();
 
-					AgentLease Lease = new AgentLease(LeaseId, "Updating workspaces", null, null, Log.Id, LeaseState.Pending, null, true, Payload);
-					return TaskSubscription.FromResult(Lease);
+					return new AgentLease(LeaseId, "Updating workspaces", null, null, Log.Id, LeaseState.Pending, null, true, Payload);
 				}
 			}
 
 			if (Agent.RequestConform)
 			{
-				return TaskSubscription.FromResult(null);
+				return AgentLease.Drain;
 			}
 
 			return null;
@@ -211,7 +210,7 @@ namespace HordeServer.Tasks.Impl
 		}
 
 		/// <inheritdoc/>
-		public override Task OnLeaseCompletedAsync(IAgent Agent, ObjectId LeaseId, ConformTask Task, LeaseOutcome Outcome, ReadOnlyMemory<byte> Output)
+		public override Task OnLeaseFinishedAsync(IAgent Agent, ObjectId LeaseId, Any Payload, LeaseOutcome Outcome, ReadOnlyMemory<byte> Output)
 		{
 			return ReleaseConformLeaseAsync(LeaseId);
 		}
diff --git a/Engine/Source/Programs/Horde/HordeServer/Tasks/Impl/JobTaskSource.cs b/Engine/Source/Programs/Horde/HordeServer/Tasks/Impl/JobTaskSource.cs
--- a/Engine/Source/Programs/Horde/HordeServer/Tasks/Impl/JobTaskSource.cs
+++ b/Engine/Source/Programs/Horde/HordeServer/Tasks/Impl/JobTaskSource.cs
@@ -38,7 +38,7 @@ namespace HordeServer.Tasks.Impl
 	/// <summary>
 	/// Background service to dispatch pending work to agents in priority order.
 	/// </summary>
-	public class JobTaskSource : TickedBackgroundService, ITaskSource
+	public sealed class JobTaskSource : TaskSourceBase<ExecuteJobTask>, IDisposable
 	{
 		/// <summary>
 		/// An item in the queue to be executed
@@ -147,66 +147,30 @@ namespace HordeServer.Tasks.Impl
 		/// <summary>
 		/// Information about an agent waiting for work
 		/// </summary>
-		class QueueWaiter : ITaskListener
+		class QueueWaiter
 		{
 			/// <summary>
-			/// The dispatch service instance
+			/// The agent performing the wait
 			/// </summary>
-			JobTaskSource DispatchService;
+			public IAgent Agent { get; }
 
 			/// <summary>
-			/// The agent performing the wait
+			/// Task to wait for a lease to be assigned
 			/// </summary>
-			public IAgent Agent { get; }
+			public Task<AgentLease?> Task => LeaseSource.Task;
 
 			/// <summary>
 			/// Completion source for the waiting agent. If a new queue item becomes available, the result will be passed through this.
 			/// </summary>
-			public TaskCompletionSource<NewLeaseInfo?> LeaseSource { get; } = new TaskCompletionSource<NewLeaseInfo?>();
-
-			/// <inheritdoc/>
-			public bool Accepted { get; set; }
+			public TaskCompletionSource<AgentLease?> LeaseSource { get; } = new TaskCompletionSource<AgentLease?>();
 
 			/// <summary>
 			/// Constructor
 			/// </summary>
-			/// <param name="DispatchService"></param>
 			/// <param name="Agent">The agent waiting for a task</param>
-			public QueueWaiter(JobTaskSource DispatchService, IAgent Agent)
+			public QueueWaiter(IAgent Agent)
 			{
-				this.DispatchService = DispatchService;
 				this.Agent = Agent;
-
-				lock (DispatchService)
-				{
-					DispatchService.Waiters.Add(this);
-				}
-			}
-
-			public Task<NewLeaseInfo?> LeaseTask => LeaseSource.Task;
-
-			public void Accept()
-			{
-				Accepted = true;
-			}
-
-			public async ValueTask DisposeAsync()
-			{
-				lock (DispatchService.LockObject)
-				{
-					DispatchService.Waiters.Remove(this);
-				}
-
-				if (!LeaseSource.TrySetResult(null) && !Accepted)
-				{
-					NewLeaseInfo? NewLeaseInfo = await LeaseSource.Task;
-					if(NewLeaseInfo != null)
-					{
-						Any Any = Any.Parser.ParseFrom(NewLeaseInfo.Lease.Payload);
-						ExecuteJobTask Task = Any.Unpack<ExecuteJobTask>();
-						await DispatchService.CancelLeaseAsync(Agent, Task.JobId.ToObjectId(), Task.BatchId.ToSubResourceId());
-					}
-				}
 			}
 		}
 
@@ -265,11 +229,6 @@ namespace HordeServer.Tasks.Impl
 		/// </summary>
 		IHostApplicationLifetime ApplicationLifetime;
 
-		/// <summary>
-		/// Task which waits for the application stopping event to start
-		/// </summary>
-		CancellationTask StoppingTask;
-
 		/// <summary>
 		/// Settings instance
 		/// </summary>
@@ -315,14 +274,16 @@ namespace HordeServer.Tasks.Impl
 		/// </summary>
 		private Dictionary<StreamId, IStream> Streams = new Dictionary<StreamId, IStream>();
 
+		/// <summary>
+		/// Background tick registration
+		/// </summary>
+		BackgroundTick Ticker;
+
 		/// <summary>
 		/// Interval between querying the database for jobs to execute
 		/// </summary>
 		static readonly TimeSpan RefreshInterval = TimeSpan.FromSeconds(5.0);
 
-		/// <inheritdoc/>
-		public MessageDescriptor Descriptor => ExecuteJobTask.Descriptor;
-
 		/// <summary>
 		/// Constructor
 		/// </summary>
@@ -340,7 +301,6 @@ namespace HordeServer.Tasks.Impl
 		/// <param name="Settings">Settings for the server</param>
 		/// <param name="Logger">Log writer</param>
 		public JobTaskSource(DatabaseService DatabaseService, IAgentCollection Agents, IJobCollection Jobs, IJobStepRefCollection JobStepRefs, IGraphCollection Graphs, IPoolCollection Pools, IUgsMetadataCollection UgsMetadataCollection, StreamService StreamService, ILogFileService LogFileService, PerforceLoadBalancer PerforceLoadBalancer, IHostApplicationLifetime ApplicationLifetime, IOptionsMonitor<ServerSettings> Settings, ILogger<JobTaskSource> Logger)
-			: base(RefreshInterval, Logger)
 		{
 			this.DatabaseService = DatabaseService;
 			this.AgentsCollection = Agents;
@@ -353,11 +313,23 @@ namespace HordeServer.Tasks.Impl
 			this.LogFileService = LogFileService;
 			this.PerforceLoadBalancer = PerforceLoadBalancer;
 			this.ApplicationLifetime = ApplicationLifetime;
-			this.StoppingTask = new CancellationTask(ApplicationLifetime.ApplicationStopping);
+			this.Ticker = new BackgroundTick(TickAsync, RefreshInterval, Logger);
 			this.Settings = Settings;
 			this.Logger = Logger;
 		}
 
+		/// <inheritdoc/>
+		public async ValueTask DisposeAsync()
+		{
+			await Ticker.DisposeAsync();
+		}
+
+		/// <inheritdoc/>
+		public void Dispose()
+		{
+			Ticker.Dispose();
+		}
+
 		/// <summary>
 		/// Gets an object containing the stats of the queue for diagnostic purposes.
 		/// </summary>
@@ -404,7 +376,7 @@ namespace HordeServer.Tasks.Impl
 		/// </summary>
 		/// <param name="StoppingToken">Token that indicates that the service should shut down</param>
 		/// <returns>Async task</returns>
-		protected async override Task TickAsync(CancellationToken StoppingToken)
+		async Task TickAsync(CancellationToken StoppingToken)
 		{
 			// Set the NewBatchIdToQueueItem member, so we capture any updated jobs during the DB query.
 			lock (LockObject)
@@ -720,14 +692,38 @@ namespace HordeServer.Tasks.Impl
 		}
 
 		/// <inheritdoc/>
-		public Task<ITaskListener?> SubscribeAsync(IAgent Agent)
+		public override async Task<AgentLease?> AssignLeaseAsync(IAgent Agent, CancellationToken CancellationToken)
 		{
-			QueueWaiter Waiter = new QueueWaiter(this, Agent);
+			QueueWaiter Waiter = new QueueWaiter(Agent);
 			lock (LockObject)
 			{
 				AssignAnyQueueItemToWaiter(Waiter);
+				if (Waiter.LeaseSource.Task.IsCompleted)
+				{
+					return Waiter.Task.Result;
+				}
+				Waiters.Add(Waiter);
+			}
+
+			AgentLease? Lease;
+			using (CancellationToken.Register(() => Waiter.LeaseSource.TrySetResult(null)))
+			{
+				Lease = await Waiter.Task;
 			}
-			return Task.FromResult<ITaskListener?>(Waiter);
+
+			lock (LockObject)
+			{
+				Waiters.Remove(Waiter);
+			}
+
+			return Lease;
+		}
+
+		/// <inheritdoc/>
+		public override Task CancelLeaseAsync(IAgent Agent, ObjectId LeaseId, Any Payload)
+		{
+			ExecuteJobTask Task = Payload.Unpack<ExecuteJobTask>();
+			return CancelLeaseAsync(Agent, Task.JobId.ToObjectId(), Task.BatchId.ToSubResourceId());
 		}
 
 		/// <summary>
@@ -776,7 +772,7 @@ namespace HordeServer.Tasks.Impl
 
 					// Create the lease and try to set it on the waiter. If this fails, the waiter has already moved on, and the lease can be cancelled.
 					AgentLease Lease = new AgentLease(LeaseId, LeaseName.ToString(), Job.StreamId, Item.PoolId, LogId, LeaseState.Pending, null, true, Payload);
-					if (Waiter.LeaseSource.TrySetResult(new NewLeaseInfo(Lease)))
+					if (Waiter.LeaseSource.TrySetResult(Lease))
 					{
 						Logger.LogDebug("Assigned lease {LeaseId} to agent {AgentId}", LeaseId, Agent.Id);
 						return Lease;
@@ -976,7 +972,7 @@ namespace HordeServer.Tasks.Impl
 		}
 
 		/// <inheritdoc/>
-		public async Task OnLeaseCompletedAsync(IAgent Agent, ObjectId LeaseId, Any Any, LeaseOutcome Outcome, ReadOnlyMemory<byte> Output)
+		public override async Task OnLeaseFinishedAsync(IAgent Agent, ObjectId LeaseId, Any Any, LeaseOutcome Outcome, ReadOnlyMemory<byte> Output)
 		{
 			if (Outcome != LeaseOutcome.Success)
 			{
diff --git a/Engine/Source/Programs/Horde/HordeServer/Tasks/Impl/RestartTaskSource.cs b/Engine/Source/Programs/Horde/HordeServer/Tasks/Impl/RestartTaskSource.cs
--- a/Engine/Source/Programs/Horde/HordeServer/Tasks/Impl/RestartTaskSource.cs
+++ b/Engine/Source/Programs/Horde/HordeServer/Tasks/Impl/RestartTaskSource.cs
@@ -12,6 +12,7 @@ using MongoDB.Bson;
 using System;
 using System.Collections.Generic;
 using System.Linq;
+using System.Threading;
 using System.Threading.Tasks;
 
 namespace HordeServer.Tasks.Impl
@@ -25,7 +26,7 @@ namespace HordeServer.Tasks.Impl
 			this.LogService = LogService;
 		}
 
-		public override async Task<ITaskListener?> SubscribeAsync(IAgent Agent)
+		public override async Task<AgentLease?> AssignLeaseAsync(IAgent Agent, CancellationToken CancellationToken)
 		{
 			if (!Agent.RequestRestart)
 			{
@@ -33,7 +34,7 @@ namespace HordeServer.Tasks.Impl
 			}
 			if (Agent.Leases.Count > 0)
 			{
-				return TaskSubscription.FromResult(null);
+				return AgentLease.Drain;
 			}
 
 			ILogFile Log = await LogService.CreateLogFileAsync(ObjectId.Empty, Agent.SessionId, LogType.Json);
@@ -43,8 +44,7 @@ namespace HordeServer.Tasks.Impl
 
 			byte[] Payload = Any.Pack(Task).ToByteArray();
 
-			AgentLease? Lease = new AgentLease(ObjectId.GenerateNewId(), "Restart", null, null, Log.Id, LeaseState.Pending, null, true, Payload);
-			return TaskSubscription.FromResult(Lease);
+			return new AgentLease(ObjectId.GenerateNewId(), "Restart", null, null, Log.Id, LeaseState.Pending, null, true, Payload);
 		}
 	}
 }
diff --git a/Engine/Source/Programs/Horde/HordeServer/Tasks/Impl/ShutdownTaskSource.cs b/Engine/Source/Programs/Horde/HordeServer/Tasks/Impl/ShutdownTaskSource.cs
--- a/Engine/Source/Programs/Horde/HordeServer/Tasks/Impl/ShutdownTaskSource.cs
+++ b/Engine/Source/Programs/Horde/HordeServer/Tasks/Impl/ShutdownTaskSource.cs
@@ -1,5 +1,6 @@
 // Copyright Epic Games, Inc. All Rights Reserved.
 
+using EpicGames.Core;
 using Google.Protobuf;
 using Google.Protobuf.Reflection;
 using Google.Protobuf.WellKnownTypes;
@@ -12,6 +13,7 @@ using MongoDB.Bson;
 using System;
 using System.Collections.Generic;
 using System.Linq;
+using System.Threading;
 using System.Threading.Tasks;
 
 namespace HordeServer.Tasks.Impl
@@ -25,7 +27,7 @@ namespace HordeServer.Tasks.Impl
 			this.LogService = LogService;
 		}
 
-		public override async Task<ITaskListener?> SubscribeAsync(IAgent Agent)
+		public override async Task<AgentLease?> AssignLeaseAsync(IAgent Agent, CancellationToken CancellationToken)
 		{
 			if (!Agent.RequestShutdown)
 			{
@@ -33,7 +35,7 @@ namespace HordeServer.Tasks.Impl
 			}
 			if (Agent.Leases.Count > 0)
 			{
-				return TaskSubscription.FromResult(null);
+				return AgentLease.Drain;
 			}
 
 			ILogFile Log = await LogService.CreateLogFileAsync(ObjectId.Empty, Agent.SessionId, LogType.Json);
@@ -43,8 +45,7 @@ namespace HordeServer.Tasks.Impl
 
 			byte[] Payload = Any.Pack(Task).ToByteArray();
 
-			AgentLease? Lease = new AgentLease(ObjectId.GenerateNewId(), "Shutdown", null, null, Log.Id, LeaseState.Pending, null, true, Payload);
-			return TaskSubscription.FromResult(Lease);
+			return new AgentLease(ObjectId.GenerateNewId(), "Shutdown", null, null, Log.Id, LeaseState.Pending, null, true, Payload);
 		}
 	}
 }
diff --git a/Engine/Source/Programs/Horde/HordeServer/Tasks/Impl/UpgradeTaskSource.cs b/Engine/Source/Programs/Horde/HordeServer/Tasks/Impl/UpgradeTaskSource.cs
--- a/Engine/Source/Programs/Horde/HordeServer/Tasks/Impl/UpgradeTaskSource.cs
+++ b/Engine/Source/Programs/Horde/HordeServer/Tasks/Impl/UpgradeTaskSource.cs
@@ -13,6 +13,7 @@ using MongoDB.Bson;
 using System;
 using System.Collections.Generic;
 using System.Linq;
+using System.Threading;
 using System.Threading.Tasks;
 
 namespace HordeServer.Tasks.Impl
@@ -32,7 +33,7 @@ namespace HordeServer.Tasks.Impl
 			this.Clock = Clock;
 		}
 
-		public async override Task<ITaskListener?> SubscribeAsync(IAgent Agent)
+		public async override Task<AgentLease?> AssignLeaseAsync(IAgent Agent, CancellationToken CancellationToken)
 		{
 			string? RequiredVersion = await GetRequiredSoftwareVersion(Agent);
 			if (RequiredVersion != null && Agent.Version != RequiredVersion)
@@ -61,7 +62,7 @@ namespace HordeServer.Tasks.Impl
 
 					Lease = new AgentLease(ObjectId.GenerateNewId(), $"Upgrade to {RequiredVersion}", null, null, LogFile.Id, LeaseState.Pending, null, true, Payload);
 				}
-				return TaskSubscription.FromResult(Lease);
+				return Lease;
 			}
 			return null;
 		}
