Horde: Missing files.

diff --git a/Engine/Source/Programs/Horde/HordeServer/Services/RpcService.cs b/Engine/Source/Programs/Horde/HordeServer/Services/RpcService.cs
--- a/Engine/Source/Programs/Horde/HordeServer/Services/RpcService.cs
+++ b/Engine/Source/Programs/Horde/HordeServer/Services/RpcService.cs
@@ -22,6 +22,7 @@ using System.Security.Claims;
 using System.Text;
 using System.Threading.Tasks;
 using HordeServer.Tasks.Impl;
+using System.Threading;
 
 [assembly: InternalsVisibleTo("HordeAgentTests")]
 
@@ -392,7 +393,9 @@ namespace HordeServer.Services
 				}
 
 				// Get a task for moving to the next item. This will only complete once the call has closed.
+				using CancellationTokenSource CancellationSource = new CancellationTokenSource();
 				NextRequestTask = Reader.MoveNext();
+				NextRequestTask = NextRequestTask.ContinueWith(Task => { CancellationSource.Cancel(); return Task.Result; }, TaskScheduler.Current);
 
 				// Get the current agent state
 				IAgent? Agent = await AgentService.GetAgentAsync(new AgentId(Request.AgentId));
@@ -412,15 +415,8 @@ namespace HordeServer.Services
 						GetCapabilities(Request.Capabilities, out Properties, out Resources);
 					}
 
-					// Capture the current agent update index. This allows us to detect if anything has changed.
-					uint UpdateIndex = Agent.UpdateIndex;
-
-					// Update the agent session and return to the caller if anything changes
-					Agent = await AgentService.UpdateSessionAsync(Agent, Request.SessionId.ToObjectId(), Request.Status, Properties, Resources, Request.Leases);
-					if (Agent != null && Agent.UpdateIndex == UpdateIndex && (Agent.Leases.Count > 0 || Agent.Status != AgentStatus.Stopping))
-					{
-						Agent = await AgentService.WaitForLeaseAsync(Agent, NextRequestTask);
-					}
+					// Update the session
+					Agent = await AgentService.UpdateSessionWithWaitAsync(Agent, Request.SessionId.ToObjectId(), Request.Status, Properties, Resources, Request.Leases, CancellationSource.Token);
 				}
 
 				// Handle the invalid agent case
diff --git a/Engine/Source/Programs/Horde/HordeServer/Startup.cs b/Engine/Source/Programs/Horde/HordeServer/Startup.cs
--- a/Engine/Source/Programs/Horde/HordeServer/Startup.cs
+++ b/Engine/Source/Programs/Horde/HordeServer/Startup.cs
@@ -528,7 +528,6 @@ namespace HordeServer
 			Services.AddHostedService(Provider => Provider.GetRequiredService<CommitService>());
 			Services.AddHostedService(Provider => Provider.GetRequiredService<ConsistencyService>());
 			Services.AddHostedService(Provider => (DowntimeService)Provider.GetRequiredService<IDowntimeService>());
-			Services.AddHostedService(Provider => Provider.GetRequiredService<JobTaskSource>());
 			Services.AddHostedService(Provider => Provider.GetRequiredService<IIssueService>());
 			Services.AddHostedService(Provider => (LogFileService)Provider.GetRequiredService<ILogFileService>());
 			Services.AddHostedService(Provider => (NotificationService)Provider.GetRequiredService<INotificationService>());
@@ -545,7 +544,7 @@ namespace HordeServer
 			Services.AddHostedService<TelemetryService>();
 			Services.AddHostedService(Provider => Provider.GetRequiredService<DeviceService>());
 
-			// Task sources. Order of registration is important here; it dictates the order in which sources are served.
+			// Task sources. Order of registration is important here; it dictates the priority in which sources are served.
 			Services.AddSingleton<JobTaskSource>();
 			Services.AddSingleton<ConformTaskSource>();
 			Services.AddSingleton<IComputeService, ComputeService>();
@@ -555,7 +554,7 @@ namespace HordeServer
 			Services.AddSingleton<ITaskSource, RestartTaskSource>();
 			Services.AddSingleton<ITaskSource, ConformTaskSource>(Provider => Provider.GetRequiredService<ConformTaskSource>());
 			Services.AddSingleton<ITaskSource, JobTaskSource>(Provider => Provider.GetRequiredService<JobTaskSource>());
-			Services.AddSingleton<ITaskSource>(Provider => new NewTaskSourceWrapper(Provider.GetRequiredService<IComputeService>()));
+			Services.AddSingleton<ITaskSource, ComputeService>();
 
 			Services.AddHostedService(Provider => Provider.GetRequiredService<ConformTaskSource>());
 
diff --git a/Engine/Source/Programs/Shared/EpicGames.Core/AsyncUtils.cs b/Engine/Source/Programs/Shared/EpicGames.Core/AsyncUtils.cs
--- a/Engine/Source/Programs/Shared/EpicGames.Core/AsyncUtils.cs
+++ b/Engine/Source/Programs/Shared/EpicGames.Core/AsyncUtils.cs
@@ -3,6 +3,7 @@
 using System;
 using System.Collections.Generic;
 using System.Text;
+using System.Threading;
 using System.Threading.Tasks;
 
 namespace EpicGames.Core
@@ -12,6 +13,27 @@ namespace EpicGames.Core
 	/// </summary>
 	public static class AsyncUtils
 	{
+		/// <summary>
+		/// Converts a cancellation token to a waitable task
+		/// </summary>
+		/// <param name="Token">Cancellation token</param>
+		/// <returns></returns>
+		public static Task AsTask(this CancellationToken Token)
+		{
+			return Task.Delay(-1, Token).ContinueWith(x => { });
+		}
+
+		/// <summary>
+		/// Waits for a time period to elapse or the task to be cancelled, without throwing an cancellation exception
+		/// </summary>
+		/// <param name="Time">Time to wait</param>
+		/// <param name="Token">Cancellation token</param>
+		/// <returns></returns>
+		public static Task DelayNoThrow(TimeSpan Time, CancellationToken Token)
+		{
+			return Task.Delay(Time, Token).ContinueWith(x => { });
+		}
+
 		/// <summary>
 		/// Removes all the complete tasks from a list, allowing each to throw exceptions as necessary
 		/// </summary>
