Horde: Add a generic log for operations performed by an agent, which includes a trace of any setting changes (with the user that made them) as well as any work performed, sessions started/ended, autoscaling operations and so on.

The /api/v1/agents/XXX/history endpoint returns a list of log messages, queryable by time range and supporting pagination.

diff --git a/Engine/Source/Programs/Horde/HordeServer/Collections/IAgentCollection.cs b/Engine/Source/Programs/Horde/HordeServer/Collections/IAgentCollection.cs
--- a/Engine/Source/Programs/Horde/HordeServer/Collections/IAgentCollection.cs
+++ b/Engine/Source/Programs/Horde/HordeServer/Collections/IAgentCollection.cs
@@ -149,5 +149,12 @@ namespace HordeServer.Collections
 		/// <param name="LeaseIdx">Index of the lease</param>
 		/// <returns>New agent state if it succeeded, otherwise null</returns>
 		Task<IAgent?> TryCancelLeaseAsync(IAgent Agent, int LeaseIdx);
+
+		/// <summary>
+		/// Gets the log channel for an agent
+		/// </summary>
+		/// <param name="AgentId"></param>
+		/// <returns></returns>
+		IAuditLogChannel<AgentId> GetLogger(AgentId AgentId);
 	}
 }
diff --git a/Engine/Source/Programs/Horde/HordeServer/Collections/IAuditLog.cs b/Engine/Source/Programs/Horde/HordeServer/Collections/IAuditLog.cs
--- a/Engine/Source/Programs/Horde/HordeServer/Collections/IAuditLog.cs
+++ b/Engine/Source/Programs/Horde/HordeServer/Collections/IAuditLog.cs
@@ -24,16 +24,10 @@ using System.Text;
 namespace HordeServer.Collections
 {
 	/// <summary>
-	/// Message in an audit log
+	/// Message from an audit log
 	/// </summary>
-	/// <typeparam name="TTargetId">Type of entity that the log is for</typeparam>
-	public interface IAuditLogMessage<TTargetId>
+	public interface IAuditLogMessage
 	{
-		/// <summary>
-		/// Unique id for the entity
-		/// </summary>
-		public TTargetId TargetId { get; }
-
 		/// <summary>
 		/// Timestamp for the event
 		/// </summary>
@@ -51,118 +45,77 @@ namespace HordeServer.Collections
 	}
 
 	/// <summary>
-	/// Concrete implementation of <see cref="IAuditLogMessage{T}"/> used to add new entries
+	/// Channel for a particular entity
 	/// </summary>
-	/// <typeparam name="TTargetId"></typeparam>
-	public class NewAuditLogMessage<TTargetId> : IAuditLogMessage<TTargetId>
+	public interface IAuditLogChannel : ILogger
 	{
-		/// <inheritdoc/>
-		public TTargetId TargetId { get; set; }
-
-		/// <inheritdoc/>
-		public DateTime TimeUtc { get; set; }
-
-		/// <inheritdoc/>
-		public LogLevel Level { get; }
-
-		/// <inheritdoc/>
-		public string Data { get; set; }
-
-		/// <summary>
-		/// Constructor
-		/// </summary>
-		public NewAuditLogMessage(TTargetId EntityId, DateTime TimeUtc, LogLevel Level, string Data)
-		{
-			this.TargetId = EntityId;
-			this.TimeUtc = TimeUtc;
-			this.Level = Level;
-			this.Data = Data;
-		}
-	}
-
-	/// <summary>
-	/// Interface for a collection of log messages for a particular document type
-	/// </summary>
-	public interface IAuditLog<TTargetId>
-	{
-		/// <summary>
-		/// Adds a message with the given properties
-		/// </summary>
-		/// <param name="NewMessage">The message to add</param>
-		Task AddAsync(NewAuditLogMessage<TTargetId> NewMessage);
-
 		/// <summary>
 		/// Finds messages matching certain criteria
 		/// </summary>
-		/// <param name="Id">Identifier for the </param>
 		/// <param name="MinTime"></param>
 		/// <param name="MaxTime"></param>
 		/// <param name="Index"></param>
 		/// <param name="Count"></param>
 		/// <returns></returns>
-		IAsyncEnumerable<IAuditLogMessage<TTargetId>> FindAsync(TTargetId Id, DateTime? MinTime = null, DateTime? MaxTime = null, int? Index = null, int? Count = null);
+		IAsyncEnumerable<IAuditLogMessage> FindAsync(DateTime? MinTime = null, DateTime? MaxTime = null, int? Index = null, int? Count = null);
 
 		/// <summary>
 		/// Deletes messages between a given time range for a particular object
 		/// </summary>
-		/// <param name="Id">Id of the object</param>
 		/// <param name="MinTime">Minimum time to remove</param>
 		/// <param name="MaxTime">Maximum time to remove</param>
 		/// <returns>Async task</returns>
-		Task<long> DeleteAsync(TTargetId Id, DateTime? MinTime = null, DateTime? MaxTime = null);
+		Task<long> DeleteAsync(DateTime? MinTime = null, DateTime? MaxTime = null);
 	}
 
 	/// <summary>
-	/// Factory for instantiating audit log instances
+	/// Channel for a particular entity
 	/// </summary>
-	/// <typeparam name="TTargetId"></typeparam>
-	public interface IAuditLogFactory<TTargetId>
+	public interface IAuditLogChannel<TSubject> : IAuditLogChannel
 	{
 		/// <summary>
-		/// Create a new audit log instance, with the given database name
+		/// Identifier for the subject of this channel
 		/// </summary>
-		/// <param name="Name"></param>
-		/// <returns></returns>
-		IAuditLog<TTargetId> Create(string Name);
+		TSubject Subject { get; }
 	}
 
 	/// <summary>
-	/// Extension methods for AuditLog instances
+	/// Message in an audit log
 	/// </summary>
-	public static class AuditLogExtensions
+	/// <typeparam name="TTargetId">Type of entity that the log is for</typeparam>
+	public interface IAuditLogMessage<TTargetId> : IAuditLogMessage
 	{
 		/// <summary>
-		/// Write a message to an audit log
+		/// Unique id for the entity
 		/// </summary>
-		public static Task WriteAsync<TTargetId>(this IAuditLog<TTargetId> AuditLog, TTargetId Id, LogLevel Level, string Format, params object[] Args)
-		{
-			return WriteAsync(AuditLog, Id, DateTime.UtcNow, Level, new EventId(0), Format, Args);
-		}
+		public TTargetId Subject { get; }
+	}
 
+	/// <summary>
+	/// Interface for a collection of log messages for a particular document type
+	/// </summary>
+	public interface IAuditLog<TSubject>
+	{
 		/// <summary>
-		/// Write a message to an audit log
+		/// Get the channel for a particular subject
 		/// </summary>
-		public static Task WriteAsync<TTargetId>(this IAuditLog<TTargetId> AuditLog, TTargetId Id, LogLevel Level, EventId EventId, string Format, params object[] Args)
-		{
-			return WriteAsync(AuditLog, Id, DateTime.UtcNow, Level, EventId, Format, Args);
-		}
+		/// <param name="Subject"></param>
+		/// <returns></returns>
+		IAuditLogChannel<TSubject> this[TSubject Subject] { get; }
+	}
 
+	/// <summary>
+	/// Factory for instantiating audit log instances
+	/// </summary>
+	/// <typeparam name="TSubject"></typeparam>
+	public interface IAuditLogFactory<TSubject>
+	{
 		/// <summary>
-		/// Write a message to an audit log
+		/// Create a new audit log instance, with the given database name
 		/// </summary>
-		public static Task WriteAsync<TTargetId>(this IAuditLog<TTargetId> AuditLog, TTargetId Id, DateTime TimeUtc, LogLevel Level, EventId EventId, string Format, params object[] Args)
-		{
-			ArrayBufferWriter<byte> Buffer = new ArrayBufferWriter<byte>();
-			using (Utf8JsonWriter Writer = new Utf8JsonWriter(Buffer))
-			{
-				Dictionary<string, object> Properties = new Dictionary<string, object>();
-				MessageTemplate.ParsePropertyValues(Format, Args, Properties);
-				string Text = MessageTemplate.Render(Format, Properties);
-				MessageTemplate.SerializeJson(Writer, TimeUtc, Level, EventId, Format, Properties);
-			}
-
-			NewAuditLogMessage<TTargetId> NewMessage = new NewAuditLogMessage<TTargetId>(Id, TimeUtc, Level, Encoding.UTF8.GetString(Buffer.WrittenSpan));
-			return AuditLog.AddAsync(NewMessage);
-		}
+		/// <param name="CollectionName"></param>
+		/// <param name="SubjectProperty"></param>
+		/// <returns></returns>
+		IAuditLog<TSubject> Create(string CollectionName, string SubjectProperty);
 	}
 }
diff --git a/Engine/Source/Programs/Horde/HordeServer/Collections/Impl/AgentCollection.cs b/Engine/Source/Programs/Horde/HordeServer/Collections/Impl/AgentCollection.cs
--- a/Engine/Source/Programs/Horde/HordeServer/Collections/Impl/AgentCollection.cs
+++ b/Engine/Source/Programs/Horde/HordeServer/Collections/Impl/AgentCollection.cs
@@ -6,6 +6,7 @@ using HordeCommon.Rpc.Tasks;
 using HordeServer.Models;
 using HordeServer.Services;
 using HordeServer.Utilities;
+using Microsoft.Extensions.Logging;
 using MongoDB.Bson;
 using MongoDB.Bson.Serialization.Attributes;
 using MongoDB.Driver;
@@ -131,18 +132,17 @@ namespace HordeServer.Collections.Impl
 			}
 		}
 
-		/// <summary>
-		/// Collection of agent documents
-		/// </summary>
 		readonly IMongoCollection<AgentDocument> Agents;
+		readonly IAuditLog<AgentId> AuditLog;
 
 		/// <summary>
 		/// Constructor
 		/// </summary>
-		/// <param name="DatabaseService">The database service instance</param>
-		public AgentCollection(DatabaseService DatabaseService)
+		public AgentCollection(DatabaseService DatabaseService, IAuditLog<AgentId> AuditLog)
 		{
-			Agents = DatabaseService.GetCollection<AgentDocument>("Agents");
+			this.Agents = DatabaseService.GetCollection<AgentDocument>("Agents");
+			this.AuditLog = AuditLog;
+
 			if (!DatabaseService.ReadOnlyMode)
 			{
 				Agents.Indexes.CreateOne(new CreateIndexModel<AgentDocument>(Builders<AgentDocument>.IndexKeys.Ascending(x => x.Deleted).Ascending(x => x.Id).Ascending(x => x.Pools)));
@@ -483,5 +483,11 @@ namespace HordeServer.Collections.Impl
 			UpdateDefinition<AgentDocument> Update = Builders<AgentDocument>.Update.Set(x => x.Leases![LeaseIdx].State, LeaseState.Cancelled);
 			return await TryUpdateAsync(Agent, Update);
 		}
+
+		/// <inheritdoc/>
+		public IAuditLogChannel<AgentId> GetLogger(AgentId AgentId)
+		{
+			return AuditLog[AgentId];
+		}
 	}
 }
diff --git a/Engine/Source/Programs/Horde/HordeServer/Collections/Impl/AuditLog.cs b/Engine/Source/Programs/Horde/HordeServer/Collections/Impl/AuditLog.cs
--- a/Engine/Source/Programs/Horde/HordeServer/Collections/Impl/AuditLog.cs
+++ b/Engine/Source/Programs/Horde/HordeServer/Collections/Impl/AuditLog.cs
@@ -1,5 +1,6 @@
 // Copyright Epic Games, Inc. All Rights Reserved.
 
+using EpicGames.Core;
 using Google.Protobuf.WellKnownTypes;
 using HordeCommon;
 using HordeCommon.Rpc.Tasks;
@@ -11,24 +12,24 @@ using MongoDB.Bson;
 using MongoDB.Bson.Serialization.Attributes;
 using MongoDB.Driver;
 using System;
+using System.Buffers;
 using System.Collections.Generic;
 using System.Linq;
 using System.Linq.Expressions;
+using System.Text.Json;
+using System.Threading.Channels;
 using System.Threading.Tasks;
 
 namespace HordeServer.Collections.Impl
 {
-	/// <summary>
-	/// Collection of agent documents
-	/// </summary>
-	class AuditLog<TTargetId> : IAuditLog<TTargetId>
+	class AuditLog<TSubject> : IAuditLog<TSubject>, IAsyncDisposable, IDisposable
 	{
-		class AuditLogMessage : IAuditLogMessage<TTargetId>
+		class AuditLogMessage : IAuditLogMessage<TSubject>
 		{
 			public ObjectId Id { get; set; }
 
-			[BsonElement("tid")]
-			public TTargetId TargetId { get; set; }
+			[BsonElement("s")]
+			public TSubject Subject { get; set; }
 
 			[BsonElement("t")]
 			public DateTime TimeUtc { get; set; }
@@ -41,40 +42,107 @@ namespace HordeServer.Collections.Impl
 
 			public AuditLogMessage()
 			{
-				TargetId = default!;
+				Subject = default!;
 				Data = String.Empty;
 			}
 
-			public AuditLogMessage(IAuditLogMessage<TTargetId> Other)
+			public AuditLogMessage(TSubject Subject, DateTime TimeUtc, LogLevel Level, string Data)
 			{
 				this.Id = ObjectId.GenerateNewId();
-				this.TargetId = Other.TargetId;
-				this.TimeUtc = Other.TimeUtc;
-				this.Level = Other.Level;
-				this.Data = Other.Data;
+				this.Subject = Subject;
+				this.TimeUtc = TimeUtc;
+				this.Level = Level;
+				this.Data = Data;
 			}
 		}
 
+		class AuditLogChannel : IAuditLogChannel<TSubject>
+		{
+			public readonly AuditLog<TSubject> Outer;
+			public TSubject Subject { get; }
+
+			public AuditLogChannel(AuditLog<TSubject> Outer, TSubject Subject)
+			{
+				this.Outer = Outer;
+				this.Subject = Subject;
+			}
+
+			public IDisposable? BeginScope<TState>(TState state) => null;
+
+			public bool IsEnabled(LogLevel logLevel) => true;
+
+			public void Log<TState>(LogLevel LogLevel, EventId EventId, TState State, Exception Exception, Func<TState, Exception, string> Formatter)
+			{
+				DateTime Time = DateTime.UtcNow;
+
+				string Data = MessageTemplate.Serialize(LogLevel, EventId, State, Exception, Formatter);
+				AuditLogMessage Message = new AuditLogMessage(Subject, Time, LogLevel, Data);
+				Outer.MessageChannel.Writer.TryWrite(Message);
+
+				using (IDisposable _ = Outer.Logger.BeginScope($"Subject: {{{Outer.SubjectProperty}}}", Subject))
+				{
+					Outer.Logger.Log(LogLevel, EventId, State, Exception, Formatter);
+				}
+			}
+
+			public IAsyncEnumerable<IAuditLogMessage> FindAsync(DateTime? MinTime, DateTime? MaxTime, int? Index, int? Count) => Outer.FindAsync(Subject, MinTime, MaxTime, Index, Count);
+
+			public Task<long> DeleteAsync(DateTime? MinTime, DateTime? MaxTime) => Outer.DeleteAsync(Subject, MinTime, MaxTime);
+		}
+
 		IMongoCollection<AuditLogMessage> Messages;
+		Channel<AuditLogMessage> MessageChannel;
+		string SubjectProperty;
+		ILogger Logger;
+		Task BackgroundTask;
+
+		public IAuditLogChannel<TSubject> this[TSubject Subject] => new AuditLogChannel(this, Subject);
 
-		public AuditLog(DatabaseService DatabaseService, string CollectionName)
+		public AuditLog(DatabaseService DatabaseService, string CollectionName, string SubjectProperty, ILogger Logger)
 		{
-			Messages = DatabaseService.GetCollection<AuditLogMessage>(CollectionName);
+			this.Messages = DatabaseService.GetCollection<AuditLogMessage>(CollectionName);
+			this.MessageChannel = Channel.CreateUnbounded<AuditLogMessage>();
+			this.SubjectProperty = SubjectProperty;
+			this.Logger = Logger;
+
 			if (!DatabaseService.ReadOnlyMode)
 			{
 				Messages.Indexes.CreateOne(new CreateIndexModel<AuditLogMessage>(Builders<AuditLogMessage>.IndexKeys.Ascending(x => x.Id).Descending(x => x.TimeUtc)));
 			}
+
+			BackgroundTask = Task.Run(() => WriteMessagesAsync());
+		}
+
+		public async ValueTask DisposeAsync()
+		{
+			MessageChannel.Writer.TryComplete();
+			await BackgroundTask;
+		}
+
+		public void Dispose()
+		{
+			DisposeAsync().AsTask().Wait();
 		}
 
-		public Task AddAsync(NewAuditLogMessage<TTargetId> NewMessage)
+		async Task WriteMessagesAsync()
 		{
-			AuditLogMessage Message = new AuditLogMessage(NewMessage);
-			return Messages.InsertOneAsync(Message);
+			while (await MessageChannel.Reader.WaitToReadAsync())
+			{
+				List<AuditLogMessage> NewMessages = new List<AuditLogMessage>();
+				while (MessageChannel.Reader.TryRead(out AuditLogMessage? NewMessage))
+				{
+					NewMessages.Add(NewMessage);
+				}
+				if (NewMessages.Count > 0)
+				{
+					await Messages.InsertManyAsync(NewMessages);
+				}
+			}
 		}
 
-		public async IAsyncEnumerable<IAuditLogMessage<TTargetId>> FindAsync(TTargetId Id, DateTime? MinTime = null, DateTime? MaxTime = null, int? Index = null, int? Count = null)
+		async IAsyncEnumerable<IAuditLogMessage<TSubject>> FindAsync(TSubject Subject, DateTime? MinTime = null, DateTime? MaxTime = null, int? Index = null, int? Count = null)
 		{
-			FilterDefinition<AuditLogMessage> Filter = Builders<AuditLogMessage>.Filter.Eq(x => x.TargetId, Id);
+			FilterDefinition<AuditLogMessage> Filter = Builders<AuditLogMessage>.Filter.Eq(x => x.Subject, Subject);
 			if (MinTime != null)
 			{
 				Filter &= Builders<AuditLogMessage>.Filter.Gte(x => x.TimeUtc, MinTime.Value);
@@ -84,7 +152,7 @@ namespace HordeServer.Collections.Impl
 				Filter &= Builders<AuditLogMessage>.Filter.Lte(x => x.TimeUtc, MaxTime.Value);
 			}
 
-			using (IAsyncCursor<AuditLogMessage> Cursor = await Messages.Find(Filter).SortByDescending(x => x.TimeUtc).ToCursorAsync())
+			using (IAsyncCursor<AuditLogMessage> Cursor = await Messages.Find(Filter).SortByDescending(x => x.TimeUtc).Range(Index, Count).ToCursorAsync())
 			{
 				while (await Cursor.MoveNextAsync())
 				{
@@ -96,9 +164,9 @@ namespace HordeServer.Collections.Impl
 			}
 		}
 
-		public async Task<long> DeleteAsync(TTargetId Id, DateTime? MinTime = null, DateTime? MaxTime = null)
+		async Task<long> DeleteAsync(TSubject Subject, DateTime? MinTime = null, DateTime? MaxTime = null)
 		{
-			FilterDefinition<AuditLogMessage> Filter = Builders<AuditLogMessage>.Filter.Eq(x => x.TargetId, Id);
+			FilterDefinition<AuditLogMessage> Filter = Builders<AuditLogMessage>.Filter.Eq(x => x.Subject, Subject);
 			if (MinTime != null)
 			{
 				Filter &= Builders<AuditLogMessage>.Filter.Gte(x => x.TimeUtc, MinTime.Value);
@@ -113,27 +181,20 @@ namespace HordeServer.Collections.Impl
 		}
 	}
 
-	/// <summary>
-	/// Factory for instantiating audit log instances
-	/// </summary>
-	/// <typeparam name="TTargetId"></typeparam>
-	public class AuditLogFactory<TTargetId> : IAuditLogFactory<TTargetId>
+	class AuditLogFactory<TSubject> : IAuditLogFactory<TSubject>
 	{
 		DatabaseService DatabaseService;
+		ILogger<AuditLog<TSubject>> Logger;
 
-		/// <summary>
-		/// Constructor
-		/// </summary>
-		/// <param name="DatabaseService"></param>
-		public AuditLogFactory(DatabaseService DatabaseService)
+		public AuditLogFactory(DatabaseService DatabaseService, ILogger<AuditLog<TSubject>> Logger)
 		{
 			this.DatabaseService = DatabaseService;
+			this.Logger = Logger;
 		}
 
-		/// <inheritdoc/>
-		public IAuditLog<TTargetId> Create(string Name)
+		public IAuditLog<TSubject> Create(string CollectionName, string SubjectProperty)
 		{
-			return new AuditLog<TTargetId>(DatabaseService, Name);
+			return new AuditLog<TSubject>(DatabaseService, CollectionName, SubjectProperty, Logger);
 		}
 	}
 }
diff --git a/Engine/Source/Programs/Horde/HordeServer/Compute/Impl/ComputeService.cs b/Engine/Source/Programs/Horde/HordeServer/Compute/Impl/ComputeService.cs
--- a/Engine/Source/Programs/Horde/HordeServer/Compute/Impl/ComputeService.cs
+++ b/Engine/Source/Programs/Horde/HordeServer/Compute/Impl/ComputeService.cs
@@ -128,9 +128,9 @@ namespace HordeServer.Compute.Impl
 	/// <summary>
 	/// Dispatches remote actions. Does not implement any cross-pod communication to satisfy leases; only agents connected to this server instance will be stored.
 	/// </summary>
-	class ComputeService : IComputeService, IDisposable
+	class ComputeService : TaskSourceBase<ComputeTaskMessage>, IComputeService, IDisposable
 	{
-		public MessageDescriptor Descriptor => ComputeTaskMessage.Descriptor;
+		public override string Type => "Compute";
 
 		public static NamespaceId DefaultNamespaceId { get; } = new NamespaceId("default");
 
@@ -210,7 +210,7 @@ namespace HordeServer.Compute.Impl
 		}
 
 		/// <inheritdoc/>
-		public async Task<AgentLease?> AssignLeaseAsync(IAgent Agent, CancellationToken CancellationToken)
+		public override async Task<AgentLease?> AssignLeaseAsync(IAgent Agent, CancellationToken CancellationToken)
 		{
 			// If the agent is disabled, just block until we cancel
 			if (!Agent.Enabled)
@@ -247,9 +247,8 @@ namespace HordeServer.Compute.Impl
 		}
 
 		/// <inheritdoc/>
-		public Task CancelLeaseAsync(IAgent Agent, ObjectId LeaseId, Any Payload)
+		public override Task CancelLeaseAsync(IAgent Agent, ObjectId LeaseId, ComputeTaskMessage Message)
 		{
-			ComputeTaskMessage Message = Payload.Unpack<ComputeTaskMessage>();
 			ComputeTaskInfo TaskInfo = new ComputeTaskInfo(Message.Task, new ChannelId(Message.ChannelId));
 			return TaskScheduler.EnqueueAsync(Message.Requirements.Hash, TaskInfo, true);
 		}
@@ -267,19 +266,16 @@ namespace HordeServer.Compute.Impl
 			return Messages.ConvertAll<IComputeTaskStatus>(x => x);
 		}
 
-		public Task OnLeaseStartedAsync(IAgent Agent, ObjectId LeaseId, Any Payload)
+		public override Task OnLeaseStartedAsync(IAgent Agent, ObjectId LeaseId, ComputeTaskMessage ComputeTask, ILogger Logger)
 		{
-			ComputeTaskMessage ComputeTask = Payload.Unpack<ComputeTaskMessage>();
 			Logger.LogInformation("Compute lease started (lease: {LeaseId}, task: {TaskHash}, agent: {AgentId}, channel: {ChannelId})", LeaseId, ComputeTask.Task.Hash, Agent.Id, ComputeTask.ChannelId);
 
 			ComputeTaskStatus Status = new ComputeTaskStatus(ComputeTask.Task, ComputeTaskState.Executing, Agent.Id, LeaseId);
 			return MessageQueue.PostAsync(ComputeTask.ChannelId, Status);
 		}
 
-		public Task OnLeaseFinishedAsync(IAgent Agent, ObjectId LeaseId, Any Payload, LeaseOutcome Outcome, ReadOnlyMemory<byte> Output)
+		public override Task OnLeaseFinishedAsync(IAgent Agent, ObjectId LeaseId, ComputeTaskMessage ComputeTask, LeaseOutcome Outcome, ReadOnlyMemory<byte> Output, ILogger Logger)
 		{
-			ComputeTaskMessage ComputeTask = Payload.Unpack<ComputeTaskMessage>();
-
 			ComputeTaskResultMessage Message = ComputeTaskResultMessage.Parser.ParseFrom(Output.ToArray());
 
 			ComputeTaskStatus Status = new ComputeTaskStatus(ComputeTask.Task, ComputeTaskState.Complete, Agent.Id, LeaseId);
diff --git a/Engine/Source/Programs/Horde/HordeServer/Controllers/AgentsController.cs b/Engine/Source/Programs/Horde/HordeServer/Controllers/AgentsController.cs
--- a/Engine/Source/Programs/Horde/HordeServer/Controllers/AgentsController.cs
+++ b/Engine/Source/Programs/Horde/HordeServer/Controllers/AgentsController.cs
@@ -1,30 +1,24 @@
 // Copyright Epic Games, Inc. All Rights Reserved.
 
 using HordeServer.Api;
+using HordeServer.Collections;
 using HordeServer.Models;
 using HordeServer.Services;
 using HordeServer.Utilities;
-using Microsoft.AspNetCore.Authentication;
 using Microsoft.AspNetCore.Authorization;
 using Microsoft.AspNetCore.Mvc;
-using Microsoft.AspNetCore.Mvc.ModelBinding;
-using Microsoft.Extensions.Options;
-using Microsoft.IdentityModel.Tokens;
+using Microsoft.Extensions.Logging;
 using MongoDB.Bson;
 using System;
 using System.Collections.Generic;
-using System.ComponentModel.DataAnnotations;
-using System.IdentityModel.Tokens.Jwt;
 using System.Linq;
-using System.Reflection;
-using System.Security.Claims;
-using System.Text.Json;
+using System.Text;
 using System.Threading.Tasks;
 
 namespace HordeServer.Controllers
 {
-	using PoolId = StringId<IPool>;
 	using AgentSoftwareChannelName = StringId<AgentSoftwareChannels>;
+	using PoolId = StringId<IPool>;
 
 	/// <summary>
 	/// Controller for the /api/v1/agents endpoint
@@ -126,8 +120,6 @@ namespace HordeServer.Controllers
 				return NotFound();
 			}
 
-			List<IPool> Pools = await PoolService.GetPoolsAsync();
-
 			GlobalPermissionsCache Cache = new GlobalPermissionsCache();
 			bool bIncludeAcl = await AgentService.AuthorizeAsync(Agent, AclAction.ViewPermissions, User, Cache);
 			return new GetAgentResponse(Agent, bIncludeAcl).ApplyFilter(Filter);
@@ -143,24 +135,65 @@ namespace HordeServer.Controllers
 		[Route("/api/v1/agents/{AgentId}")]
 		public async Task<ActionResult> UpdateAgentAsync(string AgentId, [FromBody] UpdateAgentRequest Update)
 		{
-			IAgent? Agent = await AgentService.GetAgentAsync(new AgentId(AgentId));
-			if (Agent == null)
-			{
-				return NotFound();
-			}
+			List<PoolId>? UpdatePools = Update.Pools?.ConvertAll(x => new PoolId(x));
+			AgentSoftwareChannelName? Channel = String.IsNullOrEmpty(Update.Channel) ? (AgentSoftwareChannelName?)null : new AgentSoftwareChannelName(Update.Channel);
+
+			string UserName = User.GetUserName() ?? "Unknown";
 
 			GlobalPermissionsCache Cache = new GlobalPermissionsCache();
-			if (!await AgentService.AuthorizeAsync(Agent, AclAction.UpdateAgent, User, Cache))
+			for (; ; )
 			{
-				return Forbid();
+				IAgent? Agent = await AgentService.GetAgentAsync(new AgentId(AgentId));
+				if (Agent == null)
+				{
+					return NotFound();
+				}
+				if (!await AgentService.AuthorizeAsync(Agent, AclAction.UpdateAgent, User, Cache))
+				{
+					return Forbid();
+				}
+				if (Update.Acl != null && !await AgentService.AuthorizeAsync(Agent, AclAction.ChangePermissions, User, Cache))
+				{
+					return Forbid();
+				}
+
+				IAgent? NewAgent = await AgentService.Agents.TryUpdateSettingsAsync(Agent, Update.Enabled, Update.RequestConform, Update.RequestRestart, Update.RequestShutdown, Channel, Update.Pools?.ConvertAll(x => new PoolId(x)), Acl.Merge(Agent.Acl, Update.Acl), Update.Comment);
+				if (NewAgent == null)
+				{
+					continue;
+				}
+
+				IAuditLogChannel<AgentId> Logger = AgentService.Agents.GetLogger(Agent.Id);
+				if (Agent.Enabled != NewAgent.Enabled)
+				{
+					Logger.LogInformation("Setting changed: Enabled = {State} ({UserName})", NewAgent.Enabled, UserName);
+				}
+				if (Agent.RequestConform != NewAgent.RequestConform)
+				{
+					Logger.LogInformation("Setting changed: RequestConfig = {State} ({UserName})", NewAgent.RequestConform, UserName);
+				}
+				if (Agent.RequestRestart != NewAgent.RequestRestart)
+				{
+					Logger.LogInformation("Setting changed: RequestRestart = {State} ({UserName})", NewAgent.RequestRestart, UserName);
+				}
+				if (Agent.RequestShutdown != NewAgent.RequestShutdown)
+				{
+					Logger.LogInformation("Setting changed: RequestShutdown = {State} ({UserName})", NewAgent.RequestShutdown, UserName);
+				}
+				if (Agent.Comment != NewAgent.Comment)
+				{
+					Logger.LogInformation("Setting changed: Comment = \"{Comment}\" ({UserName})", Update.Comment, UserName);
+				}
+				foreach (PoolId AddedPool in NewAgent.ExplicitPools.Except(Agent.ExplicitPools))
+				{
+					Logger.LogInformation("Added to pool {PoolId} ({UserName})", AddedPool, UserName);
+				}
+				foreach (PoolId RemovedPool in Agent.ExplicitPools.Except(NewAgent.ExplicitPools))
+				{
+					Logger.LogInformation("Removed from pool {PoolId} ({UserName})", RemovedPool, UserName);
+				}
+				break;
 			}
-			if (Update.Acl != null && !await AgentService.AuthorizeAsync(Agent, AclAction.ChangePermissions, User, Cache))
-			{
-				return Forbid();
-			}
-
-			AgentSoftwareChannelName? Channel = String.IsNullOrEmpty(Update.Channel) ? (AgentSoftwareChannelName?)null : new AgentSoftwareChannelName(Update.Channel);
-			await AgentService.UpdateAgentAsync(Agent, Update.Enabled, Update.RequestConform, Update.RequestRestart, Update.RequestShutdown, Channel, Update.Pools?.ConvertAll(x => new PoolId(x)), Acl.Merge(Agent.Acl, Update.Acl), Update.Comment);
 			return Ok();
 		}
 
@@ -187,6 +220,41 @@ namespace HordeServer.Controllers
 			return new OkResult();
 		}
 
+		/// <summary>
+		/// Retrieve historical information about a specific agent
+		/// </summary>
+		/// <param name="AgentId">Id of the agent to get information about</param>
+		/// <param name="MinTime">Minimum time for records to return</param>
+		/// <param name="MaxTime">Maximum time for records to return</param>
+		/// <param name="Index">Offset of the first result</param>
+		/// <param name="Count">Number of records to return</param>
+		/// <returns>Information about the requested agent</returns>
+		[HttpGet]
+		[Route("/api/v1/agents/{AgentId}/history")]
+		[ProducesResponseType(typeof(GetAgentHistoryResponse), 200)]
+		public async Task GetAgentHistoryAsync(string AgentId, [FromQuery] DateTime? MinTime = null, [FromQuery] DateTime? MaxTime = null, [FromQuery] int Index = 0, [FromQuery] int Count = 50)
+		{
+			IAuditLogChannel<AgentId> Channel = AgentService.Agents.GetLogger(AgentId.ToAgentId());
+
+			Response.ContentType = "application/json";
+			Response.StatusCode = 200;
+			await Response.StartAsync();
+
+			string Prefix = "{\n\t\"entries\":\n\t[";
+			await Response.BodyWriter.WriteAsync(Encoding.UTF8.GetBytes(Prefix));
+
+			string Separator = "";
+			await foreach (IAuditLogMessage<AgentId> Message in Channel.FindAsync(MinTime, MaxTime, Index, Count))
+			{
+				string Line = $"{Separator}\n\t\t{Message.Data}";
+				await Response.Body.WriteAsync(Encoding.UTF8.GetBytes(Line));
+				Separator = ",";
+			}
+
+			string Suffix = "\n\t]\n}";
+			await Response.BodyWriter.WriteAsync(Encoding.UTF8.GetBytes(Suffix));
+		}
+
 		/// <summary>
 		/// Find all the sessions of a particular agent
 		/// </summary>
diff --git a/Engine/Source/Programs/Horde/HordeServer/Services/AgentService.cs b/Engine/Source/Programs/Horde/HordeServer/Services/AgentService.cs
--- a/Engine/Source/Programs/Horde/HordeServer/Services/AgentService.cs
+++ b/Engine/Source/Programs/Horde/HordeServer/Services/AgentService.cs
@@ -1,6 +1,7 @@
 // Copyright Epic Games, Inc. All Rights Reserved.
 
 using EpicGames.Core;
+using Google.Protobuf;
 using Google.Protobuf.WellKnownTypes;
 using HordeCommon;
 using HordeServer.Collections;
@@ -58,7 +59,7 @@ namespace HordeServer.Services
 		/// <summary>
 		/// Collection of agent documents
 		/// </summary>
-		IAgentCollection Agents;
+		public IAgentCollection Agents { get; }
 
 		/// <summary>
 		/// Collection of lease documents
@@ -85,6 +86,11 @@ namespace HordeServer.Services
 		/// </summary>
 		IHostApplicationLifetime ApplicationLifetime;
 
+		/// <summary>
+		/// Logger for individual agents
+		/// </summary>
+		IAuditLog<AgentId> AgentLog;
+
 		/// <summary>
 		/// Log output writer
 		/// </summary>
@@ -103,18 +109,7 @@ namespace HordeServer.Services
 		/// <summary>
 		/// Constructor
 		/// </summary>
-		/// <param name="Agents">Collection of agent documents</param>
-		/// <param name="Leases">Collection of lease documents</param>
-		/// <param name="Sessions">Collection of session documents</param>
-		/// <param name="AclService">The ACL service</param>
-		/// <param name="DowntimeService">The downtime service</param>
-		/// <param name="PoolCollection">Pool collection</param>
-		/// <param name="DogStatsd">The DogStatsd metric client</param>
-		/// <param name="TaskSources">The list of task sources</param>
-		/// <param name="ApplicationLifetime"></param>
-		/// <param name="Logger">Log output writer</param>
-		/// <param name="Clock">Clock</param>
-		public AgentService(IAgentCollection Agents, ILeaseCollection Leases, ISessionCollection Sessions, AclService AclService, IDowntimeService DowntimeService, IPoolCollection PoolCollection, IDogStatsd DogStatsd, IEnumerable<ITaskSource> TaskSources, IHostApplicationLifetime ApplicationLifetime, ILogger<AgentService> Logger, IClock Clock)
+		public AgentService(IAgentCollection Agents, ILeaseCollection Leases, ISessionCollection Sessions, AclService AclService, IDowntimeService DowntimeService, IPoolCollection PoolCollection, IDogStatsd DogStatsd, IEnumerable<ITaskSource> TaskSources, IHostApplicationLifetime ApplicationLifetime, ILogger<AgentService> Logger, IAuditLog<AgentId> AgentLog, IClock Clock)
 			: base(TimeSpan.FromSeconds(30.0), Logger)
 		{
 			this.Agents = Agents;
@@ -126,6 +121,7 @@ namespace HordeServer.Services
 			this.DogStatsd = DogStatsd;
 			this.TaskSources = TaskSources.ToArray();
 			this.ApplicationLifetime = ApplicationLifetime;
+			this.AgentLog = AgentLog;
 			this.Logger = Logger;
 			this.Clock = Clock;
 		}
@@ -179,44 +175,6 @@ namespace HordeServer.Services
 			return Agents.FindAsync(Pool, null, ModifiedAfter, null, Index, Count);
 		}
 
-		/// <summary>
-		/// Update an agent's settings
-		/// </summary>
-		/// <param name="Agent">Agent instance</param>
-		/// <param name="bNewEnabled">Whether the agent is enabled or not</param>
-		/// <param name="bNewRequestConform">Whether to request a conform job be run</param>
-		/// <param name="bNewRequestRestart">Whether to request the machine be restarted</param>
-		/// <param name="bNewRequestShutdown">Whether to request the machine be shutdown</param>
-		/// <param name="NewChannel">Override for the desired client version</param>
-		/// <param name="NewPools">List of pools for the agent</param>
-		/// <param name="NewAcl">New ACL for this agent</param>
-		/// <param name="NewComment">New comment</param>
-		/// <returns>Version of the software that needs to be installed on the agent. Null if the agent is running the correct version.</returns>
-		public async Task<IAgent?> UpdateAgentAsync(IAgent? Agent, bool? bNewEnabled, bool? bNewRequestConform, bool? bNewRequestRestart, bool? bNewRequestShutdown, AgentSoftwareChannelName? NewChannel, List<PoolId>? NewPools, Acl? NewAcl, string? NewComment)
-		{
-			while (Agent != null)
-			{
-				// Apply the update
-				IAgent? NewAgent = await Agents.TryUpdateSettingsAsync(Agent, bNewEnabled, bNewRequestConform, bNewRequestRestart, bNewRequestShutdown, NewChannel, NewPools, NewAcl, NewComment);
-				if (NewAgent != null)
-				{
-					Agent = NewAgent;
-
-					// If we're requesting a conform, allow the agent to break out of a long poll for work immediately to pick it up
-//					if ((bNewRequestConform ?? false) || (bNewRequestRestart ?? false) || (bNewRequestShutdown ?? false))
-//					{
-//						DispatchService.CancelLongPollForAgent(Agent.Id);
-//					}
-
-					return Agent;
-				}
-
-				// Update the agent
-				Agent = await GetAgentAsync(Agent.Id);
-			}
-			return null;
-		}
-
 		/// <summary>
 		/// Update the current workspaces for an agent.
 		/// </summary>
@@ -277,6 +235,8 @@ namespace HordeServer.Services
 		{
 			for (; ; )
 			{
+				IAuditLogChannel<AgentId> AgentLogger = Agents.GetLogger(Agent.Id);
+
 				// Check if there's already a session running for this agent.
 				IAgent? NewAgent;
 				if (Agent.SessionId != null)
@@ -291,7 +251,7 @@ namespace HordeServer.Services
 					// Remove any outstanding leases
 					foreach (AgentLease Lease in Agent.Leases)
 					{
-						Logger.LogDebug("Removing outstanding lease {LeaseId} during session create...", Lease.Id);
+						AgentLogger.LogInformation("Removing outstanding lease {LeaseId}", Lease.Id);
 						await RemoveLeaseAsync(Agent, Lease, UtcNow, LeaseOutcome.Failed, null);
 					}
 
@@ -307,6 +267,7 @@ namespace HordeServer.Services
 					if(NewAgent != null)
 					{
 						Agent = NewAgent;
+						AgentLogger.LogInformation("Session {SessionId} started", NewSession.Id);
 						break;
 					}
 
@@ -456,7 +417,7 @@ namespace HordeServer.Services
 				IAgent? NewAgent = await Agents.TryAddLeaseAsync(Agent, Lease);
 				if (NewAgent != null)
 				{
-					await Source.OnLeaseStartedAsync(NewAgent, Lease.Id, Any.Parser.ParseFrom(Lease.Payload));
+					await Source.OnLeaseStartedAsync(NewAgent, Lease.Id, Any.Parser.ParseFrom(Lease.Payload), Agents.GetLogger(Agent.Id));
 					await CreateLeaseAsync(Agent, Lease);
 					return NewAgent;
 				}
@@ -677,11 +638,12 @@ namespace HordeServer.Services
 				// Remove any outstanding leases
 				foreach(AgentLease Lease in Leases)
 				{
-					Logger.LogDebug("Removing lease {LeaseId} during session terminate...", Lease.Id);
+					Agents.GetLogger(Agent.Id).LogInformation("Removing lease {LeaseId} during session terminate...", Lease.Id);
 					await RemoveLeaseAsync(Agent, Lease, FinishTime, LeaseOutcome.Failed, null);
 				}
 
 				// Update the session document
+				Agents.GetLogger(Agent.Id).LogInformation("Terminated session {SessionId}", SessionId);
 				await Sessions.UpdateAsync(SessionId, FinishTime, Agent.Properties, Agent.Resources);
 				return true;
 			}
@@ -757,7 +719,7 @@ namespace HordeServer.Services
 				{
 					if (Any.Is(TaskSource.Descriptor))
 					{
-						await TaskSource.OnLeaseFinishedAsync(Agent, Lease.Id, Any, Outcome, Output);
+						await TaskSource.OnLeaseFinishedAsync(Agent, Lease.Id, Any, Outcome, Output, Agents.GetLogger(Agent.Id));
 						break;
 					}
 				}
diff --git a/Engine/Source/Programs/Horde/HordeServer/Services/AutoscaleService.cs b/Engine/Source/Programs/Horde/HordeServer/Services/AutoscaleService.cs
--- a/Engine/Source/Programs/Horde/HordeServer/Services/AutoscaleService.cs
+++ b/Engine/Source/Programs/Horde/HordeServer/Services/AutoscaleService.cs
@@ -101,13 +101,6 @@ namespace HordeServer.Services
 		/// <summary>
 		/// Constructor
 		/// </summary>
-		/// <param name="DatabaseService"></param>
-		/// <param name="AgentCollection"></param>
-		/// <param name="PoolCollection"></param>
-		/// <param name="LeaseCollection"></param>
-		/// <param name="FleetManager"></param>
-		/// <param name="DogStatsd"></param>
-		/// <param name="Logger"></param>
 		public AutoscaleService(DatabaseService DatabaseService, IAgentCollection AgentCollection, IPoolCollection PoolCollection, ILeaseCollection LeaseCollection, IFleetManager FleetManager, IDogStatsd DogStatsd, ILogger<AutoscaleService> Logger)
 			: base(DatabaseService, new ObjectId("5fdbadf968de915b2cac0d15"), Logger)
 		{
diff --git a/Engine/Source/Programs/Horde/HordeServer/Services/Impl/AwsFleetManager.cs b/Engine/Source/Programs/Horde/HordeServer/Services/Impl/AwsFleetManager.cs
--- a/Engine/Source/Programs/Horde/HordeServer/Services/Impl/AwsFleetManager.cs
+++ b/Engine/Source/Programs/Horde/HordeServer/Services/Impl/AwsFleetManager.cs
@@ -92,7 +92,7 @@ namespace HordeServer.Services.Impl
 					{
 						if (await AgentCollection.TryUpdateSettingsAsync(NewAgent, bRequestShutdown: true) != null)
 						{
-							Logger.LogInformation("Shutting down agent {AgentId}", Agent.Id);
+							AgentCollection.GetLogger(Agent.Id).LogInformation("Shutting down due to autoscaler");
 							Count--;
 							break;
 						}
diff --git a/Engine/Source/Programs/Horde/HordeServer/Startup.cs b/Engine/Source/Programs/Horde/HordeServer/Startup.cs
--- a/Engine/Source/Programs/Horde/HordeServer/Startup.cs
+++ b/Engine/Source/Programs/Horde/HordeServer/Startup.cs
@@ -340,6 +340,9 @@ namespace HordeServer
 			Services.AddSingleton<IUserCollection, UserCollectionV2>();
 			Services.AddSingleton<IDeviceCollection, DeviceCollection>();
 
+			// Auditing
+			Services.AddSingleton<IAuditLog<AgentId>>(SP => SP.GetRequiredService<IAuditLogFactory<AgentId>>().Create("Agents.Log", "AgentId"));
+
 			// Storage
 			Services.AddSingleton<IBlobCollection>(SP => new CachingBlobCollection(new BlobCollection(SP.GetRequiredService<IStorageBackend<BlobCollection>>()), 256 * 1024 * 1024));
 			Services.AddSingleton<IBucketCollection, BucketCollection>();
diff --git a/Engine/Source/Programs/Horde/HordeServer/Tasks/ITaskSource.cs b/Engine/Source/Programs/Horde/HordeServer/Tasks/ITaskSource.cs
--- a/Engine/Source/Programs/Horde/HordeServer/Tasks/ITaskSource.cs
+++ b/Engine/Source/Programs/Horde/HordeServer/Tasks/ITaskSource.cs
@@ -6,10 +6,13 @@ using Google.Protobuf.WellKnownTypes;
 using HordeCommon;
 using HordeServer.Models;
 using HordeServer.Services;
+using Microsoft.Extensions.Logging;
 using MongoDB.Bson;
 using System;
 using System.Collections.Generic;
 using System.Linq;
+using System.Reflection;
+using System.Text;
 using System.Threading;
 using System.Threading.Tasks;
 
@@ -20,6 +23,11 @@ namespace HordeServer.Tasks
 	/// </summary>
 	public interface ITaskSource
 	{
+		/// <summary>
+		/// The task type
+		/// </summary>
+		string Type { get; }
+
 		/// <summary>
 		/// Descriptor for this lease type
 		/// </summary>
@@ -48,7 +56,8 @@ namespace HordeServer.Tasks
 		/// <param name="Agent">The agent executing the lease</param>
 		/// <param name="LeaseId">The lease id</param>
 		/// <param name="Payload">Payload for the lease</param>
-		Task OnLeaseStartedAsync(IAgent Agent, ObjectId LeaseId, Any Payload);
+		/// <param name="Logger">Logger for the agent</param>
+		Task OnLeaseStartedAsync(IAgent Agent, ObjectId LeaseId, Any Payload, ILogger Logger);
 
 		/// <summary>
 		/// Notification that a task has completed
@@ -58,7 +67,8 @@ namespace HordeServer.Tasks
 		/// <param name="Payload">The lease payload</param>
 		/// <param name="Outcome">Outcome of the lease</param>
 		/// <param name="Output">Output from the task</param>
-		Task OnLeaseFinishedAsync(IAgent Agent, ObjectId LeaseId, Any Payload, LeaseOutcome Outcome, ReadOnlyMemory<byte> Output);
+		/// <param name="Logger">Logger for the agent</param>
+		Task OnLeaseFinishedAsync(IAgent Agent, ObjectId LeaseId, Any Payload, LeaseOutcome Outcome, ReadOnlyMemory<byte> Output, ILogger Logger);
 	}
 	
 	/// <summary>
@@ -69,6 +79,9 @@ namespace HordeServer.Tasks
 	{
 		static TMessage Message = new TMessage();
 
+		/// <inheritdoc/>
+		public abstract string Type { get; }
+
 		/// <inheritdoc/>
 		public MessageDescriptor Descriptor => Message.Descriptor;
 
@@ -76,12 +89,68 @@ namespace HordeServer.Tasks
 		public abstract Task<AgentLease?> AssignLeaseAsync(IAgent Agent, CancellationToken CancellationToken);
 
 		/// <inheritdoc/>
-		public virtual Task CancelLeaseAsync(IAgent Agent, ObjectId LeaseId, Any Payload) => Task.CompletedTask;
+		public Task CancelLeaseAsync(IAgent Agent, ObjectId LeaseId, Any Payload) => CancelLeaseAsync(Agent, LeaseId, Payload.Unpack<TMessage>());
 
 		/// <inheritdoc/>
-		public virtual Task OnLeaseStartedAsync(IAgent Agent, ObjectId LeaseId, Any Payload) => Task.CompletedTask;
+		public Task OnLeaseStartedAsync(IAgent Agent, ObjectId LeaseId, Any Payload, ILogger Logger) => OnLeaseStartedAsync(Agent, LeaseId, Payload.Unpack<TMessage>(), Logger);
 
 		/// <inheritdoc/>
-		public virtual Task OnLeaseFinishedAsync(IAgent Agent, ObjectId LeaseId, Any Payload, LeaseOutcome Outcome, ReadOnlyMemory<byte> Output) => Task.CompletedTask;
+		public Task OnLeaseFinishedAsync(IAgent Agent, ObjectId LeaseId, Any Payload, LeaseOutcome Outcome, ReadOnlyMemory<byte> Output, ILogger Logger) => OnLeaseFinishedAsync(Agent, LeaseId, Payload.Unpack<TMessage>(), Outcome, Output, Logger);
+
+		/// <inheritdoc cref="ITaskSource.CancelLeaseAsync(IAgent, ObjectId, Any)"/>
+		public virtual Task CancelLeaseAsync(IAgent Agent, ObjectId LeaseId, TMessage Payload) => Task.CompletedTask;
+
+		/// <inheritdoc cref="ITaskSource.OnLeaseStartedAsync(IAgent, ObjectId, Any, ILogger)"/>
+		public virtual Task OnLeaseStartedAsync(IAgent Agent, ObjectId LeaseId, TMessage Payload, ILogger Logger)
+		{
+			LogLeaseStartedInfo<TMessage>(LeaseId, Payload, Logger);
+			return Task.CompletedTask;
+		}
+
+		/// <inheritdoc cref="ITaskSource.OnLeaseFinishedAsync(IAgent, ObjectId, Any, LeaseOutcome, ReadOnlyMemory{byte}, ILogger)"/>
+		public virtual Task OnLeaseFinishedAsync(IAgent Agent, ObjectId LeaseId, TMessage Payload, LeaseOutcome Outcome, ReadOnlyMemory<byte> Output, ILogger Logger)
+		{
+			Logger.LogInformation("Lease {LeaseId} complete", LeaseId);
+			return Task.CompletedTask;
+		}
+
+		class TypeInfo<T>
+		{
+			public static readonly PropertyInfo[] Properties = typeof(T).GetProperties(BindingFlags.Public | BindingFlags.Instance).ToArray();
+			public static readonly string FormatString = CreateFormatString(Properties);
+
+			static string CreateFormatString(PropertyInfo[] Properties)
+			{
+				StringBuilder Message = new StringBuilder($"Lease {{LeaseId}} started (Type={{Type}}");
+				foreach (PropertyInfo Property in Properties)
+				{
+					Message.Append($", {Property.Name}={{{Property.Name}}}");
+				}
+				Message.Append(")");
+				return Message.ToString();
+			}
+		}
+
+		/// <summary>
+		/// Helper method to log information about a lease starting
+		/// </summary>
+		/// <typeparam name="T"></typeparam>
+		/// <param name="LeaseId"></param>
+		/// <param name="Payload"></param>
+		/// <param name="Logger"></param>
+		protected void LogLeaseStartedInfo<T>(ObjectId LeaseId, T Payload, ILogger Logger)
+		{
+			PropertyInfo[] Properties = TypeInfo<T>.Properties;
+
+			object?[] Values = new object?[Properties.Length + 2];
+			Values[0] = LeaseId;
+			Values[1] = Type;
+			for (int Idx = 0; Idx < Properties.Length; Idx++)
+			{
+				Values[Idx + 2] = Properties[Idx].GetValue(Payload);
+			}
+
+			Logger.LogInformation(TypeInfo<T>.FormatString, Values);
+		}
 	}
 }
diff --git a/Engine/Source/Programs/Horde/HordeServer/Tasks/Impl/ConformTaskSource.cs b/Engine/Source/Programs/Horde/HordeServer/Tasks/Impl/ConformTaskSource.cs
--- a/Engine/Source/Programs/Horde/HordeServer/Tasks/Impl/ConformTaskSource.cs
+++ b/Engine/Source/Programs/Horde/HordeServer/Tasks/Impl/ConformTaskSource.cs
@@ -26,6 +26,9 @@ namespace HordeServer.Tasks.Impl
 	/// </summary>
 	public sealed class ConformTaskSource : TaskSourceBase<ConformTask>, IHostedService, IDisposable
 	{
+		/// <inheritdoc/>
+		public override string Type => "Conform";
+
 		DatabaseService DatabaseService;
 		IAgentCollection AgentCollection;
 		PoolService PoolService;
@@ -210,7 +213,7 @@ namespace HordeServer.Tasks.Impl
 		}
 
 		/// <inheritdoc/>
-		public override Task OnLeaseFinishedAsync(IAgent Agent, ObjectId LeaseId, Any Payload, LeaseOutcome Outcome, ReadOnlyMemory<byte> Output)
+		public override Task OnLeaseFinishedAsync(IAgent Agent, ObjectId LeaseId, ConformTask Payload, LeaseOutcome Outcome, ReadOnlyMemory<byte> Output, ILogger Logger)
 		{
 			return ReleaseConformLeaseAsync(LeaseId);
 		}
diff --git a/Engine/Source/Programs/Horde/HordeServer/Tasks/Impl/JobTaskSource.cs b/Engine/Source/Programs/Horde/HordeServer/Tasks/Impl/JobTaskSource.cs
--- a/Engine/Source/Programs/Horde/HordeServer/Tasks/Impl/JobTaskSource.cs
+++ b/Engine/Source/Programs/Horde/HordeServer/Tasks/Impl/JobTaskSource.cs
@@ -40,6 +40,9 @@ namespace HordeServer.Tasks.Impl
 	/// </summary>
 	public sealed class JobTaskSource : TaskSourceBase<ExecuteJobTask>, IDisposable
 	{
+		/// <inheritdoc/>
+		public override string Type => "Job";
+
 		/// <summary>
 		/// An item in the queue to be executed
 		/// </summary>
@@ -720,9 +723,8 @@ namespace HordeServer.Tasks.Impl
 		}
 
 		/// <inheritdoc/>
-		public override Task CancelLeaseAsync(IAgent Agent, ObjectId LeaseId, Any Payload)
+		public override Task CancelLeaseAsync(IAgent Agent, ObjectId LeaseId, ExecuteJobTask Task)
 		{
-			ExecuteJobTask Task = Payload.Unpack<ExecuteJobTask>();
 			return CancelLeaseAsync(Agent, Task.JobId.ToObjectId(), Task.BatchId.ToSubResourceId());
 		}
 
@@ -972,12 +974,11 @@ namespace HordeServer.Tasks.Impl
 		}
 
 		/// <inheritdoc/>
-		public override async Task OnLeaseFinishedAsync(IAgent Agent, ObjectId LeaseId, Any Any, LeaseOutcome Outcome, ReadOnlyMemory<byte> Output)
+		public override async Task OnLeaseFinishedAsync(IAgent Agent, ObjectId LeaseId, ExecuteJobTask Task, LeaseOutcome Outcome, ReadOnlyMemory<byte> Output, ILogger Logger)
 		{
 			if (Outcome != LeaseOutcome.Success)
 			{
 				AgentId AgentId = Agent.Id;
-				ExecuteJobTask Task = Any.Unpack<ExecuteJobTask>();
 				ObjectId JobId = Task.JobId.ToObjectId();
 				SubResourceId BatchId = Task.BatchId.ToSubResourceId();
 
diff --git a/Engine/Source/Programs/Horde/HordeServer/Tasks/Impl/RestartTaskSource.cs b/Engine/Source/Programs/Horde/HordeServer/Tasks/Impl/RestartTaskSource.cs
--- a/Engine/Source/Programs/Horde/HordeServer/Tasks/Impl/RestartTaskSource.cs
+++ b/Engine/Source/Programs/Horde/HordeServer/Tasks/Impl/RestartTaskSource.cs
@@ -19,6 +19,8 @@ namespace HordeServer.Tasks.Impl
 {
 	class RestartTaskSource : TaskSourceBase<RestartTask>
 	{
+		public override string Type => "Restart";
+
 		ILogFileService LogService;
 
 		public RestartTaskSource(ILogFileService LogService)
diff --git a/Engine/Source/Programs/Horde/HordeServer/Tasks/Impl/ShutdownTaskSource.cs b/Engine/Source/Programs/Horde/HordeServer/Tasks/Impl/ShutdownTaskSource.cs
--- a/Engine/Source/Programs/Horde/HordeServer/Tasks/Impl/ShutdownTaskSource.cs
+++ b/Engine/Source/Programs/Horde/HordeServer/Tasks/Impl/ShutdownTaskSource.cs
@@ -20,6 +20,8 @@ namespace HordeServer.Tasks.Impl
 {
 	class ShutdownTaskSource : TaskSourceBase<ShutdownTask>
 	{
+		public override string Type => "Shutdown";
+
 		ILogFileService LogService;
 
 		public ShutdownTaskSource(ILogFileService LogService)
diff --git a/Engine/Source/Programs/Horde/HordeServer/Tasks/Impl/UpgradeTaskSource.cs b/Engine/Source/Programs/Horde/HordeServer/Tasks/Impl/UpgradeTaskSource.cs
--- a/Engine/Source/Programs/Horde/HordeServer/Tasks/Impl/UpgradeTaskSource.cs
+++ b/Engine/Source/Programs/Horde/HordeServer/Tasks/Impl/UpgradeTaskSource.cs
@@ -22,6 +22,8 @@ namespace HordeServer.Tasks.Impl
 
 	class UpgradeTaskSource : TaskSourceBase<UpgradeTask>
 	{
+		public override string Type => "Upgrade";
+
 		AgentSoftwareService AgentSoftwareService;
 		ILogFileService LogService;
 		IClock Clock;
diff --git a/Engine/Source/Programs/Shared/EpicGames.Core/MessageTemplate.cs b/Engine/Source/Programs/Shared/EpicGames.Core/MessageTemplate.cs
--- a/Engine/Source/Programs/Shared/EpicGames.Core/MessageTemplate.cs
+++ b/Engine/Source/Programs/Shared/EpicGames.Core/MessageTemplate.cs
@@ -114,16 +114,59 @@ namespace EpicGames.Core
 		/// <summary>
 		/// Serialize a message template to JOSN
 		/// </summary>
-		public static void SerializeJson(Utf8JsonWriter Writer, DateTime Time, LogLevel LogLevel, EventId EventId, string Format, IEnumerable<KeyValuePair<string, object>>? Properties, Exception? Exception = null)
+		public static string Serialize<TState>(LogLevel LogLevel, EventId EventId, TState State, Exception Exception, Func<TState, Exception, string> Formatter)
 		{
-			string Text = Render(Format, Properties);
-			SerializeJson(Writer, Time, LogLevel, Text, EventId, 0, 1, Format, Properties, Exception);
+			ArrayBufferWriter<byte> Buffer = new ArrayBufferWriter<byte>();
+			using (Utf8JsonWriter Writer = new Utf8JsonWriter(Buffer))
+			{
+				Serialize(Writer, LogLevel, EventId, State, Exception, Formatter);
+			}
+			return Encoding.UTF8.GetString(Buffer.WrittenSpan);
+		}
+
+		/// <summary>
+		/// Serialize a message template to JOSN
+		/// </summary>
+		public static void Serialize<TState>(Utf8JsonWriter Writer, LogLevel LogLevel, EventId EventId, TState State, Exception Exception, Func<TState, Exception, string> Formatter)
+		{
+			// Render the message
+			string Text = Formatter(State, Exception);
+
+			// Try to log the event
+			IEnumerable<KeyValuePair<string, object>>? Values = State as IEnumerable<KeyValuePair<string, object>>;
+			if (Values != null)
+			{
+				KeyValuePair<string, object>? Format = Values.FirstOrDefault(x => x.Key.Equals(MessageTemplate.FormatPropertyName, StringComparison.Ordinal));
+				if (Format != null)
+				{
+					// Format all the other values
+					Serialize(Writer, DateTime.UtcNow, LogLevel, EventId, Text, Format.Value.Value?.ToString() ?? String.Empty, Values.Where(x => !x.Key.Equals(MessageTemplate.FormatPropertyName, StringComparison.Ordinal)), Exception);
+				}
+				else
+				{
+					// Include all the data, but don't use the format string
+					Serialize(Writer, DateTime.UtcNow, LogLevel, EventId, Text, null, Values, Exception);
+				}
+			}
+			else
+			{
+				// Format as a string
+				Serialize(Writer, DateTime.UtcNow, LogLevel, EventId, Text, null, null, Exception);
+			}
+		}
+
+		/// <summary>
+		/// Serialize a message template to JSON
+		/// </summary>
+		public static void Serialize(Utf8JsonWriter Writer, DateTime Time, LogLevel LogLevel, EventId EventId, string Text, string? Format, IEnumerable<KeyValuePair<string, object>>? Properties, Exception? Exception = null)
+		{
+			Serialize(Writer, Time, LogLevel, EventId, Text, 0, 1, Format, Properties, Exception);
 		}
 
 		/// <summary>
 		/// Serialize a message template to JSON
 		/// </summary>
-		public static void SerializeJson(Utf8JsonWriter Writer, DateTime Time, LogLevel LogLevel, string Text, EventId EventId, int LineIndex, int LineCount, string? Format, IEnumerable<KeyValuePair<string, object>>? Properties, Exception? Exception = null)
+		public static void Serialize(Utf8JsonWriter Writer, DateTime Time, LogLevel LogLevel, EventId EventId, string Text, int LineIndex, int LineCount, string? Format, IEnumerable<KeyValuePair<string, object>>? Properties, Exception? Exception = null)
 		{
 			Writer.WriteStartObject();
 			Writer.WriteString("time", Time.ToString("s", CultureInfo.InvariantCulture));
