Horde: Add a notification system for agent update events. Allows agents to react immediately to setting changes, rather than having to wait for the next session update to commence.

diff --git a/Engine/Source/Programs/Horde/HordeServer/Collections/IAgentCollection.cs b/Engine/Source/Programs/Horde/HordeServer/Collections/IAgentCollection.cs
--- a/Engine/Source/Programs/Horde/HordeServer/Collections/IAgentCollection.cs
+++ b/Engine/Source/Programs/Horde/HordeServer/Collections/IAgentCollection.cs
@@ -156,5 +156,12 @@ namespace HordeServer.Collections
 		/// <param name="AgentId"></param>
 		/// <returns></returns>
 		IAuditLogChannel<AgentId> GetLogger(AgentId AgentId);
+
+		/// <summary>
+		/// Subscribe to notifications on agent states being updated
+		/// </summary>
+		/// <param name="OnUpdate">Callback for updates</param>
+		/// <returns>Disposable subscription object</returns>
+		Task<IDisposable> SubscribeToUpdateEventsAsync(Action<AgentId> OnUpdate);
 	}
 }
diff --git a/Engine/Source/Programs/Horde/HordeServer/Collections/Impl/AgentCollection.cs b/Engine/Source/Programs/Horde/HordeServer/Collections/Impl/AgentCollection.cs
--- a/Engine/Source/Programs/Horde/HordeServer/Collections/Impl/AgentCollection.cs
+++ b/Engine/Source/Programs/Horde/HordeServer/Collections/Impl/AgentCollection.cs
@@ -1,5 +1,6 @@
 // Copyright Epic Games, Inc. All Rights Reserved.
 
+using EpicGames.Redis;
 using Google.Protobuf.WellKnownTypes;
 using HordeCommon;
 using HordeCommon.Rpc.Tasks;
@@ -134,13 +135,17 @@ namespace HordeServer.Collections.Impl
 
 		readonly IMongoCollection<AgentDocument> Agents;
 		readonly IAuditLog<AgentId> AuditLog;
+		readonly RedisService RedisService;
+		readonly RedisChannel<AgentId> UpdateEventChannel;
 
 		/// <summary>
 		/// Constructor
 		/// </summary>
-		public AgentCollection(DatabaseService DatabaseService, IAuditLog<AgentId> AuditLog)
+		public AgentCollection(DatabaseService DatabaseService, RedisService RedisService, IAuditLog<AgentId> AuditLog)
 		{
 			this.Agents = DatabaseService.GetCollection<AgentDocument>("Agents");
+			this.RedisService = RedisService;
+			this.UpdateEventChannel = new RedisChannel<AgentId>("agents/notify");
 			this.AuditLog = AuditLog;
 
 			if (!DatabaseService.ReadOnlyMode)
@@ -307,7 +312,15 @@ namespace HordeServer.Collections.Impl
 			}
 
 			// Apply the update
-			return await TryUpdateAsync(Agent, UpdateBuilder.Combine(Updates));
+			IAgent? NewAgent = await TryUpdateAsync(Agent, UpdateBuilder.Combine(Updates));
+			if (NewAgent != null)
+			{
+				if (NewAgent.RequestRestart != Agent.RequestRestart || NewAgent.RequestConform != Agent.RequestConform || NewAgent.RequestShutdown != Agent.RequestShutdown || NewAgent.Channel != Agent.Channel)
+				{
+					await PublishUpdateEventAsync(Agent.Id);
+				}
+			}
+			return NewAgent;
 		}
 
 		/// <inheritdoc/>
@@ -481,7 +494,12 @@ namespace HordeServer.Collections.Impl
 			AgentDocument Agent = (AgentDocument)AgentInterface;
 
 			UpdateDefinition<AgentDocument> Update = Builders<AgentDocument>.Update.Set(x => x.Leases![LeaseIdx].State, LeaseState.Cancelled);
-			return await TryUpdateAsync(Agent, Update);
+			IAgent? NewAgent = await TryUpdateAsync(Agent, Update);
+			if (NewAgent != null)
+			{
+				await PublishUpdateEventAsync(Agent.Id);
+			}
+			return NewAgent;
 		}
 
 		/// <inheritdoc/>
@@ -489,5 +507,17 @@ namespace HordeServer.Collections.Impl
 		{
 			return AuditLog[AgentId];
 		}
+
+		/// <inheritdoc/>
+		Task PublishUpdateEventAsync(AgentId AgentId)
+		{
+			return RedisService.Database.PublishAsync(UpdateEventChannel, AgentId);
+		}
+
+		/// <inheritdoc/>
+		public async Task<IDisposable> SubscribeToUpdateEventsAsync(Action<AgentId> OnUpdate)
+		{
+			return await RedisService.Multiplexer.GetSubscriber().SubscribeAsync(UpdateEventChannel, (Channel, AgentId) => OnUpdate(AgentId));
+		}
 	}
 }
diff --git a/Engine/Source/Programs/Horde/HordeServer/Services/AgentService.cs b/Engine/Source/Programs/Horde/HordeServer/Services/AgentService.cs
--- a/Engine/Source/Programs/Horde/HordeServer/Services/AgentService.cs
+++ b/Engine/Source/Programs/Horde/HordeServer/Services/AgentService.cs
@@ -29,7 +29,7 @@ namespace HordeServer.Services
 	/// <summary>
 	/// Wraps funtionality for manipulating agents
 	/// </summary>
-	public class AgentService : TickedBackgroundService
+	public sealed class AgentService : TickedBackgroundService
 	{
 		/// <summary>
 		/// Maximum time between updates for an agent to be considered online
@@ -86,11 +86,6 @@ namespace HordeServer.Services
 		/// </summary>
 		IHostApplicationLifetime ApplicationLifetime;
 
-		/// <summary>
-		/// Logger for individual agents
-		/// </summary>
-		IAuditLog<AgentId> AgentLog;
-
 		/// <summary>
 		/// Log output writer
 		/// </summary>
@@ -106,10 +101,20 @@ namespace HordeServer.Services
 		/// </summary>
 		AsyncCachedValue<List<IPool>> PoolsList;
 
+		/// <summary>
+		/// All the agents currently performing a long poll for work on this server
+		/// </summary>
+		Dictionary<AgentId, CancellationTokenSource> WaitingAgents = new Dictionary<AgentId, CancellationTokenSource>();
+
+		/// <summary>
+		/// Subscription for update events
+		/// </summary>
+		IDisposable Subscription;
+
 		/// <summary>
 		/// Constructor
 		/// </summary>
-		public AgentService(IAgentCollection Agents, ILeaseCollection Leases, ISessionCollection Sessions, AclService AclService, IDowntimeService DowntimeService, IPoolCollection PoolCollection, IDogStatsd DogStatsd, IEnumerable<ITaskSource> TaskSources, IHostApplicationLifetime ApplicationLifetime, ILogger<AgentService> Logger, IAuditLog<AgentId> AgentLog, IClock Clock)
+		public AgentService(IAgentCollection Agents, ILeaseCollection Leases, ISessionCollection Sessions, AclService AclService, IDowntimeService DowntimeService, IPoolCollection PoolCollection, IDogStatsd DogStatsd, IEnumerable<ITaskSource> TaskSources, IHostApplicationLifetime ApplicationLifetime, ILogger<AgentService> Logger, IClock Clock)
 			: base(TimeSpan.FromSeconds(30.0), Logger)
 		{
 			this.Agents = Agents;
@@ -121,9 +126,17 @@ namespace HordeServer.Services
 			this.DogStatsd = DogStatsd;
 			this.TaskSources = TaskSources.ToArray();
 			this.ApplicationLifetime = ApplicationLifetime;
-			this.AgentLog = AgentLog;
 			this.Logger = Logger;
 			this.Clock = Clock;
+
+			Subscription = Agents.SubscribeToUpdateEventsAsync(OnAgentUpdate).Result;
+		}
+
+		/// <inheritdoc/>
+		public override void Dispose()
+		{
+			base.Dispose();
+			Subscription.Dispose();
 		}
 
 		/// <summary>
@@ -222,6 +235,21 @@ namespace HordeServer.Services
 			return NewDynamicPools;
 		}
 
+		/// <summary>
+		/// Callback for an agents 
+		/// </summary>
+		/// <param name="AgentId"></param>
+		void OnAgentUpdate(AgentId AgentId)
+		{
+			lock (WaitingAgents)
+			{
+				if (WaitingAgents.TryGetValue(AgentId, out CancellationTokenSource? CancellationSource))
+				{
+					CancellationSource.Cancel();
+				}
+			}
+		}
+
 		/// <summary>
 		/// Creates a new agent session
 		/// </summary>
@@ -286,21 +314,6 @@ namespace HordeServer.Services
 			return Agent;
 		}
 
-		static AgentLease? SafeGetLeaseResultInternal(Task<AgentLease?> Task, ILogger Logger)
-		{
-			if (!Task.IsCompleted)
-			{
-				Logger.LogError("Task is not complete");
-				return null;
-			}
-			if (Task.IsFaulted)
-			{
-				Logger.LogError(Task.Exception, "Task has faulted");
-				return null;
-			}
-			return Task.Result;
-		}
-
 		async Task<(ITaskSource, AgentLease)?> GuardedWaitForLeaseAsync(ITaskSource Source, IAgent Agent, CancellationToken CancellationToken)
 		{
 			try
@@ -346,54 +359,71 @@ namespace HordeServer.Services
 				using CancellationTokenSource CancellationSource = CancellationTokenSource.CreateLinkedTokenSource(ApplicationLifetime.ApplicationStopping, CancellationToken);
 				CancellationSource.CancelAfter(MaxWaitTime);
 
-				// If we're in a maintenance window, just wait for the time to expire
-				if (DowntimeService.IsDowntimeActive)
+				// Assign a new lease
+				(ITaskSource, AgentLease)? Result = null;
+				try
 				{
-					await AsyncUtils.DelayNoThrow(MaxWaitTime, CancellationToken);
-					break;
-				}
+					// Add the cancellation source to the set of waiting agents
+					lock (WaitingAgents)
+					{
+						WaitingAgents[Agent.Id] = CancellationSource;
+					}
 
-				// Create all the tasks to wait for
-				List<Task<(ITaskSource, AgentLease)?>> Tasks = new List<Task<(ITaskSource, AgentLease)?>>();
-				foreach (ITaskSource TaskSource in TaskSources)
-				{
-					Tasks.Add(GuardedWaitForLeaseAsync(TaskSource, Agent, CancellationSource.Token));
-				}
+					// If we're in a maintenance window, just wait for the time to expire
+					if (DowntimeService.IsDowntimeActive)
+					{
+						await AsyncUtils.DelayNoThrow(MaxWaitTime, CancellationToken);
+						break;
+					}
 
-				// Wait for a lease to be available
-				(ITaskSource, AgentLease)? Result = null;
-				while (Tasks.Count > 0)
-				{
-					await Task.WhenAny(Tasks);
+					// Create all the tasks to wait for
+					List<Task<(ITaskSource, AgentLease)?>> Tasks = new List<Task<(ITaskSource, AgentLease)?>>();
+					foreach (ITaskSource TaskSource in TaskSources)
+					{
+						Tasks.Add(GuardedWaitForLeaseAsync(TaskSource, Agent, CancellationSource.Token));
+					}
 
-					for (int Idx = 0; Idx < Tasks.Count; Idx++)
+					// Wait for a lease to be available
+					while (Tasks.Count > 0)
 					{
-						Task<(ITaskSource, AgentLease)?> Task = Tasks[Idx];
-						if (!Task.IsCompleted)
+						await Task.WhenAny(Tasks);
+
+						for (int Idx = 0; Idx < Tasks.Count; Idx++)
 						{
-							continue;
-						}
+							Task<(ITaskSource, AgentLease)?> Task = Tasks[Idx];
+							if (!Task.IsCompleted)
+							{
+								continue;
+							}
 
-						Tasks.RemoveAt(Idx--);
+							Tasks.RemoveAt(Idx--);
 
-						(ITaskSource, AgentLease)? TaskResult = Task.Result;
-						if (!TaskResult.HasValue)
-						{
-							continue;
-						}
+							(ITaskSource, AgentLease)? TaskResult = Task.Result;
+							if (!TaskResult.HasValue)
+							{
+								continue;
+							}
 
-						if (Result == null)
-						{
-							Result = TaskResult;
-							CancellationSource.Cancel();
-						}
-						else
-						{
-							(ITaskSource TaskSource, AgentLease TaskLease) = TaskResult.Value;
-							await TaskSource.CancelLeaseAsync(Agent, TaskLease.Id, Any.Parser.ParseFrom(TaskLease.Payload));
+							if (Result == null)
+							{
+								Result = TaskResult;
+								CancellationSource.Cancel();
+							}
+							else
+							{
+								(ITaskSource TaskSource, AgentLease TaskLease) = TaskResult.Value;
+								await TaskSource.CancelLeaseAsync(Agent, TaskLease.Id, Any.Parser.ParseFrom(TaskLease.Payload));
+							}
 						}
 					}
 				}
+				finally
+				{
+					lock (WaitingAgents)
+					{
+						WaitingAgents.Remove(Agent.Id);
+					}
+				}
 
 				// Exit if we didn't find any work to do. It may be that all the task sources returned null, in which case wait for the time period to expire.
 				if (Result == null)
@@ -453,7 +483,6 @@ namespace HordeServer.Services
 			}
 
 			await Agents.TryCancelLeaseAsync(Agent, Index);
-
 			return true;
 		}
 
diff --git a/Engine/Source/Programs/Horde/HordeServer/Utilities/AgentId.cs b/Engine/Source/Programs/Horde/HordeServer/Utilities/AgentId.cs
--- a/Engine/Source/Programs/Horde/HordeServer/Utilities/AgentId.cs
+++ b/Engine/Source/Programs/Horde/HordeServer/Utilities/AgentId.cs
@@ -1,9 +1,11 @@
 // Copyright Epic Games, Inc. All Rights Reserved.
 
+using EpicGames.Redis;
 using HordeServer.Models;
 using MongoDB.Bson.Serialization;
 using MongoDB.Bson.Serialization.Attributes;
 using MongoDB.Bson.Serialization.Serializers;
+using StackExchange.Redis;
 using System;
 using System.Collections.Generic;
 using System.ComponentModel;
@@ -20,6 +22,7 @@ namespace HordeServer.Utilities
 	/// <summary>
 	/// Normalized hostname of an agent
 	/// </summary>
+	[RedisConverter(typeof(AgentIdRedisConverter))]
 	[BsonSerializer(typeof(AgentIdBsonSerializer))]
 	public struct AgentId : IEquatable<AgentId>
 	{
@@ -125,6 +128,18 @@ namespace HordeServer.Utilities
 		}
 	}
 
+	/// <summary>
+	/// Converter to/from Redis values
+	/// </summary>
+	public sealed class AgentIdRedisConverter : IRedisConverter<AgentId>
+	{
+		/// <inheritdoc/>
+		public AgentId FromRedisValue(RedisValue Value) => new AgentId((string)Value);
+
+		/// <inheritdoc/>
+		public RedisValue ToRedisValue(AgentId Value) => Value.ToString();
+	}
+
 	/// <summary>
 	/// Serializer for StringId objects
 	/// </summary>
diff --git a/Engine/Source/Programs/Shared/EpicGames.Redis/RedisSerializer.cs b/Engine/Source/Programs/Shared/EpicGames.Redis/RedisSerializer.cs
--- a/Engine/Source/Programs/Shared/EpicGames.Redis/RedisSerializer.cs
+++ b/Engine/Source/Programs/Shared/EpicGames.Redis/RedisSerializer.cs
@@ -14,7 +14,7 @@ namespace EpicGames.Redis
 	/// <summary>
 	/// Attribute specifying the converter type to use for a class
 	/// </summary>
-	[AttributeUsage(AttributeTargets.Class)]
+	[AttributeUsage(AttributeTargets.Class | AttributeTargets.Struct)]
 	public class RedisConverterAttribute : Attribute
 	{
 		/// <summary>
