Horde: Build fix.

diff --git a/Engine/Source/Programs/Horde/HordeServer/Controllers/AgentsController.cs b/Engine/Source/Programs/Horde/HordeServer/Controllers/AgentsController.cs
--- a/Engine/Source/Programs/Horde/HordeServer/Controllers/AgentsController.cs
+++ b/Engine/Source/Programs/Horde/HordeServer/Controllers/AgentsController.cs
@@ -231,7 +231,6 @@ namespace HordeServer.Controllers
 		/// <returns>Information about the requested agent</returns>
 		[HttpGet]
 		[Route("/api/v1/agents/{AgentId}/history")]
-		[ProducesResponseType(typeof(GetAgentHistoryResponse), 200)]
 		public async Task GetAgentHistoryAsync(string AgentId, [FromQuery] DateTime? MinTime = null, [FromQuery] DateTime? MaxTime = null, [FromQuery] int Index = 0, [FromQuery] int Count = 50)
 		{
 			IAuditLogChannel<AgentId> Channel = AgentService.Agents.GetLogger(AgentId.ToAgentId());
