Horde: Fix tests.

diff --git a/Engine/Source/Programs/Horde/HordeServer.Compute.Tests/TestSetup.cs b/Engine/Source/Programs/Horde/HordeServer.Compute.Tests/TestSetup.cs
--- a/Engine/Source/Programs/Horde/HordeServer.Compute.Tests/TestSetup.cs
+++ b/Engine/Source/Programs/Horde/HordeServer.Compute.Tests/TestSetup.cs
@@ -136,8 +136,12 @@ namespace HordeServerTests
 		{
 			Services.AddLogging(Builder => Builder.AddConsole());
 
+			Services.AddSingleton<RedisService>();
 			Services.AddSingleton<StackExchange.Redis.IDatabase>(new Mock<StackExchange.Redis.IDatabase>().Object);
 
+			Services.AddSingleton(typeof(IAuditLogFactory<>), typeof(AuditLogFactory<>));
+			Services.AddSingleton<IAuditLog<AgentId>>(SP => SP.GetRequiredService<IAuditLogFactory<AgentId>>().Create("Agents.Log", "AgentId"));
+
 			Services.AddSingleton<IAgentCollection, AgentCollection>();
 			Services.AddSingleton<IAgentSoftwareCollection, AgentSoftwareCollection>();
 			Services.AddSingleton<ICounterCollection, CounterCollection>();
