Horde: Allow filtering which users can receive notifications on Slack (and limit dev environment to just me).

diff --git a/Engine/Source/Programs/Horde/HordeServer/Notifications/Impl/SlackNotificationSink.cs b/Engine/Source/Programs/Horde/HordeServer/Notifications/Impl/SlackNotificationSink.cs
--- a/Engine/Source/Programs/Horde/HordeServer/Notifications/Impl/SlackNotificationSink.cs
+++ b/Engine/Source/Programs/Horde/HordeServer/Notifications/Impl/SlackNotificationSink.cs
@@ -200,6 +200,7 @@ namespace HordeServer.Notifications.Impl
 		ServerSettings Settings;
 		IMongoCollection<MessageStateDocument> MessageStates;
 		IMongoCollection<SlackUser> SlackUsers;
+		HashSet<string>? AllowUsers;
 		ILogger Logger;
 
 		/// <summary>
@@ -227,6 +228,11 @@ namespace HordeServer.Notifications.Impl
 			this.MessageStates = DatabaseService.Database.GetCollection<MessageStateDocument>("Slack");
 			this.SlackUsers = DatabaseService.Database.GetCollection<SlackUser>("Slack.UsersV2");
 			this.Logger = Logger;
+
+			if (!String.IsNullOrEmpty(Settings.Value.SlackUsers))
+			{
+				AllowUsers = new HashSet<string>(Settings.Value.SlackUsers.Split(','), StringComparer.OrdinalIgnoreCase);
+			}
 		}
 
 		/// <inheritdoc/>
@@ -1134,6 +1140,12 @@ namespace HordeServer.Notifications.Impl
 
 		private async Task SendMessageAsync(string Recipient, string? Text = null, BlockBase[]? Blocks = null, BlockKitAttachment[]? Attachments = null)
 		{
+			if (AllowUsers != null && !AllowUsers.Contains(Recipient))
+			{
+				Logger.LogDebug("Suppressing message to {Recipient}: {Text}", Recipient, Text);
+				return;
+			}
+
 			BlockKitMessage Message = new BlockKitMessage();
 			Message.Recipient = Recipient;
 			Message.Text = Text;
@@ -1151,6 +1163,12 @@ namespace HordeServer.Notifications.Impl
 
 		private async Task SendOrUpdateMessageAsync(string Recipient, string EventId, ObjectId? UserId, string? Text = null, BlockBase[]? Blocks = null, BlockKitAttachment[]? Attachments = null)
 		{
+			if (AllowUsers != null && !AllowUsers.Contains(Recipient))
+			{
+				Logger.LogDebug("Suppressing message to {Recipient}: {Text}", Recipient, Text);
+				return;
+			}
+
 			BlockKitMessage Message = new BlockKitMessage();
 			Message.Text = Text;
 			if (Blocks != null)
diff --git a/Engine/Source/Programs/Horde/HordeServer/ServerSettings.cs b/Engine/Source/Programs/Horde/HordeServer/ServerSettings.cs
--- a/Engine/Source/Programs/Horde/HordeServer/ServerSettings.cs
+++ b/Engine/Source/Programs/Horde/HordeServer/ServerSettings.cs
@@ -368,6 +368,11 @@ namespace HordeServer
 		/// </summary>
 		public string? SlackSocketToken { get; set; }
 
+		/// <summary>
+		/// Filtered list of slack users to send notifications to. Should be Slack user ids, separated by commas.
+		/// </summary>
+		public string? SlackUsers { get; set; }
+
 		/// <summary>
 		/// Channel to send stream notification update failures to
 		/// </summary>
