Horde: Remove aggregates endpoint. Not used by dashboard.

diff --git a/Engine/Source/Programs/Horde/HordeServer/Controllers/JobsController.cs b/Engine/Source/Programs/Horde/HordeServer/Controllers/JobsController.cs
--- a/Engine/Source/Programs/Horde/HordeServer/Controllers/JobsController.cs
+++ b/Engine/Source/Programs/Horde/HordeServer/Controllers/JobsController.cs
@@ -1234,104 +1234,6 @@ namespace HordeServer.Controllers
 			return NotFound();
 		}
 
-		/// <summary>
-		/// Adds an array of aggregates to a job
-		/// </summary>
-		/// <param name="JobId">Unique id for the job</param>
-		/// <param name="Requests">Properties of the new aggregates</param>
-		/// <returns>Response object with the first index of the new aggregates</returns>
-		[HttpPost]
-		[Route("/api/v1/jobs/{JobId}/aggregates")]
-		public async Task<ActionResult> CreateAggregatesAsync(string JobId, [FromBody] List<NewAggregate> Requests)
-		{
-			ObjectId JobIdValue = JobId.ToObjectId();
-			for (; ; )
-			{
-				IJob? Job = await JobService.GetJobAsync(JobIdValue);
-				if (Job == null)
-				{
-					return NotFound();
-				}
-				if (!await JobService.AuthorizeAsync(Job, AclAction.ExecuteJob, User, null))
-				{
-					return Forbid();
-				}
-
-				IGraph Graph = await JobService.GetGraphAsync(Job);
-				Graph = await Graphs.AppendAsync(Graph, null, Requests, null);
-
-				if (await JobService.TryUpdateGraphAsync(Job, Graph))
-				{
-					return Ok();
-				}
-			}
-		}
-
-		/// <summary>
-		/// Gets the aggregates for a job
-		/// </summary>
-		/// <param name="JobId">Unique id for the job</param>
-		/// <param name="Filter">Filter for the properties to return</param>
-		/// <returns>List of aggregates for the job</returns>
-		[HttpGet]
-		[Route("/api/v1/jobs/{JobId}/aggregates")]
-		[ProducesResponseType(typeof(List<GetAggregateResponse>), 200)]
-		public async Task<ActionResult<List<object>>> GetAggregatesAsync(string JobId, [FromQuery] PropertyFilter? Filter = null)
-		{
-			IJob? Job = await JobService.GetJobAsync(JobId.ToObjectId());
-			if (Job == null)
-			{
-				return NotFound();
-			}
-			if (!await JobService.AuthorizeAsync(Job, AclAction.ViewJob, User, null))
-			{
-				return Forbid();
-			}
-
-			IGraph Graph = await JobService.GetGraphAsync(Job);
-
-			List<GetLabelStateResponse> Aggregates = new List<GetLabelStateResponse>();
-			Job.GetLabelStateResponses(Graph, Aggregates);
-			return Aggregates.ConvertAll(x => PropertyFilter.Apply(x, Filter));
-		}
-
-		/// <summary>
-		/// Gets an aggregate in a job
-		/// </summary>
-		/// <param name="JobId">Unique id for the job</param>
-		/// <param name="AggregateIdx">The aggregate index</param>
-		/// <param name="Filter">Filter for the properties to return</param>
-		/// <returns>Aggregate information</returns>
-		[HttpGet]
-		[Route("/api/v1/jobs/{JobId}/aggregates/{AggregateIdx}")]
-		[ProducesResponseType(typeof(GetLabelStateResponse), 200)]
-		public async Task<ActionResult<object>> GetAggregateAsync(string JobId, int AggregateIdx, [FromQuery] PropertyFilter? Filter = null)
-		{
-			IJob? Job = await JobService.GetJobAsync(JobId.ToObjectId());
-			if (Job == null)
-			{
-				return NotFound();
-			}
-			if (!await JobService.AuthorizeAsync(Job, AclAction.ViewJob, User, null))
-			{
-				return Forbid();
-			}
-
-			IGraph Graph = await JobService.GetGraphAsync(Job);
-
-			List<GetLabelStateResponse> Responses = new List<GetLabelStateResponse>();
-			Job.GetLabelStateResponses(Graph, Responses);
-
-			if (AggregateIdx < 0 || AggregateIdx >= Responses.Count)
-			{
-				return NotFound();
-			}
-			else
-			{
-				return PropertyFilter.Apply(Responses[AggregateIdx], Filter);
-			}
-		}
-
 		/// <summary>
 		/// Updates notifications for a specific label.
 		/// </summary>
