Horde: Rename model interfaces to files starting with "I".

diff --git a/Engine/Source/Programs/Horde/HordeServer/Models/Agent.cs b/Engine/Source/Programs/Horde/HordeServer/Models/Agent.cs
deleted file mode 100644
diff --git a/Engine/Source/Programs/Horde/HordeServer/Models/Device.cs b/Engine/Source/Programs/Horde/HordeServer/Models/Device.cs
deleted file mode 100644
diff --git a/Engine/Source/Programs/Horde/HordeServer/Models/Event.cs b/Engine/Source/Programs/Horde/HordeServer/Models/Event.cs
deleted file mode 100644
diff --git a/Engine/Source/Programs/Horde/HordeServer/Models/Graph.cs b/Engine/Source/Programs/Horde/HordeServer/Models/Graph.cs
deleted file mode 100644
diff --git a/Engine/Source/Programs/Horde/HordeServer/Models/IAgent.cs b/Engine/Source/Programs/Horde/HordeServer/Models/IAgent.cs
new file mode 100644
--- /dev/null
+++ b/Engine/Source/Programs/Horde/HordeServer/Models/IAgent.cs
@@ -0,0 +1,803 @@
+// Copyright Epic Games, Inc. All Rights Reserved.
+
+using Google.Protobuf.WellKnownTypes;
+using EpicGames.Core;
+using HordeServer.Api;
+using HordeCommon;
+using HordeCommon.Rpc.Tasks;
+using HordeServer.Services;
+using HordeServer.Utilities;
+using MongoDB.Bson;
+using MongoDB.Bson.Serialization;
+using MongoDB.Bson.Serialization.Attributes;
+using System;
+using System.Collections.Generic;
+using System.Diagnostics.CodeAnalysis;
+using System.IO;
+using System.Linq;
+using System.Security.Cryptography;
+using Serilog;
+using System.Diagnostics;
+using System.Threading.Tasks;
+using HordeServer.Collections;
+using EpicGames.Horde.Compute;
+using EpicGames.Horde.Common;
+
+namespace HordeServer.Models
+{
+	using PoolId = StringId<IPool>;
+	using StreamId = StringId<IStream>;
+	using AgentSoftwareVersion = StringId<IAgentSoftwareCollection>;
+	using AgentSoftwareChannelName = StringId<AgentSoftwareChannels>;
+
+	/// <summary>
+	/// Information about a workspace synced to an agent
+	/// </summary>
+	public class AgentWorkspace
+	{
+		/// <summary>
+		/// Name of the Perforce cluster to use
+		/// </summary>
+		public string? Cluster { get; set; }
+
+		/// <summary>
+		/// User to log into Perforce with (eg. buildmachine)
+		/// </summary>
+		public string? UserName { get; set; }
+
+		/// <summary>
+		/// Identifier to distinguish this workspace from other workspaces
+		/// </summary>
+		public string Identifier { get; set; }
+
+		/// <summary>
+		/// The stream to sync
+		/// </summary>
+		public string Stream { get; set; }
+
+		/// <summary>
+		/// Custom view for the workspace
+		/// </summary>
+		public List<string>? View { get; set; }
+
+		/// <summary>
+		/// Whether to use an incremental workspace
+		/// </summary>
+		public bool bIncremental { get; set; }
+
+		/// <summary>
+		/// Constructor
+		/// </summary>
+		/// <param name="Cluster">Name of the Perforce cluster</param>
+		/// <param name="UserName">User to log into Perforce with (eg. buildmachine)</param>
+		/// <param name="Identifier">Identifier to distinguish this workspace from other workspaces</param>
+		/// <param name="Stream">The stream to sync</param>
+		/// <param name="View">Custom view for the workspace</param>
+		/// <param name="bIncremental">Whether to use an incremental workspace</param>
+		public AgentWorkspace(string? Cluster, string? UserName, string Identifier, string Stream, List<string>? View, bool bIncremental)
+		{
+			if (!String.IsNullOrEmpty(Cluster))
+			{
+				this.Cluster = Cluster;
+			}
+			if (!String.IsNullOrEmpty(UserName))
+			{
+				this.UserName = UserName;
+			}
+			this.Identifier = Identifier;
+			this.Stream = Stream;
+			this.View = View;
+			this.bIncremental = bIncremental;
+		}
+
+		/// <summary>
+		/// Constructor
+		/// </summary>
+		/// <param name="Workspace">RPC message to construct from</param>
+		public AgentWorkspace(HordeCommon.Rpc.Messages.AgentWorkspace Workspace)
+			: this(Workspace.ConfiguredCluster, Workspace.ConfiguredUserName, Workspace.Identifier, Workspace.Stream, (Workspace.View.Count > 0) ? Workspace.View.ToList() : null, Workspace.Incremental)
+		{
+		}
+
+		/// <summary>
+		/// Gets a digest of the settings for this workspace
+		/// </summary>
+		/// <returns>Digest for the workspace settings</returns>
+		public string GetDigest()
+		{
+#pragma warning disable CA5351 // Do Not Use Broken Cryptographic Algorithms
+			using (MD5 Hasher = MD5.Create())
+			{
+				byte[] Data = BsonExtensionMethods.ToBson(this);
+				return BitConverter.ToString(Hasher.ComputeHash(Data)).Replace("-", "", StringComparison.Ordinal);
+			}
+#pragma warning restore CA5351
+		}
+
+		/// <inheritdoc/>
+		public override bool Equals(object? Obj)
+		{
+			AgentWorkspace? Other = Obj as AgentWorkspace;
+			if (Other == null)
+			{
+				return false;
+			}
+			if (Cluster != Other.Cluster || UserName != Other.UserName || Identifier != Other.Identifier || Stream != Other.Stream || bIncremental != Other.bIncremental)
+			{
+				return false;
+			}
+			if (!Enumerable.SequenceEqual(View ?? new List<string>(), Other.View ?? new List<string>()))
+			{
+				return false;
+			}
+			return true;
+		}
+
+		/// <inheritdoc/>
+		public override int GetHashCode()
+		{
+			return HashCode.Combine(Cluster, UserName, Identifier, Stream, bIncremental); // Ignore 'View' for now
+		}
+
+		/// <summary>
+		/// Checks if two workspace sets are equivalent, ignoring order
+		/// </summary>
+		/// <param name="WorkspacesA">First list of workspaces</param>
+		/// <param name="WorkspacesB">Second list of workspaces</param>
+		/// <returns>True if the sets are equivalent</returns>
+		public static bool SetEquals(IReadOnlyList<AgentWorkspace> WorkspacesA, IReadOnlyList<AgentWorkspace> WorkspacesB)
+		{
+			HashSet<AgentWorkspace> WorkspacesSetA = new HashSet<AgentWorkspace>(WorkspacesA);
+			return WorkspacesSetA.SetEquals(WorkspacesB);
+		}
+
+		/// <summary>
+		/// Converts this workspace to an RPC message
+		/// </summary>
+		/// <param name="Server">The Perforce server</param>
+		/// <param name="Credentials">Credentials for the server</param>
+		/// <returns>The RPC message</returns>
+		public HordeCommon.Rpc.Messages.AgentWorkspace ToRpcMessage(IPerforceServer Server, PerforceCredentials? Credentials)
+		{
+			// Construct the message
+			HordeCommon.Rpc.Messages.AgentWorkspace Result = new HordeCommon.Rpc.Messages.AgentWorkspace();
+			Result.ConfiguredCluster = Cluster;
+			Result.ConfiguredUserName = UserName;
+			Result.ServerAndPort = Server.ServerAndPort;
+			Result.UserName = Credentials?.UserName ?? UserName;
+			Result.Password = Credentials?.Password;
+			Result.Identifier = Identifier;
+			Result.Stream = Stream;
+			if (View != null)
+			{
+				Result.View.AddRange(View);
+			}
+			Result.Incremental = bIncremental;
+			return Result;
+		}
+	}
+
+	/// <summary>
+	/// Document describing an active lease
+	/// </summary>
+	public class AgentLease
+	{
+		/// <summary>
+		/// Lease returned when a task source wants to run, but needs to wait for other tasks to finish first.
+		/// </summary>
+		public static AgentLease Drain { get; } = new AgentLease(ObjectId.Empty, String.Empty, null, null, null, LeaseState.Pending, null, true, null);
+
+		/// <summary>
+		/// Name of this lease
+		/// </summary>
+		[BsonRequired]
+		public ObjectId Id { get; set; }
+
+		/// <summary>
+		/// Name of this lease
+		/// </summary>
+		public string Name { get; set; }
+
+		/// <summary>
+		/// The current state of the lease
+		/// </summary>
+		public LeaseState State { get; set; }
+
+		/// <summary>
+		/// The stream for this lease
+		/// </summary>
+		public StreamId? StreamId { get; set; }
+
+		/// <summary>
+		/// The pool for this lease
+		/// </summary>
+		public PoolId? PoolId { get; set; }
+
+		/// <summary>
+		/// Optional log for this lease
+		/// </summary>
+		public ObjectId? LogId { get; set; }
+
+		/// <summary>
+		/// Time at which the lease started
+		/// </summary>
+		[BsonRequired]
+		public DateTime StartTime { get; set; }
+
+		/// <summary>
+		/// Time at which the lease should be terminated
+		/// </summary>
+		public DateTime? ExpiryTime { get; set; }
+
+		/// <summary>
+		/// Flag indicating whether this lease has been accepted by the agent
+		/// </summary>
+		public bool Active { get; set; }
+
+		/// <summary>
+		/// Resources used by this lease
+		/// </summary>
+		public IReadOnlyDictionary<string, int>? Resources { get; set; }
+
+		/// <summary>
+		/// Whether the lease requires exclusive access to the agent
+		/// </summary>
+		public bool Exclusive { get; set; }
+
+		/// <summary>
+		/// For leases in the pending state, encodes an "any" protobuf containing the payload for the agent to execute the lease.
+		/// </summary>
+		public byte[]? Payload { get; set; }
+
+		/// <summary>
+		/// Private constructor
+		/// </summary>
+		[BsonConstructor]
+		private AgentLease()
+		{
+			Name = String.Empty;
+		}
+
+		/// <summary>
+		/// Constructor
+		/// </summary>
+		/// <param name="Id">Identifier for the lease</param>
+		/// <param name="Name">Name of this lease</param>
+		/// <param name="StreamId"></param>
+		/// <param name="PoolId"></param>
+		/// <param name="LogId">Unique id for the log</param>
+		/// <param name="State">State for the lease</param>
+		/// <param name="Resources">Resources required for this lease</param>
+		/// <param name="Exclusive">Whether to reserve the entire device</param>
+		/// <param name="Payload">Encoded "any" protobuf describing the contents of the payload</param>
+		public AgentLease(ObjectId Id, string Name, StreamId? StreamId, PoolId? PoolId, ObjectId? LogId, LeaseState State, IReadOnlyDictionary<string, int>? Resources, bool Exclusive, byte[]? Payload)
+		{
+			this.Id = Id;
+			this.Name = Name;
+			this.StreamId = StreamId;
+			this.PoolId = PoolId;
+			this.LogId = LogId;
+			this.State = State;
+			this.Resources = Resources;
+			this.Exclusive = Exclusive;
+			this.Payload = Payload;
+			StartTime = DateTime.UtcNow;
+		}
+
+		/// <summary>
+		/// Determines if this is a conform lease
+		/// </summary>
+		/// <returns>True if this is a conform lease</returns>
+		public bool IsConformLease()
+		{
+			if (Payload != null)
+			{
+				Any BasePayload = Any.Parser.ParseFrom(Payload);
+				if (BasePayload.Is(ConformTask.Descriptor))
+				{
+					return true;
+				}
+			}
+			return false;
+		}
+
+		/// <summary>
+		/// Gets user-readable payload information
+		/// </summary>
+		/// <param name="Payload">The payload data</param>
+		/// <returns>Dictionary of key/value pairs for the payload</returns>
+		public static Dictionary<string, string>? GetPayloadDetails(ReadOnlyMemory<byte>? Payload)
+		{
+			Dictionary<string, string>? Details = null;
+			if (Payload != null)
+			{
+				Any BasePayload = Any.Parser.ParseFrom(Payload.Value.ToArray());
+
+				Details = new Dictionary<string, string>();
+
+				ConformTask ConformTask;
+				if(BasePayload.TryUnpack(out ConformTask))
+				{
+					Details["Type"] = "Conform";
+					Details["LogId"] = ConformTask.LogId;
+				}
+
+				ExecuteJobTask JobTask;
+				if (BasePayload.TryUnpack(out JobTask))
+				{
+					Details["Type"] = "Job";
+					Details["JobId"] = JobTask.JobId;
+					Details["BatchId"] = JobTask.BatchId;
+					Details["LogId"] = JobTask.LogId;
+				}
+
+				UpgradeTask UpgradeTask;
+				if (BasePayload.TryUnpack(out UpgradeTask))
+				{
+					Details["Type"] = "Upgrade";
+					Details["SoftwareId"] = UpgradeTask.SoftwareId;
+					Details["LogId"] = UpgradeTask.LogId;
+				}
+			}
+			return Details;
+		}
+
+		/// <summary>
+		/// Converts this lease to an RPC message
+		/// </summary>
+		/// <returns>RPC message</returns>
+		public HordeCommon.Rpc.Messages.Lease ToRpcMessage()
+		{
+			HordeCommon.Rpc.Messages.Lease Lease = new HordeCommon.Rpc.Messages.Lease();
+			Lease.Id = Id.ToString();
+			Lease.Payload = Google.Protobuf.WellKnownTypes.Any.Parser.ParseFrom(Payload); 
+			Lease.State = State;
+			return Lease;
+		}
+	}
+
+	/// <summary>
+	/// Well-known property names for agents
+	/// </summary>
+	static class KnownPropertyNames
+	{
+		/// <summary>
+		/// The agent id
+		/// </summary>
+		public const string Id = "Id";
+
+		/// <summary>
+		/// The operating system (Linux, MacOS, Windows)
+		/// </summary>
+		public const string OSFamily = "OSFamily";
+
+		/// <summary>
+		/// Pools that this agent belongs to
+		/// </summary>
+		public const string Pool = "Pool";
+
+		/// <summary>
+		/// Number of logical cores
+		/// </summary>
+		public const string LogicalCores = "LogicalCores";
+
+		/// <summary>
+		/// Amount of RAM, in GB
+		/// </summary>
+		public const string RAM = "RAM";
+	}
+
+	/// <summary>
+	/// Mirrors an Agent document in the database
+	/// </summary>
+	public interface IAgent
+	{
+		/// <summary>
+		/// Randomly generated unique id for this agent.
+		/// </summary>
+		public AgentId Id { get; }
+
+		/// <summary>
+		/// The current session id, if it's online
+		/// </summary>
+		public ObjectId? SessionId { get; }
+
+		/// <summary>
+		/// Time at which the current session expires. 
+		/// </summary>
+		public DateTime? SessionExpiresAt { get; }
+
+		/// <summary>
+		/// Current status of this agent
+		/// </summary>
+		public AgentStatus Status { get; }
+
+		/// <summary>
+		/// Whether the agent is enabled
+		/// </summary>
+		public bool Enabled { get; }
+
+		/// <summary>
+		/// Whether the agent should be included on the dashboard. This is set to true for ephemeral agents once they are no longer online, or agents that are explicitly deleted.
+		/// </summary>
+		public bool Deleted { get; }
+
+		/// <summary>
+		/// Version of the software running on this agent
+		/// </summary>
+		public string? Version { get; }
+
+		/// <summary>
+		/// Arbitrary comment for the agent (useful for disable reasons etc)
+		/// </summary>
+		public string? Comment { get; }
+
+		/// <summary>
+		/// List of properties for this agent
+		/// </summary>
+		public IReadOnlyList<string> Properties { get; }
+
+		/// <summary>
+		/// List of resources available to the agent
+		/// </summary>
+		public IReadOnlyDictionary<string, int> Resources { get; }
+
+		/// <summary>
+		/// Channel for the software running on this agent. Uses <see cref="AgentSoftwareService.DefaultChannelName"/> if not specified
+		/// </summary>
+		public AgentSoftwareChannelName? Channel { get; }
+
+		/// <summary>
+		/// Last upgrade that was attempted
+		/// </summary>
+		public string? LastUpgradeVersion { get; }
+
+		/// <summary>
+		/// Time that which the last upgrade was attempted
+		/// </summary>
+		public DateTime? LastUpgradeTime { get; }
+
+		/// <summary>
+		/// Dynamically applied pools
+		/// </summary>
+		public IReadOnlyList<PoolId> DynamicPools { get; }
+
+		/// <summary>
+		/// List of manually assigned pools for agent
+		/// </summary>
+		public IReadOnlyList<PoolId> ExplicitPools { get; }
+
+		/// <summary>
+		/// Whether a conform is requested
+		/// </summary>
+		public bool RequestConform { get; }
+
+		/// <summary>
+		/// Whether a machine restart is requested
+		/// </summary>
+		public bool RequestRestart { get; }
+
+		/// <summary>
+		/// Whether the machine should be shutdown
+		/// </summary>
+		public bool RequestShutdown { get; }
+
+		/// <summary>
+		/// List of workspaces currently synced to this machine
+		/// </summary>
+		public IReadOnlyList<AgentWorkspace> Workspaces { get; }
+
+		/// <summary>
+		/// Time at which the last conform job ran
+		/// </summary>
+		public DateTime LastConformTime { get; }
+
+		/// <summary>
+		/// Number of times a conform job has failed
+		/// </summary>
+		public int? ConformAttemptCount { get; }
+
+		/// <summary>
+		/// Array of active leases.
+		/// </summary>
+		public IReadOnlyList<AgentLease> Leases { get; }
+
+		/// <summary>
+		/// ACL for modifying this agent
+		/// </summary>
+		public Acl? Acl { get; }
+
+		/// <summary>
+		/// Last time that the agent was modified
+		/// </summary>
+		public DateTime UpdateTime { get; }
+
+		/// <summary>
+		/// Update counter for this document. Any updates should compare-and-swap based on the value of this counter, or increment it in the case of server-side updates.
+		/// </summary>
+		public uint UpdateIndex { get; }
+	}
+
+	/// <summary>
+	/// Extension methods for IAgent
+	/// </summary>
+	public static class AgentExtensions
+	{
+		/// <summary>
+		/// Determines whether this agent is online
+		/// </summary>
+		/// <returns></returns>
+		public static bool IsSessionValid(this IAgent Agent, DateTime UtcNow)
+		{
+			return Agent.SessionId.HasValue && Agent.SessionExpiresAt.HasValue && UtcNow < Agent.SessionExpiresAt.Value;
+		}
+
+		/// <summary>
+		/// Tests whether an agent is in the given pool
+		/// </summary>
+		/// <param name="Agent"></param>
+		/// <param name="PoolId"></param>
+		/// <returns></returns>
+		public static bool IsInPool(this IAgent Agent, PoolId PoolId)
+		{
+			return Agent.DynamicPools.Contains(PoolId) || Agent.ExplicitPools.Contains(PoolId);
+		}
+
+		/// <summary>
+		/// Get all the pools for each agent
+		/// </summary>
+		/// <param name="Agent">The agent to query</param>
+		/// <returns></returns>
+		public static IEnumerable<PoolId> GetPools(this IAgent Agent)
+		{
+			foreach (PoolId PoolId in Agent.DynamicPools)
+			{
+				yield return PoolId;
+			}
+			foreach (PoolId PoolId in Agent.ExplicitPools)
+			{
+				yield return PoolId;
+			}
+		}
+
+		/// <summary>
+		/// Tests whether an agent has a particular property
+		/// </summary>
+		/// <param name="Agent"></param>
+		/// <param name="Property"></param>
+		/// <returns></returns>
+		public static bool HasProperty(this IAgent Agent, string Property)
+		{
+			return Agent.Properties.BinarySearch(Property, StringComparer.OrdinalIgnoreCase) >= 0;
+		}
+
+		/// <summary>
+		/// Finds property values from a sorted list of Name=Value pairs
+		/// </summary>
+		/// <param name="Agent">The agent to query</param>
+		/// <param name="Name">Name of the property to find</param>
+		/// <returns>Property values</returns>
+		public static IEnumerable<string> GetPropertyValues(this IAgent Agent, string Name)
+		{
+			if (Name.Equals(KnownPropertyNames.Id, StringComparison.OrdinalIgnoreCase))
+			{
+				yield return Agent.Id.ToString();
+			}
+			else if (Name.Equals(KnownPropertyNames.Pool, StringComparison.OrdinalIgnoreCase))
+			{
+				foreach (PoolId PoolId in Agent.GetPools())
+				{
+					yield return PoolId.ToString();
+				}
+			}
+			else
+			{
+				int Index = Agent.Properties.BinarySearch(Name, StringComparer.OrdinalIgnoreCase);
+				if (Index < 0)
+				{
+					Index = ~Index;
+					while (Index < Agent.Properties.Count)
+					{
+						string Property = Agent.Properties[Index];
+						if (Property.Length <= Name.Length || !Property.StartsWith(Name, StringComparison.OrdinalIgnoreCase) || Property[Name.Length] != '=')
+						{
+							break;
+						}
+						yield return Property.Substring(Name.Length + 1);
+					}
+				}
+			}
+		}
+
+		/// <summary>
+		/// Evaluates a condition against an agent
+		/// </summary>
+		/// <param name="Agent">The agent to evaluate</param>
+		/// <param name="Condition">The condition to evaluate</param>
+		/// <returns>True if the agent satisfies the condition</returns>
+		public static bool SatisfiesCondition(this IAgent Agent, Condition Condition)
+		{
+			return Condition.Evaluate(x => Agent.GetPropertyValues(x));
+		}
+
+		/// <summary>
+		/// Determine whether it's possible to add a lease for the given resources
+		/// </summary>
+		/// <param name="Agent">The agent to create a lease for</param>
+		/// <param name="Requirements">Requirements for the lease</param>
+		/// <returns>True if the new lease can be granted</returns>
+		public static bool MeetsRequirements(this IAgent Agent, Requirements Requirements)
+		{
+			return MeetsRequirements(Agent, Requirements.Condition, Requirements.Resources, Requirements.Exclusive);
+		}
+
+		/// <summary>
+		/// Determine whether it's possible to add a lease for the given resources
+		/// </summary>
+		/// <param name="Agent">The agent to create a lease for</param>
+		/// <param name="Exclusive">Whether t</param>
+		/// <param name="Condition">Condition to satisfy</param>
+		/// <param name="Resources">Resources required to execute</param>
+		/// <returns>True if the new lease can be granted</returns>
+		public static bool MeetsRequirements(this IAgent Agent, Condition? Condition, Dictionary<string, int>? Resources, bool Exclusive)
+		{
+			if (!Agent.Enabled || Agent.Status != AgentStatus.Ok)
+			{
+				return false;
+			}
+			if (Agent.Leases.Any(x => x.Exclusive))
+			{
+				return false;
+			}
+			if (Exclusive && Agent.Leases.Any())
+			{
+				return false;
+			}
+			if (Condition != null && !Agent.SatisfiesCondition(Condition))
+			{
+				return false;
+			}
+			if (Resources != null)
+			{
+				foreach ((string Name, int Count) in Resources)
+				{
+					int RemainingCount;
+					if (!Agent.Resources.TryGetValue(Name, out RemainingCount))
+					{
+						return false;
+					}
+					foreach (AgentLease Lease in Agent.Leases)
+					{
+						if (Lease.Resources != null)
+						{
+							int LeaseCount;
+							Lease.Resources.TryGetValue(Name, out LeaseCount);
+							RemainingCount -= LeaseCount;
+						}
+					}
+					if (RemainingCount < Count)
+					{
+						return false;
+					}
+				}
+			}
+			return true;
+		}
+
+		/// <summary>
+		/// Gets all the autosdk workspaces required for an agent
+		/// </summary>
+		/// <param name="Agent"></param>
+		/// <param name="Globals"></param>
+		/// <param name="Workspaces"></param>
+		/// <returns></returns>
+		public static HashSet<AgentWorkspace> GetAutoSdkWorkspaces(this IAgent Agent, Globals Globals, List<AgentWorkspace> Workspaces)
+		{
+			HashSet<AgentWorkspace> AutoSdkWorkspaces = new HashSet<AgentWorkspace>();
+			foreach (string? ClusterName in Workspaces.Select(x => x.Cluster).Distinct())
+			{
+				PerforceCluster? Cluster = Globals.FindPerforceCluster(ClusterName);
+				if (Cluster != null)
+				{
+					AgentWorkspace? AutoSdkWorkspace = GetAutoSdkWorkspace(Agent, Cluster);
+					if (AutoSdkWorkspace != null)
+					{
+						AutoSdkWorkspaces.Add(AutoSdkWorkspace);
+					}
+				}
+			}
+			return AutoSdkWorkspaces;
+		}
+
+		/// <summary>
+		/// Get the AutoSDK workspace required for an agent
+		/// </summary>
+		/// <param name="Agent"></param>
+		/// <param name="Cluster">The perforce cluster to get a workspace for</param>
+		/// <returns></returns>
+		public static AgentWorkspace? GetAutoSdkWorkspace(this IAgent Agent, PerforceCluster Cluster)
+		{
+			foreach (AutoSdkWorkspace AutoSdk in Cluster.AutoSdk)
+			{
+				if (AutoSdk.Stream != null && AutoSdk.Properties.All(x => Agent.Properties.Contains(x)))
+				{
+					return new AgentWorkspace(Cluster.Name, AutoSdk.UserName, AutoSdk.Name ?? "AutoSDK", AutoSdk.Stream!, null, true);
+				}
+			}
+			return null;
+		}
+
+		/// <summary>
+		/// Converts this workspace to an RPC message
+		/// </summary>
+		/// <param name="Agent">The agent to get a workspace for</param>
+		/// <param name="Workspace">The workspace definition</param>
+		/// <param name="Cluster">The global state</param>
+		/// <param name="LoadBalancer">The Perforce load balancer</param>
+		/// <param name="WorkspaceMessages">List of messages</param>
+		/// <returns>The RPC message</returns>
+		public static async Task<bool> TryAddWorkspaceMessage(this IAgent Agent, AgentWorkspace Workspace, PerforceCluster Cluster, PerforceLoadBalancer LoadBalancer, IList<HordeCommon.Rpc.Messages.AgentWorkspace> WorkspaceMessages)
+		{
+			// Find a matching server, trying to use a previously selected one if possible
+			string? BaseServerAndPort;
+			string? ServerAndPort;
+
+			HordeCommon.Rpc.Messages.AgentWorkspace? ExistingWorkspace = WorkspaceMessages.FirstOrDefault(x => x.ConfiguredCluster == Workspace.Cluster);
+			if(ExistingWorkspace != null)
+			{
+				BaseServerAndPort = ExistingWorkspace.BaseServerAndPort;
+				ServerAndPort = ExistingWorkspace.ServerAndPort;
+			}
+			else
+			{
+				if (Cluster == null)
+				{
+					return false;
+				}
+
+				IPerforceServer? Server = await LoadBalancer.SelectServerAsync(Cluster, Agent);
+				if (Server == null)
+				{
+					return false;
+				}
+
+				BaseServerAndPort = Server.BaseServerAndPort;
+				ServerAndPort = Server.ServerAndPort;
+			}
+
+			// Find the matching credentials for the desired user
+			PerforceCredentials? Credentials = null;
+			if (Cluster != null)
+			{
+				if (Workspace.UserName == null)
+				{
+					Credentials = Cluster.Credentials.FirstOrDefault();
+				}
+				else
+				{
+					Credentials = Cluster.Credentials.FirstOrDefault(x => String.Equals(x.UserName, Workspace.UserName, StringComparison.OrdinalIgnoreCase));
+				}
+			}
+
+			// Construct the message
+			HordeCommon.Rpc.Messages.AgentWorkspace Result = new HordeCommon.Rpc.Messages.AgentWorkspace();
+			Result.ConfiguredCluster = Workspace.Cluster;
+			Result.ConfiguredUserName = Workspace.UserName;
+			Result.Cluster = Cluster?.Name;
+			Result.BaseServerAndPort = BaseServerAndPort;
+			Result.ServerAndPort = ServerAndPort;
+			Result.UserName = Credentials?.UserName ?? Workspace.UserName;
+			Result.Password = Credentials?.Password;
+			Result.Identifier = Workspace.Identifier;
+			Result.Stream = Workspace.Stream;
+			if (Workspace.View != null)
+			{
+				Result.View.AddRange(Workspace.View);
+			}
+			Result.Incremental = Workspace.bIncremental;
+			WorkspaceMessages.Add(Result);
+
+			return true;
+		}
+	}
+}
diff --git a/Engine/Source/Programs/Horde/HordeServer/Models/IDevice.cs b/Engine/Source/Programs/Horde/HordeServer/Models/IDevice.cs
new file mode 100644
--- /dev/null
+++ b/Engine/Source/Programs/Horde/HordeServer/Models/IDevice.cs
@@ -0,0 +1,277 @@
+// Copyright Epic Games, Inc. All Rights Reserved.
+
+using System;
+using MongoDB.Bson;
+using MongoDB.Bson.Serialization.Attributes;
+
+using System.Collections.Generic;
+using HordeServer.Utilities;
+
+namespace HordeServer.Models
+{
+	using DeviceId = StringId<IDevice>;
+	using DevicePlatformId = StringId<IDevicePlatform>;
+	using DevicePoolId = StringId<IDevicePool>;
+
+	/// <summary>
+	/// A reservation containing one or more devices
+	/// </summary>
+	public interface IDeviceReservation
+	{
+		/// <summary>
+		/// Randomly generated unique id for this reservation
+		/// </summary>
+		public ObjectId Id { get; }
+
+		/// <summary>
+		/// Which device pool the reservation is in
+		/// </summary>
+		public DevicePoolId PoolId { get; }
+
+		/// <summary>
+		/// JobID holding reservation
+		/// </summary>
+		public string? JobId { get; }
+
+		/// <summary>
+		/// Job step id holding reservation
+		/// </summary>
+		public string? StepId { get; }
+
+		/// <summary>
+		/// Reservations held by a user, requires a token
+		/// </summary>
+		public ObjectId? UserId { get; }
+
+		/// <summary>
+		/// The hostname of machine holding reservation
+		/// </summary>
+		string? Hostname { get; }
+
+		/// <summary>
+		/// The hostname of machine holding reservation
+		/// </summary>
+		string? ReservationDetails { get; }
+
+		/// <summary>
+		/// The UTC time when the reservation was created
+		/// </summary>
+		DateTime CreateTimeUtc { get; }
+
+		/// <summary>
+		/// The last update time for reservation renewal (for expiration)
+		/// </summary>
+		DateTime UpdateTimeUtc { get; }
+
+		/// <summary>
+		/// The reserved devices
+		/// </summary>
+		public List<DeviceId> Devices { get; }
+
+		/// <summary>
+		/// The legacy reservation system guid, to be removed once can update Gauntlet client in all streams
+		/// </summary>
+		public string LegacyGuid { get; }
+	}
+
+
+	/// <summary>
+	/// A device platform 
+	/// </summary>
+	public interface IDevicePlatform
+	{
+		/// <summary>
+		/// Unique identifier of device platform 
+		/// </summary>
+		public DevicePlatformId Id { get; }
+		
+		/// <summary>
+		/// Device platform name, for example Android, PS5, etc
+		/// </summary>
+		string Name { get; }
+
+		/// <summary>
+		/// A list of valid models for the platform
+		/// </summary>
+		public IReadOnlyList<string>? Models { get; }
+
+	}
+
+    /// <summary>
+    /// The type of device pool
+    /// </summary>
+    public enum DevicePoolType
+    {
+        /// <summary>
+        /// Available to CIS jobs
+        /// </summary>
+        Automation,
+
+        /// <summary>
+        /// Shared by users with remote checking and checkouts
+        /// </summary>
+        Shared
+
+    }
+
+    /// <summary>
+    /// A logical pool of devices
+    /// </summary>
+    public interface IDevicePool
+	{
+		/// <summary>
+		/// Unique identifier of pool
+		/// </summary>
+		public DevicePoolId Id { get; }
+
+		/// <summary>
+		/// The type of pool 
+		/// </summary>
+		public DevicePoolType PoolType { get; }
+
+        /// <summary>
+        /// Friendly name of the pool
+        /// </summary>
+        string Name { get; }
+
+		/// <summary>
+		/// Acl access to the pool
+		/// </summary>
+		public Acl? Acl { get; }
+	}
+
+	/// <summary>
+	/// A device utilization snapshot
+	/// </summary>
+	public class DeviceUtilizationTelemetry
+	{
+		/// <summary>
+		/// The job id which utilized device
+		/// </summary>
+		public string? JobId { get; set;}
+
+		/// <summary>
+		/// The job's step id
+		/// </summary>
+		public string? StepId { get; set;}
+
+		/// <summary>
+		/// The time device was reserved
+		/// </summary>
+		public DateTime ReservationStartUtc { get; set;}
+
+		/// <summary>
+		/// The time device was freed
+		/// </summary>
+		public DateTime? ReservationFinishUtc { get; set;}
+
+		/// <summary>
+		/// Private constructor
+		/// </summary>
+		[BsonConstructor]
+		private DeviceUtilizationTelemetry()
+		{
+		}
+
+		/// <summary>
+		/// Constructor
+		/// </summary>
+		/// <param name="ReservationStartUtc"></param>
+		public DeviceUtilizationTelemetry( DateTime ReservationStartUtc )
+		{
+			this.ReservationStartUtc = ReservationStartUtc;
+		}
+
+	}
+
+
+	/// <summary>
+	/// A physical device
+	/// </summary>
+	public interface IDevice
+	{
+		/// <summary>
+		/// The unique id of the device
+		/// </summary>
+		public DeviceId Id { get; }
+
+		/// <summary>
+		/// The platform of the device
+		/// </summary>
+		public DevicePlatformId PlatformId { get; }
+
+		/// <summary>
+		/// Which pool the device belongs to
+		/// </summary>
+		public DevicePoolId PoolId { get; }
+
+		/// <summary>
+		/// Friendly name of the device
+		/// </summary>
+		public string Name { get; }
+
+		/// <summary>
+		/// The model of the device
+		/// </summary>
+		public string? ModelId { get; }
+
+
+		/// <summary>
+		/// Address if the device supports shared network connections
+		/// </summary>
+		public string? Address { get; }
+
+		/// <summary>
+		/// Whether the device is currently enabled
+		/// </summary>
+		public bool Enabled { get; }
+
+		/// <summary>
+		/// Id of the user that last modified this device
+		/// </summary>
+		public string? ModifiedByUser { get; }
+
+		/// <summary>
+		/// Id of the user that has this device checked out
+		/// </summary>
+		public string? CheckedOutByUser { get; }
+
+		/// <summary>
+		/// The last time this device was checked out
+		/// </summary>
+		public DateTime? CheckOutTime { get; }
+
+
+		/// <summary>
+		/// The last time a problem was reported
+		/// </summary>
+		public DateTime? ProblemTimeUtc { get; }
+
+		/// <summary>
+		/// The time device was marked for maintenance
+		/// </summary>
+		public DateTime? MaintenanceTimeUtc { get; }
+
+		/// <summary>
+		/// Markdown notes for device if any 
+		/// </summary>
+		public string? Notes { get; }
+
+		/// <summary>
+		/// Device job utilization history 
+		/// </summary>
+		public List<DeviceUtilizationTelemetry>? Utilization { get; set; }
+
+		/// <summary>
+		/// ACL for modifying this device
+		/// </summary>
+		public Acl? Acl { get; }
+	}
+}
+
+
+
+
+
+
+
diff --git a/Engine/Source/Programs/Horde/HordeServer/Models/IEvent.cs b/Engine/Source/Programs/Horde/HordeServer/Models/IEvent.cs
new file mode 100644
--- /dev/null
+++ b/Engine/Source/Programs/Horde/HordeServer/Models/IEvent.cs
@@ -0,0 +1,85 @@
+// Copyright Epic Games, Inc. All Rights Reserved.
+
+using HordeServer.Api;
+using HordeServer.Utilities;
+using System;
+using System.Collections.Generic;
+using System.Diagnostics.CodeAnalysis;
+using System.Linq;
+using System.Threading.Tasks;
+
+namespace HordeServer.Models
+{
+	using StreamId = StringId<IStream>;
+	using TemplateRefId = StringId<TemplateRef>;
+
+	/// <summary>
+	/// Base interface for events that users can subscribe for notifications to
+	/// </summary>
+	public interface IEvent
+	{
+		/// <summary>
+		/// Converts this event into an API event record
+		/// </summary>
+		/// <returns>The API event record</returns>
+		public EventRecord ToRecord();
+	}
+
+	/// <summary>
+	/// Event that occurs when a job completes
+	/// </summary>
+	public interface IJobCompleteEvent : IEvent
+	{
+		/// <summary>
+		/// The stream id
+		/// </summary>
+		StreamId StreamId { get; }
+
+		/// <summary>
+		/// The template id
+		/// </summary>
+		TemplateRefId TemplateId { get; }
+	}
+
+	/// <summary>
+	/// Event that occurs when a label completes
+	/// </summary>
+	public interface ILabelCompleteEvent : IEvent
+	{
+		/// <summary>
+		/// The stream id
+		/// </summary>
+		StreamId StreamId { get; }
+
+		/// <summary>
+		/// The template id
+		/// </summary>
+		TemplateRefId TemplateId { get; }
+
+		/// <summary>
+		/// Name of the label to monitor
+		/// </summary>
+		string LabelName { get; }
+	}
+
+	/// <summary>
+	/// Event that occurs when a job step completes
+	/// </summary>
+	public interface IStepCompleteEvent : IEvent
+	{
+		/// <summary>
+		/// The stream id
+		/// </summary>
+		StreamId StreamId { get; }
+
+		/// <summary>
+		/// The template id
+		/// </summary>
+		TemplateRefId TemplateId { get; }
+
+		/// <summary>
+		/// Name of the step to monitor
+		/// </summary>
+		string StepName { get; }
+	}
+}
diff --git a/Engine/Source/Programs/Horde/HordeServer/Models/IGraph.cs b/Engine/Source/Programs/Horde/HordeServer/Models/IGraph.cs
new file mode 100644
--- /dev/null
+++ b/Engine/Source/Programs/Horde/HordeServer/Models/IGraph.cs
@@ -0,0 +1,554 @@
+// Copyright Epic Games, Inc. All Rights Reserved.
+
+using EpicGames.Core;
+using HordeServer.Api;
+using HordeCommon;
+using HordeServer.Utilities;
+using MongoDB.Bson;
+using MongoDB.Bson.Serialization.Attributes;
+using Newtonsoft.Json;
+using System;
+using System.Collections.Generic;
+using System.Diagnostics;
+using System.Diagnostics.CodeAnalysis;
+using System.Linq;
+
+namespace HordeServer.Models
+{
+	/// <summary>
+	/// Represents a node in the graph
+	/// </summary>
+	public interface INode
+	{
+		/// <summary>
+		/// The name of this node 
+		/// </summary>
+		public string Name { get; }
+
+		/// <summary>
+		/// Indices of nodes which must have succeeded for this node to run
+		/// </summary>
+		public NodeRef[] InputDependencies { get; }
+
+		/// <summary>
+		/// Indices of nodes which must have completed for this node to run
+		/// </summary>
+		public NodeRef[] OrderDependencies { get; }
+
+		/// <summary>
+		/// The priority that this node should be run at, within this job
+		/// </summary>
+		public Priority Priority { get; }
+
+		/// <summary>
+		/// Whether this node can be run multiple times
+		/// </summary>
+		public bool AllowRetry { get; }
+
+		/// <summary>
+		/// This node can start running early, before dependencies of other nodes in the same group are complete
+		/// </summary>
+		public bool RunEarly { get; }
+
+		/// <summary>
+		/// Whether to include warnings in the output (defaults to true)
+		/// </summary>
+		public bool Warnings { get; }
+
+		/// <summary>
+		/// List of credentials required for this node. Each entry maps an environment variable name to a credential in the form "CredentialName.PropertyName".
+		/// </summary>
+		public IReadOnlyDictionary<string, string>? Credentials { get; }
+
+		/// <summary>
+		/// Properties for this node
+		/// </summary>
+		public IReadOnlyDictionary<string, string>? Properties { get; }
+	}
+
+	/// <summary>
+	/// Information about a sequence of nodes which can execute on a single agent
+	/// </summary>
+	public interface INodeGroup
+	{
+		/// <summary>
+		/// The type of agent to execute this group
+		/// </summary>
+		public string AgentType { get; }
+
+		/// <summary>
+		/// Nodes in this group
+		/// </summary>
+		public IReadOnlyList<INode> Nodes { get; }
+	}
+
+	/// <summary>
+	/// Reference to a node within another grup
+	/// </summary>
+	[DebuggerDisplay("Group: {GroupIdx}, Node: {NodeIdx}")]
+	public class NodeRef
+	{
+		/// <summary>
+		/// The group index of the referenced node
+		/// </summary>
+		public int GroupIdx { get; set; }
+
+		/// <summary>
+		/// The node index of the referenced node
+		/// </summary>
+		public int NodeIdx { get; set; }
+
+		/// <summary>
+		/// Private constructor for serialization
+		/// </summary>
+		[BsonConstructor]
+		private NodeRef()
+		{
+			GroupIdx = 0;
+			NodeIdx = 0;
+		}
+
+		/// <summary>
+		/// Constructor
+		/// </summary>
+		/// <param name="GroupIdx">Index of thr group containing the node</param>
+		/// <param name="NodeIdx">Index of the node within the group</param>
+		public NodeRef(int GroupIdx, int NodeIdx)
+		{
+			this.GroupIdx = GroupIdx;
+			this.NodeIdx = NodeIdx;
+		}
+
+		/// <inheritdoc/>
+		public override bool Equals(object? Other)
+		{
+			NodeRef? OtherNodeRef = Other as NodeRef;
+			return OtherNodeRef != null && OtherNodeRef.GroupIdx == GroupIdx && OtherNodeRef.NodeIdx == NodeIdx;
+		}
+
+		/// <inheritdoc/>
+		public override int GetHashCode()
+		{
+			return HashCode.Combine(GroupIdx, NodeIdx);
+		}
+
+		/// <summary>
+		/// Converts this reference to a node name
+		/// </summary>
+		/// <param name="Groups">List of groups that this reference points to</param>
+		/// <returns>Name of the referenced node</returns>
+		public INode ToNode(IReadOnlyList<INodeGroup> Groups)
+		{
+			return Groups[GroupIdx].Nodes[NodeIdx];
+		}
+	}
+
+	/// <summary>
+	/// An collection of node references
+	/// </summary>
+	public interface IAggregate
+	{
+		/// <summary>
+		/// Name of the aggregate
+		/// </summary>
+		public string Name { get; }
+
+		/// <summary>
+		/// List of nodes for the aggregate to be valid
+		/// </summary>
+		public IReadOnlyList<NodeRef> Nodes { get; }
+	}
+
+	/// <summary>
+	/// Label indicating the status of a set of nodes
+	/// </summary>
+	public interface ILabel
+	{
+		/// <summary>
+		/// Label to show in the dashboard. Null if does not need to be shown.
+		/// </summary>
+		public string? DashboardName { get; }
+
+		/// <summary>
+		/// Category for the label. May be null.
+		/// </summary>
+		public string? DashboardCategory { get; }
+
+		/// <summary>
+		/// Name to display for this label in UGS
+		/// </summary>
+		public string? UgsName { get; }
+
+		/// <summary>
+		/// Project which this label applies to, for UGS
+		/// </summary>
+		public string? UgsProject { get; }
+
+		/// <summary>
+		/// Which change to display the label on
+		/// </summary>
+		public LabelChange Change { get; }
+
+		/// <summary>
+		/// List of required nodes for the aggregate to be valid
+		/// </summary>
+		public List<NodeRef> RequiredNodes { get; }
+
+		/// <summary>
+		/// List of optional nodes to include in the aggregate state
+		/// </summary>
+		public List<NodeRef> IncludedNodes { get; }
+	}
+
+	/// <summary>
+	/// Extension methods for ILabel
+	/// </summary>
+	public static class LabelExtensions
+	{
+		/// <summary>
+		/// Enumerate all the required dependencies of this node group
+		/// </summary>
+		/// <param name="Label">The label instance</param>
+		/// <param name="Groups">List of groups for the job containing this aggregate</param>
+		/// <returns>Sequence of nodes</returns>
+		public static IEnumerable<INode> GetDependencies(this ILabel Label, IReadOnlyList<INodeGroup> Groups)
+		{
+			foreach (NodeRef RequiredNode in Label.RequiredNodes)
+			{
+				yield return RequiredNode.ToNode(Groups);
+			}
+			foreach (NodeRef IncludedNode in Label.IncludedNodes)
+			{
+				yield return IncludedNode.ToNode(Groups);
+			}
+		}
+	}
+
+	/// <summary>
+	/// A unique dependency graph instance
+	/// </summary>
+	public interface IGraph
+	{
+		/// <summary>
+		/// Hash of this graph
+		/// </summary>
+		public ContentHash Id { get; }
+
+		/// <summary>
+		/// Schema version for this document
+		/// </summary>
+		public int Schema { get; }
+
+		/// <summary>
+		/// List of groups for this graph
+		/// </summary>
+		public IReadOnlyList<INodeGroup> Groups { get; }
+
+		/// <summary>
+		/// List of aggregates for this graph
+		/// </summary>
+		public IReadOnlyList<IAggregate> Aggregates { get; }
+
+		/// <summary>
+		/// Status labels for this graph
+		/// </summary>
+		public IReadOnlyList<ILabel> Labels { get; }
+	}
+
+	/// <summary>
+	/// Extension methods for graphs
+	/// </summary>
+	public static class GraphExtensions
+	{
+		/// <summary>
+		/// Gets the node from a node reference
+		/// </summary>
+		/// <param name="Graph">The graph instance</param>
+		/// <param name="Ref">The node reference</param>
+		/// <returns>The node for the given reference</returns>
+		public static INode GetNode(this IGraph Graph, NodeRef Ref)
+		{
+			return Graph.Groups[Ref.GroupIdx].Nodes[Ref.NodeIdx];
+		}
+
+		/// <summary>
+		/// Tries to find a node by name
+		/// </summary>
+		/// <param name="Graph">The graph to search</param>
+		/// <param name="NodeName">Name of the node</param>
+		/// <param name="Ref">Receives the node reference</param>
+		/// <returns>True if the node was found, false otherwise</returns>
+		public static bool TryFindNode(this IGraph Graph, string NodeName, out NodeRef Ref)
+		{
+			for (int GroupIdx = 0; GroupIdx < Graph.Groups.Count; GroupIdx++)
+			{
+				INodeGroup Group = Graph.Groups[GroupIdx];
+				for (int NodeIdx = 0; NodeIdx < Group.Nodes.Count; NodeIdx++)
+				{
+					INode Node = Group.Nodes[NodeIdx];
+					if (String.Equals(Node.Name, NodeName, StringComparison.OrdinalIgnoreCase))
+					{
+						Ref = new NodeRef(GroupIdx, NodeIdx);
+						return true;
+					}
+				}
+			}
+
+			Ref = new NodeRef(0, 0);
+			return false;
+		}
+
+		/// <summary>
+		/// Tries to find a node by name
+		/// </summary>
+		/// <param name="Graph">The graph to search</param>
+		/// <param name="NodeName">Name of the node</param>
+		/// <param name="Node">Receives the node</param>
+		/// <returns>True if the node was found, false otherwise</returns>
+		public static bool TryFindNode(this IGraph Graph, string NodeName, out INode? Node)
+		{
+			NodeRef Ref;
+			if (TryFindNode(Graph, NodeName, out Ref))
+			{
+				Node = Graph.Groups[Ref.GroupIdx].Nodes[Ref.NodeIdx];
+				return true;
+			}
+			else
+			{
+				Node = null;
+				return false;
+			}
+		}
+
+		/// <summary>
+		/// Tries to find a node by name
+		/// </summary>
+		/// <param name="Graph">The graph to search</param>
+		/// <param name="Name">Name of the node</param>
+		/// <param name="AggregateIdx">Receives the aggregate index</param>
+		/// <returns>True if the node was found, false otherwise</returns>
+		public static bool TryFindAggregate(this IGraph Graph, string Name, out int AggregateIdx)
+		{
+			AggregateIdx = Graph.Aggregates.FindIndex(x => x.Name.Equals(Name, StringComparison.OrdinalIgnoreCase));
+			return AggregateIdx != -1;
+		}
+
+		/// <summary>
+		/// Tries to find a node by name
+		/// </summary>
+		/// <param name="Graph">The graph to search</param>
+		/// <param name="Name">Name of the node</param>
+		/// <param name="Aggregate">Receives the aggregate</param>
+		/// <returns>True if the node was found, false otherwise</returns>
+		public static bool TryFindAggregate(this IGraph Graph, string Name, [NotNullWhen(true)] out IAggregate? Aggregate)
+		{
+			int AggregateIdx;
+			if (TryFindAggregate(Graph, Name, out AggregateIdx))
+			{
+				Aggregate = Graph.Aggregates[AggregateIdx];
+				return true;
+			}
+			else
+			{
+				Aggregate = null;
+				return false;
+			}
+		}
+
+		/// <summary>
+		/// Gets a list of dependencies for the given node
+		/// </summary>
+		/// <param name="Graph">The graph instance</param>
+		/// <param name="Node">The node to return dependencies for</param>
+		/// <returns>List of dependencies</returns>
+		public static IEnumerable<INode> GetDependencies(this IGraph Graph, INode Node)
+		{
+			return Enumerable.Concat(Node.InputDependencies, Node.OrderDependencies).Select(x => Graph.GetNode(x));
+		}
+	}
+
+	/// <summary>
+	/// Information required to create a node
+	/// </summary>
+	public class NewNode
+	{
+		/// <summary>
+		/// The name of this node 
+		/// </summary>
+		public string Name { get; set; } = null!;
+
+		/// <summary>
+		/// List of nodes which must succeed for this node to run
+		/// </summary>
+		public List<string>? InputDependencies { get; set; }
+
+		/// <summary>
+		/// List of nodes which must have completed for this node to run
+		/// </summary>
+		public List<string>? OrderDependencies { get; set; }
+
+		/// <summary>
+		/// The priority of this node
+		/// </summary>
+		public Priority? Priority { get; set; }
+
+		/// <summary>
+		/// This node can be run multiple times
+		/// </summary>
+		public bool? AllowRetry { get; set; }
+
+		/// <summary>
+		/// This node can start running early, before dependencies of other nodes in the same group are complete
+		/// </summary>
+		public bool? RunEarly { get; set; }
+
+		/// <summary>
+		/// Whether to include warnings in the diagnostic output
+		/// </summary>
+		public bool? Warnings { get; set; }
+
+		/// <summary>
+		/// Credentials required for this node to run. This dictionary maps from environment variable names to a credential property in the format 'CredentialName.PropertyName'.
+		/// </summary>
+		public Dictionary<string, string>? Credentials { get; set; }
+
+		/// <summary>
+		/// Properties for this node
+		/// </summary>
+		public Dictionary<string, string>? Properties { get; set; }
+
+		/// <summary>
+		/// Constructor
+		/// </summary>
+		/// <param name="Name">Name of the node</param>
+		/// <param name="InputDependencies">List of nodes which must have completed succesfully for this node to run</param>
+		/// <param name="OrderDependencies">List of nodes which must have completed for this node to run</param>
+		/// <param name="Priority">Priority of this node</param>
+		/// <param name="AllowRetry">Whether the node can be run multiple times</param>
+		/// <param name="RunEarly">Whether the node can run early, before dependencies of other nodes in the same group complete</param>
+		/// <param name="Warnings">Whether to include warnings in the diagnostic output (defaults to true)</param>
+		/// <param name="Credentials">Credentials required for this node to run</param>
+		/// <param name="Properties">Properties for the node</param>
+		public NewNode(string Name, List<string>? InputDependencies = null, List<string>? OrderDependencies = null, Priority? Priority = null, bool? AllowRetry = null, bool? RunEarly = null, bool? Warnings = null, Dictionary<string, string>? Credentials = null, Dictionary<string, string>? Properties = null)
+		{
+			this.Name = Name;
+			this.InputDependencies = InputDependencies;
+			this.OrderDependencies = OrderDependencies;
+			this.Priority = Priority;
+			this.AllowRetry = AllowRetry;
+			this.RunEarly = RunEarly;
+			this.Warnings = Warnings;
+			this.Credentials = Credentials;
+			this.Properties = Properties;
+		}
+	}
+
+	/// <summary>
+	/// Information about a group of nodes
+	/// </summary>
+	public class NewGroup
+	{
+		/// <summary>
+		/// The type of agent to execute this group
+		/// </summary>
+		public string AgentType { get; set; }
+
+		/// <summary>
+		/// Nodes in the group
+		/// </summary>
+		public List<NewNode> Nodes { get; set; }
+
+		/// <summary>
+		/// Constructor
+		/// </summary>
+		/// <param name="AgentType">The type of agent to execute this group</param>
+		/// <param name="Nodes">Nodes in this group</param>
+		public NewGroup(string AgentType, List<NewNode> Nodes)
+		{
+			this.AgentType = AgentType;
+			this.Nodes = Nodes;
+		}
+	}
+
+	/// <summary>
+	/// Information about a group of nodes
+	/// </summary>
+	public class NewLabel
+	{
+		/// <summary>
+		/// Category for this label
+		/// </summary>
+		[Obsolete("Use DashboardCategory instead")]
+		public string? Category => DashboardCategory;
+
+		/// <summary>
+		/// Name of the aggregate
+		/// </summary>
+		[Obsolete("Use DashboardName instead")]
+		public string? Name => DashboardName;
+
+		/// <summary>
+		/// Name of the aggregate
+		/// </summary>
+		public string? DashboardName { get; set; }
+
+		/// <summary>
+		/// Category for this label
+		/// </summary>
+		public string? DashboardCategory { get; set; }
+
+		/// <summary>
+		/// Name of the badge in UGS
+		/// </summary>
+		public string? UgsName { get; set; }
+
+		/// <summary>
+		/// Project to show this label for in UGS
+		/// </summary>
+		public string? UgsProject { get; set; }
+
+		/// <summary>
+		/// Which change the label applies to
+		/// </summary>
+		public LabelChange Change { get; set; }
+
+		/// <summary>
+		/// Nodes which must be part of the job for the aggregate to be valid
+		/// </summary>
+		public List<string> RequiredNodes { get; set; } = new List<string>();
+
+		/// <summary>
+		/// Nodes which must be part of the job for the aggregate to be valid
+		/// </summary>
+		public List<string> IncludedNodes { get; set; } = new List<string>();
+	}
+
+	/// <summary>
+	/// Information about a group of nodes
+	/// </summary>
+	public class NewAggregate
+	{
+		/// <summary>
+		/// Name of the aggregate
+		/// </summary>
+		public string Name { get; set; }
+
+		/// <summary>
+		/// Nodes which must be part of the job for the aggregate to be valid
+		/// </summary>
+		public List<string> Nodes { get; set; }
+
+		/// <summary>
+		/// Constructor
+		/// </summary>
+		/// <param name="Name">Name of this aggregate</param>
+		/// <param name="Nodes">Nodes which must be part of the job for the aggregate to be shown</param>
+		public NewAggregate(string Name, List<string> Nodes)
+		{
+			this.Name = Name;
+			this.Nodes = Nodes;
+		}
+	}
+}
+
+
diff --git a/Engine/Source/Programs/Horde/HordeServer/Models/IJob.cs b/Engine/Source/Programs/Horde/HordeServer/Models/IJob.cs
new file mode 100644
--- /dev/null
+++ b/Engine/Source/Programs/Horde/HordeServer/Models/IJob.cs
@@ -0,0 +1,1491 @@
+// Copyright Epic Games, Inc. All Rights Reserved.
+
+using EpicGames.Core;
+using HordeServer.Api;
+using HordeCommon;
+using HordeServer.Utilities;
+using MongoDB.Bson;
+using MongoDB.Bson.Serialization.Attributes;
+using MongoDB.Driver;
+using System;
+using System.Collections.Generic;
+using System.ComponentModel;
+using System.Diagnostics.CodeAnalysis;
+using System.Linq;
+using System.Runtime.CompilerServices;
+using System.Text;
+using System.Text.Json.Serialization;
+using System.Threading.Tasks;
+
+namespace HordeServer.Models
+{
+	using ReportPlacement = HordeCommon.Rpc.ReportPlacement;
+	using StreamId = StringId<HordeServer.Models.IStream>;
+	using TemplateRefId = StringId<HordeServer.Models.TemplateRef>;
+	using PoolId = StringId<HordeServer.Models.IPool>;
+
+	/// <summary>
+	/// Report for a job or jobstep
+	/// </summary>
+	public interface IReport
+	{
+		/// <summary>
+		/// Name of the report
+		/// </summary>
+		string Name { get; }
+
+		/// <summary>
+		/// Where to render the report
+		/// </summary>
+		ReportPlacement Placement { get; }
+
+		/// <summary>
+		/// The artifact id
+		/// </summary>
+		ObjectId ArtifactId { get; }
+	}
+
+	/// <summary>
+	/// Implementation of IReport
+	/// </summary>
+	public class Report : IReport
+	{
+		/// <inheritdoc/>
+		public string Name { get; set; } = String.Empty;
+
+		/// <inheritdoc/>
+		public ReportPlacement Placement { get; set; }
+
+		/// <inheritdoc/>
+		public ObjectId ArtifactId { get; set; }
+	}
+
+	/// <summary>
+	/// Embedded jobstep document
+	/// </summary>
+	public interface IJobStep
+	{
+		/// <summary>
+		/// Unique ID assigned to this jobstep. A new id is generated whenever a jobstep's order is changed.
+		/// </summary>
+		public SubResourceId Id { get; }
+
+		/// <summary>
+		/// Index of the node which this jobstep is to execute
+		/// </summary>
+		public int NodeIdx { get; }
+
+		/// <summary>
+		/// Current state of the job step. This is updated automatically when runs complete.
+		/// </summary>
+		public JobStepState State { get; }
+
+		/// <summary>
+		/// Current outcome of the jobstep
+		/// </summary>
+		public JobStepOutcome Outcome { get; }
+
+		/// <summary>
+		/// The log id for this step
+		/// </summary>
+		public ObjectId? LogId { get; }
+
+		/// <summary>
+		/// Unique id for notifications
+		/// </summary>
+		public ObjectId? NotificationTriggerId { get; }
+
+		/// <summary>
+		/// Time at which the batch transitioned to the ready state (UTC).
+		/// </summary>
+		public DateTime? ReadyTimeUtc { get; }
+
+		/// <summary>
+		/// Time at which the batch transitioned to the executing state (UTC).
+		/// </summary>
+		public DateTime? StartTimeUtc { get; }
+
+		/// <summary>
+		/// Time at which the run finished (UTC)
+		/// </summary>
+		public DateTime? FinishTimeUtc { get; }
+
+		/// <summary>
+		/// Override for the priority of this step
+		/// </summary>
+		public Priority? Priority { get; }
+
+		/// <summary>
+		/// If a retry is requested, stores the name of the user that requested it
+		/// </summary>
+		public string? RetryByUser { get; }
+		
+		/// <summary>
+		/// Signal if a step should be aborted
+		/// </summary>
+		public bool AbortRequested { get; }
+		
+		/// <summary>
+		/// If an abort is requested, stores the name of the user that requested it
+		/// </summary>
+		public string? AbortByUser { get; }
+
+		/// <summary>
+		/// List of reports for this step
+		/// </summary>
+		public IReadOnlyList<IReport>? Reports { get; }
+
+		/// <summary>
+		/// Reports for this jobstep.
+		/// </summary>
+		public Dictionary<string, string>? Properties { get; }
+	}
+
+	/// <summary>
+	/// Extension methods for job steps
+	/// </summary>
+	public static class JobStepExtensions
+	{
+		/// <summary>
+		/// Determines if a jobstep has failed or is skipped. Can be used to determine whether dependent steps will be able to run.
+		/// </summary>
+		/// <returns>True if the step is failed or skipped</returns>
+		public static bool IsFailedOrSkipped(this IJobStep Step)
+		{
+			return Step.State == JobStepState.Skipped || Step.Outcome == JobStepOutcome.Failure;
+		}
+
+		/// <summary>
+		/// Determines if a jobstep is done by checking to see if it is completed, skipped, or aborted.
+		/// </summary>
+		/// <returns>True if the step is completed, skipped, or aborted</returns>
+		public static bool IsPending(this IJobStep Step)
+		{
+			return Step.State != JobStepState.Aborted && Step.State != JobStepState.Completed && Step.State != JobStepState.Skipped;
+		}
+	}
+
+	/// <summary>
+	/// Stores information about a batch of job steps
+	/// </summary>
+	public interface IJobStepBatch
+	{
+		/// <summary>
+		/// Unique id for this group
+		/// </summary>
+		public SubResourceId Id { get; }
+
+		/// <summary>
+		/// The log file id for this batch
+		/// </summary>
+		public ObjectId? LogId { get; }
+
+		/// <summary>
+		/// Index of the group being executed
+		/// </summary>
+		public int GroupIdx { get; }
+
+		/// <summary>
+		/// The state of this group
+		/// </summary>
+		public JobStepBatchState State { get; }
+
+		/// <summary>
+		/// Error associated with this group
+		/// </summary>
+		public JobStepBatchError Error { get; }
+
+		/// <summary>
+		/// Steps within this run
+		/// </summary>
+		public IReadOnlyList<IJobStep> Steps { get; }
+
+		/// <summary>
+		/// The pool that this agent was taken from
+		/// </summary>
+		public PoolId? PoolId { get; }
+
+		/// <summary>
+		/// The agent assigned to execute this group
+		/// </summary>
+		public AgentId? AgentId { get; }
+
+		/// <summary>
+		/// The agent session that is executing this group
+		/// </summary>
+		public ObjectId? SessionId { get; }
+
+		/// <summary>
+		/// The lease that's executing this group
+		/// </summary>
+		public ObjectId? LeaseId { get; }
+
+		/// <summary>
+		/// The weighted priority of this batch for the scheduler
+		/// </summary>
+		public int SchedulePriority { get; }
+
+		/// <summary>
+		/// Time at which the group became ready (UTC).
+		/// </summary>
+		public DateTime? ReadyTimeUtc { get; }
+
+		/// <summary>
+		/// Time at which the group started (UTC).
+		/// </summary>
+		public DateTime? StartTimeUtc { get; }
+
+		/// <summary>
+		/// Time at which the group finished (UTC)
+		/// </summary>
+		public DateTime? FinishTimeUtc { get; }
+	}
+
+	/// <summary>
+	/// Extension methods for IJobStepBatch
+	/// </summary>
+	public static class JobStepBatchExtensions
+	{
+		/// <summary>
+		/// Attempts to get a step with the given id
+		/// </summary>
+		/// <param name="Batch">The batch to search</param>
+		/// <param name="StepId">The step id</param>
+		/// <param name="Step">On success, receives the step object</param>
+		/// <returns>True if the step was found</returns>
+		public static bool TryGetStep(this IJobStepBatch Batch, SubResourceId StepId, [NotNullWhen(true)] out IJobStep? Step)
+		{
+			Step = Batch.Steps.FirstOrDefault(x => x.Id == StepId);
+			return Step != null;
+		}
+
+		/// <summary>
+		/// Determines if new steps can be appended to this batch. We do not allow this after the last step has been completed, because the agent is shutting down.
+		/// </summary>
+		/// <param name="Batch">The batch to search</param>
+		/// <returns>True if new steps can be appended to this batch</returns>
+		public static bool CanBeAppendedTo(this IJobStepBatch Batch)
+		{
+			return Batch.State <= JobStepBatchState.Running;
+		}
+
+		/// <summary>
+		/// Gets the wait time for this batch
+		/// </summary>
+		/// <param name="Batch">The batch to search</param>
+		/// <returns>Wait time for the batch</returns>
+		public static TimeSpan? GetWaitTime(this IJobStepBatch Batch)
+		{
+			if (Batch.StartTimeUtc == null || Batch.ReadyTimeUtc == null)
+			{
+				return null;
+			}
+			else
+			{
+				return Batch.StartTimeUtc.Value - Batch.ReadyTimeUtc.Value;
+			}
+		}
+
+		/// <summary>
+		/// Gets the initialization time for this batch
+		/// </summary>
+		/// <param name="Batch">The batch to search</param>
+		/// <returns>Initialization time for this batch</returns>
+		public static TimeSpan? GetInitTime(this IJobStepBatch Batch)
+		{
+			if (Batch.StartTimeUtc != null)
+			{
+				foreach (IJobStep Step in Batch.Steps)
+				{
+					if (Step.StartTimeUtc != null)
+					{
+						return Step.StartTimeUtc - Batch.StartTimeUtc.Value;
+					}
+				}
+			}
+			return null;
+		}
+
+		/// <summary>
+		/// Get the dependencies required for this batch to start, taking run-early nodes into account
+		/// </summary>
+		/// <param name="Batch">The batch to search</param>
+		/// <param name="Groups">List of node groups</param>
+		/// <returns>Set of nodes that must have completed for this batch to start</returns>
+		public static HashSet<INode> GetStartDependencies(this IJobStepBatch Batch, IReadOnlyList<INodeGroup> Groups)
+		{
+			// Find all the nodes that this group will start with.
+			List<INode> Nodes = Batch.Steps.ConvertAll(x => Groups[Batch.GroupIdx].Nodes[x.NodeIdx]);
+			if (Nodes.Any(x => x.RunEarly))
+			{
+				Nodes.RemoveAll(x => !x.RunEarly);
+			}
+
+			// Find all their dependencies
+			HashSet<INode> Dependencies = new HashSet<INode>();
+			foreach (INode Node in Nodes)
+			{
+				Dependencies.UnionWith(Node.InputDependencies.Select(x => Groups[x.GroupIdx].Nodes[x.NodeIdx]));
+				Dependencies.UnionWith(Node.OrderDependencies.Select(x => Groups[x.GroupIdx].Nodes[x.NodeIdx]));
+			}
+
+			// Exclude all the dependencies within the same group
+			Dependencies.ExceptWith(Groups[Batch.GroupIdx].Nodes);
+			return Dependencies;
+		}
+	}
+
+	/// <summary>
+	/// Cumulative timing information to reach a certain point in a job
+	/// </summary>
+	public class TimingInfo
+	{
+		/// <summary>
+		/// Wait time on the critical path
+		/// </summary>
+		public TimeSpan? TotalWaitTime { get; set; }
+
+		/// <summary>
+		/// Sync time on the critical path
+		/// </summary>
+		public TimeSpan? TotalInitTime { get; set; }
+
+		/// <summary>
+		/// Duration to this point
+		/// </summary>
+		public TimeSpan? TotalTimeToComplete { get; set; }
+
+		/// <summary>
+		/// Average wait time to this point
+		/// </summary>
+		public TimeSpan? AverageTotalWaitTime { get; set; }
+
+		/// <summary>
+		/// Average sync time to this point
+		/// </summary>
+		public TimeSpan? AverageTotalInitTime { get; set; }
+
+		/// <summary>
+		/// Average duration to this point
+		/// </summary>
+		public TimeSpan? AverageTotalTimeToComplete { get; set; }
+
+		/// <summary>
+		/// Individual step timing information
+		/// </summary>
+		public IJobStepTiming? StepTiming { get; set; }
+
+		/// <summary>
+		/// Constructor
+		/// </summary>
+		public TimingInfo()
+		{
+			TotalWaitTime = TimeSpan.Zero;
+			TotalInitTime = TimeSpan.Zero;
+			TotalTimeToComplete = TimeSpan.Zero;
+
+			AverageTotalWaitTime = TimeSpan.Zero;
+			AverageTotalInitTime = TimeSpan.Zero;
+			AverageTotalTimeToComplete = TimeSpan.Zero;
+		}
+
+		/// <summary>
+		/// Copy constructor
+		/// </summary>
+		/// <param name="Other">The timing info object to copy from</param>
+		public TimingInfo(TimingInfo Other)
+		{
+			TotalWaitTime = Other.TotalWaitTime;
+			TotalInitTime = Other.TotalInitTime;
+			TotalTimeToComplete = Other.TotalTimeToComplete;
+
+			AverageTotalWaitTime = Other.AverageTotalWaitTime;
+			AverageTotalInitTime = Other.AverageTotalInitTime;
+			AverageTotalTimeToComplete = Other.AverageTotalTimeToComplete;
+		}
+
+		/// <summary>
+		/// Modifies this timing to wait for another timing
+		/// </summary>
+		/// <param name="Other">The other node to wait for</param>
+		public void WaitFor(TimingInfo Other)
+		{
+			if (TotalTimeToComplete != null)
+			{
+				if (Other.TotalTimeToComplete == null || Other.TotalTimeToComplete.Value > TotalTimeToComplete.Value)
+				{
+					TotalInitTime = Other.TotalInitTime;
+					TotalWaitTime = Other.TotalWaitTime;
+					TotalTimeToComplete = Other.TotalTimeToComplete;
+				}
+			}
+
+			if (AverageTotalTimeToComplete != null)
+			{
+				if (Other.AverageTotalTimeToComplete == null || Other.AverageTotalTimeToComplete.Value > AverageTotalTimeToComplete.Value)
+				{
+					AverageTotalInitTime = Other.AverageTotalInitTime;
+					AverageTotalWaitTime = Other.AverageTotalWaitTime;
+					AverageTotalTimeToComplete = Other.AverageTotalTimeToComplete;
+				}
+			}
+		}
+
+		/// <summary>
+		/// Waits for all the given timing info objects to complete
+		/// </summary>
+		/// <param name="Others">Other timing info objects to wait for</param>
+		public void WaitForAll(IEnumerable<TimingInfo> Others)
+		{
+			foreach(TimingInfo Other in Others)
+			{
+				WaitFor(Other);
+			}
+		}
+
+		/// <summary>
+		/// Constructs a new TimingInfo object which represents the last TimingInfo to finish
+		/// </summary>
+		/// <param name="Others">TimingInfo objects to wait for</param>
+		/// <returns>New TimingInfo instance</returns>
+		public static TimingInfo Max(IEnumerable<TimingInfo> Others)
+		{
+			TimingInfo TimingInfo = new TimingInfo();
+			TimingInfo.WaitForAll(Others);
+			return TimingInfo;
+		}
+	}
+
+	/// <summary>
+	/// Information about a chained job trigger
+	/// </summary>
+	public interface IChainedJob
+	{
+		/// <summary>
+		/// The target to monitor
+		/// </summary>
+		public string Target { get; }
+
+		/// <summary>
+		/// The template to trigger on success
+		/// </summary>
+		public TemplateRefId TemplateRefId { get; }
+
+		/// <summary>
+		/// The triggered job id
+		/// </summary>
+		public ObjectId? JobId { get; }
+	}
+
+	/// <summary>
+	/// Document describing a job
+	/// </summary>
+	[SuppressMessage("Compiler", "CA1056:URI parameters should not be strings")]
+	public interface IJob
+	{
+		/// <summary>
+		/// Job argument indicating a target that should be built
+		/// </summary>
+		public const string TargetArgumentPrefix = "-Target=";
+
+		/// <summary>
+		/// Name of the node which parses the buildgraph script
+		/// </summary>
+		public const string SetupNodeName = "Setup Build";
+
+		/// <summary>
+		/// Identifier for the job. Randomly generated.
+		/// </summary>
+		public ObjectId Id { get; }
+
+		/// <summary>
+		/// The stream that this job belongs to
+		/// </summary>
+		public StreamId StreamId { get; }
+
+		/// <summary>
+		/// The template ref id
+		/// </summary>
+		public TemplateRefId TemplateId { get; }
+
+		/// <summary>
+		/// The template that this job was created from
+		/// </summary>
+		public ContentHash? TemplateHash { get; }
+
+		/// <summary>
+		/// Hash of the graph definition
+		/// </summary>
+		public ContentHash GraphHash { get; }
+
+		/// <summary>
+		/// Id of the user that started this job
+		/// </summary>
+		public ObjectId? StartedByUserId { get; }
+
+		/// <summary>
+		/// Name of the user that started this job
+		/// </summary>
+		public string? StartedByUser { get; }
+
+		/// <summary>
+		/// Name of the user that aborted this job. Set to null if the job is not aborted.
+		/// </summary>
+		public string? AbortedByUser { get; }
+
+		/// <summary>
+		/// Name of the job.
+		/// </summary>
+		public string Name { get; }
+
+		/// <summary>
+		/// The changelist number to build
+		/// </summary>
+		public int Change { get; }
+
+		/// <summary>
+		/// The code changelist number for this build
+		/// </summary>
+		public int CodeChange { get; }
+
+		/// <summary>
+		/// The preflight changelist number
+		/// </summary>
+		public int PreflightChange { get; }
+
+		/// <summary>
+		/// The cloned preflight changelist number (if the prefight change is duplicated via p4 reshelve)
+		/// </summary>
+		public int ClonedPreflightChange { get; }
+
+		/// <summary>
+		/// Priority of this job
+		/// </summary>
+		public Priority Priority { get; }
+
+		/// <summary>
+		/// For preflights, submit the change if the job is successful
+		/// </summary>
+		public bool AutoSubmit { get; }
+
+		/// <summary>
+		/// The submitted changelist number
+		/// </summary>
+		public int? AutoSubmitChange { get; }
+
+		/// <summary>
+		/// Message produced by trying to auto-submit the change
+		/// </summary>
+		public string? AutoSubmitMessage { get; }
+
+		/// <summary>
+		/// Whether to update issues based on the outcome of this job
+		/// </summary>
+		public bool UpdateIssues { get; }
+
+		/// <summary>
+		/// Time that the job was created (in UTC)
+		/// </summary>
+		public DateTime CreateTimeUtc { get; }
+
+		/// <summary>
+		/// Largest value of the CombinedPriority value for batches in the ready state.
+		/// </summary>
+		public int SchedulePriority { get; }
+
+		/// <summary>
+		/// Array of jobstep runs
+		/// </summary>
+		public IReadOnlyList<IJobStepBatch> Batches { get; }
+
+		/// <summary>
+		/// Optional user-defined properties for this job
+		/// </summary>
+		public IReadOnlyList<string> Arguments { get; }
+
+		/// <summary>
+		/// Issues associated with this job
+		/// </summary>
+		public IReadOnlyList<int> Issues { get; }
+
+		/// <summary>
+		/// Unique id for notifications
+		/// </summary>
+		public ObjectId? NotificationTriggerId { get; }
+
+		/// <summary>
+		/// Whether to show badges in UGS for this job
+		/// </summary>
+		public bool ShowUgsBadges { get; }
+
+		/// <summary>
+		/// Whether to show alerts in UGS for this job
+		/// </summary>
+		public bool ShowUgsAlerts { get; }
+
+		/// <summary>
+		/// Notification channel for this job.
+		/// </summary>
+		public string? NotificationChannel { get; }
+		
+		/// <summary>
+		/// Notification channel filter for this job.
+		/// </summary>
+		public string? NotificationChannelFilter { get; }
+		
+		/// <summary>
+		/// URL to a Helix Swarm server to call once job finishes. Used for adding build outcome to Swarm reviews.
+		/// </summary>
+		public string? HelixSwarmCallbackUrl { get; }
+
+		/// <summary>
+		/// Mapping of label ids to notification trigger ids for notifications
+		/// </summary>
+		public IReadOnlyDictionary<int, ObjectId> LabelIdxToTriggerId { get; }
+
+		/// <summary>
+		/// List of reports for this step
+		/// </summary>
+		public IReadOnlyList<IReport>? Reports { get; }
+
+		/// <summary>
+		/// List of downstream job triggers
+		/// </summary>
+		public IReadOnlyList<IChainedJob> ChainedJobs { get; }
+
+		/// <summary>
+		/// Next id for batches or groups
+		/// </summary>
+		public SubResourceId NextSubResourceId { get; }
+
+		/// <summary>
+		/// The last update time
+		/// </summary>
+		public DateTime UpdateTimeUtc { get; }
+
+		/// <summary>
+		/// The ACL for this job.
+		/// </summary>
+		public Acl? Acl { get; }
+
+		/// <summary>
+		/// Update counter for this document. Any updates should compare-and-swap based on the value of this counter, or increment it in the case of server-side updates.
+		/// </summary>
+		public int UpdateIndex { get; }
+	}
+
+	/// <summary>
+	/// Extension methods for jobs
+	/// </summary>
+	public static class JobExtensions
+	{
+		/// <summary>
+		/// Attempts to get a step with a given ID
+		/// </summary>
+		/// <param name="Job">The job document</param>
+		/// <param name="StepId">The step id</param>
+		/// <param name="Step">On success, receives the step object</param>
+		/// <returns>True if the step was found</returns>
+		public static bool TryGetStep(this IJob Job, SubResourceId StepId, [NotNullWhen(true)] out IJobStep? Step)
+		{
+			foreach (IJobStepBatch Batch in Job.Batches)
+			{
+				if (Batch.TryGetStep(StepId, out Step))
+				{
+					return true;
+				}
+			}
+
+			Step = null;
+			return false;
+		}
+
+		/// <summary>
+		/// Gets the current job state
+		/// </summary>
+		/// <param name="Job">The job document</param>
+		/// <returns>Job state</returns>
+		public static JobState GetState(this IJob Job)
+		{
+			bool bWaiting = false;
+			foreach (IJobStepBatch Batch in Job.Batches)
+			{
+				foreach (IJobStep Step in Batch.Steps)
+				{
+					if (Step.State == JobStepState.Running)
+					{
+						return JobState.Running;
+					}
+					else if (Step.State == JobStepState.Ready || Step.State == JobStepState.Waiting)
+					{
+						if (Batch.State == JobStepBatchState.Starting || Batch.State == JobStepBatchState.Running)
+						{
+							return JobState.Running;
+						}
+						else
+						{
+							bWaiting = true;
+						}
+					}
+				}
+			}
+			return bWaiting ? JobState.Waiting : JobState.Complete;
+		}
+
+		/// <summary>
+		/// Gets the outcome for a particular named target. May be an aggregate or node name.
+		/// </summary>
+		/// <param name="Job">The job to check</param>
+		/// <returns>The step outcome</returns>
+		public static (JobStepState, JobStepOutcome) GetTargetState(this IJob Job)
+		{
+			IReadOnlyDictionary<NodeRef, IJobStep> NodeToStep = GetStepForNodeMap(Job);
+			return GetTargetState(NodeToStep.Values);
+		}
+
+		/// <summary>
+		/// Gets the outcome for a particular named target. May be an aggregate or node name.
+		/// </summary>
+		/// <param name="Job">The job to check</param>
+		/// <param name="Graph">Graph for the job</param>
+		/// <param name="Target">Target to find an outcome for</param>
+		/// <returns>The step outcome</returns>
+		public static (JobStepState, JobStepOutcome)? GetTargetState(this IJob Job, IGraph Graph, string? Target)
+		{
+			if (Target == null)
+			{
+				return GetTargetState(Job);
+			}
+
+			NodeRef NodeRef;
+			if (Graph.TryFindNode(Target, out NodeRef))
+			{
+				IJobStep? Step;
+				if (Job.TryGetStepForNode(NodeRef, out Step))
+				{
+					return (Step.State, Step.Outcome);
+				}
+				else
+				{
+					return null;
+				}
+			}
+
+			IAggregate? Aggregate;
+			if (Graph.TryFindAggregate(Target, out Aggregate))
+			{
+				IReadOnlyDictionary<NodeRef, IJobStep> StepForNode = GetStepForNodeMap(Job);
+
+				List<IJobStep> Steps = new List<IJobStep>();
+				foreach (NodeRef AggregateNodeRef in Aggregate.Nodes)
+				{
+					IJobStep? Step;
+					if (!StepForNode.TryGetValue(AggregateNodeRef, out Step))
+					{
+						return null;
+					}
+					Steps.Add(Step);
+				}
+
+				return GetTargetState(Steps);
+			}
+
+			return null;
+		}
+
+		/// <summary>
+		/// Gets the outcome for a particular named target. May be an aggregate or node name.
+		/// </summary>
+		/// <param name="Steps">Steps to include</param>
+		/// <returns>The step outcome</returns>
+		public static (JobStepState, JobStepOutcome) GetTargetState(IEnumerable<IJobStep> Steps)
+		{
+			bool bAnySkipped = false;
+			bool bAnyWarnings = false;
+			bool bAnyFailed = false;
+			bool bAnyPending = false;
+			foreach (IJobStep Step in Steps)
+			{
+				bAnyPending |= Step.IsPending();
+				bAnySkipped |= Step.State == JobStepState.Aborted || Step.State == JobStepState.Skipped;
+				bAnyFailed |= (Step.Outcome == JobStepOutcome.Failure);
+				bAnyWarnings |= (Step.Outcome == JobStepOutcome.Warnings);
+			}
+
+			JobStepState NewState = bAnyPending ? JobStepState.Running : JobStepState.Completed;
+			JobStepOutcome NewOutcome = bAnyFailed ? JobStepOutcome.Failure : bAnyWarnings ? JobStepOutcome.Warnings : bAnySkipped ? JobStepOutcome.Unspecified : JobStepOutcome.Success;
+			return (NewState, NewOutcome);
+		}
+
+		/// <summary>
+		/// Gets the outcome for a particular named target. May be an aggregate or node name.
+		/// </summary>
+		/// <param name="Job">The job to check</param>
+		/// <param name="Graph">Graph for the job</param>
+		/// <param name="Target">Target to find an outcome for</param>
+		/// <returns>The step outcome</returns>
+		public static JobStepOutcome GetTargetOutcome(this IJob Job, IGraph Graph, string Target)
+		{
+			NodeRef NodeRef;
+			if (Graph.TryFindNode(Target, out NodeRef))
+			{
+				IJobStep? Step;
+				if (Job.TryGetStepForNode(NodeRef, out Step))
+				{
+					return Step.Outcome;
+				}
+				else
+				{
+					return JobStepOutcome.Unspecified;
+				}
+			}
+
+			IAggregate? Aggregate;
+			if (Graph.TryFindAggregate(Target, out Aggregate))
+			{
+				IReadOnlyDictionary<NodeRef, IJobStep> StepForNode = GetStepForNodeMap(Job);
+
+				bool Warnings = false;
+				foreach (NodeRef AggregateNodeRef in Aggregate.Nodes)
+				{
+					IJobStep? Step;
+					if (!StepForNode.TryGetValue(AggregateNodeRef, out Step))
+					{
+						return JobStepOutcome.Unspecified;
+					}
+					if (Step.Outcome == JobStepOutcome.Failure)
+					{
+						return JobStepOutcome.Failure;
+					}
+					Warnings |= (Step.Outcome == JobStepOutcome.Warnings);
+				}
+				return Warnings ? JobStepOutcome.Warnings : JobStepOutcome.Success;
+			}
+
+			return JobStepOutcome.Unspecified;
+		}
+
+		/// <summary>
+		/// Gets the job step for a particular node
+		/// </summary>
+		/// <param name="Job">The job to search</param>
+		/// <param name="NodeRef">The node ref</param>
+		/// <param name="JobStep">Receives the jobstep on success</param>
+		/// <returns>True if the jobstep was founds</returns>
+		public static bool TryGetStepForNode(this IJob Job, NodeRef NodeRef, [NotNullWhen(true)] out IJobStep? JobStep)
+		{
+			JobStep = null;
+			foreach (IJobStepBatch Batch in Job.Batches)
+			{
+				if (Batch.GroupIdx == NodeRef.GroupIdx)
+				{
+					foreach (IJobStep BatchStep in Batch.Steps)
+					{
+						if (BatchStep.NodeIdx == NodeRef.NodeIdx)
+						{
+							JobStep = BatchStep;
+						}
+					}
+				}
+			}
+			return JobStep != null;
+		}
+
+		/// <summary>
+		/// Gets a dictionary that maps <see cref="NodeRef"/> objects to their associated
+		/// <see cref="IJobStep"/> objects on a <see cref="IJob"/>.
+		/// </summary>
+		/// <param name="Job">The job document</param>
+		/// <returns>Map of <see cref="NodeRef"/> to <see cref="IJobStep"/></returns>
+		public static IReadOnlyDictionary<NodeRef, IJobStep> GetStepForNodeMap(this IJob Job)
+		{
+			Dictionary<NodeRef, IJobStep> StepForNode = new Dictionary<NodeRef, IJobStep>();
+			foreach (IJobStepBatch Batch in Job.Batches)
+			{
+				foreach (IJobStep BatchStep in Batch.Steps)
+				{
+					NodeRef BatchNodeRef = new NodeRef(Batch.GroupIdx, BatchStep.NodeIdx);
+					StepForNode[BatchNodeRef] = BatchStep;
+				}
+			}
+			return StepForNode;
+		}
+
+		/// <summary>
+		/// Gets the estimated timing info for all nodes in the job
+		/// </summary>
+		/// <param name="Job">The job document</param>
+		/// <param name="Graph">Graph for this job</param>
+		/// <param name="JobTiming">Job timing information</param>
+		/// <returns>Map of node to expected timing info</returns>
+		public static Dictionary<INode, TimingInfo> GetTimingInfo(this IJob Job, IGraph Graph, IJobTiming JobTiming)
+		{
+			TimeSpan CurrentTime = DateTime.UtcNow - Job.CreateTimeUtc;
+
+			Dictionary<INode, TimingInfo> NodeToTimingInfo = Graph.Groups.SelectMany(x => x.Nodes).ToDictionary(x => x, x => new TimingInfo());
+			foreach (IJobStepBatch Batch in Job.Batches)
+			{
+				INodeGroup Group = Graph.Groups[Batch.GroupIdx];
+
+				// Step through the batch, keeping track of the time that things finish.
+				TimingInfo TimingInfo = new TimingInfo();
+
+				// Wait for the dependencies for the batch to start
+				HashSet<INode> DependencyNodes = Batch.GetStartDependencies(Graph.Groups);
+				TimingInfo.WaitForAll(DependencyNodes.Select(x => NodeToTimingInfo[x]));
+
+				// If the batch has actually started, correct the expected time to use this instead
+				if (Batch.StartTimeUtc != null)
+				{
+					TimingInfo.TotalTimeToComplete = Batch.StartTimeUtc - Job.CreateTimeUtc;
+				}
+
+				// Get the average times for this batch
+				TimeSpan? AverageWaitTime = GetAverageWaitTime(Graph, Batch, JobTiming);
+				TimeSpan? AverageInitTime = GetAverageInitTime(Graph, Batch, JobTiming);
+
+				// Update the wait times and initialization times along this path
+				TimingInfo.TotalWaitTime = TimingInfo.TotalWaitTime + (Batch.GetWaitTime() ?? AverageWaitTime);
+				TimingInfo.TotalInitTime = TimingInfo.TotalInitTime + (Batch.GetInitTime() ?? AverageInitTime);
+
+				// Update the average wait and initialization times too
+				TimingInfo.AverageTotalWaitTime = TimingInfo.AverageTotalWaitTime + AverageWaitTime;
+				TimingInfo.AverageTotalInitTime = TimingInfo.AverageTotalInitTime + AverageInitTime;
+
+				// Step through the batch, updating the expected times as we go
+				foreach (IJobStep Step in Batch.Steps)
+				{
+					INode Node = Group.Nodes[Step.NodeIdx];
+
+					// Get the timing for this step
+					IJobStepTiming? StepTimingInfo;
+					JobTiming.TryGetStepTiming(Node.Name, out StepTimingInfo);
+
+					// If the step has already started, update the actual time to reach this point
+					if(Step.StartTimeUtc != null)
+					{
+						TimingInfo.TotalTimeToComplete = Step.StartTimeUtc.Value - Job.CreateTimeUtc;
+					}
+
+					// If the step hasn't started yet, make sure the start time is later than the current time
+					if (Step.StartTimeUtc == null && CurrentTime > TimingInfo.TotalTimeToComplete)
+					{
+						TimingInfo.TotalTimeToComplete = CurrentTime;
+					}
+
+					// Wait for all the node dependencies to complete
+					TimingInfo.WaitForAll(Graph.GetDependencies(Node).Select(x => NodeToTimingInfo[x]));
+
+					// If the step has actually finished, correct the time to use that instead
+					if (Step.FinishTimeUtc != null)
+					{
+						TimingInfo.TotalTimeToComplete = Step.FinishTimeUtc.Value - Job.CreateTimeUtc;
+					}
+					else
+					{
+						TimingInfo.TotalTimeToComplete = TimingInfo.TotalTimeToComplete + NullableTimeSpanFromSeconds(StepTimingInfo?.AverageDuration);
+					}
+
+					// If the step hasn't finished yet, make sure the start time is later than the current time
+					if (Step.FinishTimeUtc == null && CurrentTime > TimingInfo.TotalTimeToComplete)
+					{
+						TimingInfo.TotalTimeToComplete = CurrentTime;
+					}
+
+					// Update the average time to complete
+					TimingInfo.AverageTotalTimeToComplete = TimingInfo.AverageTotalTimeToComplete + NullableTimeSpanFromSeconds(StepTimingInfo?.AverageDuration);
+
+					// Add it to the lookup
+					TimingInfo NodeTimingInfo = new TimingInfo(TimingInfo);
+					NodeTimingInfo.StepTiming = StepTimingInfo;
+					NodeToTimingInfo[Node] = NodeTimingInfo;
+				}
+			}
+			return NodeToTimingInfo;
+		}
+
+		/// <summary>
+		/// Gets the average wait time for this batch
+		/// </summary>
+		/// <param name="Graph">Graph for the job</param>
+		/// <param name="Batch">The batch to get timing info for</param>
+		/// <param name="JobTiming">The job timing information</param>
+		/// <returns>Wait time for the batch</returns>
+		public static TimeSpan? GetAverageWaitTime(IGraph Graph, IJobStepBatch Batch, IJobTiming JobTiming)
+		{
+			TimeSpan? WaitTime = null;
+			foreach (IJobStep Step in Batch.Steps)
+			{
+				INode Node = Graph.Groups[Batch.GroupIdx].Nodes[Step.NodeIdx];
+				if(JobTiming.TryGetStepTiming(Node.Name, out IJobStepTiming? TimingInfo))
+				{
+					if (TimingInfo.AverageWaitTime != null)
+					{
+						TimeSpan StepWaitTime = TimeSpan.FromSeconds(TimingInfo.AverageWaitTime.Value);
+						if (WaitTime == null || StepWaitTime > WaitTime.Value)
+						{
+							WaitTime = StepWaitTime;
+						}
+					}
+				}
+			}
+			return WaitTime;
+		}
+
+		/// <summary>
+		/// Gets the average initialization time for this batch
+		/// </summary>
+		/// <param name="Graph">Graph for the job</param>
+		/// <param name="Batch">The batch to get timing info for</param>
+		/// <param name="JobTiming">The job timing information</param>
+		/// <returns>Initialization time for this batch</returns>
+		public static TimeSpan? GetAverageInitTime(IGraph Graph, IJobStepBatch Batch, IJobTiming JobTiming)
+		{
+			TimeSpan? InitTime = null;
+			foreach (IJobStep Step in Batch.Steps)
+			{
+				INode Node = Graph.Groups[Batch.GroupIdx].Nodes[Step.NodeIdx];
+				if (JobTiming.TryGetStepTiming(Node.Name, out IJobStepTiming? TimingInfo))
+				{
+					if (TimingInfo.AverageInitTime != null)
+					{
+						TimeSpan StepInitTime = TimeSpan.FromSeconds(TimingInfo.AverageInitTime.Value);
+						if (InitTime == null || StepInitTime > InitTime.Value)
+						{
+							InitTime = StepInitTime;
+						}
+					}
+				}
+			}
+			return InitTime;
+		}
+
+		/// <summary>
+		/// Creates a nullable timespan from a nullable number of seconds
+		/// </summary>
+		/// <param name="Seconds">The number of seconds to construct from</param>
+		/// <returns>TimeSpan object</returns>
+		static TimeSpan? NullableTimeSpanFromSeconds(float? Seconds)
+		{
+			if (Seconds == null)
+			{
+				return null;
+			}
+			else
+			{
+				return TimeSpan.FromSeconds(Seconds.Value);
+			}
+		}
+
+		/// <summary>
+		/// Attempts to get a batch with the given id
+		/// </summary>
+		/// <param name="Job">The job document</param>
+		/// <param name="BatchId">The batch id</param>
+		/// <param name="Batch">On success, receives the batch object</param>
+		/// <returns>True if the batch was found</returns>
+		public static bool TryGetBatch(this IJob Job, SubResourceId BatchId, [NotNullWhen(true)] out IJobStepBatch? Batch)
+		{
+			Batch = Job.Batches.FirstOrDefault(x => x.Id == BatchId);
+			return Batch != null;
+		}
+
+		/// <summary>
+		/// Attempts to get a batch with the given id
+		/// </summary>
+		/// <param name="Job">The job document</param>
+		/// <param name="BatchId">The batch id</param>
+		/// <param name="StepId">The step id</param>
+		/// <param name="Step">On success, receives the step object</param>
+		/// <returns>True if the batch was found</returns>
+		public static bool TryGetStep(this IJob Job, SubResourceId BatchId, SubResourceId StepId, [NotNullWhen(true)] out IJobStep? Step)
+		{
+			IJobStepBatch? Batch;
+			if (!TryGetBatch(Job, BatchId, out Batch))
+			{
+				Step = null;
+				return false;
+			}
+			return Batch.TryGetStep(StepId, out Step);
+		}
+
+		/// <summary>
+		/// Finds the set of nodes affected by a label
+		/// </summary>
+		/// <param name="Job">The job document</param>
+		/// <param name="Graph">Graph definition for the job</param>
+		/// <param name="LabelIdx">Index of the label. -1 or Graph.Labels.Count are treated as referring to the default lable.</param>
+		/// <returns>Set of nodes affected by the given label</returns>
+		public static HashSet<NodeRef> GetNodesForLabel(this IJob Job, IGraph Graph, int LabelIdx)
+		{
+			if (LabelIdx != -1 && LabelIdx != Graph.Labels.Count)
+			{
+				// Return all the nodes included by the label
+				return new HashSet<NodeRef>(Graph.Labels[LabelIdx].IncludedNodes);
+			}
+			else
+			{
+				// Set of nodes which are not covered by an existing label, initially containing everything
+				HashSet<NodeRef> UnlabeledNodes = new HashSet<NodeRef>();
+				for (int GroupIdx = 0; GroupIdx < Graph.Groups.Count; GroupIdx++)
+				{
+					INodeGroup Group = Graph.Groups[GroupIdx];
+					for (int NodeIdx = 0; NodeIdx < Group.Nodes.Count; NodeIdx++)
+					{
+						UnlabeledNodes.Add(new NodeRef(GroupIdx, NodeIdx));
+					}
+				}
+
+				// Remove all the nodes that are part of an active label
+				IReadOnlyDictionary<NodeRef, IJobStep> StepForNode = Job.GetStepForNodeMap();
+				foreach (ILabel Label in Graph.Labels)
+				{
+					if (Label.RequiredNodes.Any(x => StepForNode.ContainsKey(x)))
+					{
+						UnlabeledNodes.ExceptWith(Label.IncludedNodes);
+					}
+				}
+				return UnlabeledNodes;
+			}
+		}
+
+		/// <summary>
+		/// Create a list of aggregate responses, combining the graph definitions with the state of the job
+		/// </summary>
+		/// <param name="Job">The job document</param>
+		/// <param name="Graph">Graph definition for the job</param>
+		/// <param name="Responses">List to receive all the responses</param>
+		/// <returns>The default label state</returns>
+		public static GetDefaultLabelStateResponse? GetLabelStateResponses(this IJob Job, IGraph Graph, List<GetLabelStateResponse> Responses)
+		{
+			// Create a lookup from noderef to step information
+			IReadOnlyDictionary<NodeRef, IJobStep> StepForNode = Job.GetStepForNodeMap();
+
+			// Set of nodes which are not covered by an existing label, initially containing everything
+			HashSet<NodeRef> UnlabeledNodes = new HashSet<NodeRef>();
+			for (int GroupIdx = 0; GroupIdx < Graph.Groups.Count; GroupIdx++)
+			{
+				INodeGroup Group = Graph.Groups[GroupIdx];
+				for (int NodeIdx = 0; NodeIdx < Group.Nodes.Count; NodeIdx++)
+				{
+					UnlabeledNodes.Add(new NodeRef(GroupIdx, NodeIdx));
+				}
+			}
+
+			// Create the responses
+			foreach (ILabel Label in Graph.Labels)
+			{
+				// Refresh the state for this label
+				LabelState NewState = LabelState.Unspecified;
+				foreach (NodeRef RequiredNodeRef in Label.RequiredNodes)
+				{
+					if (StepForNode.ContainsKey(RequiredNodeRef))
+					{
+						NewState = LabelState.Complete;
+						break;
+					}
+				}
+
+				// Refresh the outcome
+				LabelOutcome NewOutcome = LabelOutcome.Success;
+				if (NewState == LabelState.Complete)
+				{
+					GetLabelState(Label.IncludedNodes, StepForNode, out NewState, out NewOutcome);
+					UnlabeledNodes.ExceptWith(Label.IncludedNodes);
+				}
+
+				// Create the response
+				Responses.Add(new GetLabelStateResponse(NewState, NewOutcome));
+			}
+
+			// Remove all the nodes that don't have a step
+			UnlabeledNodes.RemoveWhere(x => !StepForNode.ContainsKey(x));
+
+			// Remove successful "setup build" nodes from the list
+			if (Graph.Groups.Count > 1 && Graph.Groups[0].Nodes.Count > 0)
+			{
+				INode Node = Graph.Groups[0].Nodes[0];
+				if (Node.Name == IJob.SetupNodeName)
+				{
+					NodeRef NodeRef = new NodeRef(0, 0);
+					if (UnlabeledNodes.Contains(NodeRef))
+					{
+						IJobStep Step = StepForNode[NodeRef];
+						if (Step.State == JobStepState.Completed && Step.Outcome == JobStepOutcome.Success && Responses.Count > 0)
+						{
+							UnlabeledNodes.Remove(NodeRef);
+						}
+					}
+				}
+			}
+
+			// Add a response for everything not included elsewhere.
+			GetLabelState(UnlabeledNodes, StepForNode, out LabelState OtherState, out LabelOutcome OtherOutcome);
+			return new GetDefaultLabelStateResponse(OtherState, OtherOutcome, UnlabeledNodes.Select(x => Graph.GetNode(x).Name).ToList());
+		}
+
+		/// <summary>
+		/// Get the states of all labels for this job
+		/// </summary>
+		/// <param name="Job">The job to get states for</param>
+		/// <param name="Graph">The graph for this job</param>
+		/// <returns>Collection of label states by label index</returns>
+		public static IReadOnlyList<(LabelState, LabelOutcome)> GetLabelStates(this IJob Job, IGraph Graph)
+		{
+			IReadOnlyDictionary<NodeRef, IJobStep> StepForNodeRef = Job.GetStepForNodeMap();
+
+			List<(LabelState, LabelOutcome)> States = new List<(LabelState, LabelOutcome)>();
+			for (int Idx = 0; Idx < Graph.Labels.Count; Idx++)
+			{
+				ILabel Label = Graph.Labels[Idx];
+
+				// Default the label to the unspecified state
+				LabelState NewState = LabelState.Unspecified;
+				LabelOutcome NewOutcome = LabelOutcome.Unspecified;
+
+				// Check if the label should be included
+				if (Label.RequiredNodes.Any(x => StepForNodeRef.ContainsKey(x)))
+				{
+					// Combine the state of the steps contributing towards this label
+					bool bAnySkipped = false;
+					bool bAnyWarnings = false;
+					bool bAnyFailed = false;
+					bool bAnyPending = false;
+					foreach (NodeRef IncludedNode in Label.IncludedNodes)
+					{
+						IJobStep? Step;
+						if (StepForNodeRef.TryGetValue(IncludedNode, out Step))
+						{
+							bAnyPending |= Step.IsPending();
+							bAnySkipped |= Step.State == JobStepState.Aborted || Step.State == JobStepState.Skipped;
+							bAnyFailed |= (Step.Outcome == JobStepOutcome.Failure);
+							bAnyWarnings |= (Step.Outcome == JobStepOutcome.Warnings);
+						}
+					}
+
+					// Figure out the overall label state
+					NewState = bAnyPending ? LabelState.Running : LabelState.Complete;
+					NewOutcome = bAnyFailed ? LabelOutcome.Failure : bAnyWarnings ? LabelOutcome.Warnings : bAnySkipped? LabelOutcome.Unspecified : LabelOutcome.Success;
+				}
+
+				States.Add((NewState, NewOutcome));
+			}
+			return States;	
+		}
+
+		/// <summary>
+		/// Get the states of all UGS badges for this job
+		/// </summary>
+		/// <param name="Job">The job to get states for</param>
+		/// <param name="Graph">The graph for this job</param>
+		/// <returns>List of badge states</returns>
+		public static Dictionary<int, UgsBadgeState> GetUgsBadgeStates(this IJob Job, IGraph Graph)
+		{
+			IReadOnlyList<(LabelState, LabelOutcome)> LabelStates = GetLabelStates(Job, Graph);
+			return Job.GetUgsBadgeStates(Graph, LabelStates);
+		}
+
+		/// <summary>
+		/// Get the states of all UGS badges for this job
+		/// </summary>
+		/// <param name="Job">The job to get states for</param>
+		/// <param name="Graph">The graph for this job</param>
+		/// <param name="LabelStates">The existing label states to get the UGS badge states from</param>
+		/// <returns>List of badge states</returns>
+		public static Dictionary<int, UgsBadgeState> GetUgsBadgeStates(this IJob Job, IGraph Graph, IReadOnlyList<(LabelState, LabelOutcome)> LabelStates)
+		{
+			Dictionary<int, UgsBadgeState> UgsBadgeStates = new Dictionary<int, UgsBadgeState>();
+			for (int LabelIdx = 0; LabelIdx < LabelStates.Count; ++LabelIdx)
+			{
+				if (Graph.Labels[LabelIdx].UgsName == null)
+				{
+					continue;
+				}
+
+				(LabelState State, LabelOutcome Outcome) Label = LabelStates[LabelIdx];
+				switch (Label.State)
+				{
+					case LabelState.Complete:
+					{
+						switch (Label.Outcome)
+						{
+							case LabelOutcome.Success:
+							{
+								UgsBadgeStates.Add(LabelIdx, UgsBadgeState.Success);
+								break;
+							}
+
+							case LabelOutcome.Warnings:
+							{
+								UgsBadgeStates.Add(LabelIdx, UgsBadgeState.Warning);
+								break;
+							}
+
+							case LabelOutcome.Failure:
+							{
+								UgsBadgeStates.Add(LabelIdx, UgsBadgeState.Failure);
+								break;
+							}
+
+							case LabelOutcome.Unspecified:
+							{
+								UgsBadgeStates.Add(LabelIdx, UgsBadgeState.Skipped);
+								break;
+							}
+						}
+						break;
+					}
+
+					case LabelState.Running:
+					{
+						UgsBadgeStates.Add(LabelIdx, UgsBadgeState.Starting);
+						break;
+					}
+
+					case LabelState.Unspecified:
+					{
+						UgsBadgeStates.Add(LabelIdx, UgsBadgeState.Skipped);
+						break;
+					}
+				}
+			}
+			return UgsBadgeStates;
+		}
+
+		/// <summary>
+		/// Gets the state of a job, as a label that includes all steps
+		/// </summary>
+		/// <param name="Job">The job to query</param>
+		/// <param name="StepForNode">Map from node to step</param>
+		/// <param name="NewState">Receives the state of the label</param>
+		/// <param name="NewOutcome">Receives the outcome of the label</param>
+		public static void GetJobState(this IJob Job, IReadOnlyDictionary<NodeRef, IJobStep> StepForNode, out LabelState NewState, out LabelOutcome NewOutcome)
+		{
+			List<NodeRef> Nodes = new List<NodeRef>();
+			foreach (IJobStepBatch Batch in Job.Batches)
+			{
+				foreach (IJobStep Step in Batch.Steps)
+				{
+					Nodes.Add(new NodeRef(Batch.GroupIdx, Step.NodeIdx));
+				}
+			}
+			GetLabelState(Nodes, StepForNode, out NewState, out NewOutcome);
+		}
+		
+		/// <summary>
+		/// Gets the state of a label
+		/// </summary>
+		/// <param name="IncludedNodes">Nodes to include in this label</param>
+		/// <param name="StepForNode">Map from node to step</param>
+		/// <param name="NewState">Receives the state of the label</param>
+		/// <param name="NewOutcome">Receives the outcome of the label</param>
+		public static void GetLabelState(IEnumerable<NodeRef> IncludedNodes, IReadOnlyDictionary<NodeRef, IJobStep> StepForNode, out LabelState NewState, out LabelOutcome NewOutcome)
+		{
+			NewState = LabelState.Complete;
+			NewOutcome = LabelOutcome.Success;
+			foreach (NodeRef IncludedNodeRef in IncludedNodes)
+			{
+				IJobStep? IncludedStep;
+				if (StepForNode.TryGetValue(IncludedNodeRef, out IncludedStep))
+				{
+					// Update the state
+					if (IncludedStep.State != JobStepState.Completed && IncludedStep.State != JobStepState.Skipped && IncludedStep.State != JobStepState.Aborted)
+					{
+						NewState = LabelState.Running;
+					}
+
+					// Update the outcome
+					if (IncludedStep.State == JobStepState.Skipped || IncludedStep.State == JobStepState.Aborted || IncludedStep.Outcome == JobStepOutcome.Failure)
+					{
+						NewOutcome = LabelOutcome.Failure;
+					}
+					else if (IncludedStep.Outcome == JobStepOutcome.Warnings && NewOutcome == LabelOutcome.Success)
+					{
+						NewOutcome = LabelOutcome.Warnings;
+					}
+				}
+			}
+		}
+
+		/// <summary>
+		/// Creates a response containing information about this object
+		/// </summary>
+		/// <param name="Job">The job document</param>
+		/// <param name="Graph">The graph definition</param>
+		/// <param name="bIncludeBatches">Whether to include the job batches in the response</param>
+		/// <param name="bIncludeLabels">Whether to include the job aggregates in the response</param>
+		/// <param name="bIncludeAcl">Whether to include the ACL in the response</param>
+		/// <returns>The response object</returns>
+		public static GetJobResponse ToResponse(this IJob Job, IGraph Graph, bool bIncludeBatches, bool bIncludeLabels, bool bIncludeAcl)
+		{
+			GetAclResponse? AclResponse = (bIncludeAcl && Job.Acl != null) ? new GetAclResponse(Job.Acl) : null;
+			GetJobResponse Response = new GetJobResponse(Job, AclResponse);
+			if (bIncludeBatches || bIncludeLabels)
+			{
+				if (bIncludeBatches)
+				{
+					Response.Batches = Job.Batches.ConvertAll(x => new GetBatchResponse(x));
+				}
+				if (bIncludeLabels)
+				{
+					Response.Labels = new List<GetLabelStateResponse>();
+					Response.DefaultLabel = GetLabelStateResponses(Job, Graph, Response.Labels);
+				}
+			}
+			return Response;
+		}
+
+		/// <summary>
+		/// Creates a response containing information about this object
+		/// </summary>
+		/// <param name="Job">The job document</param>
+		/// <param name="Graph">The graph for this job</param>
+		/// <param name="bIncludeAcl">Whether to include the ACL in the response</param>
+		/// <param name="Filter">Filter for the properties to return</param>
+		/// <returns>Object containing the requested properties</returns>
+		public static object ToResponse(this IJob Job, IGraph Graph, bool bIncludeAcl, PropertyFilter? Filter)
+		{
+			if(Filter == null)
+			{
+				return Job.ToResponse(Graph, true, true, bIncludeAcl);
+			}
+			else
+			{
+				return Filter.ApplyTo(Job.ToResponse(Graph, Filter.Includes(nameof(GetJobResponse.Batches)), Filter.Includes(nameof(GetJobResponse.Labels)) || Filter.Includes(nameof(GetJobResponse.DefaultLabel)), bIncludeAcl));
+			}
+		}
+
+		/// <summary>
+		/// Creates an RPC response object
+		/// </summary>
+		/// <param name="Job">The job document</param>
+		/// <returns></returns>
+		public static HordeCommon.Rpc.GetJobResponse ToRpcResponse(this IJob Job)
+		{
+			HordeCommon.Rpc.GetJobResponse Response = new HordeCommon.Rpc.GetJobResponse();
+			Response.StreamId = Job.StreamId.ToString();
+			Response.Change = Job.Change;
+			Response.CodeChange = Job.CodeChange;
+			Response.PreflightChange = Job.PreflightChange;
+			Response.ClonedPreflightChange = Job.ClonedPreflightChange;
+			Response.Arguments.Add(Job.Arguments);
+			return Response;
+		}
+	}
+
+	/// <summary>
+	/// Projection of a job definition to just include permissions info
+	/// </summary>
+	public interface IJobPermissions
+	{
+		/// <summary>
+		/// ACL for the job
+		/// </summary>
+		public Acl? Acl { get; }
+
+		/// <summary>
+		/// The stream containing 
+		/// </summary>
+		public StreamId StreamId { get; }
+	}
+}
diff --git a/Engine/Source/Programs/Horde/HordeServer/Models/IJobStepRef.cs b/Engine/Source/Programs/Horde/HordeServer/Models/IJobStepRef.cs
new file mode 100644
--- /dev/null
+++ b/Engine/Source/Programs/Horde/HordeServer/Models/IJobStepRef.cs
@@ -0,0 +1,185 @@
+// Copyright Epic Games, Inc. All Rights Reserved.
+
+using HordeCommon;
+using HordeServer.Utilities;
+using MongoDB.Bson;
+using MongoDB.Bson.Serialization;
+using MongoDB.Bson.Serialization.Attributes;
+using System;
+using System.Collections.Generic;
+using System.Linq;
+using System.Threading.Tasks;
+
+namespace HordeServer.Models
+{
+	using StreamId = StringId<IStream>;
+	using TemplateRefId = StringId<TemplateRef>;
+	using PoolId = StringId<IPool>;
+
+	/// <summary>
+	/// Unique id struct for JobStepRef objects. Includes a job id, batch id, and step id to uniquely identify the step.
+	/// </summary>
+	[BsonSerializer(typeof(JobStepRefIdSerializer))]
+	public struct JobStepRefId
+	{
+		/// <summary>
+		/// The job id
+		/// </summary>
+		public ObjectId JobId { get; set; }
+
+		/// <summary>
+		/// The batch id within the job
+		/// </summary>
+		public SubResourceId BatchId { get; set; }
+
+		/// <summary>
+		/// The step id
+		/// </summary>
+		public SubResourceId StepId { get; set; }
+
+		/// <summary>
+		/// Constructor
+		/// </summary>
+		/// <param name="JobId">The job id</param>
+		/// <param name="BatchId">The batch id within the job</param>
+		/// <param name="StepId">The step id</param>
+		public JobStepRefId(ObjectId JobId, SubResourceId BatchId, SubResourceId StepId)
+		{
+			this.JobId = JobId;
+			this.BatchId = BatchId;
+			this.StepId = StepId;
+		}
+
+		/// <summary>
+		/// Parse a job step id from a string
+		/// </summary>
+		/// <param name="Text">Text to parse</param>
+		/// <returns>The parsed id</returns>
+		public static JobStepRefId Parse(string Text)
+		{
+			string[] Components = Text.Split(':');
+			return new JobStepRefId(ObjectId.Parse(Components[0]), SubResourceId.Parse(Components[1]), SubResourceId.Parse(Components[2]));
+		}
+
+		/// <summary>
+		/// Formats this id as a string
+		/// </summary>
+		/// <returns>Formatted id</returns>
+		public override string ToString()
+		{
+			return $"{JobId}:{BatchId}:{StepId}";
+		}
+	}
+
+	/// <summary>
+	/// Serializer for JobStepRefId objects
+	/// </summary>
+	public sealed class JobStepRefIdSerializer : IBsonSerializer<JobStepRefId>
+	{
+		/// <inheritdoc/>
+		public Type ValueType
+		{
+			get { return typeof(JobStepRefId); }
+		}
+
+		/// <inheritdoc/>
+		void IBsonSerializer.Serialize(BsonSerializationContext Context, BsonSerializationArgs Args, object Value)
+		{
+			Serialize(Context, Args, (JobStepRefId)Value);
+		}
+
+		/// <inheritdoc/>
+		object IBsonSerializer.Deserialize(BsonDeserializationContext Context, BsonDeserializationArgs Args)
+		{
+			return ((IBsonSerializer<JobStepRefId>)this).Deserialize(Context, Args);
+		}
+
+		/// <inheritdoc/>
+		public JobStepRefId Deserialize(BsonDeserializationContext Context, BsonDeserializationArgs Args)
+		{
+			return JobStepRefId.Parse(Context.Reader.ReadString());
+		}
+
+		/// <inheritdoc/>
+		public void Serialize(BsonSerializationContext Context, BsonSerializationArgs Args, JobStepRefId Value)
+		{
+			Context.Writer.WriteString(((JobStepRefId)Value).ToString());
+		}
+	}
+
+	/// <summary>
+	/// Searchable reference to a jobstep
+	/// </summary>
+	public interface IJobStepRef
+	{
+		/// <summary>
+		/// Globally unique identifier for the jobstep being referenced
+		/// </summary>
+		public JobStepRefId Id { get; }
+
+		/// <summary>
+		/// Name of the job
+		/// </summary>
+		public string JobName { get; }
+
+		/// <summary>
+		/// Name of the name
+		/// </summary>
+		public string NodeName { get; }
+
+		/// <summary>
+		/// Unique id of the stream containing the job
+		/// </summary>
+		public StreamId StreamId { get; }
+
+		/// <summary>
+		/// Template for the job being executed
+		/// </summary>
+		public TemplateRefId TemplateId { get; }
+
+		/// <summary>
+		/// The change number being built
+		/// </summary>
+		public int Change { get; }
+
+		/// <summary>
+		/// Log for this step
+		/// </summary>
+		public ObjectId? LogId { get; }
+
+		/// <summary>
+		/// The agent type
+		/// </summary>
+		public PoolId? PoolId { get; }
+
+		/// <summary>
+		/// The agent id
+		/// </summary>
+		public AgentId? AgentId { get; }
+
+		/// <summary>
+		/// Outcome of the step, once complete.
+		/// </summary>
+		public JobStepOutcome? Outcome { get; }
+
+		/// <summary>
+		/// Time taken for the batch containing this batch to start after it became ready
+		/// </summary>
+		public float BatchWaitTime { get; }
+
+		/// <summary>
+		/// Time taken for this batch to initialize
+		/// </summary>
+		public float BatchInitTime { get; }
+
+		/// <summary>
+		/// Time at which the step started.
+		/// </summary>
+		public DateTime StartTimeUtc { get; }
+
+		/// <summary>
+		/// Time at which the step finished.
+		/// </summary>
+		public DateTime? FinishTimeUtc { get; }
+	}
+}
diff --git a/Engine/Source/Programs/Horde/HordeServer/Models/IJobTiming.cs b/Engine/Source/Programs/Horde/HordeServer/Models/IJobTiming.cs
new file mode 100644
--- /dev/null
+++ b/Engine/Source/Programs/Horde/HordeServer/Models/IJobTiming.cs
@@ -0,0 +1,45 @@
+// Copyright Epic Games, Inc. All Rights Reserved.
+
+using System;
+using System.Collections.Generic;
+using System.Diagnostics.CodeAnalysis;
+using System.Linq;
+using System.Threading.Tasks;
+
+namespace HordeServer.Models
+{
+	/// <summary>
+	/// Information about the timing for an individual step
+	/// </summary>
+	public interface IJobStepTiming
+	{
+		/// <summary>
+		/// Wait time before executing the group containing this node
+		/// </summary>
+		public float? AverageWaitTime { get; }
+
+		/// <summary>
+		/// Time taken for the group containing this node to initialize
+		/// </summary>
+		public float? AverageInitTime { get; }
+
+		/// <summary>
+		/// Time spent executing this node
+		/// </summary>
+		public float? AverageDuration { get; }
+	}
+
+	/// <summary>
+	/// Timing information for a particular job
+	/// </summary>
+	public interface IJobTiming
+	{
+		/// <summary>
+		/// Gets timing information for a particular step
+		/// </summary>
+		/// <param name="Name">Name of the node being executed</param>
+		/// <param name="Timing">Receives the timing information for the given step</param>
+		/// <returns>True if the timing was found</returns>
+		public bool TryGetStepTiming(string Name, [NotNullWhen(true)] out IJobStepTiming? Timing);
+	}
+}
diff --git a/Engine/Source/Programs/Horde/HordeServer/Models/ILease.cs b/Engine/Source/Programs/Horde/HordeServer/Models/ILease.cs
new file mode 100644
--- /dev/null
+++ b/Engine/Source/Programs/Horde/HordeServer/Models/ILease.cs
@@ -0,0 +1,113 @@
+// Copyright Epic Games, Inc. All Rights Reserved.
+
+using HordeServer.Api;
+using HordeCommon;
+using HordeServer.Utilities;
+using MongoDB.Bson;
+using MongoDB.Bson.Serialization.Attributes;
+using Serilog;
+using System;
+using System.Collections.Generic;
+using System.Linq;
+using System.Threading.Tasks;
+using Google.Protobuf.WellKnownTypes;
+using Google.Protobuf;
+
+namespace HordeServer.Models
+{
+	using PoolId = StringId<IPool>;
+	using StreamId = StringId<IStream>;
+
+	/// <summary>
+	/// Document describing a lease. This exists to permanently record a lease; the agent object tracks internal state of any active leases through AgentLease objects.
+	/// </summary>
+	public interface ILease
+	{
+		/// <summary>
+		/// The unique id of this lease
+		/// </summary>
+		public ObjectId Id { get; }
+
+		/// <summary>
+		/// Name of this lease
+		/// </summary>
+		public string Name { get; }
+
+		/// <summary>
+		/// Unique id of the agent 
+		/// </summary>
+		public AgentId AgentId { get; }
+
+		/// <summary>
+		/// Unique id of the agent session
+		/// </summary>
+		public ObjectId SessionId { get; }
+
+		/// <summary>
+		/// The stream this lease belongs to
+		/// </summary>
+		public StreamId? StreamId { get; }
+
+		/// <summary>
+		/// Pool for the work being executed
+		/// </summary>
+		public PoolId? PoolId { get; }
+
+		/// <summary>
+		/// The log for this lease, if applicable
+		/// </summary>
+		public ObjectId? LogId { get; }
+
+		/// <summary>
+		/// Time at which this lease started
+		/// </summary>
+		public DateTime StartTime { get; }
+
+		/// <summary>
+		/// Time at which this lease completed
+		/// </summary>
+		public DateTime? FinishTime { get; }
+
+		/// <summary>
+		/// Payload for this lease. A packed Google.Protobuf.Any object.
+		/// </summary>
+		public ReadOnlyMemory<byte> Payload { get; }
+
+		/// <summary>
+		/// Outcome of the lease
+		/// </summary>
+		public LeaseOutcome Outcome { get; }
+
+		/// <summary>
+		/// Output from executing the lease
+		/// </summary>
+		public ReadOnlyMemory<byte> Output { get; }
+	}
+
+	/// <summary>
+	/// Extension methods for leases
+	/// </summary>
+	public static class LeaseExtensions
+	{
+		/// <summary>
+		/// Gets the task from a lease, encoded as an Any protobuf object
+		/// </summary>
+		/// <param name="Lease">The lease to query</param>
+		/// <returns>The task definition encoded as a protobuf Any object</returns>
+		public static Any GetTask(this ILease Lease)
+		{
+			return Any.Parser.ParseFrom(Lease.Payload.ToArray());
+		}
+
+		/// <summary>
+		/// Gets a typed task object from a lease
+		/// </summary>
+		/// <typeparam name="T">Type of the protobuf message to return</typeparam>
+		/// <param name="Lease">The lease to query</param>
+		/// <returns>The task definition</returns>
+		public static T GetTask<T>(this ILease Lease) where T : IMessage<T>, new()
+		{
+			return GetTask(Lease).Unpack<T>();
+		}
+	}
+}
diff --git a/Engine/Source/Programs/Horde/HordeServer/Models/ILogEvent.cs b/Engine/Source/Programs/Horde/HordeServer/Models/ILogEvent.cs
new file mode 100644
--- /dev/null
+++ b/Engine/Source/Programs/Horde/HordeServer/Models/ILogEvent.cs
@@ -0,0 +1,129 @@
+// Copyright Epic Games, Inc. All Rights Reserved.
+
+using HordeCommon;
+using HordeCommon.Rpc;
+using Microsoft.Extensions.Logging;
+using MongoDB.Bson;
+using System;
+using System.Collections.Generic;
+using System.Diagnostics.CodeAnalysis;
+using System.Linq;
+using System.Text.Json;
+using System.Threading.Tasks;
+
+namespace HordeServer.Models
+{
+	/// <summary>
+	/// Interface for a log event line
+	/// </summary>
+	public interface ILogEventLine
+	{
+		/// <summary>
+		/// Log level
+		/// </summary>
+		LogLevel Level { get; }
+
+		/// <summary>
+		/// The event id
+		/// </summary>
+		EventId? EventId { get; }
+
+		/// <summary>
+		/// The complete rendered message, in plaintext
+		/// </summary>
+		string Message { get; }
+
+		/// <summary>
+		/// Data for the line
+		/// </summary>
+		JsonElement Data { get; }
+	}
+
+	/// <summary>
+	/// Interface for event data
+	/// </summary>
+	public interface ILogEventData
+	{
+		/// <summary>
+		/// The log level
+		/// </summary>
+		EventSeverity Severity { get; }
+
+		/// <summary>
+		/// The type of event
+		/// </summary>
+		EventId? EventId { get; }
+
+		/// <summary>
+		/// The complete rendered message, in plaintext
+		/// </summary>
+		string Message { get; }
+
+		/// <summary>
+		/// Gets this event data as a BSON document
+		/// </summary>
+		IReadOnlyList<ILogEventLine> Lines { get; }
+	}
+
+	/// <summary>
+	/// Represents a node in the graph
+	/// </summary>
+	public interface ILogEvent
+	{
+		/// <summary>
+		/// Unique id of the log containing this event
+		/// </summary>
+		public ObjectId LogId { get; }
+
+		/// <summary>
+		/// Severity of the event
+		/// </summary>
+		public EventSeverity Severity { get; }
+
+		/// <summary>
+		/// Index of the first line for this event
+		/// </summary>
+		public int LineIndex { get; }
+
+		/// <summary>
+		/// Number of lines in the event
+		/// </summary>
+		public int LineCount { get; }
+
+		/// <summary>
+		/// Span id for this log event
+		/// </summary>
+		public ObjectId? SpanId { get; }
+	}
+
+	/// <summary>
+	/// Represents a node in the graph
+	/// </summary>
+	public class NewLogEventData
+	{
+		/// <summary>
+		/// Unique id of the log containing this event
+		/// </summary>
+		public ObjectId LogId { get; set; }
+
+		/// <summary>
+		/// Severity of the event
+		/// </summary>
+		public EventSeverity Severity { get; set; }
+
+		/// <summary>
+		/// Index of the first line for this event
+		/// </summary>
+		public int LineIndex { get; set; }
+
+		/// <summary>
+		/// Number of lines in the event
+		/// </summary>
+		public int LineCount { get; set; }
+
+		/// <summary>
+		/// The span this this event belongs to
+		/// </summary>
+		public ObjectId? SpanId { get; set; }
+	}
+}
diff --git a/Engine/Source/Programs/Horde/HordeServer/Models/ILogFile.cs b/Engine/Source/Programs/Horde/HordeServer/Models/ILogFile.cs
new file mode 100644
--- /dev/null
+++ b/Engine/Source/Programs/Horde/HordeServer/Models/ILogFile.cs
@@ -0,0 +1,121 @@
+// Copyright Epic Games, Inc. All Rights Reserved.
+
+using EpicGames.Core;
+using HordeServer.Api;
+using HordeServer.Services;
+using HordeServer.Utilities;
+using MongoDB.Bson;
+using MongoDB.Bson.Serialization.Attributes;
+using System;
+using System.Collections.Generic;
+using System.Diagnostics.CodeAnalysis;
+using System.Linq;
+using System.Threading.Tasks;
+
+namespace HordeServer.Models
+{
+	/// <summary>
+	/// Information about a log file chunk
+	/// </summary>
+	public interface ILogChunk
+	{
+		/// <summary>
+		/// Offset of the chunk within the log
+		/// </summary>
+		long Offset { get; }
+
+		/// <summary>
+		/// Length of this chunk. If zero, the chunk is still being written to.
+		/// </summary>
+		int Length { get; }
+
+		/// <summary>
+		/// Index of the first line within this chunk. If a line straddles two chunks, this is the index of the split line.
+		/// </summary>
+		int LineIndex { get; }
+
+		/// <summary>
+		/// If the chunk has yet to be pushed to persistent storage, includes the name of the server that is currently storing it.
+		/// </summary>
+		string? Server { get; }
+	}
+
+	/// <summary>
+	/// Information about a log file
+	/// </summary>
+	public interface ILogFile
+	{
+		/// <summary>
+		/// Identifier for the LogFile. Randomly generated.
+		/// </summary>
+		public ObjectId Id { get; }
+
+		/// <summary>
+		/// Unique id of the job containing this log
+		/// </summary>
+		public ObjectId JobId { get; }
+
+		/// <summary>
+		/// The session allowed to write to this log
+		/// </summary>
+		public ObjectId? SessionId { get; }
+
+		/// <summary>
+		/// Maximum line index in the file
+		/// </summary>
+		public int? MaxLineIndex { get; }
+
+		/// <summary>
+		/// Length of the file which is indexed
+		/// </summary>
+		public long? IndexLength { get; }
+
+		/// <summary>
+		/// Type of data stored in this log 
+		/// </summary>
+		public LogType Type { get; }
+
+		/// <summary>
+		/// Chunks within this file
+		/// </summary>
+		public IReadOnlyList<ILogChunk> Chunks { get; }
+	}
+
+	/// <summary>
+	/// Extension methods for log files
+	/// </summary>
+	public static class LogFileExtensions
+	{
+		/// <summary>
+		/// Gets the chunk index containing the given offset.
+		/// </summary>
+		/// <param name="Chunks">The chunks to search</param>
+		/// <param name="Offset">The offset to search for</param>
+		/// <returns>The chunk index containing the given offset</returns>
+		public static int GetChunkForOffset(this IReadOnlyList<ILogChunk> Chunks, long Offset)
+		{
+			int ChunkIndex = Chunks.BinarySearch(x => x.Offset, Offset);
+			if (ChunkIndex < 0)
+			{
+				ChunkIndex = ~ChunkIndex - 1;
+			}
+			return ChunkIndex;
+		}
+
+		/// <summary>
+		/// Gets the starting chunk index for the given line
+		/// </summary>
+		/// <param name="Chunks">The chunks to search</param>
+		/// <param name="LineIndex">Index of the line to query</param>
+		/// <returns>Index of the chunk to fetch</returns>
+		public static int GetChunkForLine(this IReadOnlyList<ILogChunk> Chunks, int LineIndex)
+		{
+			int ChunkIndex = Chunks.BinarySearch(x => x.LineIndex, LineIndex);
+			if(ChunkIndex < 0)
+			{
+				ChunkIndex = ~ChunkIndex - 1;
+			}
+			return ChunkIndex;
+		}
+	}
+}
diff --git a/Engine/Source/Programs/Horde/HordeServer/Models/INotificationTrigger.cs b/Engine/Source/Programs/Horde/HordeServer/Models/INotificationTrigger.cs
new file mode 100644
--- /dev/null
+++ b/Engine/Source/Programs/Horde/HordeServer/Models/INotificationTrigger.cs
@@ -0,0 +1,51 @@
+// Copyright Epic Games, Inc. All Rights Reserved.
+
+using HordeServer.Api;
+using MongoDB.Bson;
+using MongoDB.Bson.Serialization.Attributes;
+using System.Collections.Generic;
+
+namespace HordeServer.Models
+{
+	/// <summary>
+	/// An individual subscription
+	/// </summary>
+	public interface INotificationSubscription
+	{
+		/// <summary>
+		/// Name of the user
+		/// </summary>
+		public ObjectId UserId { get; }
+
+		/// <summary>
+		/// Whether to receive email notifications
+		/// </summary>
+		public bool Email { get; }
+
+		/// <summary>
+		/// Whether to receive Slack notifications
+		/// </summary>
+		public bool Slack { get; }
+	}
+
+	/// <summary>
+	/// Trigger for notifications to be sent
+	/// </summary>
+	public interface INotificationTrigger
+	{
+		/// <summary>
+		/// Unique id for this subscription list
+		/// </summary>
+		ObjectId Id { get; }
+
+		/// <summary>
+		/// Whether this trigger has been fired
+		/// </summary>
+		bool Fired { get; }
+
+		/// <summary>
+		/// List of subscriptions to this event
+		/// </summary>
+		public IReadOnlyList<INotificationSubscription> Subscriptions { get; }
+	}
+}
diff --git a/Engine/Source/Programs/Horde/HordeServer/Models/IPool.cs b/Engine/Source/Programs/Horde/HordeServer/Models/IPool.cs
new file mode 100644
--- /dev/null
+++ b/Engine/Source/Programs/Horde/HordeServer/Models/IPool.cs
@@ -0,0 +1,78 @@
+// Copyright Epic Games, Inc. All Rights Reserved.
+
+using EpicGames.Horde.Common;
+using HordeServer.Api;
+using HordeServer.Utilities;
+using MongoDB.Bson;
+using MongoDB.Bson.Serialization.Attributes;
+using System;
+using System.Collections.Generic;
+using System.Linq;
+using System.Security.Permissions;
+using System.Threading.Tasks;
+
+namespace HordeServer.Models
+{
+	using PoolId = StringId<IPool>;
+
+	/// <summary>
+	/// A pool of machines
+	/// </summary>
+	public interface IPool
+	{
+		/// <summary>
+		/// Unique id for this pool
+		/// </summary>
+		public PoolId Id { get; }
+
+		/// <summary>
+		/// Name of the pool
+		/// </summary>
+		public string Name { get; }
+
+		/// <summary>
+		/// Condition for agents to automatically be included in this pool
+		/// </summary>
+		public Condition? Condition { get; }
+
+		/// <summary>
+		/// List of workspaces currently assigned to this pool
+		/// </summary>
+		public IReadOnlyList<AgentWorkspace> Workspaces { get; }
+
+		/// <summary>
+		/// Arbitrary properties related to this pool
+		/// </summary>
+		public IReadOnlyDictionary<string, string> Properties { get; }
+
+		/// <summary>
+		/// Whether to enable autoscaling for this pool
+		/// </summary>
+		public bool EnableAutoscaling { get; }
+
+		/// <summary>
+		/// The minimum number of agents to keep in the pool
+		/// </summary>
+		public int? MinAgents { get; }
+
+		/// <summary>
+		/// The minimum number of idle agents to hold in reserve
+		/// </summary>
+		public int? NumReserveAgents { get; }
+
+		/// <summary>
+		/// Last time the pool was (auto) scaled up
+		/// </summary>
+		public DateTime? LastScaleUpTime { get; }
+		
+		/// <summary>
+		/// Last time the pool was (auto) scaled down
+		/// </summary>
+		public DateTime? LastScaleDownTime { get; }
+
+		/// <summary>
+		/// Update index for this document
+		/// </summary>
+		public int UpdateIndex { get; }
+	}
+}
diff --git a/Engine/Source/Programs/Horde/HordeServer/Models/IProject.cs b/Engine/Source/Programs/Horde/HordeServer/Models/IProject.cs
new file mode 100644
--- /dev/null
+++ b/Engine/Source/Programs/Horde/HordeServer/Models/IProject.cs
@@ -0,0 +1,248 @@
+// Copyright Epic Games, Inc. All Rights Reserved.
+
+using EpicGames.Core;
+using HordeServer.Api;
+using HordeServer.Utilities;
+using MongoDB.Bson;
+using MongoDB.Bson.Serialization.Attributes;
+using MongoDB.Driver;
+using System;
+using System.Collections.Generic;
+using System.Linq;
+using System.Text.RegularExpressions;
+using System.Threading.Tasks;
+
+namespace HordeServer.Models
+{
+	using ProjectId = StringId<IProject>;
+
+	/// <summary>
+	/// Specifies a category of streams to display on the dashboard
+	/// </summary>
+	public class StreamCategory
+	{
+		/// <summary>
+		/// Name of this group
+		/// </summary>
+		[BsonRequired]
+		public string Name { get; set; }
+
+		/// <summary>
+		/// Row to display this category on
+		/// </summary>
+		public int Row { get; set; }
+
+		/// <summary>
+		/// Whether to show this category on the nav menu
+		/// </summary>
+		public bool ShowOnNavMenu { get; set; } = true;
+
+		/// <summary>
+		/// List of stream name patterns to include for this project
+		/// </summary>
+		public List<string> IncludePatterns { get; set; } = new List<string>();
+
+		/// <summary>
+		/// Patterns for stream names to be excluded
+		/// </summary>
+		public List<string> ExcludePatterns { get; set; } = new List<string>();
+
+		/// <summary>
+		/// Private constructor
+		/// </summary>
+		private StreamCategory()
+		{
+			Name = null!;
+		}
+
+		/// <summary>
+		/// Constructor
+		/// </summary>
+		/// <param name="Name">Name of this group</param>
+		public StreamCategory(string Name)
+		{
+			this.Name = Name;
+		}
+
+		/// <summary>
+		/// Constructs a category from a request object
+		/// </summary>
+		/// <param name="Request">The request object</param>
+		public StreamCategory(CreateProjectCategoryRequest Request)
+		{
+			this.Name = Request.Name;
+			this.Row = Request.Row;
+			this.ShowOnNavMenu = Request.ShowOnNavMenu;
+			this.IncludePatterns = Request.IncludePatterns;
+			this.ExcludePatterns = Request.ExcludePatterns;
+		}
+
+		/// <summary>
+		/// Tests if a given name matches a pattern
+		/// </summary>
+		/// <param name="Name">The name to test</param>
+		/// <param name="Pattern">The pattern to match against</param>
+		/// <returns>True if the pattern matches, false otherwise</returns>
+		public static bool MatchPattern(string Name, string Pattern)
+		{
+			return Regex.IsMatch(Name, Pattern);
+		}
+	}
+
+	/// <summary>
+	/// Represents a project
+	/// </summary>
+	public interface IProject
+	{
+		/// <summary>
+		/// Identifier for the project.
+		/// </summary>
+		public ProjectId Id { get; }
+
+		/// <summary>
+		/// Name of the project
+		/// </summary>
+		public string Name { get; }
+
+		/// <summary>
+		/// Path to the config file used to configure this project
+		/// </summary>
+		public string? ConfigPath { get; }
+
+		/// <summary>
+		/// Revision of the config file used to configure this project
+		/// </summary>
+		public string? ConfigRevision { get; }
+
+		/// <summary>
+		/// Order to display on the dashboard
+		/// </summary>
+		public int Order { get; }
+
+		/// <summary>
+		/// The ACL for this object
+		/// </summary>
+		public Acl? Acl { get; }
+
+		/// <summary>
+		/// List of stream categories for this project. Controls how streams are displayed on the dashboard.
+		/// </summary>
+		public IReadOnlyList<StreamCategory> Categories { get; }
+	}
+	
+	/// <summary>
+	/// Interface for a document containing a project logo
+	/// </summary>
+	public interface IProjectLogo
+	{
+		/// <summary>
+		/// The project id
+		/// </summary>
+		public ProjectId Id { get; }
+
+		/// <summary>
+		/// Path to the logo
+		/// </summary>
+		public string Path { get; }
+
+		/// <summary>
+		/// Revision of the file
+		/// </summary>
+		public string Revision { get; }
+
+		/// <summary>
+		/// Mime type for the image
+		/// </summary>
+		public string MimeType { get; }
+
+		/// <summary>
+		/// Image data
+		/// </summary>
+		public byte[] Data { get; }
+	}
+
+	/// <summary>
+	/// Extension methods for projects
+	/// </summary>
+	public static class ProjectExtensions
+	{
+		/// <summary>
+		/// Converts this object to a public response
+		/// </summary>
+		/// <param name="Project">The project instance</param>
+		/// <param name="bIncludeStreams">Whether to include streams in the response</param>
+		/// <param name="bIncludeCategories">Whether to include categories in the response</param>
+		/// <param name="Streams">The list of streams</param>
+		/// <param name="bIncludeAcl">Whether to include the ACL in the response</param>
+		/// <returns>Response instance</returns>
+		public static GetProjectResponse ToResponse(this IProject Project, bool bIncludeStreams, bool bIncludeCategories, List<IStream>? Streams, bool bIncludeAcl)
+		{
+			List<GetProjectStreamResponse>? StreamResponses = null;
+			if(bIncludeStreams)
+			{
+				StreamResponses = Streams!.ConvertAll(x => new GetProjectStreamResponse(x.Id.ToString(), x.Name));
+			}
+
+			List<GetProjectCategoryResponse>? CategoryResponses = null;
+			if (bIncludeCategories)
+			{
+				CategoryResponses = Project.Categories.ConvertAll(x => new GetProjectCategoryResponse(x));
+				if (Streams != null)
+				{
+					foreach (IStream Stream in Streams)
+					{
+						GetProjectCategoryResponse? CategoryResponse = CategoryResponses.FirstOrDefault(x => MatchCategory(Stream.Name, x));
+						if(CategoryResponse == null)
+						{
+							int Row = (CategoryResponses.Count > 0) ? CategoryResponses.Max(x => x.Row) : 0;
+							if (CategoryResponses.Count(x => x.Row == Row) >= 3)
+							{
+								Row++;
+							}
+
+							StreamCategory OtherCategory = new StreamCategory("Other");
+							OtherCategory.Row = Row;
+							OtherCategory.IncludePatterns.Add(".*");
+
+							CategoryResponse = new GetProjectCategoryResponse(OtherCategory);
+							CategoryResponses.Add(CategoryResponse);
+						}
+						CategoryResponse.Streams!.Add(Stream.Id.ToString());
+					}
+				}
+			}
+
+			GetAclResponse? AclResponse = (bIncludeAcl && Project.Acl != null) ? new GetAclResponse(Project.Acl) : null;
+			return new GetProjectResponse(Project.Id.ToString(), Project.Name, Project.Order, StreamResponses, CategoryResponses, AclResponse);
+		}
+
+		/// <summary>
+		/// Tests if a category response matches a given stream name
+		/// </summary>
+		/// <param name="Name">The stream name</param>
+		/// <param name="Category">The category response</param>
+		/// <returns>True if the category matches</returns>
+		static bool MatchCategory(string Name, GetProjectCategoryResponse Category)
+		{
+			if (Category.IncludePatterns.Any(x => StreamCategory.MatchPattern(Name, x)))
+			{
+				if (!Category.ExcludePatterns.Any(x => StreamCategory.MatchPattern(Name, x)))
+				{
+					return true;
+				}
+			}
+			return false;
+		}
+	}
+
+	/// <summary>
+	/// Projection of a project definition to just include permissions info
+	/// </summary>
+	public interface IProjectPermissions
+	{
+		/// <summary>
+		/// ACL for the project
+		/// </summary>
+		public Acl? Acl { get; }
+	}
+}
diff --git a/Engine/Source/Programs/Horde/HordeServer/Models/ISession.cs b/Engine/Source/Programs/Horde/HordeServer/Models/ISession.cs
new file mode 100644
--- /dev/null
+++ b/Engine/Source/Programs/Horde/HordeServer/Models/ISession.cs
@@ -0,0 +1,49 @@
+// Copyright Epic Games, Inc. All Rights Reserved.
+
+using HordeServer.Api;
+using HordeServer.Utilities;
+using MongoDB.Bson;
+using MongoDB.Bson.Serialization.Attributes;
+using System;
+using System.Collections.Generic;
+using System.Linq;
+using System.Threading.Tasks;
+
+namespace HordeServer.Models
+{
+	/// <summary>
+	/// Information about an agent session.
+	/// </summary>
+	public interface ISession
+	{
+		/// <summary>
+		/// Unique id for this session
+		/// </summary>
+		public ObjectId Id { get; }
+
+		/// <summary>
+		/// The agent id
+		/// </summary>
+		public AgentId AgentId { get; }
+
+		/// <summary>
+		/// Start time for this session
+		/// </summary>
+		public DateTime StartTime { get; }
+
+		/// <summary>
+		/// Finishing time for this session
+		/// </summary>
+		public DateTime? FinishTime { get; }
+
+		/// <summary>
+		/// Properties of this agent at the time the session started
+		/// </summary>
+		public IReadOnlyList<string>? Properties { get; }
+
+		/// <summary>
+		/// Version of the agent software
+		/// </summary>
+		public string? Version { get; }
+	}
+}
diff --git a/Engine/Source/Programs/Horde/HordeServer/Models/IStream.cs b/Engine/Source/Programs/Horde/HordeServer/Models/IStream.cs
new file mode 100644
--- /dev/null
+++ b/Engine/Source/Programs/Horde/HordeServer/Models/IStream.cs
@@ -0,0 +1,704 @@
+// Copyright Epic Games, Inc. All Rights Reserved.
+
+using EpicGames.Core;
+using Google.Protobuf.WellKnownTypes;
+using HordeServer.Api;
+using HordeCommon;
+using HordeServer.Utilities;
+using MongoDB.Bson;
+using MongoDB.Bson.Serialization.Attributes;
+using MongoDB.Driver;
+using System;
+using System.Collections.Generic;
+using System.Diagnostics.CodeAnalysis;
+using System.IO;
+using System.Linq;
+using System.Security.Cryptography;
+using System.Text;
+using System.Text.RegularExpressions;
+using System.Threading.Tasks;
+using System.Runtime.CompilerServices;
+
+namespace HordeServer.Models
+{
+	using PoolId = StringId<IPool>;
+	using ProjectId = StringId<IProject>;
+	using StreamId = StringId<IStream>;
+	using TemplateRefId = StringId<TemplateRef>;
+
+	/// <summary>
+	/// Exception thrown when stream validation fails
+	/// </summary>
+	public class InvalidStreamException : Exception
+	{
+		/// <inheritdoc/>
+		public InvalidStreamException()
+		{
+		}
+
+		/// <inheritdoc/>
+		public InvalidStreamException(string Message) : base(Message)
+		{
+		}
+
+		/// <inheritdoc/>
+		public InvalidStreamException(string Message, Exception InnerEx) : base(Message, InnerEx)
+		{
+		}
+	}
+
+	/// <summary>
+	/// Mapping from a BuildGraph agent type to a set of machines on the farm
+	/// </summary>
+	public class AgentType
+	{
+		/// <summary>
+		/// Name of the pool of agents to use
+		/// </summary>
+		[BsonRequired]
+		public PoolId Pool { get; set; }
+
+		/// <summary>
+		/// Name of the workspace to execute on
+		/// </summary>
+		public string? Workspace { get; set; }
+
+		/// <summary>
+		/// Path to the temporary storage dir
+		/// </summary>
+		public string? TempStorageDir { get; set; }
+
+		/// <summary>
+		/// Environment variables to be set when executing the job
+		/// </summary>
+		public Dictionary<string, string>? Environment { get; set; }
+
+		/// <summary>
+		/// Default constructor
+		/// </summary>
+		[BsonConstructor]
+		private AgentType()
+		{
+		}
+
+		/// <summary>
+		/// Constructor
+		/// </summary>
+		/// <param name="Pool">The pool for this agent</param>
+		/// <param name="Workspace">Name of the workspace to use</param>
+		/// <param name="TempStorageDir">Path to the temp storage directory</param>
+		public AgentType(PoolId Pool, string Workspace, string? TempStorageDir)
+		{
+			this.Pool = Pool;
+			this.Workspace = Workspace;
+			this.TempStorageDir = TempStorageDir;
+		}
+
+		/// <summary>
+		/// Constructor
+		/// </summary>
+		/// <param name="Request">The object to construct from</param>
+		public AgentType(CreateAgentTypeRequest Request)
+		{
+			Pool = new PoolId(Request.Pool);
+			Workspace = Request.Workspace;
+			TempStorageDir = Request.TempStorageDir;
+			Environment = Request.Environment;
+		}
+
+		/// <summary>
+		/// Constructs an AgentType object from an optional request
+		/// </summary>
+		/// <param name="Request">The request object</param>
+		/// <returns>New agent type object</returns>
+		[return: NotNullIfNotNull("Request")]
+		public static AgentType? FromRequest(CreateAgentTypeRequest? Request)
+		{
+			return (Request != null) ? new AgentType(Request) : null;
+		}
+
+		/// <summary>
+		/// Creates an API response object from this stream
+		/// </summary>
+		/// <returns>The response object</returns>
+		public Api.GetAgentTypeResponse ToApiResponse()
+		{
+			return new Api.GetAgentTypeResponse(Pool.ToString(), Workspace, TempStorageDir, Environment);
+		}
+
+		/// <summary>
+		/// Creates an API response object from this stream
+		/// </summary>
+		/// <returns>The response object</returns>
+		public HordeCommon.Rpc.GetAgentTypeResponse ToRpcResponse()
+		{
+			HordeCommon.Rpc.GetAgentTypeResponse Response = new HordeCommon.Rpc.GetAgentTypeResponse();
+			if (TempStorageDir != null)
+			{
+				Response.TempStorageDir = TempStorageDir;
+			}
+			if (Environment != null)
+			{
+				Response.Environment.Add(Environment);
+			}
+			return Response;
+		}
+	}
+
+	/// <summary>
+	/// Information about a workspace type
+	/// </summary>
+	public class WorkspaceType
+	{
+		/// <summary>
+		/// Name of the Perforce cluster to use
+		/// </summary>
+		[BsonIgnoreIfNull]
+		public string? Cluster { get; set; }
+
+		/// <summary>
+		/// The Perforce server and port
+		/// </summary>
+		[BsonIgnoreIfNull]
+		public string? ServerAndPort { get; set; }
+
+		/// <summary>
+		/// The Perforce username for syncing this workspace
+		/// </summary>
+		[BsonIgnoreIfNull]
+		public string? UserName { get; set; }
+
+		/// <summary>
+		/// The Perforce password for syncing this workspace
+		/// </summary>
+		[BsonIgnoreIfNull]
+		public string? Password { get; set; }
+
+		/// <summary>
+		/// Identifier to distinguish this workspace from other workspaces. Defaults to the workspace type name.
+		/// </summary>
+		[BsonIgnoreIfNull]
+		public string? Identifier { get; set; }
+
+		/// <summary>
+		/// Override for the stream to sync
+		/// </summary>
+		[BsonIgnoreIfNull]
+		public string? Stream { get; set; }
+
+		/// <summary>
+		/// Custom view for the workspace
+		/// </summary>
+		[BsonIgnoreIfNull]
+		public List<string>? View { get; set; }
+
+		/// <summary>
+		/// Whether to use an incrementally synced workspace
+		/// </summary>
+		public bool Incremental { get; set; }
+
+		/// <summary>
+		/// Default constructor
+		/// </summary>
+		public WorkspaceType()
+		{
+		}
+
+		/// <summary>
+		/// Constructor
+		/// </summary>
+		/// <param name="Request">The object to construct from</param>
+		public WorkspaceType(CreateWorkspaceTypeRequest Request)
+		{
+			Cluster = Request.Cluster;
+			ServerAndPort = Request.ServerAndPort;
+			UserName = Request.UserName;
+			Password = Request.Password;
+			Identifier = Request.Identifier;
+			Stream = Request.Stream;
+			View = Request.View;
+			Incremental = Request.Incremental;
+		}
+
+		/// <summary>
+		/// Constructs an AgentType object from an optional request
+		/// </summary>
+		/// <param name="Request">The request object</param>
+		/// <returns>New agent type object</returns>
+		[return: NotNullIfNotNull("Request")]
+		public static WorkspaceType? FromRequest(CreateWorkspaceTypeRequest? Request)
+		{
+			return (Request != null) ? new WorkspaceType(Request) : null;
+		}
+
+		/// <summary>
+		/// Creates an API response object from this stream
+		/// </summary>
+		/// <returns>The response object</returns>
+		public Api.GetWorkspaceTypeResponse ToApiResponse()
+		{
+			return new Api.GetWorkspaceTypeResponse(Cluster, ServerAndPort, UserName, Identifier, Stream, View, Incremental);
+		}
+	}
+
+	/// <summary>
+	/// Allows triggering another downstream job on succesful completion of a step or aggregate
+	/// </summary>
+	public class ChainedJobTemplate
+	{
+		/// <summary>
+		/// Name of the target that needs to complete successfully
+		/// </summary>
+		public string Trigger { get; set; } = String.Empty;
+
+		/// <summary>
+		/// The new template to trigger
+		/// </summary>
+		public TemplateRefId TemplateRefId { get; set; }
+
+		/// <summary>
+		/// Default constructor for serialization
+		/// </summary>
+		private ChainedJobTemplate()
+		{
+		}
+
+		/// <summary>
+		/// Constructor
+		/// </summary>
+		/// <param name="Trigger">Name of the target that needs to complete</param>
+		/// <param name="TemplateRefId">The new template to trigger</param>
+		public ChainedJobTemplate(string Trigger, TemplateRefId TemplateRefId)
+		{
+			this.Trigger = Trigger;
+			this.TemplateRefId = TemplateRefId;
+		}
+
+		/// <summary>
+		/// Constructor
+		/// </summary>
+		/// <param name="Request">Request to construct from</param>
+		public ChainedJobTemplate(CreateChainedJobTemplateRequest Request)
+			: this(Request.Trigger, new TemplateRefId(Request.TemplateId))
+		{
+		}
+	}
+
+	/// <summary>
+	/// Reference to a template
+	/// </summary>
+	public class TemplateRef
+	{
+		/// <summary>
+		/// The template name (duplicated from the template object)
+		/// </summary>
+		public string Name { get; set; }
+
+		/// <summary>
+		/// Hash of the template definition
+		/// </summary>
+		public ContentHash Hash { get; set; }
+
+		/// <summary>
+		/// Whether to show badges in UGS for this schedule
+		/// </summary>
+		public bool ShowUgsBadges { get; set; }
+
+		/// <summary>
+		/// Whether to show desktop alerts for build health issues created from jobs this type
+		/// </summary>
+		public bool ShowUgsAlerts { get; set; }
+
+		/// <summary>
+		/// Notification channel for this template. Overrides the stream channel if set.
+		/// </summary>
+		public string? NotificationChannel { get; set; }
+
+		/// <summary>
+		/// Notification channel filter for this template. Errors|Warnings|Success
+		/// </summary>
+		public string? NotificationChannelFilter { get; set; }
+
+		/// <summary>
+		/// Channel for triage notification messages
+		/// </summary>
+		public string? TriageChannel { get; set; }
+
+		/// <summary>
+		/// List of schedules for this template
+		/// </summary>
+		[BsonIgnoreIfNull]
+		public Schedule? Schedule { get; set; }
+
+		/// <summary>
+		/// List of downstream templates to trigger at the same change
+		/// </summary>
+		[BsonIgnoreIfNull]
+		public List<ChainedJobTemplate>? ChainedJobs { get; set; }
+
+		/// <summary>
+		/// Custom permissions for this template
+		/// </summary>
+		public Acl? Acl { get; set; }
+
+		/// <summary>
+		/// Private constructor for serialization
+		/// </summary>
+		private TemplateRef()
+		{
+			this.Name = null!;
+			this.Hash = null!;
+		}
+
+		/// <summary>
+		/// Constructor
+		/// </summary>
+		/// <param name="Template">The template being referenced</param>
+		/// <param name="ShowUgsBadges">Whether to show badges in UGS for this job</param>
+		/// <param name="ShowUgsAlerts">Whether to show alerts in UGS for this job</param>
+		/// <param name="NotificationChannel">Notification channel for this template</param>
+		/// <param name="NotificationChannelFilter">Notification channel filter for this template</param>
+		/// <param name="TriageChannel"></param>
+		/// <param name="Schedule">Schedule for this template</param>
+		/// <param name="Triggers">List of downstream templates to trigger</param>
+		/// <param name="Acl">ACL for this template</param>
+		public TemplateRef(ITemplate Template, bool ShowUgsBadges = false, bool ShowUgsAlerts = false, string? NotificationChannel = null, string? NotificationChannelFilter = null, string? TriageChannel = null, Schedule? Schedule = null, List<ChainedJobTemplate>? Triggers = null, Acl? Acl = null)
+		{
+			this.Name = Template.Name;
+			this.Hash = Template.Id;
+			this.ShowUgsBadges = ShowUgsBadges;
+			this.ShowUgsAlerts = ShowUgsAlerts;
+			this.NotificationChannel = NotificationChannel;
+			this.NotificationChannelFilter = NotificationChannelFilter;
+			this.TriageChannel = TriageChannel;
+			this.Schedule = Schedule;
+			this.ChainedJobs = Triggers;
+			this.Acl = Acl;
+		}
+	}
+
+	/// <summary>
+	/// Query used to identify a base changelist for a preflight
+	/// </summary>
+	public class ChangeQuery
+	{
+		/// <summary>
+		/// Template to search for
+		/// </summary>
+		public TemplateRefId? TemplateRefId { get; set; }
+
+		/// <summary>
+		/// The target to look at the status for
+		/// </summary>
+		public string? Target { get; set; }
+
+		/// <summary>
+		/// Whether to match a job that contains warnings
+		/// </summary>
+		public List<JobStepOutcome>? Outcomes { get; set; }
+
+		/// <summary>
+		/// Convert to a request object
+		/// </summary>
+		/// <returns></returns>
+		public ChangeQueryRequest ToRequest()
+		{
+			return new ChangeQueryRequest { TemplateId = TemplateRefId?.ToString(), Target = Target, Outcomes = Outcomes };
+		}
+	}
+
+	/// <summary>
+	/// Definition of a query to execute to find the changelist to run a build at
+	/// </summary>
+	public class DefaultPreflight
+	{
+		/// <summary>
+		/// The template id to execute
+		/// </summary>
+		public TemplateRefId? TemplateRefId { get; set; }
+
+		/// <summary>
+		/// Query specifying a changelist to use
+		/// </summary>
+		public ChangeQuery? Change { get; set; }
+
+		/// <summary>
+		/// The job type to query for the change to use
+		/// </summary>
+		[Obsolete("Use Change.TemplateRefId instead")]
+		public TemplateRefId? ChangeTemplateRefId { get; set; }
+
+		/// <summary>
+		/// Constructor
+		/// </summary>
+		/// <param name="TemplateRefId"></param>
+		/// <param name="Change">The job type to query for the change to use</param>
+		public DefaultPreflight(TemplateRefId? TemplateRefId, ChangeQuery? Change)
+		{
+			this.TemplateRefId = TemplateRefId;
+			this.Change = Change;
+		}
+
+		/// <summary>
+		/// Convert to a request object
+		/// </summary>
+		/// <returns></returns>
+		public DefaultPreflightRequest ToRequest()
+		{
+#pragma warning disable CS0618 // Type or member is obsolete
+			ChangeQueryRequest? ChangeRequest = null;
+			if (Change != null)
+			{
+				ChangeRequest = Change.ToRequest();
+			}
+			else if (ChangeTemplateRefId != null)
+			{
+				ChangeRequest = new ChangeQueryRequest { TemplateId = ChangeTemplateRefId?.ToString() };
+			}
+
+			return new DefaultPreflightRequest { TemplateId = TemplateRefId?.ToString(), Change = ChangeRequest, ChangeTemplateId = ChangeRequest?.TemplateId };
+#pragma warning restore CS0618 // Type or member is obsolete
+		}
+	}
+
+	/// <summary>
+	/// Extension methods for template refs
+	/// </summary>
+	static class TemplateRefExtensions
+	{
+		/// <summary>
+		/// Adds a new template ref to a list
+		/// </summary>
+		/// <param name="TemplateRefs">List of template refs</param>
+		/// <param name="TemplateRef">The template ref to add</param>
+		public static void AddRef(this Dictionary<TemplateRefId, TemplateRef> TemplateRefs, TemplateRef TemplateRef)
+		{
+			TemplateRefs.Add(new TemplateRefId(TemplateRef.Name), TemplateRef);
+		}
+	}
+
+	/// <summary>
+	/// Information about a stream
+	/// </summary>
+	public interface IStream
+	{
+		/// <summary>
+		/// Name of the stream.
+		/// </summary>
+		public StreamId Id { get; }
+
+		/// <summary>
+		/// The project that this stream belongs to
+		/// </summary>
+		public ProjectId ProjectId { get; }
+
+		/// <summary>
+		/// The stream name
+		/// </summary>
+		public string Name { get; }
+
+		/// <summary>
+		/// Name of the perforce cluster
+		/// </summary>
+		public string ClusterName { get; }
+
+		/// <summary>
+		/// Path to the configuration file for this stream
+		/// </summary>
+		public string ConfigPath { get; }
+
+		/// <summary>
+		/// The revision of config file used for this stream
+		/// </summary>
+		public string ConfigRevision { get; }
+
+		/// <summary>
+		/// Order to display on the dashboard's drop-down list
+		/// </summary>
+		public int Order { get; }
+
+		/// <summary>
+		/// Notification channel for all jobs in this stream
+		/// </summary>
+		public string? NotificationChannel { get; }
+
+		/// <summary>
+		/// Notification channel filter for all jobs in this stream. Errors|Warnings|Success
+		/// </summary>
+		public string? NotificationChannelFilter { get; }
+
+		/// <summary>
+		/// Channel to post issue triage notifications
+		/// </summary>
+		public string? TriageChannel { get; }
+
+		/// <summary>
+		/// Default template to use for preflights
+		/// </summary>
+		public DefaultPreflight? DefaultPreflight { get; }
+
+		/// <summary>
+		/// List of pages to display in the dashboard
+		/// </summary>
+		public IReadOnlyList<StreamTab> Tabs { get; }
+
+		/// <summary>
+		/// Dictionary of agent types
+		/// </summary>
+		public IReadOnlyDictionary<string, AgentType> AgentTypes { get; }
+
+		/// <summary>
+		/// Dictionary of workspace types
+		/// </summary>
+		public IReadOnlyDictionary<string, WorkspaceType> WorkspaceTypes { get; }
+
+		/// <summary>
+		/// List of templates available for this stream
+		/// </summary>
+		public IReadOnlyDictionary<TemplateRefId, TemplateRef> Templates { get; }
+
+		/// <summary>
+		/// Stream is paused for builds until specified time
+		/// </summary>
+		public DateTime? PausedUntil { get; }
+		
+		/// <summary>
+		/// Comment/reason for why the stream was paused
+		/// </summary>
+		public string? PauseComment { get; }
+
+		/// <summary>
+		/// The ACL for this object
+		/// </summary>
+		public Acl? Acl { get; }
+	}
+
+	/// <summary>
+	/// Extension methods for streams
+	/// </summary>
+	static class StreamExtensions
+	{
+		/// <summary>
+		/// Tries to get an agent workspace definition from the given type name
+		/// </summary>
+		/// <param name="Stream">The stream object</param>
+		/// <param name="AgentType">The agent type</param>
+		/// <param name="Workspace">Receives the agent workspace definition</param>
+		/// <returns>True if the agent type was valid, and an agent workspace could be created</returns>
+		public static bool TryGetAgentWorkspace(this IStream Stream, AgentType AgentType, [NotNullWhen(true)] out AgentWorkspace? Workspace)
+		{
+			// Get the workspace settings
+			if (AgentType.Workspace == null)
+			{
+				// Use the default settings (fast switching workspace, clean 
+				Workspace = new AgentWorkspace(null, null, GetDefaultWorkspaceIdentifier(Stream), Stream.Name, null, false);
+				return true;
+			}
+			else
+			{
+				// Try to get the matching workspace type
+				WorkspaceType? WorkspaceType;
+				if (!Stream.WorkspaceTypes.TryGetValue(AgentType.Workspace, out WorkspaceType))
+				{
+					Workspace = null;
+					return false;
+				}
+
+				// Get the workspace identifier
+				string Identifier;
+				if (WorkspaceType.Identifier != null && !String.IsNullOrEmpty(WorkspaceType.Identifier))
+				{
+					Identifier = WorkspaceType.Identifier;
+				}
+				else if (WorkspaceType.Incremental)
+				{
+					Identifier = $"{Stream.GetEscapedName()}+{AgentType.Workspace}";
+				}
+				else
+				{
+					Identifier = GetDefaultWorkspaceIdentifier(Stream);
+				}
+
+				// Create the new workspace
+				Workspace = new AgentWorkspace(WorkspaceType.Cluster, WorkspaceType.UserName, Identifier, WorkspaceType.Stream ?? Stream.Name, WorkspaceType.View, WorkspaceType.Incremental);
+				return true;
+			}
+		}
+
+		/// <summary>
+		/// The escaped name of this stream. Removes all non-identifier characters.
+		/// </summary>
+		/// <param name="Stream">The stream object</param>
+		/// <returns>Escaped name for the stream</returns>
+		public static string GetEscapedName(this IStream Stream)
+		{
+			return Regex.Replace(Stream.Name, @"[^a-zA-Z0-9_]", "+");
+		}
+
+		/// <summary>
+		/// Gets the default identifier for workspaces created for this stream. Just includes an escaped depot name.
+		/// </summary>
+		/// <param name="Stream">The stream object</param>
+		/// <returns>The default workspace identifier</returns>
+		private static string GetDefaultWorkspaceIdentifier(IStream Stream)
+		{
+			return Regex.Replace(Stream.GetEscapedName(), @"^(\+\+[^+]*).*$", "$1");
+		}
+
+		/// <summary>
+		/// Converts to a public response object
+		/// </summary>
+		/// <param name="Stream">The stream object</param>
+		/// <param name="bIncludeAcl">Whether to include the ACL in the response object</param>
+		/// <param name="ApiTemplateRefs">The template refs for this stream. Passed separately because they have their own ACL.</param>
+		/// <returns>New response instance</returns>
+		public static Api.GetStreamResponse ToApiResponse(this IStream Stream, bool bIncludeAcl, List<GetTemplateRefResponse> ApiTemplateRefs)
+		{
+			List<GetStreamTabResponse> ApiTabs = Stream.Tabs.ConvertAll(x => x.ToResponse());
+			Dictionary<string, GetAgentTypeResponse> ApiAgentTypes = Stream.AgentTypes.ToDictionary(x => x.Key, x => x.Value.ToApiResponse());
+			Dictionary<string, GetWorkspaceTypeResponse> ApiWorkspaceTypes = Stream.WorkspaceTypes.ToDictionary(x => x.Key, x => x.Value.ToApiResponse());
+			GetAclResponse? ApiAcl = (bIncludeAcl && Stream.Acl != null)? new GetAclResponse(Stream.Acl) : null;
+			return new Api.GetStreamResponse(Stream.Id.ToString(), Stream.ProjectId.ToString(), Stream.Name, Stream.ConfigPath, Stream.ConfigRevision, Stream.Order, Stream.NotificationChannel, Stream.NotificationChannelFilter, Stream.TriageChannel, Stream.DefaultPreflight?.ToRequest(), ApiTabs, ApiAgentTypes, ApiWorkspaceTypes, ApiTemplateRefs, ApiAcl, Stream.PausedUntil, Stream.PauseComment);
+		}
+
+		/// <summary>
+		/// Converts to an RPC response object
+		/// </summary>
+		/// <param name="Stream">The stream object</param>
+		/// <returns>New response instance</returns>
+		public static HordeCommon.Rpc.GetStreamResponse ToRpcResponse(this IStream Stream)
+		{
+			HordeCommon.Rpc.GetStreamResponse Response = new HordeCommon.Rpc.GetStreamResponse();
+			Response.Name = Stream.Name;
+			Response.AgentTypes.Add(Stream.AgentTypes.ToDictionary(x => x.Key, x => x.Value.ToRpcResponse()));
+			return Response;
+		}
+
+		/// <summary>
+		/// Check if stream is paused for new builds
+		/// </summary>
+		/// <param name="Stream">The stream object</param>
+		/// <param name="CurrentTime">Current time (allow tests to pass in a fake clock)</param>
+		/// <returns>If stream is paused</returns>
+		public static bool IsPaused(this IStream Stream, DateTime CurrentTime)
+		{
+			return Stream.PausedUntil != null && Stream.PausedUntil > CurrentTime;
+		}
+	}
+
+	/// <summary>
+	/// Projection of a stream definition to just include permissions info
+	/// </summary>
+	public interface IStreamPermissions
+	{
+		/// <summary>
+		/// ACL for the stream
+		/// </summary>
+		public Acl? Acl { get; }
+
+		/// <summary>
+		/// The project containing this stream
+		/// </summary>
+		public ProjectId ProjectId { get; }
+	}
+}
diff --git a/Engine/Source/Programs/Horde/HordeServer/Models/ISubscription.cs b/Engine/Source/Programs/Horde/HordeServer/Models/ISubscription.cs
new file mode 100644
--- /dev/null
+++ b/Engine/Source/Programs/Horde/HordeServer/Models/ISubscription.cs
@@ -0,0 +1,49 @@
+// Copyright Epic Games, Inc. All Rights Reserved.
+
+using HordeCommon;
+using HordeServer.Api;
+using MongoDB.Bson;
+using System;
+using System.Collections.Generic;
+using System.Linq;
+using System.Threading.Tasks;
+
+namespace HordeServer.Models
+{
+	/// <summary>
+	/// The type of notification to send
+	/// </summary>
+	public enum NotificationType
+	{
+		/// <summary>
+		/// Send a DM on Slack
+		/// </summary>
+		Slack,
+	}
+
+	/// <summary>
+	/// Subscription to a type of event
+	/// </summary>
+	public interface ISubscription
+	{
+		/// <summary>
+		/// Unique id for this subscription
+		/// </summary>
+		public string Id { get; }
+
+		/// <summary>
+		/// Name of the event to subscribe to
+		/// </summary>
+		public IEvent Event { get; }
+
+		/// <summary>
+		/// User to notify
+		/// </summary>
+		public ObjectId UserId { get; }
+
+		/// <summary>
+		/// Type of notification to receive
+		/// </summary>
+		public NotificationType NotificationType { get; }
+	}
+}
diff --git a/Engine/Source/Programs/Horde/HordeServer/Models/ITestData.cs b/Engine/Source/Programs/Horde/HordeServer/Models/ITestData.cs
new file mode 100644
--- /dev/null
+++ b/Engine/Source/Programs/Horde/HordeServer/Models/ITestData.cs
@@ -0,0 +1,60 @@
+// Copyright Epic Games, Inc. All Rights Reserved.
+
+using HordeServer.Utilities;
+using MongoDB.Bson;
+using System;
+using System.Collections.Generic;
+using System.Linq;
+using System.Threading.Tasks;
+
+namespace HordeServer.Models
+{
+	using StreamId = StringId<IStream>;
+	using TemplateRefId = StringId<TemplateRef>;
+
+	/// <summary>
+	/// Stores information about the results of a test
+	/// </summary>
+	public interface ITestData
+	{
+		/// <summary>
+		/// Unique id of the test data
+		/// </summary>
+		ObjectId Id { get; }
+
+		/// <summary>
+		/// Stream that generated the test data
+		/// </summary>
+		StreamId StreamId { get; }
+
+		/// <summary>
+		/// The template reference id
+		/// </summary>
+		TemplateRefId TemplateRefId { get; }
+
+		/// <summary>
+		/// The job which produced the data
+		/// </summary>
+		ObjectId JobId { get; }
+
+		/// <summary>
+		/// The step that ran
+		/// </summary>
+		SubResourceId StepId { get; }
+
+		/// <summary>
+		/// The changelist number that contained the data
+		/// </summary>
+		int Change { get; }
+
+		/// <summary>
+		/// Key used to identify the particular data
+		/// </summary>
+		string Key { get; }
+
+		/// <summary>
+		/// The data stored for this test
+		/// </summary>
+		BsonDocument Data { get; }
+	}
+}
diff --git a/Engine/Source/Programs/Horde/HordeServer/Models/Job.cs b/Engine/Source/Programs/Horde/HordeServer/Models/Job.cs
deleted file mode 100644
diff --git a/Engine/Source/Programs/Horde/HordeServer/Models/JobStepRef.cs b/Engine/Source/Programs/Horde/HordeServer/Models/JobStepRef.cs
deleted file mode 100644
diff --git a/Engine/Source/Programs/Horde/HordeServer/Models/JobTiming.cs b/Engine/Source/Programs/Horde/HordeServer/Models/JobTiming.cs
deleted file mode 100644
diff --git a/Engine/Source/Programs/Horde/HordeServer/Models/Lease.cs b/Engine/Source/Programs/Horde/HordeServer/Models/Lease.cs
deleted file mode 100644
diff --git a/Engine/Source/Programs/Horde/HordeServer/Models/LogEvent.cs b/Engine/Source/Programs/Horde/HordeServer/Models/LogEvent.cs
deleted file mode 100644
diff --git a/Engine/Source/Programs/Horde/HordeServer/Models/LogFile.cs b/Engine/Source/Programs/Horde/HordeServer/Models/LogFile.cs
deleted file mode 100644
diff --git a/Engine/Source/Programs/Horde/HordeServer/Models/NotificationTrigger.cs b/Engine/Source/Programs/Horde/HordeServer/Models/NotificationTrigger.cs
deleted file mode 100644
diff --git a/Engine/Source/Programs/Horde/HordeServer/Models/Pool.cs b/Engine/Source/Programs/Horde/HordeServer/Models/Pool.cs
deleted file mode 100644
diff --git a/Engine/Source/Programs/Horde/HordeServer/Models/Project.cs b/Engine/Source/Programs/Horde/HordeServer/Models/Project.cs
deleted file mode 100644
diff --git a/Engine/Source/Programs/Horde/HordeServer/Models/Session.cs b/Engine/Source/Programs/Horde/HordeServer/Models/Session.cs
deleted file mode 100644
diff --git a/Engine/Source/Programs/Horde/HordeServer/Models/Stream.cs b/Engine/Source/Programs/Horde/HordeServer/Models/Stream.cs
deleted file mode 100644
diff --git a/Engine/Source/Programs/Horde/HordeServer/Models/Subscription.cs b/Engine/Source/Programs/Horde/HordeServer/Models/Subscription.cs
deleted file mode 100644
diff --git a/Engine/Source/Programs/Horde/HordeServer/Models/TestData.cs b/Engine/Source/Programs/Horde/HordeServer/Models/TestData.cs
deleted file mode 100644
