Horde: Move Agents functionality into the same folder.

diff --git a/Engine/Source/Programs/Horde/HordeServer/Agents/IAgent.cs b/Engine/Source/Programs/Horde/HordeServer/Agents/IAgent.cs
new file mode 100644
--- /dev/null
+++ b/Engine/Source/Programs/Horde/HordeServer/Agents/IAgent.cs
@@ -0,0 +1,803 @@
+// Copyright Epic Games, Inc. All Rights Reserved.
+
+using Google.Protobuf.WellKnownTypes;
+using EpicGames.Core;
+using HordeServer.Api;
+using HordeCommon;
+using HordeCommon.Rpc.Tasks;
+using HordeServer.Services;
+using HordeServer.Utilities;
+using MongoDB.Bson;
+using MongoDB.Bson.Serialization;
+using MongoDB.Bson.Serialization.Attributes;
+using System;
+using System.Collections.Generic;
+using System.Diagnostics.CodeAnalysis;
+using System.IO;
+using System.Linq;
+using System.Security.Cryptography;
+using Serilog;
+using System.Diagnostics;
+using System.Threading.Tasks;
+using HordeServer.Collections;
+using EpicGames.Horde.Compute;
+using EpicGames.Horde.Common;
+
+namespace HordeServer.Models
+{
+	using PoolId = StringId<IPool>;
+	using StreamId = StringId<IStream>;
+	using AgentSoftwareVersion = StringId<IAgentSoftwareCollection>;
+	using AgentSoftwareChannelName = StringId<AgentSoftwareChannels>;
+
+	/// <summary>
+	/// Information about a workspace synced to an agent
+	/// </summary>
+	public class AgentWorkspace
+	{
+		/// <summary>
+		/// Name of the Perforce cluster to use
+		/// </summary>
+		public string? Cluster { get; set; }
+
+		/// <summary>
+		/// User to log into Perforce with (eg. buildmachine)
+		/// </summary>
+		public string? UserName { get; set; }
+
+		/// <summary>
+		/// Identifier to distinguish this workspace from other workspaces
+		/// </summary>
+		public string Identifier { get; set; }
+
+		/// <summary>
+		/// The stream to sync
+		/// </summary>
+		public string Stream { get; set; }
+
+		/// <summary>
+		/// Custom view for the workspace
+		/// </summary>
+		public List<string>? View { get; set; }
+
+		/// <summary>
+		/// Whether to use an incremental workspace
+		/// </summary>
+		public bool bIncremental { get; set; }
+
+		/// <summary>
+		/// Constructor
+		/// </summary>
+		/// <param name="Cluster">Name of the Perforce cluster</param>
+		/// <param name="UserName">User to log into Perforce with (eg. buildmachine)</param>
+		/// <param name="Identifier">Identifier to distinguish this workspace from other workspaces</param>
+		/// <param name="Stream">The stream to sync</param>
+		/// <param name="View">Custom view for the workspace</param>
+		/// <param name="bIncremental">Whether to use an incremental workspace</param>
+		public AgentWorkspace(string? Cluster, string? UserName, string Identifier, string Stream, List<string>? View, bool bIncremental)
+		{
+			if (!String.IsNullOrEmpty(Cluster))
+			{
+				this.Cluster = Cluster;
+			}
+			if (!String.IsNullOrEmpty(UserName))
+			{
+				this.UserName = UserName;
+			}
+			this.Identifier = Identifier;
+			this.Stream = Stream;
+			this.View = View;
+			this.bIncremental = bIncremental;
+		}
+
+		/// <summary>
+		/// Constructor
+		/// </summary>
+		/// <param name="Workspace">RPC message to construct from</param>
+		public AgentWorkspace(HordeCommon.Rpc.Messages.AgentWorkspace Workspace)
+			: this(Workspace.ConfiguredCluster, Workspace.ConfiguredUserName, Workspace.Identifier, Workspace.Stream, (Workspace.View.Count > 0) ? Workspace.View.ToList() : null, Workspace.Incremental)
+		{
+		}
+
+		/// <summary>
+		/// Gets a digest of the settings for this workspace
+		/// </summary>
+		/// <returns>Digest for the workspace settings</returns>
+		public string GetDigest()
+		{
+#pragma warning disable CA5351 // Do Not Use Broken Cryptographic Algorithms
+			using (MD5 Hasher = MD5.Create())
+			{
+				byte[] Data = BsonExtensionMethods.ToBson(this);
+				return BitConverter.ToString(Hasher.ComputeHash(Data)).Replace("-", "", StringComparison.Ordinal);
+			}
+#pragma warning restore CA5351
+		}
+
+		/// <inheritdoc/>
+		public override bool Equals(object? Obj)
+		{
+			AgentWorkspace? Other = Obj as AgentWorkspace;
+			if (Other == null)
+			{
+				return false;
+			}
+			if (Cluster != Other.Cluster || UserName != Other.UserName || Identifier != Other.Identifier || Stream != Other.Stream || bIncremental != Other.bIncremental)
+			{
+				return false;
+			}
+			if (!Enumerable.SequenceEqual(View ?? new List<string>(), Other.View ?? new List<string>()))
+			{
+				return false;
+			}
+			return true;
+		}
+
+		/// <inheritdoc/>
+		public override int GetHashCode()
+		{
+			return HashCode.Combine(Cluster, UserName, Identifier, Stream, bIncremental); // Ignore 'View' for now
+		}
+
+		/// <summary>
+		/// Checks if two workspace sets are equivalent, ignoring order
+		/// </summary>
+		/// <param name="WorkspacesA">First list of workspaces</param>
+		/// <param name="WorkspacesB">Second list of workspaces</param>
+		/// <returns>True if the sets are equivalent</returns>
+		public static bool SetEquals(IReadOnlyList<AgentWorkspace> WorkspacesA, IReadOnlyList<AgentWorkspace> WorkspacesB)
+		{
+			HashSet<AgentWorkspace> WorkspacesSetA = new HashSet<AgentWorkspace>(WorkspacesA);
+			return WorkspacesSetA.SetEquals(WorkspacesB);
+		}
+
+		/// <summary>
+		/// Converts this workspace to an RPC message
+		/// </summary>
+		/// <param name="Server">The Perforce server</param>
+		/// <param name="Credentials">Credentials for the server</param>
+		/// <returns>The RPC message</returns>
+		public HordeCommon.Rpc.Messages.AgentWorkspace ToRpcMessage(IPerforceServer Server, PerforceCredentials? Credentials)
+		{
+			// Construct the message
+			HordeCommon.Rpc.Messages.AgentWorkspace Result = new HordeCommon.Rpc.Messages.AgentWorkspace();
+			Result.ConfiguredCluster = Cluster;
+			Result.ConfiguredUserName = UserName;
+			Result.ServerAndPort = Server.ServerAndPort;
+			Result.UserName = Credentials?.UserName ?? UserName;
+			Result.Password = Credentials?.Password;
+			Result.Identifier = Identifier;
+			Result.Stream = Stream;
+			if (View != null)
+			{
+				Result.View.AddRange(View);
+			}
+			Result.Incremental = bIncremental;
+			return Result;
+		}
+	}
+
+	/// <summary>
+	/// Document describing an active lease
+	/// </summary>
+	public class AgentLease
+	{
+		/// <summary>
+		/// Lease returned when a task source wants to run, but needs to wait for other tasks to finish first.
+		/// </summary>
+		public static AgentLease Drain { get; } = new AgentLease(ObjectId.Empty, String.Empty, null, null, null, LeaseState.Pending, null, true, null);
+
+		/// <summary>
+		/// Name of this lease
+		/// </summary>
+		[BsonRequired]
+		public ObjectId Id { get; set; }
+
+		/// <summary>
+		/// Name of this lease
+		/// </summary>
+		public string Name { get; set; }
+
+		/// <summary>
+		/// The current state of the lease
+		/// </summary>
+		public LeaseState State { get; set; }
+
+		/// <summary>
+		/// The stream for this lease
+		/// </summary>
+		public StreamId? StreamId { get; set; }
+
+		/// <summary>
+		/// The pool for this lease
+		/// </summary>
+		public PoolId? PoolId { get; set; }
+
+		/// <summary>
+		/// Optional log for this lease
+		/// </summary>
+		public ObjectId? LogId { get; set; }
+
+		/// <summary>
+		/// Time at which the lease started
+		/// </summary>
+		[BsonRequired]
+		public DateTime StartTime { get; set; }
+
+		/// <summary>
+		/// Time at which the lease should be terminated
+		/// </summary>
+		public DateTime? ExpiryTime { get; set; }
+
+		/// <summary>
+		/// Flag indicating whether this lease has been accepted by the agent
+		/// </summary>
+		public bool Active { get; set; }
+
+		/// <summary>
+		/// Resources used by this lease
+		/// </summary>
+		public IReadOnlyDictionary<string, int>? Resources { get; set; }
+
+		/// <summary>
+		/// Whether the lease requires exclusive access to the agent
+		/// </summary>
+		public bool Exclusive { get; set; }
+
+		/// <summary>
+		/// For leases in the pending state, encodes an "any" protobuf containing the payload for the agent to execute the lease.
+		/// </summary>
+		public byte[]? Payload { get; set; }
+
+		/// <summary>
+		/// Private constructor
+		/// </summary>
+		[BsonConstructor]
+		private AgentLease()
+		{
+			Name = String.Empty;
+		}
+
+		/// <summary>
+		/// Constructor
+		/// </summary>
+		/// <param name="Id">Identifier for the lease</param>
+		/// <param name="Name">Name of this lease</param>
+		/// <param name="StreamId"></param>
+		/// <param name="PoolId"></param>
+		/// <param name="LogId">Unique id for the log</param>
+		/// <param name="State">State for the lease</param>
+		/// <param name="Resources">Resources required for this lease</param>
+		/// <param name="Exclusive">Whether to reserve the entire device</param>
+		/// <param name="Payload">Encoded "any" protobuf describing the contents of the payload</param>
+		public AgentLease(ObjectId Id, string Name, StreamId? StreamId, PoolId? PoolId, ObjectId? LogId, LeaseState State, IReadOnlyDictionary<string, int>? Resources, bool Exclusive, byte[]? Payload)
+		{
+			this.Id = Id;
+			this.Name = Name;
+			this.StreamId = StreamId;
+			this.PoolId = PoolId;
+			this.LogId = LogId;
+			this.State = State;
+			this.Resources = Resources;
+			this.Exclusive = Exclusive;
+			this.Payload = Payload;
+			StartTime = DateTime.UtcNow;
+		}
+
+		/// <summary>
+		/// Determines if this is a conform lease
+		/// </summary>
+		/// <returns>True if this is a conform lease</returns>
+		public bool IsConformLease()
+		{
+			if (Payload != null)
+			{
+				Any BasePayload = Any.Parser.ParseFrom(Payload);
+				if (BasePayload.Is(ConformTask.Descriptor))
+				{
+					return true;
+				}
+			}
+			return false;
+		}
+
+		/// <summary>
+		/// Gets user-readable payload information
+		/// </summary>
+		/// <param name="Payload">The payload data</param>
+		/// <returns>Dictionary of key/value pairs for the payload</returns>
+		public static Dictionary<string, string>? GetPayloadDetails(ReadOnlyMemory<byte>? Payload)
+		{
+			Dictionary<string, string>? Details = null;
+			if (Payload != null)
+			{
+				Any BasePayload = Any.Parser.ParseFrom(Payload.Value.ToArray());
+
+				Details = new Dictionary<string, string>();
+
+				ConformTask ConformTask;
+				if(BasePayload.TryUnpack(out ConformTask))
+				{
+					Details["Type"] = "Conform";
+					Details["LogId"] = ConformTask.LogId;
+				}
+
+				ExecuteJobTask JobTask;
+				if (BasePayload.TryUnpack(out JobTask))
+				{
+					Details["Type"] = "Job";
+					Details["JobId"] = JobTask.JobId;
+					Details["BatchId"] = JobTask.BatchId;
+					Details["LogId"] = JobTask.LogId;
+				}
+
+				UpgradeTask UpgradeTask;
+				if (BasePayload.TryUnpack(out UpgradeTask))
+				{
+					Details["Type"] = "Upgrade";
+					Details["SoftwareId"] = UpgradeTask.SoftwareId;
+					Details["LogId"] = UpgradeTask.LogId;
+				}
+			}
+			return Details;
+		}
+
+		/// <summary>
+		/// Converts this lease to an RPC message
+		/// </summary>
+		/// <returns>RPC message</returns>
+		public HordeCommon.Rpc.Messages.Lease ToRpcMessage()
+		{
+			HordeCommon.Rpc.Messages.Lease Lease = new HordeCommon.Rpc.Messages.Lease();
+			Lease.Id = Id.ToString();
+			Lease.Payload = Google.Protobuf.WellKnownTypes.Any.Parser.ParseFrom(Payload); 
+			Lease.State = State;
+			return Lease;
+		}
+	}
+
+	/// <summary>
+	/// Well-known property names for agents
+	/// </summary>
+	static class KnownPropertyNames
+	{
+		/// <summary>
+		/// The agent id
+		/// </summary>
+		public const string Id = "Id";
+
+		/// <summary>
+		/// The operating system (Linux, MacOS, Windows)
+		/// </summary>
+		public const string OSFamily = "OSFamily";
+
+		/// <summary>
+		/// Pools that this agent belongs to
+		/// </summary>
+		public const string Pool = "Pool";
+
+		/// <summary>
+		/// Number of logical cores
+		/// </summary>
+		public const string LogicalCores = "LogicalCores";
+
+		/// <summary>
+		/// Amount of RAM, in GB
+		/// </summary>
+		public const string RAM = "RAM";
+	}
+
+	/// <summary>
+	/// Mirrors an Agent document in the database
+	/// </summary>
+	public interface IAgent
+	{
+		/// <summary>
+		/// Randomly generated unique id for this agent.
+		/// </summary>
+		public AgentId Id { get; }
+
+		/// <summary>
+		/// The current session id, if it's online
+		/// </summary>
+		public ObjectId? SessionId { get; }
+
+		/// <summary>
+		/// Time at which the current session expires. 
+		/// </summary>
+		public DateTime? SessionExpiresAt { get; }
+
+		/// <summary>
+		/// Current status of this agent
+		/// </summary>
+		public AgentStatus Status { get; }
+
+		/// <summary>
+		/// Whether the agent is enabled
+		/// </summary>
+		public bool Enabled { get; }
+
+		/// <summary>
+		/// Whether the agent should be included on the dashboard. This is set to true for ephemeral agents once they are no longer online, or agents that are explicitly deleted.
+		/// </summary>
+		public bool Deleted { get; }
+
+		/// <summary>
+		/// Version of the software running on this agent
+		/// </summary>
+		public string? Version { get; }
+
+		/// <summary>
+		/// Arbitrary comment for the agent (useful for disable reasons etc)
+		/// </summary>
+		public string? Comment { get; }
+
+		/// <summary>
+		/// List of properties for this agent
+		/// </summary>
+		public IReadOnlyList<string> Properties { get; }
+
+		/// <summary>
+		/// List of resources available to the agent
+		/// </summary>
+		public IReadOnlyDictionary<string, int> Resources { get; }
+
+		/// <summary>
+		/// Channel for the software running on this agent. Uses <see cref="AgentSoftwareService.DefaultChannelName"/> if not specified
+		/// </summary>
+		public AgentSoftwareChannelName? Channel { get; }
+
+		/// <summary>
+		/// Last upgrade that was attempted
+		/// </summary>
+		public string? LastUpgradeVersion { get; }
+
+		/// <summary>
+		/// Time that which the last upgrade was attempted
+		/// </summary>
+		public DateTime? LastUpgradeTime { get; }
+
+		/// <summary>
+		/// Dynamically applied pools
+		/// </summary>
+		public IReadOnlyList<PoolId> DynamicPools { get; }
+
+		/// <summary>
+		/// List of manually assigned pools for agent
+		/// </summary>
+		public IReadOnlyList<PoolId> ExplicitPools { get; }
+
+		/// <summary>
+		/// Whether a conform is requested
+		/// </summary>
+		public bool RequestConform { get; }
+
+		/// <summary>
+		/// Whether a machine restart is requested
+		/// </summary>
+		public bool RequestRestart { get; }
+
+		/// <summary>
+		/// Whether the machine should be shutdown
+		/// </summary>
+		public bool RequestShutdown { get; }
+
+		/// <summary>
+		/// List of workspaces currently synced to this machine
+		/// </summary>
+		public IReadOnlyList<AgentWorkspace> Workspaces { get; }
+
+		/// <summary>
+		/// Time at which the last conform job ran
+		/// </summary>
+		public DateTime LastConformTime { get; }
+
+		/// <summary>
+		/// Number of times a conform job has failed
+		/// </summary>
+		public int? ConformAttemptCount { get; }
+
+		/// <summary>
+		/// Array of active leases.
+		/// </summary>
+		public IReadOnlyList<AgentLease> Leases { get; }
+
+		/// <summary>
+		/// ACL for modifying this agent
+		/// </summary>
+		public Acl? Acl { get; }
+
+		/// <summary>
+		/// Last time that the agent was modified
+		/// </summary>
+		public DateTime UpdateTime { get; }
+
+		/// <summary>
+		/// Update counter for this document. Any updates should compare-and-swap based on the value of this counter, or increment it in the case of server-side updates.
+		/// </summary>
+		public uint UpdateIndex { get; }
+	}
+
+	/// <summary>
+	/// Extension methods for IAgent
+	/// </summary>
+	public static class AgentExtensions
+	{
+		/// <summary>
+		/// Determines whether this agent is online
+		/// </summary>
+		/// <returns></returns>
+		public static bool IsSessionValid(this IAgent Agent, DateTime UtcNow)
+		{
+			return Agent.SessionId.HasValue && Agent.SessionExpiresAt.HasValue && UtcNow < Agent.SessionExpiresAt.Value;
+		}
+
+		/// <summary>
+		/// Tests whether an agent is in the given pool
+		/// </summary>
+		/// <param name="Agent"></param>
+		/// <param name="PoolId"></param>
+		/// <returns></returns>
+		public static bool IsInPool(this IAgent Agent, PoolId PoolId)
+		{
+			return Agent.DynamicPools.Contains(PoolId) || Agent.ExplicitPools.Contains(PoolId);
+		}
+
+		/// <summary>
+		/// Get all the pools for each agent
+		/// </summary>
+		/// <param name="Agent">The agent to query</param>
+		/// <returns></returns>
+		public static IEnumerable<PoolId> GetPools(this IAgent Agent)
+		{
+			foreach (PoolId PoolId in Agent.DynamicPools)
+			{
+				yield return PoolId;
+			}
+			foreach (PoolId PoolId in Agent.ExplicitPools)
+			{
+				yield return PoolId;
+			}
+		}
+
+		/// <summary>
+		/// Tests whether an agent has a particular property
+		/// </summary>
+		/// <param name="Agent"></param>
+		/// <param name="Property"></param>
+		/// <returns></returns>
+		public static bool HasProperty(this IAgent Agent, string Property)
+		{
+			return Agent.Properties.BinarySearch(Property, StringComparer.OrdinalIgnoreCase) >= 0;
+		}
+
+		/// <summary>
+		/// Finds property values from a sorted list of Name=Value pairs
+		/// </summary>
+		/// <param name="Agent">The agent to query</param>
+		/// <param name="Name">Name of the property to find</param>
+		/// <returns>Property values</returns>
+		public static IEnumerable<string> GetPropertyValues(this IAgent Agent, string Name)
+		{
+			if (Name.Equals(KnownPropertyNames.Id, StringComparison.OrdinalIgnoreCase))
+			{
+				yield return Agent.Id.ToString();
+			}
+			else if (Name.Equals(KnownPropertyNames.Pool, StringComparison.OrdinalIgnoreCase))
+			{
+				foreach (PoolId PoolId in Agent.GetPools())
+				{
+					yield return PoolId.ToString();
+				}
+			}
+			else
+			{
+				int Index = Agent.Properties.BinarySearch(Name, StringComparer.OrdinalIgnoreCase);
+				if (Index < 0)
+				{
+					Index = ~Index;
+					while (Index < Agent.Properties.Count)
+					{
+						string Property = Agent.Properties[Index];
+						if (Property.Length <= Name.Length || !Property.StartsWith(Name, StringComparison.OrdinalIgnoreCase) || Property[Name.Length] != '=')
+						{
+							break;
+						}
+						yield return Property.Substring(Name.Length + 1);
+					}
+				}
+			}
+		}
+
+		/// <summary>
+		/// Evaluates a condition against an agent
+		/// </summary>
+		/// <param name="Agent">The agent to evaluate</param>
+		/// <param name="Condition">The condition to evaluate</param>
+		/// <returns>True if the agent satisfies the condition</returns>
+		public static bool SatisfiesCondition(this IAgent Agent, Condition Condition)
+		{
+			return Condition.Evaluate(x => Agent.GetPropertyValues(x));
+		}
+
+		/// <summary>
+		/// Determine whether it's possible to add a lease for the given resources
+		/// </summary>
+		/// <param name="Agent">The agent to create a lease for</param>
+		/// <param name="Requirements">Requirements for the lease</param>
+		/// <returns>True if the new lease can be granted</returns>
+		public static bool MeetsRequirements(this IAgent Agent, Requirements Requirements)
+		{
+			return MeetsRequirements(Agent, Requirements.Condition, Requirements.Resources, Requirements.Exclusive);
+		}
+
+		/// <summary>
+		/// Determine whether it's possible to add a lease for the given resources
+		/// </summary>
+		/// <param name="Agent">The agent to create a lease for</param>
+		/// <param name="Exclusive">Whether t</param>
+		/// <param name="Condition">Condition to satisfy</param>
+		/// <param name="Resources">Resources required to execute</param>
+		/// <returns>True if the new lease can be granted</returns>
+		public static bool MeetsRequirements(this IAgent Agent, Condition? Condition, Dictionary<string, int>? Resources, bool Exclusive)
+		{
+			if (!Agent.Enabled || Agent.Status != AgentStatus.Ok)
+			{
+				return false;
+			}
+			if (Agent.Leases.Any(x => x.Exclusive))
+			{
+				return false;
+			}
+			if (Exclusive && Agent.Leases.Any())
+			{
+				return false;
+			}
+			if (Condition != null && !Agent.SatisfiesCondition(Condition))
+			{
+				return false;
+			}
+			if (Resources != null)
+			{
+				foreach ((string Name, int Count) in Resources)
+				{
+					int RemainingCount;
+					if (!Agent.Resources.TryGetValue(Name, out RemainingCount))
+					{
+						return false;
+					}
+					foreach (AgentLease Lease in Agent.Leases)
+					{
+						if (Lease.Resources != null)
+						{
+							int LeaseCount;
+							Lease.Resources.TryGetValue(Name, out LeaseCount);
+							RemainingCount -= LeaseCount;
+						}
+					}
+					if (RemainingCount < Count)
+					{
+						return false;
+					}
+				}
+			}
+			return true;
+		}
+
+		/// <summary>
+		/// Gets all the autosdk workspaces required for an agent
+		/// </summary>
+		/// <param name="Agent"></param>
+		/// <param name="Globals"></param>
+		/// <param name="Workspaces"></param>
+		/// <returns></returns>
+		public static HashSet<AgentWorkspace> GetAutoSdkWorkspaces(this IAgent Agent, Globals Globals, List<AgentWorkspace> Workspaces)
+		{
+			HashSet<AgentWorkspace> AutoSdkWorkspaces = new HashSet<AgentWorkspace>();
+			foreach (string? ClusterName in Workspaces.Select(x => x.Cluster).Distinct())
+			{
+				PerforceCluster? Cluster = Globals.FindPerforceCluster(ClusterName);
+				if (Cluster != null)
+				{
+					AgentWorkspace? AutoSdkWorkspace = GetAutoSdkWorkspace(Agent, Cluster);
+					if (AutoSdkWorkspace != null)
+					{
+						AutoSdkWorkspaces.Add(AutoSdkWorkspace);
+					}
+				}
+			}
+			return AutoSdkWorkspaces;
+		}
+
+		/// <summary>
+		/// Get the AutoSDK workspace required for an agent
+		/// </summary>
+		/// <param name="Agent"></param>
+		/// <param name="Cluster">The perforce cluster to get a workspace for</param>
+		/// <returns></returns>
+		public static AgentWorkspace? GetAutoSdkWorkspace(this IAgent Agent, PerforceCluster Cluster)
+		{
+			foreach (AutoSdkWorkspace AutoSdk in Cluster.AutoSdk)
+			{
+				if (AutoSdk.Stream != null && AutoSdk.Properties.All(x => Agent.Properties.Contains(x)))
+				{
+					return new AgentWorkspace(Cluster.Name, AutoSdk.UserName, AutoSdk.Name ?? "AutoSDK", AutoSdk.Stream!, null, true);
+				}
+			}
+			return null;
+		}
+
+		/// <summary>
+		/// Converts this workspace to an RPC message
+		/// </summary>
+		/// <param name="Agent">The agent to get a workspace for</param>
+		/// <param name="Workspace">The workspace definition</param>
+		/// <param name="Cluster">The global state</param>
+		/// <param name="LoadBalancer">The Perforce load balancer</param>
+		/// <param name="WorkspaceMessages">List of messages</param>
+		/// <returns>The RPC message</returns>
+		public static async Task<bool> TryAddWorkspaceMessage(this IAgent Agent, AgentWorkspace Workspace, PerforceCluster Cluster, PerforceLoadBalancer LoadBalancer, IList<HordeCommon.Rpc.Messages.AgentWorkspace> WorkspaceMessages)
+		{
+			// Find a matching server, trying to use a previously selected one if possible
+			string? BaseServerAndPort;
+			string? ServerAndPort;
+
+			HordeCommon.Rpc.Messages.AgentWorkspace? ExistingWorkspace = WorkspaceMessages.FirstOrDefault(x => x.ConfiguredCluster == Workspace.Cluster);
+			if(ExistingWorkspace != null)
+			{
+				BaseServerAndPort = ExistingWorkspace.BaseServerAndPort;
+				ServerAndPort = ExistingWorkspace.ServerAndPort;
+			}
+			else
+			{
+				if (Cluster == null)
+				{
+					return false;
+				}
+
+				IPerforceServer? Server = await LoadBalancer.SelectServerAsync(Cluster, Agent);
+				if (Server == null)
+				{
+					return false;
+				}
+
+				BaseServerAndPort = Server.BaseServerAndPort;
+				ServerAndPort = Server.ServerAndPort;
+			}
+
+			// Find the matching credentials for the desired user
+			PerforceCredentials? Credentials = null;
+			if (Cluster != null)
+			{
+				if (Workspace.UserName == null)
+				{
+					Credentials = Cluster.Credentials.FirstOrDefault();
+				}
+				else
+				{
+					Credentials = Cluster.Credentials.FirstOrDefault(x => String.Equals(x.UserName, Workspace.UserName, StringComparison.OrdinalIgnoreCase));
+				}
+			}
+
+			// Construct the message
+			HordeCommon.Rpc.Messages.AgentWorkspace Result = new HordeCommon.Rpc.Messages.AgentWorkspace();
+			Result.ConfiguredCluster = Workspace.Cluster;
+			Result.ConfiguredUserName = Workspace.UserName;
+			Result.Cluster = Cluster?.Name;
+			Result.BaseServerAndPort = BaseServerAndPort;
+			Result.ServerAndPort = ServerAndPort;
+			Result.UserName = Credentials?.UserName ?? Workspace.UserName;
+			Result.Password = Credentials?.Password;
+			Result.Identifier = Workspace.Identifier;
+			Result.Stream = Workspace.Stream;
+			if (Workspace.View != null)
+			{
+				Result.View.AddRange(Workspace.View);
+			}
+			Result.Incremental = Workspace.bIncremental;
+			WorkspaceMessages.Add(Result);
+
+			return true;
+		}
+	}
+}
diff --git a/Engine/Source/Programs/Horde/HordeServer/Agents/IAgentCollection.cs b/Engine/Source/Programs/Horde/HordeServer/Agents/IAgentCollection.cs
new file mode 100644
--- /dev/null
+++ b/Engine/Source/Programs/Horde/HordeServer/Agents/IAgentCollection.cs
@@ -0,0 +1,167 @@
+// Copyright Epic Games, Inc. All Rights Reserved.
+
+using Google.Protobuf.WellKnownTypes;
+using HordeServer.Api;
+using HordeCommon;
+using HordeServer.Models;
+using HordeCommon.Rpc.Tasks;
+using HordeServer.Services;
+using HordeServer.Utilities;
+using MongoDB.Bson;
+using MongoDB.Bson.Serialization.Attributes;
+using MongoDB.Driver;
+using System;
+using System.Collections.Generic;
+using System.Linq;
+using System.Linq.Expressions;
+using System.Threading.Tasks;
+
+namespace HordeServer.Collections
+{
+	using PoolId = StringId<IPool>;
+	using AgentSoftwareChannelName = StringId<AgentSoftwareChannels>;
+
+	/// <summary>
+	/// Interface for a collection of agent documents
+	/// </summary>
+	public interface IAgentCollection
+	{
+		/// <summary>
+		/// Adds a new agent with the given properties
+		/// </summary>
+		/// <param name="Id">Id for the new agent</param>
+		/// <param name="bEnabled">Whether the agent is enabled or not</param>
+		/// <param name="Channel">Channel to use for software run by this agent</param>
+		/// <param name="Pools">Pools for the agent</param>
+		Task<IAgent> AddAsync(AgentId Id, bool bEnabled, AgentSoftwareChannelName? Channel = null, List<PoolId>? Pools = null);
+
+		/// <summary>
+		/// Deletes an agent
+		/// </summary>
+		/// <param name="Agent">Deletes the agent</param>
+		/// <returns>Async task</returns>
+		Task<IAgent?> TryDeleteAsync(IAgent Agent);
+
+		/// <summary>
+		/// Deletes an agent
+		/// </summary>
+		/// <param name="AgentId">Unique id of the agent</param>
+		/// <returns>Async task</returns>
+		Task ForceDeleteAsync(AgentId AgentId);
+
+		/// <summary>
+		/// Gets an agent by ID
+		/// </summary>
+		/// <param name="AgentId">Unique id of the agent</param>
+		/// <returns>The agent document</returns>
+		Task<IAgent?> GetAsync(AgentId AgentId);
+
+		/// <summary>
+		/// Finds all agents matching certain criteria
+		/// </summary>
+		/// <param name="Pool">The pool containing the agent</param>
+		/// <param name="PoolId">The pool ID in string form containing the agent</param>
+		/// <param name="ModifiedAfter">If set, only returns agents modified after this time</param>
+		/// <param name="Status">Status to look for</param>
+		/// <param name="Index">Index of the first result</param>
+		/// <param name="Count">Number of results to return</param>
+		/// <returns>List of agents matching the given criteria</returns>
+		Task<List<IAgent>> FindAsync(ObjectId? Pool = null, string? PoolId = null, DateTime? ModifiedAfter = null, AgentStatus? Status = null, int? Index = null, int? Count = null);
+
+		/// <summary>
+		/// Finds all the expried agents
+		/// </summary>
+		/// <param name="UtcNow">The current time</param>
+		/// <param name="MaxAgents">Maximum number of agents to return</param>
+		/// <returns>List of agents</returns>
+		Task<List<IAgent>> FindExpiredAsync(DateTime UtcNow, int MaxAgents);
+
+		/// <summary>
+		/// Update an agent's settings
+		/// </summary>
+		/// <param name="Agent">Agent instance</param>
+		/// <param name="bEnabled">Whether the agent is enabled or not</param>
+		/// <param name="bRequestConform">Whether to request a conform job be run</param>
+		/// <param name="bRequestRestart">Whether to request the machine be restarted</param>
+		/// <param name="bRequestShutdown">Whether to request the machine be shut down</param>
+		/// <param name="Channel">Override for the desired software channel</param>
+		/// <param name="Pools">List of pools for the agent</param>
+		/// <param name="Acl">New ACL for this agent</param>
+		/// <param name="Comment">New comment</param>
+		/// <returns>Version of the software that needs to be installed on the agent. Null if the agent is running the correct version.</returns>
+		Task<IAgent?> TryUpdateSettingsAsync(IAgent Agent, bool? bEnabled = null, bool? bRequestConform = null, bool? bRequestRestart = null, bool? bRequestShutdown = null, AgentSoftwareChannelName? Channel = null, List<PoolId>? Pools = null, Acl? Acl = null, string? Comment = null);
+
+		/// <summary>
+		/// Update the current workspaces for an agent.
+		/// </summary>
+		/// <param name="Agent">The agent to update</param>
+		/// <param name="Workspaces">Current list of workspaces</param>
+		/// <param name="RequestConform">Whether the agent still needs to run another conform</param>
+		/// <returns>New agent state</returns>
+		Task<IAgent?> TryUpdateWorkspacesAsync(IAgent Agent, List<AgentWorkspace> Workspaces, bool RequestConform);
+
+		/// <summary>
+		/// Sets the current session
+		/// </summary>
+		/// <param name="Agent">The agent to update</param>
+		/// <param name="SessionId">New session id</param>
+		/// <param name="SessionExpiresAt">Expiry time for the new session</param>
+		/// <param name="Status">Status of the agent</param>
+		/// <param name="Properties">Properties for the current session</param>
+		/// <param name="Resources">Resources for the agent</param>
+		/// <param name="DynamicPools">New list of dynamic pools for the agent</param>
+		/// <param name="Version">Current version of the agent software</param>
+		/// <returns>New agent state</returns>
+		Task<IAgent?> TryStartSessionAsync(IAgent Agent, ObjectId SessionId, DateTime SessionExpiresAt, AgentStatus Status, IReadOnlyList<string> Properties, IReadOnlyDictionary<string, int> Resources, IReadOnlyList<PoolId> DynamicPools, string? Version);
+
+		/// <summary>
+		/// Attempt to update the agent state
+		/// </summary>
+		/// <param name="Agent">Agent instance</param>
+		/// <param name="Status">New status of the agent</param>
+		/// <param name="SessionExpiresAt">New expiry time for the current session</param>
+		/// <param name="Properties">Properties for the current session</param>
+		/// <param name="Resources">Resources for the agent</param>
+		/// <param name="DynamicPools">New list of dynamic pools for the agent</param>
+		/// <param name="Leases">New set of leases</param>
+		/// <returns>True if the document was updated, false if another writer updated the document first</returns>
+		Task<IAgent?> TryUpdateSessionAsync(IAgent Agent, AgentStatus? Status, DateTime? SessionExpiresAt, IReadOnlyList<string>? Properties, IReadOnlyDictionary<string, int>? Resources, IReadOnlyList<PoolId>? DynamicPools, List<AgentLease>? Leases);
+
+		/// <summary>
+		/// Terminates the current session
+		/// </summary>
+		/// <param name="Agent">The agent to terminate the session for</param>
+		/// <returns>New agent state if it succeeded, otherwise null</returns>
+		Task<IAgent?> TryTerminateSessionAsync(IAgent Agent);
+
+		/// <summary>
+		/// Attempts to add a lease to an agent
+		/// </summary>
+		/// <param name="Agent">The agent to update</param>
+		/// <param name="NewLease">The new lease document</param>
+		/// <returns>New agent state if it succeeded, otherwise null</returns>
+		Task<IAgent?> TryAddLeaseAsync(IAgent Agent, AgentLease NewLease);
+
+		/// <summary>
+		/// Attempts to cancel a lease
+		/// </summary>
+		/// <param name="Agent">The agent to update</param>
+		/// <param name="LeaseIdx">Index of the lease</param>
+		/// <returns>New agent state if it succeeded, otherwise null</returns>
+		Task<IAgent?> TryCancelLeaseAsync(IAgent Agent, int LeaseIdx);
+
+		/// <summary>
+		/// Gets the log channel for an agent
+		/// </summary>
+		/// <param name="AgentId"></param>
+		/// <returns></returns>
+		IAuditLogChannel<AgentId> GetLogger(AgentId AgentId);
+
+		/// <summary>
+		/// Subscribe to notifications on agent states being updated
+		/// </summary>
+		/// <param name="OnUpdate">Callback for updates</param>
+		/// <returns>Disposable subscription object</returns>
+		Task<IDisposable> SubscribeToUpdateEventsAsync(Action<AgentId> OnUpdate);
+	}
+}
diff --git a/Engine/Source/Programs/Horde/HordeServer/Agents/Impl/AgentCollection.cs b/Engine/Source/Programs/Horde/HordeServer/Agents/Impl/AgentCollection.cs
new file mode 100644
--- /dev/null
+++ b/Engine/Source/Programs/Horde/HordeServer/Agents/Impl/AgentCollection.cs
@@ -0,0 +1,523 @@
+// Copyright Epic Games, Inc. All Rights Reserved.
+
+using EpicGames.Redis;
+using Google.Protobuf.WellKnownTypes;
+using HordeCommon;
+using HordeCommon.Rpc.Tasks;
+using HordeServer.Models;
+using HordeServer.Services;
+using HordeServer.Utilities;
+using Microsoft.Extensions.Logging;
+using MongoDB.Bson;
+using MongoDB.Bson.Serialization.Attributes;
+using MongoDB.Driver;
+using System;
+using System.Collections.Generic;
+using System.Linq;
+using System.Linq.Expressions;
+using System.Threading.Tasks;
+
+namespace HordeServer.Collections.Impl
+{
+	using AgentSoftwareChannelName = StringId<AgentSoftwareChannels>;
+	using PoolId = StringId<IPool>;
+
+	/// <summary>
+	/// Collection of agent documents
+	/// </summary>
+	public class AgentCollection : IAgentCollection
+	{
+		/// <summary>
+		/// Legacy information about a device attached to an agent
+		/// </summary>
+		class DeviceCapabilities
+		{
+			[BsonIgnoreIfNull]
+			public HashSet<string>? Properties { get; set; }
+
+			[BsonIgnoreIfNull]
+			public Dictionary<string, int>? Resources { get; set; }
+		}
+
+		/// <summary>
+		/// Legacy capabilities of an agent
+		/// </summary>
+		class AgentCapabilities
+		{
+			public List<DeviceCapabilities> Devices { get; set; } = new List<DeviceCapabilities>();
+
+			[BsonIgnoreIfNull]
+			public HashSet<string>? Properties { get; set; }
+		}
+
+		/// <summary>
+		/// Concrete implementation of an agent document
+		/// </summary>
+		class AgentDocument : IAgent
+		{
+			static IReadOnlyList<AgentLease> EmptyLeases = new List<AgentLease>();
+
+			[BsonRequired, BsonId]
+			public AgentId Id { get; set; }
+
+			public ObjectId? SessionId { get; set; }
+			public DateTime? SessionExpiresAt { get; set; }
+
+			public AgentStatus Status { get; set; }
+
+			[BsonRequired]
+			public bool Enabled { get; set; }
+
+			public bool Ephemeral { get; set; }
+
+			[BsonIgnoreIfDefault, BsonDefaultValue(false)]
+			public bool Deleted { get; set; }
+
+			[BsonElement("Version2")]
+			public string? Version { get; set; }
+
+			public List<string>? Properties { get; set; }
+			public Dictionary<string, int>? Resources { get; set; }
+
+			[BsonIgnoreIfNull]
+			public AgentSoftwareChannelName? Channel { get; set; }
+
+			[BsonIgnoreIfNull]
+			public string? LastUpgradeVersion { get; set; }
+
+			[BsonIgnoreIfNull]
+			public DateTime? LastUpgradeTime { get; set; }
+
+			public List<PoolId> DynamicPools { get; set; } = new List<PoolId>();
+			public List<PoolId> Pools { get; set; } = new List<PoolId>();
+			public bool RequestConform { get; set; }
+
+			[BsonIgnoreIfNull]
+			public bool RequestRestart { get; set; }
+
+			[BsonIgnoreIfNull]
+			public bool RequestShutdown { get; set; }
+
+			public List<AgentWorkspace> Workspaces { get; set; } = new List<AgentWorkspace>();
+			public DateTime LastConformTime { get; set; }
+
+			[BsonIgnoreIfNull]
+			public int? ConformAttemptCount { get; set; }
+
+			public AgentCapabilities Capabilities { get; set; } = new AgentCapabilities();
+			public List<AgentLease>? Leases { get; set; }
+			public Acl? Acl { get; set; }
+			public DateTime UpdateTime { get; set; }
+			public uint UpdateIndex { get; set; }
+			public string? Comment { get; set; }
+
+			IReadOnlyList<PoolId> IAgent.DynamicPools => DynamicPools;
+			IReadOnlyList<PoolId> IAgent.ExplicitPools => Pools;
+			IReadOnlyList<AgentWorkspace> IAgent.Workspaces => Workspaces;
+			IReadOnlyList<AgentLease> IAgent.Leases => Leases ?? EmptyLeases;
+			IReadOnlyList<string> IAgent.Properties => Properties ?? Capabilities.Devices.FirstOrDefault()?.Properties?.ToList() ?? new List<string>();
+			IReadOnlyDictionary<string, int> IAgent.Resources => Resources ?? Capabilities.Devices.FirstOrDefault()?.Resources ?? new Dictionary<string, int>();
+
+			[BsonConstructor]
+			private AgentDocument()
+			{
+			}
+
+			public AgentDocument(AgentId Id, bool bEnabled, AgentSoftwareChannelName? Channel, List<PoolId> Pools)
+			{
+				this.Id = Id;
+				this.Acl = new Acl();
+				this.Enabled = bEnabled;
+				this.Channel = Channel;
+				this.Pools = Pools;
+			}
+		}
+
+		readonly IMongoCollection<AgentDocument> Agents;
+		readonly IAuditLog<AgentId> AuditLog;
+		readonly RedisService RedisService;
+		readonly RedisChannel<AgentId> UpdateEventChannel;
+
+		/// <summary>
+		/// Constructor
+		/// </summary>
+		public AgentCollection(DatabaseService DatabaseService, RedisService RedisService, IAuditLog<AgentId> AuditLog)
+		{
+			this.Agents = DatabaseService.GetCollection<AgentDocument>("Agents");
+			this.RedisService = RedisService;
+			this.UpdateEventChannel = new RedisChannel<AgentId>("agents/notify");
+			this.AuditLog = AuditLog;
+
+			if (!DatabaseService.ReadOnlyMode)
+			{
+				Agents.Indexes.CreateOne(new CreateIndexModel<AgentDocument>(Builders<AgentDocument>.IndexKeys.Ascending(x => x.Deleted).Ascending(x => x.Id).Ascending(x => x.Pools)));
+			}
+		}
+
+		/// <inheritdoc/>
+		public async Task<IAgent> AddAsync(AgentId Id, bool bEnabled, AgentSoftwareChannelName? Channel, List<PoolId>? Pools)
+		{
+			AgentDocument Agent = new AgentDocument(Id, bEnabled, Channel, Pools ?? new List<PoolId>());
+			await Agents.InsertOneAsync(Agent);
+			return Agent;
+		}
+
+		/// <inheritdoc/>
+		public async Task<IAgent?> TryDeleteAsync(IAgent AgentInterface)
+		{
+			AgentDocument Agent = (AgentDocument)AgentInterface;
+
+			UpdateDefinition<AgentDocument> Update = Builders<AgentDocument>.Update.Set(x => x.Deleted, true);
+			return await TryUpdateAsync(Agent, Update);
+		}
+
+		/// <inheritdoc/>
+		public async Task ForceDeleteAsync(AgentId AgentId)
+		{
+			await Agents.DeleteOneAsync(x => x.Id == AgentId);
+		}
+
+		/// <inheritdoc/>
+		public async Task<IAgent?> GetAsync(AgentId AgentId)
+		{
+			return await Agents.Find<AgentDocument>(x => x.Id == AgentId).FirstOrDefaultAsync();
+		}
+
+		/// <inheritdoc/>
+		public async Task<List<IAgent>> FindAsync(ObjectId? Pool, string? PoolId, DateTime? ModifiedAfter, AgentStatus? Status, int? Index, int? Count)
+		{
+			FilterDefinitionBuilder<AgentDocument> FilterBuilder = new FilterDefinitionBuilder<AgentDocument>();
+
+			FilterDefinition<AgentDocument> Filter = FilterBuilder.Ne(x => x.Deleted, true);
+			if (Pool != null)
+			{
+				Filter &= FilterBuilder.Eq(nameof(AgentDocument.Pools), Pool);
+			}
+			
+			if (PoolId != null)
+			{
+				Filter &= FilterBuilder.Eq(nameof(AgentDocument.Pools), PoolId);
+			}
+			
+			if (ModifiedAfter != null)
+			{
+				Filter &= FilterBuilder.Gt(x => x.UpdateTime, ModifiedAfter.Value);
+			}
+			
+			if (Status != null)
+			{
+				Filter &= FilterBuilder.Eq(x => x.Status, Status.Value);
+			}
+
+			IFindFluent<AgentDocument, AgentDocument> Search = Agents.Find(Filter);
+			if (Index != null)
+			{
+				Search = Search.Skip(Index.Value);
+			}
+			if (Count != null)
+			{
+				Search = Search.Limit(Count.Value);
+			}
+
+			List<AgentDocument> Results = await Search.ToListAsync();
+			return Results.ConvertAll<IAgent>(x => x);
+		}
+
+		/// <inheritdoc/>
+		public async Task<List<IAgent>> FindExpiredAsync(DateTime UtcNow, int MaxAgents)
+		{
+			List<AgentDocument> Results = await Agents.Find(x => x.SessionId.HasValue && !(x.SessionExpiresAt > UtcNow)).Limit(MaxAgents).ToListAsync();
+			return Results.ConvertAll<IAgent>(x => x);
+		}
+
+		/// <summary>
+		/// Update a single document
+		/// </summary>
+		/// <param name="Current">The document to update</param>
+		/// <param name="Update">The update definition</param>
+		/// <returns>True if the agent was updated</returns>
+		private async Task<AgentDocument?> TryUpdateAsync(AgentDocument Current, UpdateDefinition<AgentDocument> Update)
+		{
+			uint PrevUpdateIndex = Current.UpdateIndex++;
+			Current.UpdateTime = DateTime.UtcNow;
+
+			Expression<Func<AgentDocument, bool>> Filter = x => x.Id == Current.Id && x.UpdateIndex == PrevUpdateIndex;
+			UpdateDefinition<AgentDocument> UpdateWithIndex = Update.Set(x => x.UpdateIndex, Current.UpdateIndex).Set(x => x.UpdateTime, Current.UpdateTime);
+
+			return await Agents.FindOneAndUpdateAsync<AgentDocument>(Filter, UpdateWithIndex, new FindOneAndUpdateOptions<AgentDocument, AgentDocument> { ReturnDocument = ReturnDocument.After });
+		}
+
+		/// <inheritdoc/>
+		public async Task<IAgent?> TryUpdateSettingsAsync(IAgent AgentInterface, bool? Enabled = null, bool? RequestConform = null, bool? RequestRestart = null, bool? RequestShutdown = null, AgentSoftwareChannelName? Channel = null, List<PoolId>? Pools = null, Acl? Acl = null, string? Comment = null)
+		{
+			AgentDocument Agent = (AgentDocument)AgentInterface;
+
+			// Update the database
+			UpdateDefinitionBuilder<AgentDocument> UpdateBuilder = new UpdateDefinitionBuilder<AgentDocument>();
+
+			List<UpdateDefinition<AgentDocument>> Updates = new List<UpdateDefinition<AgentDocument>>();
+			if (Pools != null)
+			{
+				Updates.Add(UpdateBuilder.Set(x => x.Pools, Pools));
+			}
+			if (Enabled != null)
+			{
+				Updates.Add(UpdateBuilder.Set(x => x.Enabled, Enabled.Value));
+			}
+			if (RequestConform != null)
+			{
+				Updates.Add(UpdateBuilder.Set(x => x.RequestConform, RequestConform.Value));
+				Updates.Add(UpdateBuilder.Unset(x => x.ConformAttemptCount));
+			}
+			if (RequestRestart != null)
+			{
+				if (RequestRestart.Value)
+				{
+					Updates.Add(UpdateBuilder.Set(x => x.RequestRestart, true));
+				}
+				else
+				{
+					Updates.Add(UpdateBuilder.Unset(x => x.RequestRestart));
+				}
+			}
+			if (RequestShutdown != null)
+			{
+				if (RequestShutdown.Value)
+				{
+					Updates.Add(UpdateBuilder.Set(x => x.RequestShutdown, true));
+				}
+				else
+				{
+					Updates.Add(UpdateBuilder.Unset(x => x.RequestShutdown));
+				}
+			}
+			if (Channel != null)
+			{
+				if (Channel.Value == AgentSoftwareService.DefaultChannelName)
+				{
+					Updates.Add(UpdateBuilder.Unset(x => x.Channel));
+				}
+				else
+				{
+					Updates.Add(UpdateBuilder.Set(x => x.Channel, Channel));
+				}
+			}
+			if (Acl != null)
+			{
+				Updates.Add(Acl.CreateUpdate<AgentDocument>(x => x.Acl!, Acl));
+			}
+			if (Comment != null)
+			{
+				Updates.Add(UpdateBuilder.Set(x => x.Comment, Comment));
+			}
+
+			// Apply the update
+			IAgent? NewAgent = await TryUpdateAsync(Agent, UpdateBuilder.Combine(Updates));
+			if (NewAgent != null)
+			{
+				if (NewAgent.RequestRestart != Agent.RequestRestart || NewAgent.RequestConform != Agent.RequestConform || NewAgent.RequestShutdown != Agent.RequestShutdown || NewAgent.Channel != Agent.Channel)
+				{
+					await PublishUpdateEventAsync(Agent.Id);
+				}
+			}
+			return NewAgent;
+		}
+
+		/// <inheritdoc/>
+		public async Task<IAgent?> TryUpdateSessionAsync(IAgent AgentInterface, AgentStatus? Status, DateTime? SessionExpiresAt, IReadOnlyList<string>? Properties, IReadOnlyDictionary<string, int>? Resources, IReadOnlyList<PoolId>? DynamicPools, List<AgentLease>? Leases)
+		{
+			AgentDocument Agent = (AgentDocument)AgentInterface;
+
+			// Create an update definition for the agent
+			UpdateDefinitionBuilder<AgentDocument> UpdateBuilder = Builders<AgentDocument>.Update;
+			List<UpdateDefinition<AgentDocument>> Updates = new List<UpdateDefinition<AgentDocument>>();
+
+			if (Status != null && Agent.Status != Status.Value)
+			{
+				Updates.Add(UpdateBuilder.Set(x => x.Status, Status.Value));
+			}
+			if (SessionExpiresAt != null)
+			{
+				Updates.Add(UpdateBuilder.Set(x => x.SessionExpiresAt, SessionExpiresAt.Value));
+			}
+			if (Properties != null)
+			{
+				List<string> NewProperties = Properties.OrderBy(x => x, StringComparer.OrdinalIgnoreCase).ToList();
+				if (!AgentInterface.Properties.SequenceEqual(NewProperties, StringComparer.Ordinal))
+				{
+					Updates.Add(UpdateBuilder.Set(x => x.Properties, NewProperties));
+				}
+			}
+			if (Resources != null && !ResourcesEqual(Resources, AgentInterface.Resources))
+			{
+				Updates.Add(UpdateBuilder.Set(x => x.Resources, new Dictionary<string, int>(Resources)));
+			}
+			if (DynamicPools != null && !DynamicPools.SequenceEqual(Agent.DynamicPools))
+			{
+				Updates.Add(UpdateBuilder.Set(x => x.DynamicPools, new List<PoolId>(DynamicPools)));
+			}
+			if (Leases != null)
+			{
+				foreach (AgentLease Lease in Leases)
+				{
+					if (Lease.Payload != null && (Agent.Leases == null || !Agent.Leases.Any(x => x.Id == Lease.Id)))
+					{
+						Any Payload = Any.Parser.ParseFrom(Lease.Payload.ToArray());
+						if (Payload.TryUnpack(out ConformTask ConformTask))
+						{
+							int NewConformAttemptCount = (Agent.ConformAttemptCount ?? 0) + 1;
+							Updates.Add(UpdateBuilder.Set(x => x.ConformAttemptCount, NewConformAttemptCount));
+							Updates.Add(UpdateBuilder.Set(x => x.LastConformTime, DateTime.UtcNow));
+						}
+						else if (Payload.TryUnpack(out UpgradeTask UpgradeTask))
+						{
+							Updates.Add(UpdateBuilder.Set(x => x.LastUpgradeVersion, UpgradeTask.SoftwareId));
+							Updates.Add(UpdateBuilder.Set(x => x.LastUpgradeTime, DateTime.UtcNow));
+						}
+					}
+				}
+
+				Updates.Add(UpdateBuilder.Set(x => x.Leases, Leases));
+			}
+
+			// If there are no new updates, return immediately. This is important for preventing UpdateSession calls from returning immediately.
+			if (Updates.Count == 0)
+			{
+				return Agent;
+			}
+
+			// Update the agent, and try to create new lease documents if we succeed
+			return await TryUpdateAsync(Agent, UpdateBuilder.Combine(Updates));
+		}
+
+		static bool ResourcesEqual(IReadOnlyDictionary<string, int> DictA, IReadOnlyDictionary<string, int> DictB)
+		{
+			if (DictA.Count != DictB.Count)
+			{
+				return false;
+			}
+
+			foreach (KeyValuePair<string, int> Pair in DictA)
+			{
+				int Value;
+				if (!DictB.TryGetValue(Pair.Key, out Value) || Value != Pair.Value)
+				{
+					return false;
+				}
+			}
+
+			return true;
+		}
+
+		/// <inheritdoc/>
+		public async Task<IAgent?> TryUpdateWorkspacesAsync(IAgent AgentInterface, List<AgentWorkspace> Workspaces, bool RequestConform)
+		{
+			AgentDocument Agent = (AgentDocument)AgentInterface;
+			DateTime LastConformTime = DateTime.UtcNow;
+
+			// Set the new workspaces
+			UpdateDefinition<AgentDocument> Update = Builders<AgentDocument>.Update.Set(x => x.Workspaces, Workspaces);
+			Update = Update.Set(x => x.LastConformTime, LastConformTime);
+			Update = Update.Unset(x => x.ConformAttemptCount);
+			Update = Update.Set(x => x.RequestConform, RequestConform);
+
+			// Update the agent
+			return await TryUpdateAsync(Agent, Update);
+		}
+
+		/// <inheritdoc/>
+		public async Task<IAgent?> TryStartSessionAsync(IAgent AgentInterface, ObjectId SessionId, DateTime SessionExpiresAt, AgentStatus Status, IReadOnlyList<string> Properties, IReadOnlyDictionary<string, int> Resources, IReadOnlyList<PoolId> DynamicPools, string? Version)
+		{
+			AgentDocument Agent = (AgentDocument)AgentInterface;
+			List<string> NewProperties = Properties.OrderBy(x => x, StringComparer.OrdinalIgnoreCase).ToList();
+			Dictionary<string, int> NewResources = new Dictionary<string, int>(Resources);
+			List<PoolId> NewDynamicPools = new List<PoolId>(DynamicPools);
+
+			// Reset the agent to use the new session
+			UpdateDefinitionBuilder<AgentDocument> UpdateBuilder = Builders<AgentDocument>.Update;
+
+			List<UpdateDefinition<AgentDocument>> Updates = new List<UpdateDefinition<AgentDocument>>();
+			Updates.Add(UpdateBuilder.Set(x => x.SessionId, SessionId));
+			Updates.Add(UpdateBuilder.Set(x => x.SessionExpiresAt, SessionExpiresAt));
+			Updates.Add(UpdateBuilder.Set(x => x.Status, Status));
+			Updates.Add(UpdateBuilder.Unset(x => x.Leases));
+			Updates.Add(UpdateBuilder.Unset(x => x.Deleted));
+			Updates.Add(UpdateBuilder.Set(x => x.Properties, NewProperties));
+			Updates.Add(UpdateBuilder.Set(x => x.Resources, NewResources));
+			Updates.Add(UpdateBuilder.Set(x => x.DynamicPools, NewDynamicPools));
+			Updates.Add(UpdateBuilder.Set(x => x.Version, Version));
+			Updates.Add(UpdateBuilder.Unset(x => x.RequestRestart));
+			Updates.Add(UpdateBuilder.Unset(x => x.RequestShutdown));
+
+			// Apply the update
+			return await TryUpdateAsync(Agent, UpdateBuilder.Combine(Updates));
+		}
+
+		/// <inheritdoc/>
+		public async Task<IAgent?> TryTerminateSessionAsync(IAgent AgentInterface)
+		{
+			AgentDocument Agent = (AgentDocument)AgentInterface;
+			UpdateDefinition<AgentDocument> Update = new BsonDocument();
+
+			Update = Update.Unset(x => x.SessionId);
+			Update = Update.Unset(x => x.SessionExpiresAt);
+			Update = Update.Unset(x => x.Leases);
+
+			bool bDeleted = Agent.Deleted || Agent.Ephemeral;
+			if (bDeleted != Agent.Deleted)
+			{
+				Update = Update.Set(x => x.Deleted, Agent.Deleted);
+			}
+
+			return await TryUpdateAsync(Agent, Update);
+		}
+
+		/// <inheritdoc/>
+		public async Task<IAgent?> TryAddLeaseAsync(IAgent AgentInterface, AgentLease NewLease)
+		{
+			AgentDocument Agent = (AgentDocument)AgentInterface;
+
+			List<AgentLease> Leases = new List<AgentLease>();
+			if (Agent.Leases != null)
+			{
+				Leases.AddRange(Agent.Leases);
+			}
+			Leases.Add(NewLease);
+
+			UpdateDefinition<AgentDocument> Update = Builders<AgentDocument>.Update.Set(x => x.Leases, Leases);
+			return await TryUpdateAsync(Agent, Update);
+		}
+
+		/// <inheritdoc/>
+		public async Task<IAgent?> TryCancelLeaseAsync(IAgent AgentInterface, int LeaseIdx)
+		{
+			AgentDocument Agent = (AgentDocument)AgentInterface;
+
+			UpdateDefinition<AgentDocument> Update = Builders<AgentDocument>.Update.Set(x => x.Leases![LeaseIdx].State, LeaseState.Cancelled);
+			IAgent? NewAgent = await TryUpdateAsync(Agent, Update);
+			if (NewAgent != null)
+			{
+				await PublishUpdateEventAsync(Agent.Id);
+			}
+			return NewAgent;
+		}
+
+		/// <inheritdoc/>
+		public IAuditLogChannel<AgentId> GetLogger(AgentId AgentId)
+		{
+			return AuditLog[AgentId];
+		}
+
+		/// <inheritdoc/>
+		Task PublishUpdateEventAsync(AgentId AgentId)
+		{
+			return RedisService.Database.PublishAsync(UpdateEventChannel, AgentId);
+		}
+
+		/// <inheritdoc/>
+		public async Task<IDisposable> SubscribeToUpdateEventsAsync(Action<AgentId> OnUpdate)
+		{
+			return await RedisService.Multiplexer.GetSubscriber().SubscribeAsync(UpdateEventChannel, (Channel, AgentId) => OnUpdate(AgentId));
+		}
+	}
+}
diff --git a/Engine/Source/Programs/Horde/HordeServer/Agents/Impl/AgentService.cs b/Engine/Source/Programs/Horde/HordeServer/Agents/Impl/AgentService.cs
new file mode 100644
--- /dev/null
+++ b/Engine/Source/Programs/Horde/HordeServer/Agents/Impl/AgentService.cs
@@ -0,0 +1,888 @@
+// Copyright Epic Games, Inc. All Rights Reserved.
+
+using EpicGames.Core;
+using Google.Protobuf;
+using Google.Protobuf.WellKnownTypes;
+using HordeCommon;
+using HordeServer.Collections;
+using HordeServer.Models;
+using HordeServer.Tasks;
+using HordeServer.Utilities;
+using Microsoft.Extensions.Hosting;
+using Microsoft.Extensions.Logging;
+using MongoDB.Bson;
+using MongoDB.Driver;
+using StatsdClient;
+using System;
+using System.Collections.Generic;
+using System.Diagnostics;
+using System.Linq;
+using System.Security.Claims;
+using System.Threading;
+using System.Threading.Tasks;
+
+namespace HordeServer.Services
+{
+	using AgentSoftwareChannelName = StringId<AgentSoftwareChannels>;
+	using PoolId = StringId<IPool>;
+
+	/// <summary>
+	/// Wraps funtionality for manipulating agents
+	/// </summary>
+	public sealed class AgentService : TickedBackgroundService
+	{
+		/// <summary>
+		/// Maximum time between updates for an agent to be considered online
+		/// </summary>
+		public static readonly TimeSpan SessionExpiryTime = TimeSpan.FromMinutes(5);
+
+		/// <summary>
+		/// Time before a session expires that we will poll until
+		/// </summary>
+		public static readonly TimeSpan SessionLongPollTime = TimeSpan.FromSeconds(55);
+
+		/// <summary>
+		/// Time after which a session will be renewed
+		/// </summary>
+		public static readonly TimeSpan SessionRenewTime = TimeSpan.FromSeconds(50);
+
+		/// <summary>
+		/// The ACL service instance
+		/// </summary>
+		AclService AclService;
+
+		/// <summary>
+		/// The downtime service instance
+		/// </summary>
+		IDowntimeService DowntimeService;
+
+		/// <summary>
+		/// Collection of agent documents
+		/// </summary>
+		public IAgentCollection Agents { get; }
+
+		/// <summary>
+		/// Collection of lease documents
+		/// </summary>
+		ILeaseCollection Leases;
+
+		/// <summary>
+		/// Collection of session documents
+		/// </summary>
+		ISessionCollection Sessions;
+
+		/// <summary>
+		/// DogStatsD metric client
+		/// </summary>
+		IDogStatsd DogStatsd;
+
+		/// <summary>
+		/// List of task sources
+		/// </summary>
+		ITaskSource[] TaskSources;
+
+		/// <summary>
+		/// The application lifetime instance
+		/// </summary>
+		IHostApplicationLifetime ApplicationLifetime;
+
+		/// <summary>
+		/// Log output writer
+		/// </summary>
+		ILogger<AgentService> Logger;
+		
+		/// <summary>
+		/// Clock
+		/// </summary>
+		readonly IClock Clock;
+
+		/// <summary>
+		/// Lazily updated list of current pools
+		/// </summary>
+		AsyncCachedValue<List<IPool>> PoolsList;
+
+		/// <summary>
+		/// All the agents currently performing a long poll for work on this server
+		/// </summary>
+		Dictionary<AgentId, CancellationTokenSource> WaitingAgents = new Dictionary<AgentId, CancellationTokenSource>();
+
+		/// <summary>
+		/// Subscription for update events
+		/// </summary>
+		IDisposable Subscription;
+
+		/// <summary>
+		/// Constructor
+		/// </summary>
+		public AgentService(IAgentCollection Agents, ILeaseCollection Leases, ISessionCollection Sessions, AclService AclService, IDowntimeService DowntimeService, IPoolCollection PoolCollection, IDogStatsd DogStatsd, IEnumerable<ITaskSource> TaskSources, IHostApplicationLifetime ApplicationLifetime, ILogger<AgentService> Logger, IClock Clock)
+			: base(TimeSpan.FromSeconds(30.0), Logger)
+		{
+			this.Agents = Agents;
+			this.Leases = Leases;
+			this.Sessions = Sessions;
+			this.AclService = AclService;
+			this.DowntimeService = DowntimeService;
+			this.PoolsList = new AsyncCachedValue<List<IPool>>(() => PoolCollection.GetAsync(), TimeSpan.FromSeconds(30.0));
+			this.DogStatsd = DogStatsd;
+			this.TaskSources = TaskSources.ToArray();
+			this.ApplicationLifetime = ApplicationLifetime;
+			this.Logger = Logger;
+			this.Clock = Clock;
+
+			Subscription = Agents.SubscribeToUpdateEventsAsync(OnAgentUpdate).Result;
+		}
+
+		/// <inheritdoc/>
+		public override void Dispose()
+		{
+			base.Dispose();
+			Subscription.Dispose();
+		}
+
+		/// <summary>
+		/// Issues a bearer token for the given session id
+		/// </summary>
+		/// <param name="SessionId">The session id</param>
+		/// <returns>Bearer token for the agent</returns>
+		public string IssueSessionToken(ObjectId SessionId)
+		{
+			List<AclClaim> Claims = new List<AclClaim>();
+			Claims.Add(AclService.AgentClaim);
+			Claims.Add(AclService.GetSessionClaim(SessionId));
+			return AclService.IssueBearerToken(Claims, null);
+		}
+
+		/// <summary>
+		/// Register a new agent
+		/// </summary>
+		/// <param name="Name">Name of the agent</param>
+		/// <param name="bEnabled">Whether the agent is currently enabled</param>
+		/// <param name="Channel">Override for the desired software version</param>
+		/// <param name="Pools">Pools for this agent</param>
+		/// <returns>Unique id for the agent</returns>
+		public Task<IAgent> CreateAgentAsync(string Name, bool bEnabled, AgentSoftwareChannelName? Channel, List<PoolId>? Pools)
+		{
+			return Agents.AddAsync(new AgentId(Name), bEnabled, Channel, Pools);
+		}
+
+		/// <summary>
+		/// Gets an agent by ID
+		/// </summary>
+		/// <param name="AgentId">Unique id of the agent</param>
+		/// <returns>The agent document</returns>
+		public Task<IAgent?> GetAgentAsync(AgentId AgentId)
+		{
+			return Agents.GetAsync(AgentId);
+		}
+
+		/// <summary>
+		/// Finds all agents matching certain criteria
+		/// </summary>
+		/// <param name="Pool">The pool containing the agent</param>
+		/// <param name="ModifiedAfter">If set, only returns agents modified after this time</param>
+		/// <param name="Index">Index within the list of results</param>
+		/// <param name="Count">Number of results to return</param>
+		/// <returns>List of agents matching the given criteria</returns>
+		public Task<List<IAgent>> FindAgentsAsync(ObjectId? Pool, DateTime? ModifiedAfter, int? Index, int? Count)
+		{
+			return Agents.FindAsync(Pool, null, ModifiedAfter, null, Index, Count);
+		}
+
+		/// <summary>
+		/// Update the current workspaces for an agent.
+		/// </summary>
+		/// <param name="Agent">The agent to update</param>
+		/// <param name="Workspaces">Current list of workspaces</param>
+		/// <param name="bPendingConform">Whether the agent still needs to run another conform</param>
+		/// <returns>New agent state</returns>
+		public async Task<bool> TryUpdateWorkspacesAsync(IAgent Agent, List<AgentWorkspace> Workspaces, bool bPendingConform)
+		{
+			IAgent? NewAgent = await Agents.TryUpdateWorkspacesAsync(Agent, Workspaces, bPendingConform);
+			return NewAgent != null;
+		}
+
+		/// <summary>
+		/// Marks the agent as deleted
+		/// </summary>
+		/// <param name="Agent">The agent to delete</param>
+		/// <returns>Async task</returns>
+		public async Task DeleteAgentAsync(IAgent? Agent)
+		{
+			while (Agent != null && !Agent.Deleted)
+			{
+				IAgent? NewAgent = await Agents.TryDeleteAsync(Agent);
+				if(NewAgent != null)
+				{
+					break;
+				}
+				Agent = await GetAgentAsync(Agent.Id);
+			}
+		}
+
+		async ValueTask<List<PoolId>> GetDynamicPoolsAsync(IAgent Agent)
+		{
+			List<PoolId> NewDynamicPools = new List<PoolId>();
+
+			List<IPool> Pools = await PoolsList.GetAsync();
+			foreach (IPool Pool in Pools)
+			{
+				if (Pool.Condition != null && Agent.SatisfiesCondition(Pool.Condition))
+				{
+					NewDynamicPools.Add(Pool.Id);
+				}
+			}
+
+			return NewDynamicPools;
+		}
+
+		/// <summary>
+		/// Callback for an agents 
+		/// </summary>
+		/// <param name="AgentId"></param>
+		void OnAgentUpdate(AgentId AgentId)
+		{
+			lock (WaitingAgents)
+			{
+				if (WaitingAgents.TryGetValue(AgentId, out CancellationTokenSource? CancellationSource))
+				{
+					CancellationSource.Cancel();
+				}
+			}
+		}
+
+		/// <summary>
+		/// Creates a new agent session
+		/// </summary>
+		/// <param name="Agent">The agent to create a session for</param>
+		/// <param name="Status">Current status of the agent</param>
+		/// <param name="Properties">Properties for the agent</param>
+		/// <param name="Resources">Resources which the agent has</param>
+		/// <param name="Version">Version of the software that's running</param>
+		/// <returns>New agent state</returns>
+		public async Task<IAgent> CreateSessionAsync(IAgent Agent, AgentStatus Status, IReadOnlyList<string> Properties, IReadOnlyDictionary<string, int> Resources, string? Version)
+		{
+			for (; ; )
+			{
+				IAuditLogChannel<AgentId> AgentLogger = Agents.GetLogger(Agent.Id);
+
+				// Check if there's already a session running for this agent.
+				IAgent? NewAgent;
+				if (Agent.SessionId != null)
+				{
+					// Try to terminate the current session
+					await TryTerminateSessionAsync(Agent);
+				}
+				else
+				{
+					DateTime UtcNow = Clock.UtcNow;
+
+					// Remove any outstanding leases
+					foreach (AgentLease Lease in Agent.Leases)
+					{
+						AgentLogger.LogInformation("Removing outstanding lease {LeaseId}", Lease.Id);
+						await RemoveLeaseAsync(Agent, Lease, UtcNow, LeaseOutcome.Failed, null);
+					}
+
+					// Create a new session document
+					ISession NewSession = await Sessions.AddAsync(ObjectId.GenerateNewId(), Agent.Id, Clock.UtcNow, Properties, Resources, Version);
+					DateTime SessionExpiresAt = UtcNow + SessionExpiryTime;
+
+					// Get the new dynamic pools for the agent
+					List<PoolId> DynamicPools = await GetDynamicPoolsAsync(Agent);
+
+					// Reset the agent to use the new session
+					NewAgent = await Agents.TryStartSessionAsync(Agent, NewSession.Id, SessionExpiresAt, Status, Properties, Resources, DynamicPools, Version);
+					if(NewAgent != null)
+					{
+						Agent = NewAgent;
+						AgentLogger.LogInformation("Session {SessionId} started", NewSession.Id);
+						break;
+					}
+
+					// Remove the session we didn't use
+					await Sessions.DeleteAsync(NewSession.Id);
+				}
+
+				// Get the current agent state
+				NewAgent = await GetAgentAsync(Agent.Id);
+				if (NewAgent == null)
+				{
+					throw new InvalidOperationException($"Invalid agent id '{Agent.Id}'");
+				}
+				Agent = NewAgent;
+			}
+			return Agent;
+		}
+
+		async Task<(ITaskSource, AgentLease)?> GuardedWaitForLeaseAsync(ITaskSource Source, IAgent Agent, CancellationToken CancellationToken)
+		{
+			try
+			{
+				AgentLease? Lease = await Source.AssignLeaseAsync(Agent, CancellationToken);
+				return (Lease != null) ? (Source, Lease) : null;
+			}
+			catch (TaskCanceledException)
+			{
+				return null;
+			}
+			catch (Exception Ex)
+			{
+				Logger.LogError(Ex, "Exception while trying to assign lease"); 
+				return null;
+			}
+		}
+
+		/// <summary>
+		/// Waits for a lease to be assigned to an agent
+		/// </summary>
+		/// <param name="Agent">The agent to assign a lease to</param>
+		/// <param name="CancellationToken"></param>
+		/// <returns>True if a lease was assigned, false otherwise</returns>
+		public async Task<IAgent?> WaitForLeaseAsync(IAgent? Agent, CancellationToken CancellationToken)
+		{
+			while (Agent != null)
+			{
+				if (!Agent.SessionExpiresAt.HasValue)
+				{
+					break;
+				}
+
+				// Check we have some time to wait
+				DateTime UtcNow = Clock.UtcNow;
+				TimeSpan MaxWaitTime = (Agent.SessionExpiresAt.Value - SessionExpiryTime + SessionLongPollTime) - UtcNow;
+				if (MaxWaitTime <= TimeSpan.Zero)
+				{
+					break;
+				}
+
+				// Create a cancellation token that will expire with the session
+				using CancellationTokenSource CancellationSource = CancellationTokenSource.CreateLinkedTokenSource(ApplicationLifetime.ApplicationStopping, CancellationToken);
+				CancellationSource.CancelAfter(MaxWaitTime);
+
+				// Assign a new lease
+				(ITaskSource, AgentLease)? Result = null;
+				try
+				{
+					// Add the cancellation source to the set of waiting agents
+					lock (WaitingAgents)
+					{
+						WaitingAgents[Agent.Id] = CancellationSource;
+					}
+
+					// If we're in a maintenance window, just wait for the time to expire
+					if (DowntimeService.IsDowntimeActive)
+					{
+						await AsyncUtils.DelayNoThrow(MaxWaitTime, CancellationToken);
+						break;
+					}
+
+					// Create all the tasks to wait for
+					List<Task<(ITaskSource, AgentLease)?>> Tasks = new List<Task<(ITaskSource, AgentLease)?>>();
+					foreach (ITaskSource TaskSource in TaskSources)
+					{
+						Tasks.Add(GuardedWaitForLeaseAsync(TaskSource, Agent, CancellationSource.Token));
+					}
+
+					// Wait for a lease to be available
+					while (Tasks.Count > 0)
+					{
+						await Task.WhenAny(Tasks);
+
+						for (int Idx = 0; Idx < Tasks.Count; Idx++)
+						{
+							Task<(ITaskSource, AgentLease)?> Task = Tasks[Idx];
+							if (!Task.IsCompleted)
+							{
+								continue;
+							}
+
+							Tasks.RemoveAt(Idx--);
+
+							(ITaskSource, AgentLease)? TaskResult = Task.Result;
+							if (!TaskResult.HasValue)
+							{
+								continue;
+							}
+
+							if (Result == null)
+							{
+								Result = TaskResult;
+								CancellationSource.Cancel();
+							}
+							else
+							{
+								(ITaskSource TaskSource, AgentLease TaskLease) = TaskResult.Value;
+								await TaskSource.CancelLeaseAsync(Agent, TaskLease.Id, Any.Parser.ParseFrom(TaskLease.Payload));
+							}
+						}
+					}
+				}
+				finally
+				{
+					lock (WaitingAgents)
+					{
+						WaitingAgents.Remove(Agent.Id);
+					}
+				}
+
+				// Exit if we didn't find any work to do. It may be that all the task sources returned null, in which case wait for the time period to expire.
+				if (Result == null)
+				{
+					if (!CancellationSource.IsCancellationRequested)
+					{
+						await CancellationSource.Token.AsTask();
+					}
+					break;
+				}
+
+				// Get the resulting lease
+				(ITaskSource Source, AgentLease Lease) = Result.Value;
+				if (Lease == AgentLease.Drain)
+				{
+					await AsyncUtils.DelayNoThrow(MaxWaitTime, CancellationToken);
+					break;
+				}
+
+				// Add the new lease to the agent
+				IAgent? NewAgent = await Agents.TryAddLeaseAsync(Agent, Lease);
+				if (NewAgent != null)
+				{
+					await Source.OnLeaseStartedAsync(NewAgent, Lease.Id, Any.Parser.ParseFrom(Lease.Payload), Agents.GetLogger(Agent.Id));
+					await CreateLeaseAsync(Agent, Lease);
+					return NewAgent;
+				}
+
+				// Update the agent
+				Agent = await GetAgentAsync(Agent.Id);
+			}
+			return Agent;
+		}
+
+		/// <summary>
+		/// Cancels the specified agent lease
+		/// </summary>
+		/// <param name="Agent">Agent to cancel the lease on</param>
+		/// <param name="LeaseId">The lease id to cancel</param>
+		/// <returns></returns>
+		public async Task<bool> CancelLeaseAsync(IAgent Agent, ObjectId LeaseId)
+		{		
+			int Index = 0;
+			while (Index < Agent.Leases.Count && Agent.Leases[Index].Id != LeaseId)
+			{
+				Index++;
+			}
+
+			if (Index == Agent.Leases.Count)
+			{
+				return false;
+			}
+
+			if (Agent.Leases[Index].State == LeaseState.Cancelled)
+			{
+				return false;
+			}
+
+			await Agents.TryCancelLeaseAsync(Agent, Index);
+			return true;
+		}
+
+		/// <summary>
+		/// 
+		/// </summary>
+		public async Task<IAgent?> UpdateSessionWithWaitAsync(IAgent InAgent, ObjectId SessionId, AgentStatus Status, IReadOnlyList<string>? Properties, IReadOnlyDictionary<string, int>? Resources, IList<HordeCommon.Rpc.Messages.Lease> NewLeases, CancellationToken CancellationToken)
+		{
+			IAgent? Agent = InAgent;
+
+			// Capture the current agent update index. This allows us to detect if anything has changed.
+			uint UpdateIndex = Agent.UpdateIndex;
+
+			// Update the agent session and return to the caller if anything changes
+			Agent = await UpdateSessionAsync(Agent, SessionId, Status, Properties, Resources, NewLeases);
+			if (Agent != null && Agent.UpdateIndex == UpdateIndex && (Agent.Leases.Count > 0 || Agent.Status != AgentStatus.Stopping))
+			{
+				Agent = await WaitForLeaseAsync(Agent, CancellationToken);
+			}
+			return Agent;
+		}
+
+		/// <summary>
+		/// Updates the state of the current agent session
+		/// </summary>
+		/// <param name="InAgent">The current agent state</param>
+		/// <param name="SessionId">Id of the session</param>
+		/// <param name="Status">New status for the agent</param>
+		/// <param name="Properties">New agent properties</param>
+		/// <param name="Resources">New agent resources</param>
+		/// <param name="NewLeases">New list of leases for this session</param>
+		/// <returns>Updated agent state</returns>
+		public async Task<IAgent?> UpdateSessionAsync(IAgent InAgent, ObjectId SessionId, AgentStatus Status, IReadOnlyList<string>? Properties, IReadOnlyDictionary<string, int>? Resources, IList<HordeCommon.Rpc.Messages.Lease> NewLeases)
+		{
+			DateTime UtcNow = Clock.UtcNow;
+
+			Stopwatch Timer = Stopwatch.StartNew();
+
+			IAgent? Agent = InAgent;
+			while (Agent != null)
+			{
+				// If the agent is stopping and doesn't have any leases, we can terminate the current session.
+				if (Status == AgentStatus.Stopping && NewLeases.Count == 0)
+				{
+					// If we've already decided to terminate the session, this update is redundant but harmless
+					if (Agent.SessionId != SessionId)
+					{
+						break;
+					}
+
+					// If the session is valid, we can terminate once the agent leases are also empty
+					if (Agent.Leases.Count == 0)
+					{
+						if (!await TryTerminateSessionAsync(Agent))
+						{
+							Agent = await GetAgentAsync(Agent.Id);
+							continue;
+						}
+						break;
+					}
+				}
+
+				// Check the session id is correct.
+				if (Agent.SessionId != SessionId)
+				{
+					throw new InvalidOperationException($"Invalid agent session {SessionId}");
+				}
+
+				// Check the session hasn't expired
+				if (!Agent.IsSessionValid(UtcNow))
+				{
+					throw new InvalidOperationException("Session has already expired");
+				}
+
+				// Extend the current session time if we're within a time period of the current time expiring. This reduces
+				// unnecessary DB writes a little, but it also allows us to skip the update and jump into a long poll state
+				// if there's still time on the session left.
+				DateTime? SessionExpiresAt = null;
+				if (!Agent.SessionExpiresAt.HasValue || UtcNow > (Agent.SessionExpiresAt - SessionExpiryTime) + SessionRenewTime)
+				{
+					SessionExpiresAt = UtcNow + SessionExpiryTime;
+				}
+
+				// Flag for whether the leases array should be updated
+				bool bUpdateLeases = false;
+				List<AgentLease> Leases = new List<AgentLease>(Agent.Leases);
+
+				// Remove any completed leases from the agent
+				Dictionary<ObjectId, HordeCommon.Rpc.Messages.Lease> LeaseIdToNewState = NewLeases.ToDictionary(x => x.Id.ToObjectId(), x => x);
+				for (int Idx = 0; Idx < Leases.Count; Idx++)
+				{
+					AgentLease Lease = Leases[Idx];
+					if (Lease.State == LeaseState.Cancelled)
+					{
+						HordeCommon.Rpc.Messages.Lease? NewLease;
+						if (!LeaseIdToNewState.TryGetValue(Lease.Id, out NewLease) || NewLease.State == LeaseState.Cancelled || NewLease.State == LeaseState.Completed)
+						{
+							await RemoveLeaseAsync(Agent, Lease, UtcNow, LeaseOutcome.Cancelled, null);
+							Leases.RemoveAt(Idx--);
+							bUpdateLeases = true;
+						}
+					}
+					else
+					{
+						HordeCommon.Rpc.Messages.Lease? NewLease;
+						if (LeaseIdToNewState.TryGetValue(Lease.Id, out NewLease) && NewLease.State != Lease.State)
+						{
+							if (NewLease.State == LeaseState.Cancelled || NewLease.State == LeaseState.Completed)
+							{
+								await RemoveLeaseAsync(Agent, Lease, UtcNow, NewLease.Outcome, NewLease.Output.ToByteArray());
+								Leases.RemoveAt(Idx--);
+							}
+							else if (NewLease.State == LeaseState.Active && Lease.State == LeaseState.Pending)
+							{
+								Lease.State = LeaseState.Active;
+							}
+							bUpdateLeases = true;
+						}
+					}
+				}
+
+				// If the agent is stopping, cancel all the leases. Clear out the current session once it's complete.
+				if (Status == AgentStatus.Stopping)
+				{
+					foreach (AgentLease Lease in Leases)
+					{
+						if (Lease.State != LeaseState.Cancelled)
+						{
+							Lease.State = LeaseState.Cancelled;
+							bUpdateLeases = true;
+						}
+					}
+				}
+
+				// Get the new dynamic pools for the agent
+				List<PoolId> DynamicPools = await GetDynamicPoolsAsync(Agent);
+
+				// Update the agent, and try to create new lease documents if we succeed
+				IAgent? NewAgent = await Agents.TryUpdateSessionAsync(Agent, Status, SessionExpiresAt, Properties, Resources, DynamicPools, bUpdateLeases ? Leases : null);
+				if (NewAgent != null)
+				{
+					Agent = NewAgent;
+					break;
+				}
+
+				// Fetch the agent again
+				Agent = await GetAgentAsync(Agent.Id);
+			}
+			return Agent;
+		}
+
+		/// <summary>
+		/// Terminates an existing session. Does not update the agent itself, if it's currently 
+		/// </summary>
+		/// <param name="Agent">The agent whose current session should be terminated</param>
+		/// <returns>True if the session was terminated</returns>
+		private async Task<bool> TryTerminateSessionAsync(IAgent Agent)
+		{
+			// Make sure the agent has a valid session id
+			if (Agent.SessionId == null)
+			{
+				return true;
+			}
+
+			// Get the time that the session finishes at
+			DateTime FinishTime = Clock.UtcNow;
+			if (Agent.SessionExpiresAt.HasValue && Agent.SessionExpiresAt.Value < FinishTime)
+			{
+				FinishTime = Agent.SessionExpiresAt.Value;
+			}
+
+			// Save off the session id and current leases
+			ObjectId SessionId = Agent.SessionId.Value;
+			List<AgentLease> Leases = new List<AgentLease>(Agent.Leases);
+
+			// Clear the current session
+			IAgent? NewAgent = await Agents.TryTerminateSessionAsync(Agent);
+			if (NewAgent != null)
+			{
+				Agent = NewAgent;
+
+				// Remove any outstanding leases
+				foreach(AgentLease Lease in Leases)
+				{
+					Agents.GetLogger(Agent.Id).LogInformation("Removing lease {LeaseId} during session terminate...", Lease.Id);
+					await RemoveLeaseAsync(Agent, Lease, FinishTime, LeaseOutcome.Failed, null);
+				}
+
+				// Update the session document
+				Agents.GetLogger(Agent.Id).LogInformation("Terminated session {SessionId}", SessionId);
+				await Sessions.UpdateAsync(SessionId, FinishTime, Agent.Properties, Agent.Resources);
+				return true;
+			}
+			return false;
+		}
+
+		/// <summary>
+		/// Creates a new lease document
+		/// </summary>
+		/// <param name="Agent">Agent that will be executing the lease</param>
+		/// <param name="AgentLease">The new agent lease</param>
+		/// <returns>New lease document</returns>
+		private Task<ILease> CreateLeaseAsync(IAgent Agent, AgentLease AgentLease)
+		{
+			try
+			{
+				return Leases.AddAsync(AgentLease.Id, AgentLease.Name, Agent.Id, Agent.SessionId!.Value, AgentLease.StreamId, AgentLease.PoolId, AgentLease.LogId, AgentLease.StartTime, AgentLease.Payload!);
+			}
+			catch (Exception Ex)
+			{
+				Logger.LogError(Ex, "Unable to create lease {LeaseId} for agent {AgentId}; lease already exists?", AgentLease.Id, Agent.Id);
+				throw;
+			}
+		}
+
+		/// <summary>
+		/// Finds all leases matching a set of criteria
+		/// </summary>
+		/// <param name="AgentId">Unqiue id of the agent executing this lease</param>
+		/// <param name="SessionId">Unique id of the agent session</param>
+		/// <param name="StartTime">Start of the search window to return results for</param>
+		/// <param name="FinishTime">End of the search window to return results for</param>
+		/// <param name="Index">Index of the first result to return</param>
+		/// <param name="Count">Number of results to return</param>
+		/// <returns>List of leases matching the given criteria</returns>
+		public Task<List<ILease>> FindLeasesAsync(AgentId? AgentId, ObjectId? SessionId, DateTime? StartTime, DateTime? FinishTime, int Index, int Count)
+		{
+			return Leases.FindLeasesAsync(AgentId, SessionId, StartTime, FinishTime, Index, Count);
+		}
+
+		/// <summary>
+		/// Gets a specific lease
+		/// </summary>
+		/// <param name="LeaseId">Unique id of the lease</param>
+		/// <returns>The lease that was found, or null if it does not exist</returns>
+		public Task<ILease?> GetLeaseAsync(ObjectId LeaseId)
+		{
+			return Leases.GetAsync(LeaseId);
+		}
+
+		/// <summary>
+		/// Removes a lease with the given id. Updates the lease state in the database, and removes the item from the agent's leases array.
+		/// </summary>
+		/// <param name="Agent">The agent to remove a lease from</param>
+		/// <param name="Lease">The lease to cancel</param>
+		/// <param name="UtcNow">The current time</param>
+		/// <param name="Outcome">Final status of the lease</param>
+		/// <param name="Output">Output from executing the task</param>
+		/// <returns>Async task</returns>
+		private async Task RemoveLeaseAsync(IAgent Agent, AgentLease Lease, DateTime UtcNow, LeaseOutcome Outcome, byte[]? Output)
+		{
+			// Make sure the lease is terminated correctly
+			if (Lease.Payload == null)
+			{
+				Logger.LogWarning("Removing lease {LeaseId} (no payload)", Lease.Id);
+			}
+			else
+			{
+				Any Any = Any.Parser.ParseFrom(Lease.Payload);
+				Logger.LogInformation("Removing lease {LeaseId} ({LeaseType})", Lease.Id, Any.TypeUrl);
+
+				foreach (ITaskSource TaskSource in TaskSources)
+				{
+					if (Any.Is(TaskSource.Descriptor))
+					{
+						await TaskSource.OnLeaseFinishedAsync(Agent, Lease.Id, Any, Outcome, Output, Agents.GetLogger(Agent.Id));
+						break;
+					}
+				}
+			}
+
+			// Figure out what time the lease finished
+			DateTime FinishTime = UtcNow;
+			if (Agent.SessionExpiresAt.HasValue && Agent.SessionExpiresAt.Value < FinishTime)
+			{
+				FinishTime = Agent.SessionExpiresAt.Value;
+			}
+			if (Lease.ExpiryTime.HasValue && Lease.ExpiryTime.Value < FinishTime)
+			{
+				FinishTime = Lease.ExpiryTime.Value;
+			}
+
+			// Update the lease
+			await Leases.TrySetOutcomeAsync(Lease.Id, FinishTime, Outcome, Output);
+		}
+
+		/// <summary>
+		/// Gets information about a particular session
+		/// </summary>
+		/// <param name="SessionId">The unique session id</param>
+		/// <returns>The session information</returns>
+		public Task<ISession?> GetSessionAsync(ObjectId SessionId)
+		{
+			return Sessions.GetAsync(SessionId);
+		}
+
+		/// <summary>
+		/// Find sessions for the given agent
+		/// </summary>
+		/// <param name="AgentId">The unique agent id</param>
+		/// <param name="StartTime">Start time to include in the search</param>
+		/// <param name="FinishTime">Finish time to include in the search</param>
+		/// <param name="Index">Index of the first result to return</param>
+		/// <param name="Count">Number of results to return</param>
+		/// <returns>List of sessions matching the given criteria</returns>
+		public Task<List<ISession>> FindSessionsAsync(AgentId AgentId, DateTime? StartTime, DateTime? FinishTime, int Index, int Count)
+		{
+			return Sessions.FindAsync(AgentId, StartTime, FinishTime, Index, Count);
+		}
+
+		/// <summary>
+		/// Terminate any sessions for agents that are offline
+		/// </summary>
+		/// <param name="StoppingToken">Token indicating the service is shutting down</param>
+		/// <returns>Async task</returns>
+		protected override async Task TickAsync(CancellationToken StoppingToken)
+		{
+			while (!StoppingToken.IsCancellationRequested)
+			{
+				// Find all the agents which are ready to be expired
+				const int MaxAgents = 100;
+				DateTime UtcNow = Clock.UtcNow;
+				List<IAgent> ExpiredAgents = await Agents.FindExpiredAsync(UtcNow, MaxAgents);
+
+				// Transition each agent to being offline
+				foreach (IAgent ExpiredAgent in ExpiredAgents)
+				{
+					StoppingToken.ThrowIfCancellationRequested();
+					Logger.LogDebug("Terminating session {SessionId} for agent {Agent}", ExpiredAgent.SessionId, ExpiredAgent.Id);
+					await TryTerminateSessionAsync(ExpiredAgent);
+				}
+
+				// Try again if we didn't fetch everything
+				if(ExpiredAgents.Count < MaxAgents)
+				{
+					break;
+				}
+			}
+
+			await CollectMetrics();
+		}
+
+		private async Task CollectMetrics()
+		{
+			List<IAgent> AgentList = await Agents.FindAsync();
+			int NumAgentsTotal = AgentList.Count;
+			int NumAgentsTotalEnabled = AgentList.Count(a => a.Enabled);
+			int NumAgentsTotalDisabled = AgentList.Count(a => !a.Enabled);
+			int NumAgentsTotalOk = AgentList.Count(a => a.Enabled && a.Status == AgentStatus.Ok);
+			int NumAgentsTotalStopping = AgentList.Count(a => a.Enabled && a.Status == AgentStatus.Stopping);
+			int NumAgentsTotalUnhealthy = AgentList.Count(a => a.Enabled && a.Status == AgentStatus.Unhealthy);
+			int NumAgentsTotalUnspecified = AgentList.Count(a => a.Enabled && a.Status == AgentStatus.Unspecified);
+			
+			// TODO: utilize tags argument in a smarter way below
+			DogStatsd.Gauge("agents.total.count", NumAgentsTotal);
+			DogStatsd.Gauge("agents.total.enabled.count", NumAgentsTotalEnabled);
+			DogStatsd.Gauge("agents.total.disabled.count", NumAgentsTotalDisabled);
+			DogStatsd.Gauge("agents.total.ok.count", NumAgentsTotalOk);
+			DogStatsd.Gauge("agents.total.stopping.count", NumAgentsTotalStopping);
+			DogStatsd.Gauge("agents.total.unhealthy.count", NumAgentsTotalUnhealthy);
+			DogStatsd.Gauge("agents.total.unspecified.count", NumAgentsTotalUnspecified);
+		}
+
+		/// <summary>
+		/// Determines if the user is authorized to perform an action on a particular agent
+		/// </summary>
+		/// <param name="Agent">The agent to check</param>
+		/// <param name="Action">The action being performed</param>
+		/// <param name="User">The principal to authorize</param>
+		/// <param name="Cache">The permissions cache</param>
+		/// <returns>True if the action is authorized</returns>
+		public async Task<bool> AuthorizeAsync(IAgent Agent, AclAction Action, ClaimsPrincipal User, GlobalPermissionsCache? Cache)
+		{
+			bool? Result = Agent.Acl?.Authorize(Action, User);
+			if (Result == null)
+			{
+				return await AclService.AuthorizeAsync(Action, User, Cache);
+			}
+			else
+			{
+				return Result.Value;
+			}
+		}
+
+		/// <summary>
+		/// Determines if the user is authorized to perform an action on a particular agent
+		/// </summary>
+		/// <param name="Agent">The agent to check</param>
+		/// <param name="User">The principal to authorize</param>
+		/// <returns>True if the action is authorized</returns>
+		public bool AuthorizeSession(IAgent Agent, ClaimsPrincipal User)
+		{
+			if (Agent.SessionId != null && User.HasSessionClaim(Agent.SessionId.Value) && Agent.IsSessionValid(Clock.UtcNow))
+			{
+				return true;
+			}
+			else
+			{
+				return false;
+			}
+		}
+	}
+}
diff --git a/Engine/Source/Programs/Horde/HordeServer/Collections/IAgentCollection.cs b/Engine/Source/Programs/Horde/HordeServer/Collections/IAgentCollection.cs
deleted file mode 100644
diff --git a/Engine/Source/Programs/Horde/HordeServer/Collections/Impl/AgentCollection.cs b/Engine/Source/Programs/Horde/HordeServer/Collections/Impl/AgentCollection.cs
deleted file mode 100644
diff --git a/Engine/Source/Programs/Horde/HordeServer/Models/IAgent.cs b/Engine/Source/Programs/Horde/HordeServer/Models/IAgent.cs
deleted file mode 100644
diff --git a/Engine/Source/Programs/Horde/HordeServer/Services/AgentService.cs b/Engine/Source/Programs/Horde/HordeServer/Services/AgentService.cs
deleted file mode 100644
