Horde: Move logs and devices into their own folder.

diff --git a/Engine/Source/Programs/Horde/HordeServer/Agents/AgentId.cs b/Engine/Source/Programs/Horde/HordeServer/Agents/AgentId.cs
new file mode 100644
--- /dev/null
+++ b/Engine/Source/Programs/Horde/HordeServer/Agents/AgentId.cs
@@ -0,0 +1,185 @@
+// Copyright Epic Games, Inc. All Rights Reserved.
+
+using EpicGames.Redis;
+using HordeServer.Models;
+using MongoDB.Bson.Serialization;
+using MongoDB.Bson.Serialization.Attributes;
+using MongoDB.Bson.Serialization.Serializers;
+using StackExchange.Redis;
+using System;
+using System.Collections.Generic;
+using System.ComponentModel;
+using System.Diagnostics.CodeAnalysis;
+using System.Globalization;
+using System.Linq;
+using System.Text;
+using System.Text.Json;
+using System.Text.Json.Serialization;
+using System.Threading.Tasks;
+
+namespace HordeServer.Utilities
+{
+	/// <summary>
+	/// Normalized hostname of an agent
+	/// </summary>
+	[RedisConverter(typeof(AgentIdRedisConverter))]
+	[BsonSerializer(typeof(AgentIdBsonSerializer))]
+	public struct AgentId : IEquatable<AgentId>
+	{
+		/// <summary>
+		/// The text representing this id
+		/// </summary>
+		readonly string Name;
+
+		/// <summary>
+		/// Constructor
+		/// </summary>
+		/// <param name="Name">Hostname of the agent</param>
+		public AgentId(string Name)
+		{
+			if (Name.Length == 0)
+			{
+				throw new ArgumentException("Agent name may not be empty");
+			}
+
+			char[] Result = new char[Name.Length];
+			for (int Idx = 0; Idx < Name.Length; Idx++)
+			{
+				char Character = Name[Idx];
+				if ((Character >= 'A' && Character <= 'Z') || (Character >= '0' && Character <= '9') || Character == '_' || Character == '-')
+				{
+					Result[Idx] = Character;
+				}
+				else if (Character >= 'a' && Character <= 'z')
+				{
+					Result[Idx] = (char)('A' + (Character - 'a'));
+				}
+				else if (Character == '.')
+				{
+					if (Idx == 0)
+					{
+						throw new ArgumentException("Agent name may not begin with the '.' character");
+					}
+					else if (Idx == Name.Length - 1)
+					{
+						throw new ArgumentException("Agent name may not end with the '.' character");
+					}
+					else if (Name[Idx - 1] == '.')
+					{
+						throw new ArgumentException("Agent name may not contain two consecutive '.' characters");
+					}
+					else
+					{
+						Result[Idx] = Character;
+					}
+				}
+				else
+				{
+					throw new ArgumentException($"Character '{Character}' in agent '{Name}' is invalid");
+				}
+			}
+			this.Name = new string(Result);
+		}
+
+		/// <inheritdoc/>
+		public override bool Equals(object? Obj)
+		{
+			return Obj is AgentId && Equals((AgentId)Obj);
+		}
+
+		/// <inheritdoc/>
+		public override int GetHashCode()
+		{
+			return Name.GetHashCode(StringComparison.Ordinal);
+		}
+
+		/// <inheritdoc/>
+		public bool Equals(AgentId Other)
+		{
+			return Name.Equals(Other.Name, StringComparison.Ordinal);
+		}
+
+		/// <inheritdoc/>
+		public override string ToString()
+		{
+			return Name;
+		}
+
+		/// <summary>
+		/// Compares two string ids for equality
+		/// </summary>
+		/// <param name="Left">The first string id</param>
+		/// <param name="Right">Second string id</param>
+		/// <returns>True if the two string ids are equal</returns>
+		public static bool operator ==(AgentId Left, AgentId Right)
+		{
+			return Left.Equals(Right);
+		}
+
+		/// <summary>
+		/// Compares two string ids for inequality
+		/// </summary>
+		/// <param name="Left">The first string id</param>
+		/// <param name="Right">Second string id</param>
+		/// <returns>True if the two string ids are not equal</returns>
+		public static bool operator !=(AgentId Left, AgentId Right)
+		{
+			return !Left.Equals(Right);
+		}
+	}
+
+	/// <summary>
+	/// Converter to/from Redis values
+	/// </summary>
+	public sealed class AgentIdRedisConverter : IRedisConverter<AgentId>
+	{
+		/// <inheritdoc/>
+		public AgentId FromRedisValue(RedisValue Value) => new AgentId((string)Value);
+
+		/// <inheritdoc/>
+		public RedisValue ToRedisValue(AgentId Value) => Value.ToString();
+	}
+
+	/// <summary>
+	/// Serializer for StringId objects
+	/// </summary>
+	public sealed class AgentIdBsonSerializer : SerializerBase<AgentId>
+	{
+		/// <inheritdoc/>
+		public override AgentId Deserialize(BsonDeserializationContext Context, BsonDeserializationArgs Args)
+		{
+			string Argument;
+			if (Context.Reader.CurrentBsonType == MongoDB.Bson.BsonType.ObjectId)
+			{
+				Argument = Context.Reader.ReadObjectId().ToString();
+			}
+			else
+			{
+				Argument = Context.Reader.ReadString();
+			}
+			return new AgentId(Argument);
+		}
+
+		/// <inheritdoc/>
+		public override void Serialize(BsonSerializationContext Context, BsonSerializationArgs Args, AgentId Value)
+		{
+			Context.Writer.WriteString(Value.ToString());
+		}
+	}
+
+	/// <summary>
+	/// Extension methods for AgentId types
+	/// </summary>
+	static class AgentIdExtensions
+	{
+		/// <summary>
+		/// Create an AgentId from a string
+		/// </summary>
+		/// <param name="Text">String to parse</param>
+		/// <returns>New agent id</returns>
+		public static AgentId ToAgentId(this string Text)
+		{
+			return new AgentId(Text);
+		}
+	}
+}
diff --git a/Engine/Source/Programs/Horde/HordeServer/Collections/IDeviceCollection.cs b/Engine/Source/Programs/Horde/HordeServer/Collections/IDeviceCollection.cs
deleted file mode 100644
diff --git a/Engine/Source/Programs/Horde/HordeServer/Collections/ILogEventCollection.cs b/Engine/Source/Programs/Horde/HordeServer/Collections/ILogEventCollection.cs
deleted file mode 100644
diff --git a/Engine/Source/Programs/Horde/HordeServer/Collections/ILogFileCollection.cs b/Engine/Source/Programs/Horde/HordeServer/Collections/ILogFileCollection.cs
deleted file mode 100644
diff --git a/Engine/Source/Programs/Horde/HordeServer/Collections/Impl/DeviceCollection.cs b/Engine/Source/Programs/Horde/HordeServer/Collections/Impl/DeviceCollection.cs
deleted file mode 100644
diff --git a/Engine/Source/Programs/Horde/HordeServer/Collections/Impl/LogEventCollection.cs b/Engine/Source/Programs/Horde/HordeServer/Collections/Impl/LogEventCollection.cs
deleted file mode 100644
diff --git a/Engine/Source/Programs/Horde/HordeServer/Collections/Impl/LogFileCollection.cs b/Engine/Source/Programs/Horde/HordeServer/Collections/Impl/LogFileCollection.cs
deleted file mode 100644
diff --git a/Engine/Source/Programs/Horde/HordeServer/Devices/IDevice.cs b/Engine/Source/Programs/Horde/HordeServer/Devices/IDevice.cs
new file mode 100644
--- /dev/null
+++ b/Engine/Source/Programs/Horde/HordeServer/Devices/IDevice.cs
@@ -0,0 +1,277 @@
+// Copyright Epic Games, Inc. All Rights Reserved.
+
+using System;
+using MongoDB.Bson;
+using MongoDB.Bson.Serialization.Attributes;
+
+using System.Collections.Generic;
+using HordeServer.Utilities;
+
+namespace HordeServer.Models
+{
+	using DeviceId = StringId<IDevice>;
+	using DevicePlatformId = StringId<IDevicePlatform>;
+	using DevicePoolId = StringId<IDevicePool>;
+
+	/// <summary>
+	/// A reservation containing one or more devices
+	/// </summary>
+	public interface IDeviceReservation
+	{
+		/// <summary>
+		/// Randomly generated unique id for this reservation
+		/// </summary>
+		public ObjectId Id { get; }
+
+		/// <summary>
+		/// Which device pool the reservation is in
+		/// </summary>
+		public DevicePoolId PoolId { get; }
+
+		/// <summary>
+		/// JobID holding reservation
+		/// </summary>
+		public string? JobId { get; }
+
+		/// <summary>
+		/// Job step id holding reservation
+		/// </summary>
+		public string? StepId { get; }
+
+		/// <summary>
+		/// Reservations held by a user, requires a token
+		/// </summary>
+		public ObjectId? UserId { get; }
+
+		/// <summary>
+		/// The hostname of machine holding reservation
+		/// </summary>
+		string? Hostname { get; }
+
+		/// <summary>
+		/// The hostname of machine holding reservation
+		/// </summary>
+		string? ReservationDetails { get; }
+
+		/// <summary>
+		/// The UTC time when the reservation was created
+		/// </summary>
+		DateTime CreateTimeUtc { get; }
+
+		/// <summary>
+		/// The last update time for reservation renewal (for expiration)
+		/// </summary>
+		DateTime UpdateTimeUtc { get; }
+
+		/// <summary>
+		/// The reserved devices
+		/// </summary>
+		public List<DeviceId> Devices { get; }
+
+		/// <summary>
+		/// The legacy reservation system guid, to be removed once can update Gauntlet client in all streams
+		/// </summary>
+		public string LegacyGuid { get; }
+	}
+
+
+	/// <summary>
+	/// A device platform 
+	/// </summary>
+	public interface IDevicePlatform
+	{
+		/// <summary>
+		/// Unique identifier of device platform 
+		/// </summary>
+		public DevicePlatformId Id { get; }
+		
+		/// <summary>
+		/// Device platform name, for example Android, PS5, etc
+		/// </summary>
+		string Name { get; }
+
+		/// <summary>
+		/// A list of valid models for the platform
+		/// </summary>
+		public IReadOnlyList<string>? Models { get; }
+
+	}
+
+    /// <summary>
+    /// The type of device pool
+    /// </summary>
+    public enum DevicePoolType
+    {
+        /// <summary>
+        /// Available to CIS jobs
+        /// </summary>
+        Automation,
+
+        /// <summary>
+        /// Shared by users with remote checking and checkouts
+        /// </summary>
+        Shared
+
+    }
+
+    /// <summary>
+    /// A logical pool of devices
+    /// </summary>
+    public interface IDevicePool
+	{
+		/// <summary>
+		/// Unique identifier of pool
+		/// </summary>
+		public DevicePoolId Id { get; }
+
+		/// <summary>
+		/// The type of pool 
+		/// </summary>
+		public DevicePoolType PoolType { get; }
+
+        /// <summary>
+        /// Friendly name of the pool
+        /// </summary>
+        string Name { get; }
+
+		/// <summary>
+		/// Acl access to the pool
+		/// </summary>
+		public Acl? Acl { get; }
+	}
+
+	/// <summary>
+	/// A device utilization snapshot
+	/// </summary>
+	public class DeviceUtilizationTelemetry
+	{
+		/// <summary>
+		/// The job id which utilized device
+		/// </summary>
+		public string? JobId { get; set;}
+
+		/// <summary>
+		/// The job's step id
+		/// </summary>
+		public string? StepId { get; set;}
+
+		/// <summary>
+		/// The time device was reserved
+		/// </summary>
+		public DateTime ReservationStartUtc { get; set;}
+
+		/// <summary>
+		/// The time device was freed
+		/// </summary>
+		public DateTime? ReservationFinishUtc { get; set;}
+
+		/// <summary>
+		/// Private constructor
+		/// </summary>
+		[BsonConstructor]
+		private DeviceUtilizationTelemetry()
+		{
+		}
+
+		/// <summary>
+		/// Constructor
+		/// </summary>
+		/// <param name="ReservationStartUtc"></param>
+		public DeviceUtilizationTelemetry( DateTime ReservationStartUtc )
+		{
+			this.ReservationStartUtc = ReservationStartUtc;
+		}
+
+	}
+
+
+	/// <summary>
+	/// A physical device
+	/// </summary>
+	public interface IDevice
+	{
+		/// <summary>
+		/// The unique id of the device
+		/// </summary>
+		public DeviceId Id { get; }
+
+		/// <summary>
+		/// The platform of the device
+		/// </summary>
+		public DevicePlatformId PlatformId { get; }
+
+		/// <summary>
+		/// Which pool the device belongs to
+		/// </summary>
+		public DevicePoolId PoolId { get; }
+
+		/// <summary>
+		/// Friendly name of the device
+		/// </summary>
+		public string Name { get; }
+
+		/// <summary>
+		/// The model of the device
+		/// </summary>
+		public string? ModelId { get; }
+
+
+		/// <summary>
+		/// Address if the device supports shared network connections
+		/// </summary>
+		public string? Address { get; }
+
+		/// <summary>
+		/// Whether the device is currently enabled
+		/// </summary>
+		public bool Enabled { get; }
+
+		/// <summary>
+		/// Id of the user that last modified this device
+		/// </summary>
+		public string? ModifiedByUser { get; }
+
+		/// <summary>
+		/// Id of the user that has this device checked out
+		/// </summary>
+		public string? CheckedOutByUser { get; }
+
+		/// <summary>
+		/// The last time this device was checked out
+		/// </summary>
+		public DateTime? CheckOutTime { get; }
+
+
+		/// <summary>
+		/// The last time a problem was reported
+		/// </summary>
+		public DateTime? ProblemTimeUtc { get; }
+
+		/// <summary>
+		/// The time device was marked for maintenance
+		/// </summary>
+		public DateTime? MaintenanceTimeUtc { get; }
+
+		/// <summary>
+		/// Markdown notes for device if any 
+		/// </summary>
+		public string? Notes { get; }
+
+		/// <summary>
+		/// Device job utilization history 
+		/// </summary>
+		public List<DeviceUtilizationTelemetry>? Utilization { get; set; }
+
+		/// <summary>
+		/// ACL for modifying this device
+		/// </summary>
+		public Acl? Acl { get; }
+	}
+}
+
+
+
+
+
+
+
diff --git a/Engine/Source/Programs/Horde/HordeServer/Devices/IDeviceCollection.cs b/Engine/Source/Programs/Horde/HordeServer/Devices/IDeviceCollection.cs
new file mode 100644
--- /dev/null
+++ b/Engine/Source/Programs/Horde/HordeServer/Devices/IDeviceCollection.cs
@@ -0,0 +1,207 @@
+// Copyright Epic Games, Inc. All Rights Reserved.
+
+using HordeServer.Models;
+using HordeServer.Utilities;
+using MongoDB.Bson;
+using System.Collections.Generic;
+using System.Threading.Tasks;
+
+namespace HordeServer.Collections
+{
+	using DeviceId = StringId<IDevice>;
+	using DevicePlatformId = StringId<IDevicePlatform>;
+	using DevicePoolId = StringId<IDevicePool>;
+
+	/// <summary>
+	/// Device reservation request data
+	/// </summary>
+	public class DeviceRequestData
+	{
+		/// <summary>
+		/// The platform of the device to reserve
+		/// </summary>
+		public DevicePlatformId PlatformId { get; set; }
+
+
+		/// <summary>
+		/// Models to include for this request
+		/// </summary>
+		public List<string> IncludeModels { get; set; }
+
+		/// <summary>
+		/// Models to exclude for this request
+		/// </summary>
+		public List<string> ExcludeModels { get; set; }
+
+		/// <summary>
+		/// Constructor
+		/// </summary>
+		public DeviceRequestData(DevicePlatformId PlatformId, List<string>? IncludeModels = null, List<string>? ExcludeModels = null)
+		{
+			this.PlatformId = PlatformId;
+			this.IncludeModels = IncludeModels ?? new List<string>();
+			this.ExcludeModels = ExcludeModels ?? new List<string>();
+		}
+	}
+
+	/// <summary>
+	/// A collection for device management
+	/// </summary>
+	public interface IDeviceCollection
+	{
+		// PLATFORMS
+
+		/// <summary>
+		/// Add a platform 
+		/// </summary>
+		/// <param name="Id">The id of the platform</param>
+		/// <param name="Name">The friendly name of the platform</param>
+		Task<IDevicePlatform?> TryAddPlatformAsync(DevicePlatformId Id, string Name);
+
+		/// <summary>
+		/// Get a list of all available device platforms
+		/// </summary>		
+		Task<List<IDevicePlatform>> FindAllPlatformsAsync();
+
+		/// <summary>
+		/// Get a specific platform by id
+		/// </summary>
+		Task<IDevicePlatform?> GetPlatformAsync(DevicePlatformId PlatformId);
+
+		/// <summary>
+		/// Update a device platform
+		/// </summary>
+		/// <param name="platformId">The id of the device</param>
+		/// <param name="ModelIds">The available model ids for the platform</param>
+		Task<bool> UpdatePlatformAsync(DevicePlatformId platformId, string[]? ModelIds);
+
+		// POOLS
+
+		/// <summary>
+		/// Add a new device pool to the collection
+		/// </summary>
+		/// <param name="Id">The id of the new pool</param>
+		/// <param name="Name">The friendly name of the new pool</param>
+        /// <param name="PoolType">The pool type</param>
+		Task<IDevicePool?> TryAddPoolAsync(DevicePoolId Id, string Name, DevicePoolType PoolType);
+
+		/// <summary>
+		/// Get a pool by id
+		/// </summary>
+		Task<IDevicePool?> GetPoolAsync(DevicePoolId PoolId);
+
+		/// <summary>
+		/// Gets a list of existing device pools
+		/// </summary>		
+		Task<List<IDevicePool>> FindAllPoolsAsync();
+
+		// DEVICES
+
+		/// <summary>
+		/// Get a device by id
+		/// </summary>
+		Task<IDevice?> GetDeviceAsync(DeviceId DeviceId);
+
+		/// <summary>
+		/// Get a device by name
+		/// </summary>
+		Task<IDevice?> GetDeviceByNameAsync(string DeviceName);
+
+		/// <summary>
+		/// Get a list of all devices
+		/// </summary>
+		/// <param name="DeviceIds">Optional list of device ids to get</param>
+		Task<List<IDevice>> FindAllDevicesAsync(List<DeviceId>? DeviceIds);
+
+		/// <summary>
+		/// Adds a new device
+		/// </summary>
+		/// <param name="Id">The device id</param>
+		/// <param name="Name">The name of the device (unique)</param>
+		/// <param name="PlatformId">The platform of the device</param>
+		/// <param name="PoolId">Which pool to add it to</param>
+		/// <param name="Enabled">Whather the device is enabled by default</param>
+		/// <param name="Address">The network address or hostname the device can be reached at</param>
+		/// <param name="ModelId">The vendor model id of the device</param>
+        /// <param name="UserId">The user adding the device</param>
+		Task<IDevice?> TryAddDeviceAsync(DeviceId Id, string Name, DevicePlatformId PlatformId, DevicePoolId PoolId, bool? Enabled, string? Address, string? ModelId, ObjectId? UserId);
+
+		/// <summary>
+		/// Update a device
+		/// </summary>
+		/// <param name="DeviceId">The id of the device to update</param>
+		/// <param name="NewPoolId">The new pool to assign</param>
+		/// <param name="NewName">The new name to assign</param>		
+		/// <param name="NewAddress">The new device address or hostname</param>
+		/// <param name="NewModelId">The new model id</param>
+		/// <param name="NewNotes">The devices markdown notes</param>
+		/// <param name="NewEnabled">Whether the device is enabled or not</param>
+		/// <param name="NewProblem">Whether to set or clear problem state</param>
+		/// <param name="NewMaintenance">Whether to set or clear maintenance state</param>
+        /// <param name="ModifiedByUserId">The user who is updating the device</param>
+		Task UpdateDeviceAsync(DeviceId DeviceId, DevicePoolId? NewPoolId, string? NewName, string? NewAddress, string? NewModelId, string? NewNotes, bool? NewEnabled, bool? NewProblem, bool? NewMaintenance, ObjectId? ModifiedByUserId = null);
+
+		/// <summary>
+		/// Delete a device from the collection
+		/// </summary>
+		/// <param name="DeviceId">The id of the device to delete</param>
+		Task<bool> DeleteDeviceAsync(DeviceId DeviceId);
+
+		/// <summary>
+        /// Checkout or checkin the specified device
+        /// </summary>
+        /// <param name="DeviceId"></param>
+        /// <param name="CheckedOutByUserId"></param>
+        /// <returns></returns>
+        Task CheckoutDeviceAsync(DeviceId DeviceId, ObjectId? CheckedOutByUserId);
+
+        // RESERVATIONS
+
+        /// <summary>
+        /// Create a new reseveration in the pool with the specified devices
+        /// </summary>
+        /// <param name="PoolId">The pool of devices to use for the new reservation</param>
+        /// <param name="Request">The requested devices for the reservation</param>
+        /// <param name="Hostname">The hostname of the machine making the reservation</param>
+        /// <param name="ReservationDetails">The details of the reservation</param>
+        /// <param name="JobId">The Job Id associated with the job</param>
+        /// <param name="StepId">The Step Id associated with the job</param>
+        Task<IDeviceReservation?> TryAddReservationAsync(DevicePoolId PoolId, List<DeviceRequestData> Request, string? Hostname, string? ReservationDetails, string? JobId, string? StepId);
+
+		/// <summary>
+		/// Gets a reservation by guid for legacy clients
+		/// </summary>
+		/// <param name="LegacyGuid">YThe legacy guid of the reservation</param>
+		Task<IDeviceReservation?> TryGetReservationFromLegacyGuidAsync(string LegacyGuid);
+
+        /// <summary>
+        /// Gets a reservation by device id
+        /// </summary>
+        /// <param name="Id">A device contained in reservation</param>
+        Task<IDeviceReservation?> TryGetDeviceReservationAsync(DeviceId Id);
+
+        /// <summary>
+        /// Get a list of all reservations
+        /// </summary>
+        Task<List<IDeviceReservation>> FindAllReservationsAsync();
+
+		/// <summary>
+		/// Updates a reservation to the current time, for expiration
+		/// </summary>
+		/// <param name="Id">The id of the reservation to update</param>
+		public Task<bool> TryUpdateReservationAsync(ObjectId Id);
+
+		/// <summary>
+		/// Deletes a reservation and releases reserved devices
+		/// </summary>
+		/// <param name="Id">The id of the reservation to delete</param>
+		public Task<bool> DeleteReservationAsync(ObjectId Id);
+
+		/// <summary>
+		/// Deletes expired reservations
+		/// </summary>
+		public Task<bool> ExpireReservationsAsync();
+
+	}
+
+}
\ No newline at end of file
diff --git a/Engine/Source/Programs/Horde/HordeServer/Devices/Impl/DeviceCollection.cs b/Engine/Source/Programs/Horde/HordeServer/Devices/Impl/DeviceCollection.cs
new file mode 100644
--- /dev/null
+++ b/Engine/Source/Programs/Horde/HordeServer/Devices/Impl/DeviceCollection.cs
@@ -0,0 +1,690 @@
+// Copyright Epic Games, Inc. All Rights Reserved.
+
+using HordeServer.Services;
+using MongoDB.Bson.Serialization.Attributes;
+using System.Threading.Tasks;
+
+using HordeServer.Models;
+using MongoDB.Driver;
+using System;
+using HordeServer.Utilities;
+using System.Linq;
+
+using MongoDB.Bson;
+using System.Collections.Generic;
+
+namespace HordeServer.Collections.Impl
+{
+	using DeviceId = StringId<IDevice>;
+	using DevicePlatformId = StringId<IDevicePlatform>;
+	using DevicePoolId = StringId<IDevicePool>;
+
+	/// <summary>
+	/// Collection of device documents
+	/// </summary>
+	public class DeviceCollection : IDeviceCollection
+	{
+
+		/// <summary>
+		/// Document representing a device platform
+		/// </summary>
+		class DevicePlatformDocument : IDevicePlatform
+		{
+			[BsonRequired, BsonId]
+			public DevicePlatformId Id { get; set; }
+
+			public string Name { get; set; }
+
+			public List<string> Models { get; set; } = new List<string>();
+
+			IReadOnlyList<string> IDevicePlatform.Models => Models;
+
+			[BsonConstructor]
+			private DevicePlatformDocument()
+			{
+				Name = null!;
+			}
+
+			public DevicePlatformDocument(DevicePlatformId Id, string Name)
+			{
+				this.Id = Id;
+				this.Name = Name;
+			}
+
+		}
+
+		/// <summary>
+		/// Document representing a pool of devices
+		/// </summary>
+		class DevicePoolDocument : IDevicePool
+		{
+			[BsonRequired, BsonId]
+			public DevicePoolId Id { get; set; }
+
+			[BsonRequired]
+            public DevicePoolType PoolType { get; set; }
+
+            [BsonRequired]
+			public string Name { get; set; } = null!;
+
+			[BsonIgnoreIfNull]
+			public Acl? Acl { get; set; }
+
+			[BsonConstructor]
+			private DevicePoolDocument()
+			{
+            }
+
+			public DevicePoolDocument(DevicePoolId Id, string Name, DevicePoolType PoolType)
+			{
+				this.Id = Id;
+				this.Name = Name;
+                this.PoolType = PoolType;
+            }
+
+
+		}
+
+		/// <summary>
+		/// Document representing a reservation of devices
+		/// </summary>
+		class DeviceReservationDocument : IDeviceReservation
+		{
+			/// <summary>
+			/// Randomly generated unique id for this reservation.
+			/// </summary>
+			[BsonRequired, BsonId]
+			public ObjectId Id { get; set; }
+
+			[BsonRequired]
+			public DevicePoolId PoolId { get; set; }
+
+			[BsonIgnoreIfNull]
+			public string? JobId { get; set; }
+
+			[BsonIgnoreIfNull]
+			public string? StepId { get; set; }
+
+            /// <summary>
+            /// Reservations held by a user, requires a token like download code
+            /// </summary>
+            [BsonIgnoreIfNull]
+			public ObjectId? UserId { get; set; }
+
+            /// <summary>
+            /// The hostname of the machine which has made the reservation
+            /// </summary>
+            [BsonIgnoreIfNull]
+			public string? Hostname { get; set; }
+
+			/// <summary>
+			/// Optional string holding details about the reservation
+			/// </summary>
+			[BsonIgnoreIfNull]
+			public string? ReservationDetails { get; set; }
+
+			/// <summary>
+			/// 
+			/// </summary>
+			public List<DeviceId> Devices { get; set; } = new List<DeviceId>();
+
+			public DateTime CreateTimeUtc { get; set; }
+
+			public DateTime UpdateTimeUtc { get; set; }
+
+			// Legacy Guid
+			public string LegacyGuid { get; set; } = null!;
+
+			[BsonConstructor]
+			private DeviceReservationDocument()
+			{
+
+			}
+
+			public DeviceReservationDocument(ObjectId Id, DevicePoolId PoolId, List<DeviceId> Devices, DateTime CreateTimeUtc, string? Hostname, string? ReservationDetails, string? JobId, string? StepId)
+			{
+				this.Id = Id;
+				this.PoolId = PoolId;
+				this.Devices = Devices;
+				this.CreateTimeUtc = CreateTimeUtc;
+				this.UpdateTimeUtc = CreateTimeUtc;
+				this.Hostname = Hostname;
+				this.ReservationDetails = ReservationDetails;
+                this.JobId = JobId;
+                this.StepId = StepId;
+
+                this.LegacyGuid = Guid.NewGuid().ToString();
+
+			}
+
+		}
+
+
+		/// <summary>
+		/// Concrete implementation of an device document
+		/// </summary>
+		class DeviceDocument : IDevice
+		{
+			[BsonRequired, BsonId]
+			public DeviceId Id { get; set; }
+
+			public DevicePlatformId PlatformId { get; set; }
+
+			public DevicePoolId PoolId { get; set; }
+
+			public string Name { get; set; } = null!;
+
+			public bool Enabled { get; set; }
+
+			[BsonIgnoreIfNull]
+			public string? ModelId { get; set; }
+
+			[BsonIgnoreIfNull]
+			public string? Address { get; set; }
+
+			[BsonIgnoreIfNull]
+			public string? CheckedOutByUser { get; set; }
+
+			[BsonIgnoreIfNull]
+			public DateTime? CheckOutTime { get; set; }
+
+			[BsonIgnoreIfNull]
+			public DateTime? ProblemTimeUtc { get; set; }
+
+			[BsonIgnoreIfNull]
+			public DateTime? MaintenanceTimeUtc { get; set; }
+
+			/// <summary>
+			///  The last time this device was reserved, used to cycle devices for reservations
+			/// </summary>
+			[BsonIgnoreIfNull]
+			public DateTime? ReservationTimeUtc { get; set; }
+
+			[BsonIgnoreIfNull]
+			public string? ModifiedByUser { get; set; }
+
+			[BsonIgnoreIfNull]
+			public string? Notes { get; set; }
+
+			[BsonIgnoreIfNull]
+			public List<DeviceUtilizationTelemetry>? Utilization { get; set; }
+
+			[BsonIgnoreIfNull]
+			public Acl? Acl { get; set; }
+
+			[BsonConstructor]
+			private DeviceDocument()
+			{
+
+			}
+
+			public DeviceDocument(DeviceId Id, DevicePlatformId PlatformId, DevicePoolId PoolId, string Name, bool Enabled, string? Address, string? ModelId, ObjectId? UserId)
+			{
+				this.Id = Id;
+				this.PlatformId = PlatformId;
+				this.PoolId = PoolId;
+				this.Name = Name;
+				this.Enabled = Enabled;
+				this.Address = Address;
+				this.ModelId = ModelId;
+                this.ModifiedByUser = UserId?.ToString();
+            }
+
+		}
+
+		readonly IMongoCollection<DevicePlatformDocument> Platforms;
+
+		readonly IMongoCollection<DeviceDocument> Devices;
+
+		readonly IMongoCollection<DevicePoolDocument> Pools;
+
+		readonly IMongoCollection<DeviceReservationDocument> Reservations;
+
+		/// <summary>
+		/// Constructor
+		/// </summary>
+		public DeviceCollection(DatabaseService DatabaseService)
+		{
+			Devices = DatabaseService.GetCollection<DeviceDocument>("Devices");
+			Platforms = DatabaseService.GetCollection<DevicePlatformDocument>("Devices.Platforms");
+			Pools = DatabaseService.GetCollection<DevicePoolDocument>("Devices.Pools");
+			Reservations = DatabaseService.GetCollection<DeviceReservationDocument>("Devices.Reservations");
+
+			if (!DatabaseService.ReadOnlyMode)
+			{
+				Platforms.Indexes.CreateOne(new CreateIndexModel<DevicePlatformDocument>(Builders<DevicePlatformDocument>.IndexKeys.Ascending(x => x.Name), new CreateIndexOptions { Unique = true }));
+				Pools.Indexes.CreateOne(new CreateIndexModel<DevicePoolDocument>(Builders<DevicePoolDocument>.IndexKeys.Ascending(x => x.Name), new CreateIndexOptions { Unique = true }));
+				Devices.Indexes.CreateOne(new CreateIndexModel<DeviceDocument>(Builders<DeviceDocument>.IndexKeys.Ascending(x => x.Name), new CreateIndexOptions { Unique = true }));
+			}
+		}
+
+		/// <inheritdoc/>
+		public async Task<IDevice?> TryAddDeviceAsync(DeviceId Id, string Name, DevicePlatformId PlatformId, DevicePoolId PoolId, bool? Enabled, string? Address, string? ModelId, ObjectId? UserId)
+		{
+			DeviceDocument NewDevice = new DeviceDocument(Id, PlatformId, PoolId, Name, Enabled ?? true, Address, ModelId, UserId);
+			await Devices.InsertOneAsync(NewDevice);
+			return NewDevice;
+		}
+
+
+		/// <inheritdoc/>
+		public async Task<IDevicePlatform?> TryAddPlatformAsync(DevicePlatformId Id, string Name)
+		{
+			DevicePlatformDocument NewPlatform = new DevicePlatformDocument(Id, Name);
+
+			try
+			{
+				await Platforms.InsertOneAsync(NewPlatform);
+				return NewPlatform;
+			}
+			catch (MongoWriteException Ex)
+			{
+				if (Ex.WriteError.Category == ServerErrorCategory.DuplicateKey)
+				{
+					return null;
+				}
+				else
+				{
+					throw;
+				}
+			}
+		}
+
+		/// <inheritdoc/>
+		public async Task<List<IDevicePlatform>> FindAllPlatformsAsync()
+		{
+			List<DevicePlatformDocument> Results = await Platforms.Find(x => true).ToListAsync();
+			return Results.OrderBy(x => x.Name).Select<DevicePlatformDocument, IDevicePlatform>(x => x).ToList();
+		}
+
+		/// <inheritdoc/>
+		public async Task<bool> UpdatePlatformAsync(DevicePlatformId PlatformId, string[]? ModelIds)
+		{
+			UpdateDefinitionBuilder<DevicePlatformDocument> UpdateBuilder = Builders<DevicePlatformDocument>.Update;
+
+			List<UpdateDefinition<DevicePlatformDocument>> Updates = new List<UpdateDefinition<DevicePlatformDocument>>();
+
+			if (ModelIds != null)
+			{
+				Updates.Add(UpdateBuilder.Set(x => x.Models, ModelIds.ToList()));
+			}
+
+			if (Updates.Count > 0)
+			{
+				await Platforms.FindOneAndUpdateAsync<DevicePlatformDocument>(x => x.Id == PlatformId, UpdateBuilder.Combine(Updates));
+			}
+
+			return true;
+
+		}
+
+		/// <inheritdoc/>
+		public async Task<IDevicePool?> TryAddPoolAsync(DevicePoolId Id, string Name, DevicePoolType PoolType)
+		{
+			DevicePoolDocument NewPool = new DevicePoolDocument(Id, Name, PoolType);
+
+			try
+			{
+				await Pools.InsertOneAsync(NewPool);
+				return NewPool;
+			}
+			catch (MongoWriteException Ex)
+			{
+				if (Ex.WriteError.Category == ServerErrorCategory.DuplicateKey)
+				{
+					return null;
+				}
+				else
+				{
+					throw;
+				}
+			}
+		}
+
+		/// <inheritdoc/>
+		public async Task<List<IDevice>> FindAllDevicesAsync(List<DeviceId>? DeviceIds = null)
+		{
+			FilterDefinition<DeviceDocument> Filter = Builders<DeviceDocument>.Filter.Empty;
+			if (DeviceIds != null)
+			{
+				Filter &= Builders<DeviceDocument>.Filter.In(x => x.Id, DeviceIds);
+			}
+			List<DeviceDocument> Results = await Devices.Find(Filter).ToListAsync();
+			return Results.OrderBy(x => x.Name).Select<DeviceDocument, IDevice>(x => x).ToList();
+		}
+
+		/// <inheritdoc/>
+		public async Task<List<IDevicePool>> FindAllPoolsAsync()
+		{
+			List<DevicePoolDocument> Results = await Pools.Find(x => true).ToListAsync();
+			return Results.OrderBy(x => x.Name).Select<DevicePoolDocument, IDevicePool>(x => x).ToList();
+		}
+
+		/// <inheritdoc/>
+		public async Task<IDevicePlatform?> GetPlatformAsync(DevicePlatformId PlatformId)
+		{
+			return await Platforms.Find<DevicePlatformDocument>(x => x.Id == PlatformId).FirstOrDefaultAsync();
+		}
+
+		/// <inheritdoc/>
+		public async Task<IDevicePool?> GetPoolAsync(DevicePoolId PoolId)
+		{
+			return await Pools.Find<DevicePoolDocument>(x => x.Id == PoolId).FirstOrDefaultAsync();
+		}
+
+		/// <inheritdoc/>
+		public async Task<IDevice?> GetDeviceAsync(DeviceId DeviceId)
+		{
+			return await Devices.Find<DeviceDocument>(x => x.Id == DeviceId).FirstOrDefaultAsync();
+		}
+
+		/// <inheritdoc/>
+		public async Task<IDevice?> GetDeviceByNameAsync(string DeviceName)
+		{
+			return await Devices.Find<DeviceDocument>(x => x.Name == DeviceName).FirstOrDefaultAsync();
+		}
+
+		/// <inheritdoc/>
+		public async Task CheckoutDeviceAsync(DeviceId DeviceId, ObjectId? CheckedOutByUserId)
+		{
+			UpdateDefinitionBuilder<DeviceDocument> UpdateBuilder = Builders<DeviceDocument>.Update;
+
+			List<UpdateDefinition<DeviceDocument>> Updates = new List<UpdateDefinition<DeviceDocument>>();
+
+            string? UserId = CheckedOutByUserId?.ToString();
+
+            Updates.Add(UpdateBuilder.Set(x => x.CheckedOutByUser, string.IsNullOrEmpty(UserId) ? null : UserId));
+
+			if (CheckedOutByUserId != null)
+			{
+				Updates.Add(UpdateBuilder.Set(x => x.CheckOutTime, DateTime.UtcNow));
+			}
+
+			await Devices.FindOneAndUpdateAsync<DeviceDocument>(x => x.Id == DeviceId, UpdateBuilder.Combine(Updates));
+		}
+
+
+		/// <inheritdoc/>
+		public async Task UpdateDeviceAsync(DeviceId DeviceId, DevicePoolId? NewPoolId, string? NewName, string? NewAddress, string? NewModelId, string? NewNotes, bool? NewEnabled, bool? NewProblem, bool? NewMaintenance, ObjectId? ModifiedByUserId = null)
+		{
+			UpdateDefinitionBuilder<DeviceDocument> UpdateBuilder = Builders<DeviceDocument>.Update;
+
+			List<UpdateDefinition<DeviceDocument>> Updates = new List<UpdateDefinition<DeviceDocument>>();
+
+			DateTime UtcNow = DateTime.UtcNow;
+
+			if (ModifiedByUserId != null)
+			{
+				Updates.Add(UpdateBuilder.Set(x => x.ModifiedByUser, ModifiedByUserId.ToString()));
+			}
+
+			if (NewPoolId != null)
+			{
+				Updates.Add(UpdateBuilder.Set(x => x.PoolId, NewPoolId.Value));
+			}
+
+			if (NewName != null)
+			{
+				Updates.Add(UpdateBuilder.Set(x => x.Name, NewName));
+			}
+
+			if (NewEnabled != null)
+			{
+				Updates.Add(UpdateBuilder.Set(x => x.Enabled, NewEnabled));
+			}
+
+			if (NewAddress != null)
+			{
+				Updates.Add(UpdateBuilder.Set(x => x.Address, NewAddress));
+			}
+
+			if (!string.IsNullOrEmpty(NewModelId))
+			{
+				if (NewModelId == "Base")
+				{
+					Updates.Add(UpdateBuilder.Set(x => x.ModelId, null));
+				}
+				else
+				{
+					Updates.Add(UpdateBuilder.Set(x => x.ModelId, NewModelId));
+				}				
+			}
+
+			if (NewNotes != null)
+			{
+				Updates.Add(UpdateBuilder.Set(x => x.Notes, NewNotes));
+			}
+
+			if (NewProblem.HasValue)
+			{
+				DateTime? ProblemTime = null;
+				if (NewProblem.Value)
+				{
+					ProblemTime = UtcNow;
+				}
+
+				Updates.Add(UpdateBuilder.Set(x => x.ProblemTimeUtc, ProblemTime));
+			}
+
+			if (NewMaintenance.HasValue)
+			{
+				DateTime? MaintenanceTime = null;
+				if (NewMaintenance.Value)
+				{
+					MaintenanceTime = UtcNow;
+				}
+
+				Updates.Add(UpdateBuilder.Set(x => x.MaintenanceTimeUtc, MaintenanceTime));
+			}
+
+
+			if (Updates.Count > 0)
+			{
+				await Devices.FindOneAndUpdateAsync<DeviceDocument>(x => x.Id == DeviceId, UpdateBuilder.Combine(Updates));
+			}
+		}
+
+		/// <inheritdoc/>
+		public async Task<bool> DeleteDeviceAsync(DeviceId DeviceId)
+		{
+			FilterDefinition<DeviceDocument> Filter = Builders<DeviceDocument>.Filter.Eq(x => x.Id, DeviceId);
+			DeleteResult Result = await Devices.DeleteOneAsync(Filter);
+			return Result.DeletedCount > 0;
+
+		}
+
+		/// <inheritdoc/>
+		public async Task<List<IDeviceReservation>> FindAllDeviceReservationsAsync(DevicePoolId? PoolId = null)
+		{
+			return await Reservations.Find(A => PoolId == null || A.PoolId == PoolId).ToListAsync<DeviceReservationDocument, IDeviceReservation>();
+		}
+
+		/// <inheritdoc/>
+		public async Task<IDeviceReservation?> TryAddReservationAsync(DevicePoolId PoolId, List<DeviceRequestData> Request, string? Hostname, string? ReservationDetails, string? JobId, string? StepId)
+		{
+
+			if (Request.Count == 0)
+			{
+				return null;
+			}
+
+			DevicePoolDocument? Pool = await Pools.Find<DevicePoolDocument>(x => x.Id == PoolId).FirstOrDefaultAsync();
+
+			if (Pool == null || Pool.PoolType != DevicePoolType.Automation)
+			{
+                return null;
+            }
+			
+			HashSet<DeviceId> Allocated = new HashSet<DeviceId>();
+
+			List<DeviceReservationDocument> PoolReservations = await Reservations.Find(x => x.PoolId == PoolId).ToListAsync();
+
+			DateTime ReservationTimeUtc = DateTime.UtcNow;
+
+			// Get available devices
+			List<DeviceDocument> PoolDevices = await Devices.Find(x =>
+				x.PoolId == PoolId &&
+				x.Enabled &&
+				x.MaintenanceTimeUtc == null).ToListAsync();
+
+			// filter out problem devices
+			PoolDevices = PoolDevices.FindAll(x => (x.ProblemTimeUtc == null || ((ReservationTimeUtc - x.ProblemTimeUtc).Value.TotalMinutes > 30)));
+
+			// filter out currently reserved devices
+			PoolDevices = PoolDevices.FindAll(x => PoolReservations.FirstOrDefault(p => p.Devices.Contains(x.Id)) == null);
+
+			// sort to use last reserved first to cycle devices
+			PoolDevices.Sort((A, B) =>
+			{
+				DateTime? ATime = A.ReservationTimeUtc;
+				DateTime? BTime = B.ReservationTimeUtc;
+
+				if (ATime == BTime)
+				{
+					return 0;
+				}
+
+				if (ATime == null && BTime != null)
+				{
+					return -1;
+				}
+				if (ATime != null && BTime == null)
+				{
+					return 1;
+				}
+
+				return ATime < BTime ? -1 : 1;
+
+			});
+
+			foreach (DeviceRequestData Data in Request)
+			{
+				DeviceDocument? Device = PoolDevices.FirstOrDefault(A =>
+				{
+
+					if (Allocated.Contains(A.Id) || A.PlatformId != Data.PlatformId)
+					{
+						return false;
+					}
+
+					if (Data.IncludeModels.Count > 0 && (A.ModelId == null || !Data.IncludeModels.Contains(A.ModelId)))
+					{
+						return false;
+					}
+
+					if (Data.ExcludeModels.Count > 0 && (A.ModelId != null && Data.ExcludeModels.Contains(A.ModelId)))
+					{
+						return false;
+					}
+
+
+					return true;
+				});
+
+				if (Device == null)
+				{
+					// can't fulfill request
+					return null;
+				}
+
+				Allocated.Add(Device.Id);
+			}
+
+			// update reservation time and utilization for allocated devices 
+			foreach (DeviceId Id in Allocated)
+			{
+				DeviceDocument Device = PoolDevices.First((Device) => Device.Id == Id);
+
+				List<DeviceUtilizationTelemetry>? Utilization = Device.Utilization;
+
+				if (Utilization == null)
+				{
+					Utilization = new List<DeviceUtilizationTelemetry>();
+				}
+
+				// keep up to 100, maintaining order
+				if (Utilization.Count > 99)
+				{
+					Utilization = Utilization.GetRange(0, 99);
+				}
+
+				Utilization.Insert(0, new DeviceUtilizationTelemetry(ReservationTimeUtc) { JobId = JobId, StepId = StepId });
+
+				UpdateDefinitionBuilder<DeviceDocument> DeviceBuilder = Builders<DeviceDocument>.Update;
+				List<UpdateDefinition<DeviceDocument>> DeviceUpdates = new List<UpdateDefinition<DeviceDocument>>();
+
+				DeviceUpdates.Add(DeviceBuilder.Set(x => x.ReservationTimeUtc, ReservationTimeUtc));
+				DeviceUpdates.Add(DeviceBuilder.Set(x => x.Utilization, Utilization));
+
+				await Devices.FindOneAndUpdateAsync<DeviceDocument>(x => x.Id == Id, DeviceBuilder.Combine(DeviceUpdates));
+			}
+
+			// Create new reservation
+			DeviceReservationDocument NewReservation = new DeviceReservationDocument(ObjectId.GenerateNewId(), PoolId, Allocated.ToList(), ReservationTimeUtc, Hostname, ReservationDetails, JobId, StepId);
+			await Reservations.InsertOneAsync(NewReservation);
+
+			return NewReservation;
+
+		}
+
+		/// <inheritdoc/>
+		public async Task<bool> TryUpdateReservationAsync(ObjectId Id)
+		{
+			UpdateResult Result = await Reservations.UpdateOneAsync(x => x.Id == Id, Builders<DeviceReservationDocument>.Update.Set(x => x.UpdateTimeUtc, DateTime.UtcNow));
+			return Result.ModifiedCount == 1;
+		}
+
+		/// <inheritdoc/>
+		public async Task<bool> DeleteReservationAsync(ObjectId Id)
+		{
+			FilterDefinition<DeviceReservationDocument> Filter = Builders<DeviceReservationDocument>.Filter.Eq(x => x.Id, Id);
+			DeleteResult Result = await Reservations.DeleteOneAsync(Filter);
+			return Result.DeletedCount > 0;
+		}
+
+		/// <summary>
+		/// Deletes expired reservations
+		/// </summary>
+		public async Task<bool> ExpireReservationsAsync()
+		{
+			List<IDeviceReservation> Reserves = await Reservations.Find(A => true).ToListAsync<DeviceReservationDocument, IDeviceReservation>();
+
+			DateTime UtcNow = DateTime.UtcNow;
+
+			Reserves = Reserves.FindAll(R => (UtcNow - R.UpdateTimeUtc).TotalMinutes > 10).ToList();
+
+			if (Reserves.Count > 0)
+			{
+				DeleteResult Result = await Reservations.DeleteManyAsync(Builders<DeviceReservationDocument>.Filter.In(x => x.Id, Reserves.Select(y => y.Id)));
+				return (Result.DeletedCount > 0);
+			}
+
+			return false;
+
+		}
+
+		/// <inheritdoc/>
+		public async Task<List<IDeviceReservation>> FindAllReservationsAsync()
+		{
+			List<DeviceReservationDocument> Results = await Reservations.Find(x => true).ToListAsync();
+			return Results.OrderBy(x => x.CreateTimeUtc.Ticks).Select<DeviceReservationDocument, IDeviceReservation>(x => x).ToList();
+		}
+
+
+		/// <inheritdoc/>
+		public async Task<IDeviceReservation?> TryGetReservationFromLegacyGuidAsync(string LegacyGuid)
+		{
+			return await Reservations.Find<DeviceReservationDocument>(R => R.LegacyGuid == LegacyGuid).FirstOrDefaultAsync();
+		}
+
+		/// <inheritdoc/>
+		public async Task<IDeviceReservation?> TryGetDeviceReservationAsync(DeviceId Id)
+		{
+			List<DeviceReservationDocument> Results = await Reservations.Find(x => true).ToListAsync();
+			return Results.FirstOrDefault(R => R.Devices.Contains(Id));
+        }
+
+	}
+
+}
+
diff --git a/Engine/Source/Programs/Horde/HordeServer/Devices/Impl/DeviceService.cs b/Engine/Source/Programs/Horde/HordeServer/Devices/Impl/DeviceService.cs
new file mode 100644
--- /dev/null
+++ b/Engine/Source/Programs/Horde/HordeServer/Devices/Impl/DeviceService.cs
@@ -0,0 +1,441 @@
+// Copyright Epic Games, Inc. All Rights Reserved.
+
+using System.Threading.Tasks;
+
+using HordeServer.Models;
+using HordeServer.Collections;
+using System;
+using HordeServer.Utilities;
+using System.Threading;
+
+using Microsoft.Extensions.Logging;
+using HordeCommon;
+using MongoDB.Bson;
+using System.Collections.Generic;
+using System.Security.Claims;
+using System.Diagnostics.CodeAnalysis;
+using HordeServer.Notifications;
+using System.Linq;
+using HordeServer.Api;
+
+namespace HordeServer.Services
+{
+	using DeviceId = StringId<IDevice>;
+	using DevicePlatformId = StringId<IDevicePlatform>;
+	using DevicePoolId = StringId<IDevicePool>;
+
+	/// <summary>
+	/// Platform map required by V1 API
+	/// </summary>
+	[SingletonDocument("6165a2e26fd5f104e31e6862")]
+	public class DevicePlatformMapV1 : SingletonBase
+	{
+		/// <summary>
+		/// Platform V1 => Platform Id
+		/// </summary>
+		public Dictionary<string, DevicePlatformId> PlatformMap { get; set; } = new Dictionary<string, DevicePlatformId>();
+
+		/// <summary>
+		/// Platform Id => Platform V1
+		/// </summary>
+		public Dictionary<DevicePlatformId, string> PlatformReverseMap { get; set; } = new Dictionary<DevicePlatformId, string>();
+
+		/// <summary>
+		/// Perfspec V1 => Model
+		/// </summary>
+		public Dictionary<DevicePlatformId, string> PerfSpecHighMap { get; set; } = new Dictionary<DevicePlatformId, string>();
+
+	}
+
+
+	/// <summary>
+	/// Device management service
+	/// </summary>
+	public class DeviceService : TickedBackgroundService
+	{
+		/// <summary>
+		/// The ACL service instance
+		/// </summary>
+		AclService AclService;
+
+		/// <summary>
+		/// Instance of the notification service
+		/// </summary>
+		INotificationService NotificationService;
+
+
+		/// <summary>
+		/// Singleton instance of the job service
+		/// </summary>
+		JobService JobService;
+
+		/// <summary>
+		/// Singleton instance of the stream service
+		/// </summary>
+		StreamService StreamService;
+
+		/// <summary>
+		/// Log output writer
+		/// </summary>
+		ILogger<DeviceService> Logger;
+
+		/// <summary>
+		/// Device collection
+		/// </summary>
+		IDeviceCollection Devices;
+
+		/// <summary>
+		/// Platform map V1 singleton
+		/// </summary>
+		ISingletonDocument<DevicePlatformMapV1> PlatformMapSingleton;
+
+		/// <summary>
+		/// Device service constructor
+		/// </summary>
+		public DeviceService(IDeviceCollection Devices, ISingletonDocument<DevicePlatformMapV1> PlatformMapSingleton, JobService JobService, StreamService StreamService, AclService AclService, INotificationService NotificationService, ILogger<DeviceService> Logger)
+			: base(TimeSpan.FromMinutes(1.0), Logger)
+		{
+			this.Devices = Devices;
+            this.JobService = JobService;
+            this.StreamService = StreamService;
+            this.AclService = AclService;
+            this.NotificationService = NotificationService;
+            this.Logger = Logger;
+			
+			this.PlatformMapSingleton = PlatformMapSingleton;
+
+		}
+
+		/// <summary>
+		/// Ticks service
+		/// </summary>
+		protected override async Task TickAsync(CancellationToken StoppingToken)
+		{
+
+			if (!StoppingToken.IsCancellationRequested)
+			{
+				await Devices.ExpireReservationsAsync();
+			}
+
+		}
+
+		/// <summary>
+		/// Create a new device platform
+		/// </summary>
+		public Task<IDevicePlatform?> TryCreatePlatformAsync(DevicePlatformId Id, string Name)
+		{
+			return Devices.TryAddPlatformAsync(Id, Name);
+		}
+
+		/// <summary>
+		/// Get a list of existing device platforms
+		/// </summary>
+		public Task<List<IDevicePlatform>> GetPlatformsAsync()
+		{
+			return Devices.FindAllPlatformsAsync();
+		}
+
+		/// <summary>
+		/// Update an existing platform
+		/// </summary>
+		public Task<bool> UpdatePlatformAsync(DevicePlatformId PlatformId, string[]? ModelIds)
+		{
+			return Devices.UpdatePlatformAsync(PlatformId, ModelIds);
+		}
+
+		/// <summary>
+		/// Get a specific device platform
+		/// </summary>
+		public Task<IDevicePlatform?> GetPlatformAsync(DevicePlatformId Id)
+		{
+			return Devices.GetPlatformAsync(Id);
+		}
+
+		/// <summary>
+		/// Get a device pool by id
+		/// </summary>
+		public Task<IDevicePool?> GetPoolAsync(DevicePoolId Id)
+		{
+			return Devices.GetPoolAsync(Id);
+		}
+
+		/// <summary>
+		/// Create a new device pool
+		/// </summary>
+		public Task<IDevicePool?> TryCreatePoolAsync(DevicePoolId Id, string Name, DevicePoolType PoolType)
+		{
+			return Devices.TryAddPoolAsync(Id, Name, PoolType);
+		}
+
+		/// <summary>
+		/// Get a list of existing device pools
+		/// </summary>
+		public Task<List<IDevicePool>> GetPoolsAsync()
+		{
+			return Devices.FindAllPoolsAsync();
+		}
+
+		/// <summary>
+		/// Get a list of devices, optionally filtered to provided ids
+		/// </summary>
+		public Task<List<IDevice>> GetDevicesAsync(List<DeviceId>? DeviceIds = null)
+		{
+			return Devices.FindAllDevicesAsync(DeviceIds);
+		}
+
+		/// <summary>
+		/// Get a specific device
+		/// </summary>
+		public Task<IDevice?> GetDeviceAsync(DeviceId Id)
+		{
+			return Devices.GetDeviceAsync(Id);
+		}
+
+		/// <summary>
+		/// Get a device by name
+		/// </summary>
+		public Task<IDevice?> GetDeviceByNameAsync(string DeviceName)
+		{
+			return Devices.GetDeviceByNameAsync(DeviceName);
+		}
+
+		/// <summary>
+		/// Delete a device
+		/// </summary>
+		public async Task<bool> DeleteDeviceAsync(DeviceId Id)
+		{
+			return await Devices.DeleteDeviceAsync(Id);
+		}
+
+		/// <summary>
+		/// Create a new device
+		/// </summary>
+		/// <param name="Id">Unique id of the device</param>
+		/// <param name="Name">Friendly name of the device</param>
+		/// <param name="PlatformId">The device platform</param>
+		/// <param name="PoolId">Which pool to add the device</param>
+		/// <param name="Enabled">Whether the device is enabled</param>
+		/// <param name="Address">Address or hostname of device</param>
+		/// <param name="ModelId">Vendor model id</param>
+        /// <param name="UserId">User adding the device</param>
+		/// <returns></returns>
+		public Task<IDevice?> TryCreateDeviceAsync(DeviceId Id, string Name, DevicePlatformId PlatformId, DevicePoolId PoolId, bool? Enabled, string? Address, string? ModelId, ObjectId? UserId = null)
+		{
+			return Devices.TryAddDeviceAsync(Id, Name, PlatformId, PoolId, Enabled, Address, ModelId, UserId);
+		}
+
+		/// <summary>
+		/// Update a device
+		/// </summary>
+		public Task UpdateDeviceAsync(DeviceId DeviceId, DevicePoolId? NewPoolId = null, string? NewName = null, string? NewAddress = null, string? NewModelId = null, string? NewNotes = null, bool? NewEnabled = null, bool? NewProblem = null, bool? NewMaintenance = null, ObjectId? ModifiedByUserId = null)
+		{
+			return Devices.UpdateDeviceAsync(DeviceId, NewPoolId, NewName, NewAddress, NewModelId, NewNotes, NewEnabled, NewProblem, NewMaintenance, ModifiedByUserId);
+		}
+
+		/// <summary>
+		/// Checkout a device
+		/// </summary>
+		public Task CheckoutDeviceAsync(DeviceId DeviceId, ObjectId? UserId)
+		{
+            return Devices.CheckoutDeviceAsync(DeviceId, UserId);
+        }
+
+		/// <summary>
+		/// Try to create a reservation satisfying the specified device platforms and models
+		/// </summary>
+		public Task<IDeviceReservation?> TryCreateReservationAsync(DevicePoolId Pool, List<DeviceRequestData> Request, string? Hostname = null, string? ReservationDetails = null, string? JobId = null, string? StepId = null)
+		{
+			return Devices.TryAddReservationAsync(Pool, Request, Hostname, ReservationDetails, JobId, StepId);
+		}
+
+		/// <summary>
+		/// Update/renew an existing reservation
+		/// </summary>
+		public Task<bool> TryUpdateReservationAsync(ObjectId Id)
+		{
+			return Devices.TryUpdateReservationAsync(Id);
+		}
+
+		/// <summary>
+		///  Delete an existing reservation
+		/// </summary>
+		public Task<bool> DeleteReservationAsync(ObjectId Id)
+		{
+			return Devices.DeleteReservationAsync(Id);
+		}
+
+		/// <summary>
+		/// Get a reservation from a legacy guid
+		/// </summary>
+		public Task<IDeviceReservation?> TryGetReservationFromLegacyGuidAsync(string LegacyGuid)
+		{
+			return Devices.TryGetReservationFromLegacyGuidAsync(LegacyGuid);
+		}
+
+		/// <summary>
+		/// Get a reservation from a device id
+		/// </summary>
+		public Task<IDeviceReservation?> TryGetDeviceReservation(DeviceId DeviceId)
+		{
+			return Devices.TryGetDeviceReservationAsync(DeviceId);
+		}
+
+		/// <summary>
+		/// Get a list of existing device reservations
+		/// </summary>
+		public Task<List<IDeviceReservation>> GetReservationsAsync()
+		{
+			return Devices.FindAllReservationsAsync();
+		}
+
+		/// <summary>
+		/// 
+		/// </summary>
+		/// <param name="Message"></param>
+		/// <param name="DeviceId"></param>
+		/// <param name="JobId"></param>
+		/// <param name="StepId"></param>
+		public async Task NotifyDeviceServiceAsync(string Message, DeviceId? DeviceId = null, string? JobId = null, string? StepId = null)
+		{
+			try 
+			{
+				IDevice? Device = null;
+				IDevicePool? Pool = null;
+				IJob? Job = null;
+				IJobStep? Step = null;
+				INode? Node = null;
+				IStream? Stream = null;
+
+				if (DeviceId.HasValue)
+				{
+					Device = await GetDeviceAsync(DeviceId.Value);
+					Pool = await GetPoolAsync(Device!.PoolId);
+				}
+
+				if (JobId != null)
+				{
+					Job = await JobService.GetJobAsync(new ObjectId(JobId));
+
+					if (Job != null)
+					{
+						Stream = await StreamService.GetStreamAsync(Job.StreamId);
+
+						if (StepId != null)
+						{
+							IGraph Graph = await JobService.GetGraphAsync(Job)!;
+
+							SubResourceId StepIdValue = SubResourceId.Parse(StepId);
+							IJobStepBatch? Batch = Job.Batches.FirstOrDefault(B => B.Steps.FirstOrDefault(S => S.Id == StepIdValue) != null);
+							if (Batch != null)
+							{
+								Step = Batch.Steps.FirstOrDefault(S => S.Id == StepIdValue)!;
+								INodeGroup Group = Graph.Groups[Batch.GroupIdx];
+								Node = Group.Nodes[Step.NodeIdx];
+							}
+						}
+					}
+				}
+
+				NotificationService.NotifyDeviceService(Message, Device, Pool, Stream, Job, Step, Node);
+
+			}
+			catch (Exception Ex)
+			{
+                Logger.LogError($"Error on device notification {Ex.Message}");
+            }
+        }
+
+
+
+		/// <summary>
+		/// Authorize device action
+		/// </summary>
+		/// <param name="Action"></param>
+		/// <param name="User"></param>
+		/// <returns></returns>
+		[SuppressMessage("Usage", "CA1801:Review unused parameters")]
+		[SuppressMessage("Performance", "CA1822: Can be static ")]
+		public Task<bool> AuthorizeAsync(AclAction Action, ClaimsPrincipal User)
+		{
+			// Setup ACL's for pool and platform access, https://jira.it.epicgames.com/browse/UE-117224/			
+			// allow reads, though restrict writing to internal employees
+			if (Action == AclAction.DeviceRead)
+			{
+				return Task.FromResult(true);
+			}
+			return Task.FromResult(User.IsInRole("Internal-Employees"));			
+		}
+
+		/// <summary>
+		/// Get Platform mappings for V1 API
+		/// </summary>		
+		public async Task<DevicePlatformMapV1> GetPlatformMapV1()
+		{
+			return await PlatformMapSingleton.GetAsync();
+		}
+
+		/// <summary>
+		/// Updates the platform mapping information required for the V1 api
+		/// </summary>				
+		public async Task<bool> UpdatePlatformMapAsync(UpdatePlatformMapRequest Request)
+		{
+
+			List<IDevicePlatform> Platforms = await GetPlatformsAsync();
+
+			// Update the platform map
+			for (int i = 0; i < 10; i++)
+			{
+				DevicePlatformMapV1 Instance = await PlatformMapSingleton.GetAsync();
+
+				Instance.PlatformMap = new Dictionary<string, DevicePlatformId>();
+				Instance.PlatformReverseMap = new Dictionary<DevicePlatformId, string>();
+				Instance.PerfSpecHighMap = new Dictionary<DevicePlatformId, string>();
+
+				foreach(KeyValuePair<string, string> Entry in Request.PlatformMap)
+				{
+					IDevicePlatform? Platform = Platforms.FirstOrDefault(P => P.Id == new DevicePlatformId(Entry.Value));
+					if (Platform == null)
+					{
+						throw new Exception($"Unknowm platform in map {Entry.Key} : {Entry.Value}");
+					}
+
+					Instance.PlatformMap.Add(Entry.Key, Platform.Id);
+				}
+
+				foreach (KeyValuePair<string, string> Entry in Request.PlatformReverseMap)
+				{
+					IDevicePlatform? Platform = Platforms.FirstOrDefault(P => P.Id == new DevicePlatformId(Entry.Key));
+					if (Platform == null)
+					{
+						throw new Exception($"Unknowm platform in reverse map {Entry.Key} : {Entry.Value}");
+					}
+
+					Instance.PlatformReverseMap.Add(Platform.Id, Entry.Value);
+				}
+
+				foreach (KeyValuePair<string, string> Entry in Request.PerfSpecHighMap)
+				{
+					IDevicePlatform? Platform = Platforms.FirstOrDefault(P => P.Id == new DevicePlatformId(Entry.Key));
+					if (Platform == null)
+					{
+						throw new Exception($"Unknowm platform in spec map {Entry.Key} : {Entry.Value}");
+					}
+
+					Instance.PerfSpecHighMap.Add(Platform.Id, Entry.Value);
+				}
+
+
+				if (await PlatformMapSingleton.TryUpdateAsync(Instance))
+				{
+					return true;
+				}
+
+				Thread.Sleep(1000);
+			}
+
+			return false;
+
+		}
+
+	}
+}
\ No newline at end of file
diff --git a/Engine/Source/Programs/Horde/HordeServer/Logs/ILogEvent.cs b/Engine/Source/Programs/Horde/HordeServer/Logs/ILogEvent.cs
new file mode 100644
--- /dev/null
+++ b/Engine/Source/Programs/Horde/HordeServer/Logs/ILogEvent.cs
@@ -0,0 +1,129 @@
+// Copyright Epic Games, Inc. All Rights Reserved.
+
+using HordeCommon;
+using HordeCommon.Rpc;
+using Microsoft.Extensions.Logging;
+using MongoDB.Bson;
+using System;
+using System.Collections.Generic;
+using System.Diagnostics.CodeAnalysis;
+using System.Linq;
+using System.Text.Json;
+using System.Threading.Tasks;
+
+namespace HordeServer.Models
+{
+	/// <summary>
+	/// Interface for a log event line
+	/// </summary>
+	public interface ILogEventLine
+	{
+		/// <summary>
+		/// Log level
+		/// </summary>
+		LogLevel Level { get; }
+
+		/// <summary>
+		/// The event id
+		/// </summary>
+		EventId? EventId { get; }
+
+		/// <summary>
+		/// The complete rendered message, in plaintext
+		/// </summary>
+		string Message { get; }
+
+		/// <summary>
+		/// Data for the line
+		/// </summary>
+		JsonElement Data { get; }
+	}
+
+	/// <summary>
+	/// Interface for event data
+	/// </summary>
+	public interface ILogEventData
+	{
+		/// <summary>
+		/// The log level
+		/// </summary>
+		EventSeverity Severity { get; }
+
+		/// <summary>
+		/// The type of event
+		/// </summary>
+		EventId? EventId { get; }
+
+		/// <summary>
+		/// The complete rendered message, in plaintext
+		/// </summary>
+		string Message { get; }
+
+		/// <summary>
+		/// Gets this event data as a BSON document
+		/// </summary>
+		IReadOnlyList<ILogEventLine> Lines { get; }
+	}
+
+	/// <summary>
+	/// Represents a node in the graph
+	/// </summary>
+	public interface ILogEvent
+	{
+		/// <summary>
+		/// Unique id of the log containing this event
+		/// </summary>
+		public ObjectId LogId { get; }
+
+		/// <summary>
+		/// Severity of the event
+		/// </summary>
+		public EventSeverity Severity { get; }
+
+		/// <summary>
+		/// Index of the first line for this event
+		/// </summary>
+		public int LineIndex { get; }
+
+		/// <summary>
+		/// Number of lines in the event
+		/// </summary>
+		public int LineCount { get; }
+
+		/// <summary>
+		/// Span id for this log event
+		/// </summary>
+		public ObjectId? SpanId { get; }
+	}
+
+	/// <summary>
+	/// Represents a node in the graph
+	/// </summary>
+	public class NewLogEventData
+	{
+		/// <summary>
+		/// Unique id of the log containing this event
+		/// </summary>
+		public ObjectId LogId { get; set; }
+
+		/// <summary>
+		/// Severity of the event
+		/// </summary>
+		public EventSeverity Severity { get; set; }
+
+		/// <summary>
+		/// Index of the first line for this event
+		/// </summary>
+		public int LineIndex { get; set; }
+
+		/// <summary>
+		/// Number of lines in the event
+		/// </summary>
+		public int LineCount { get; set; }
+
+		/// <summary>
+		/// The span this this event belongs to
+		/// </summary>
+		public ObjectId? SpanId { get; set; }
+	}
+}
diff --git a/Engine/Source/Programs/Horde/HordeServer/Logs/ILogEventCollection.cs b/Engine/Source/Programs/Horde/HordeServer/Logs/ILogEventCollection.cs
new file mode 100644
--- /dev/null
+++ b/Engine/Source/Programs/Horde/HordeServer/Logs/ILogEventCollection.cs
@@ -0,0 +1,71 @@
+// Copyright Epic Games, Inc. All Rights Reserved.
+
+using HordeCommon;
+using HordeServer.Models;
+using HordeServer.Services;
+using HordeServer.Utilities;
+using Microsoft.Extensions.Logging;
+using MongoDB.Bson;
+using MongoDB.Bson.Serialization;
+using MongoDB.Bson.Serialization.Attributes;
+using MongoDB.Driver;
+using Serilog.Events;
+using System;
+using System.Collections.Generic;
+using System.Linq;
+using System.Threading.Tasks;
+
+namespace HordeServer.Collections
+{
+	/// <summary>
+	/// Interface for a collection of event documents
+	/// </summary>
+	public interface ILogEventCollection
+	{
+		/// <summary>
+		/// Creates a new event
+		/// </summary>
+		/// <param name="NewEvent">The new event to vreate</param>
+		Task AddAsync(NewLogEventData NewEvent);
+
+		/// <summary>
+		/// Creates a new event
+		/// </summary>
+		/// <param name="NewEvents">List of events to create</param>
+		Task AddManyAsync(List<NewLogEventData> NewEvents);
+
+		/// <summary>
+		/// Finds events within a log file
+		/// </summary>
+		/// <param name="LogId">Unique id of the log containing this event</param>
+		/// <param name="Index">Start index within the matching results</param>
+		/// <param name="Count">Maximum number of results to return</param>
+		/// <returns>List of events matching the query</returns>
+		Task<List<ILogEvent>> FindAsync(ObjectId LogId, int? Index = null, int? Count = null);
+
+		/// <summary>
+		/// Finds a list of events for a set of spans
+		/// </summary>
+		/// <param name="SpanIds">The span ids</param>
+		/// <param name="LogIds">List of log ids to query</param>
+		/// <param name="Index">Index of the first result to return</param>
+		/// <param name="Count">Number of results to return</param>
+		/// <returns>List of events for this issue</returns>
+		Task<List<ILogEvent>> FindEventsForSpansAsync(IEnumerable<ObjectId> SpanIds, ObjectId[]? LogIds = null, int Index = 0, int Count = 10);
+
+		/// <summary>
+		/// Delete all the events for a log file
+		/// </summary>
+		/// <param name="LogId">Unique id of the log</param>
+		/// <returns>Async task</returns>
+		Task DeleteLogAsync(ObjectId LogId);
+
+		/// <summary>
+		/// Update the span for an event
+		/// </summary>
+		/// <param name="Events">The events to update</param>
+		/// <param name="SpanId">New span id</param>
+		/// <returns>Async task</returns>
+		Task AddSpanToEventsAsync(IEnumerable<ILogEvent> Events, ObjectId SpanId);
+	}
+}
diff --git a/Engine/Source/Programs/Horde/HordeServer/Logs/ILogFile.cs b/Engine/Source/Programs/Horde/HordeServer/Logs/ILogFile.cs
new file mode 100644
--- /dev/null
+++ b/Engine/Source/Programs/Horde/HordeServer/Logs/ILogFile.cs
@@ -0,0 +1,121 @@
+// Copyright Epic Games, Inc. All Rights Reserved.
+
+using EpicGames.Core;
+using HordeServer.Api;
+using HordeServer.Services;
+using HordeServer.Utilities;
+using MongoDB.Bson;
+using MongoDB.Bson.Serialization.Attributes;
+using System;
+using System.Collections.Generic;
+using System.Diagnostics.CodeAnalysis;
+using System.Linq;
+using System.Threading.Tasks;
+
+namespace HordeServer.Models
+{
+	/// <summary>
+	/// Information about a log file chunk
+	/// </summary>
+	public interface ILogChunk
+	{
+		/// <summary>
+		/// Offset of the chunk within the log
+		/// </summary>
+		long Offset { get; }
+
+		/// <summary>
+		/// Length of this chunk. If zero, the chunk is still being written to.
+		/// </summary>
+		int Length { get; }
+
+		/// <summary>
+		/// Index of the first line within this chunk. If a line straddles two chunks, this is the index of the split line.
+		/// </summary>
+		int LineIndex { get; }
+
+		/// <summary>
+		/// If the chunk has yet to be pushed to persistent storage, includes the name of the server that is currently storing it.
+		/// </summary>
+		string? Server { get; }
+	}
+
+	/// <summary>
+	/// Information about a log file
+	/// </summary>
+	public interface ILogFile
+	{
+		/// <summary>
+		/// Identifier for the LogFile. Randomly generated.
+		/// </summary>
+		public ObjectId Id { get; }
+
+		/// <summary>
+		/// Unique id of the job containing this log
+		/// </summary>
+		public ObjectId JobId { get; }
+
+		/// <summary>
+		/// The session allowed to write to this log
+		/// </summary>
+		public ObjectId? SessionId { get; }
+
+		/// <summary>
+		/// Maximum line index in the file
+		/// </summary>
+		public int? MaxLineIndex { get; }
+
+		/// <summary>
+		/// Length of the file which is indexed
+		/// </summary>
+		public long? IndexLength { get; }
+
+		/// <summary>
+		/// Type of data stored in this log 
+		/// </summary>
+		public LogType Type { get; }
+
+		/// <summary>
+		/// Chunks within this file
+		/// </summary>
+		public IReadOnlyList<ILogChunk> Chunks { get; }
+	}
+
+	/// <summary>
+	/// Extension methods for log files
+	/// </summary>
+	public static class LogFileExtensions
+	{
+		/// <summary>
+		/// Gets the chunk index containing the given offset.
+		/// </summary>
+		/// <param name="Chunks">The chunks to search</param>
+		/// <param name="Offset">The offset to search for</param>
+		/// <returns>The chunk index containing the given offset</returns>
+		public static int GetChunkForOffset(this IReadOnlyList<ILogChunk> Chunks, long Offset)
+		{
+			int ChunkIndex = Chunks.BinarySearch(x => x.Offset, Offset);
+			if (ChunkIndex < 0)
+			{
+				ChunkIndex = ~ChunkIndex - 1;
+			}
+			return ChunkIndex;
+		}
+
+		/// <summary>
+		/// Gets the starting chunk index for the given line
+		/// </summary>
+		/// <param name="Chunks">The chunks to search</param>
+		/// <param name="LineIndex">Index of the line to query</param>
+		/// <returns>Index of the chunk to fetch</returns>
+		public static int GetChunkForLine(this IReadOnlyList<ILogChunk> Chunks, int LineIndex)
+		{
+			int ChunkIndex = Chunks.BinarySearch(x => x.LineIndex, LineIndex);
+			if(ChunkIndex < 0)
+			{
+				ChunkIndex = ~ChunkIndex - 1;
+			}
+			return ChunkIndex;
+		}
+	}
+}
diff --git a/Engine/Source/Programs/Horde/HordeServer/Logs/ILogFileCollection.cs b/Engine/Source/Programs/Horde/HordeServer/Logs/ILogFileCollection.cs
new file mode 100644
--- /dev/null
+++ b/Engine/Source/Programs/Horde/HordeServer/Logs/ILogFileCollection.cs
@@ -0,0 +1,103 @@
+// Copyright Epic Games, Inc. All Rights Reserved.
+
+using HordeServer.Api;
+using HordeServer.Models;
+using HordeServer.Services;
+using MongoDB.Bson;
+using MongoDB.Bson.Serialization.Attributes;
+using MongoDB.Driver;
+using System;
+using System.Collections.Generic;
+using System.Linq;
+using System.Net;
+using System.Threading.Tasks;
+
+namespace HordeServer.Collections
+{
+	/// <summary>
+	/// Updates a log file chunk
+	/// </summary>
+	public class CompleteLogChunkUpdate
+	{
+		/// <summary>
+		/// Index of the chunk
+		/// </summary>
+		public int Index { get; set; }
+
+		/// <summary>
+		/// New length for the chunk
+		/// </summary>
+		public int Length { get; set; }
+
+		/// <summary>
+		/// Number of lines in the chunk
+		/// </summary>
+		public int LineCount { get; set; }
+
+		/// <summary>
+		/// Constructor
+		/// </summary>
+		/// <param name="Index">Index of the chunk</param>
+		/// <param name="Length">New length for the chunk</param>
+		/// <param name="LineCount">Number of lines in the chunk</param>
+		public CompleteLogChunkUpdate(int Index, int Length, int LineCount)
+		{
+			this.Index = Index;
+			this.Length = Length;
+			this.LineCount = LineCount;
+		}
+	}
+
+	/// <summary>
+	/// Wrapper around the jobs collection in a mongo DB
+	/// </summary>
+	public interface ILogFileCollection
+	{
+		/// <summary>
+		/// Creates a new log
+		/// </summary>
+		/// <param name="JobId">Unique id of the job that owns this log file</param>
+		/// <param name="SessionId">Agent session allowed to update the log</param>
+		/// <param name="Type">Type of events to be stored in the log</param>
+		/// <returns>The new log file document</returns>
+		Task<ILogFile> CreateLogFileAsync(ObjectId JobId, ObjectId? SessionId, LogType Type);
+
+		/// <summary>
+		/// Adds a new chunk
+		/// </summary>
+		/// <param name="LogFileInterface">The current log file</param>
+		/// <param name="Offset">Offset of the new chunk</param>
+		/// <param name="LineIndex">Line index for the start of the chunk</param>
+		/// <returns>The updated log file document</returns>
+		Task<ILogFile?> TryAddChunkAsync(ILogFile LogFileInterface, long Offset, int LineIndex);
+
+		/// <summary>
+		/// Update the log file with final information about certain chunks
+		/// </summary>
+		/// <param name="LogFileInterface">The current log file</param>
+		/// <param name="Chunks">Chunks to update. New chunks will be inserted</param>
+		/// <returns>The updated log file document</returns>
+		Task<ILogFile?> TryCompleteChunksAsync(ILogFile LogFileInterface, IEnumerable<CompleteLogChunkUpdate> Chunks);
+
+		/// <summary>
+		/// Update the log file with final information about the index
+		/// </summary>
+		/// <param name="LogFileInterface">The current log file</param>
+		/// <param name="NewIndexLength">New length of the index</param>
+		/// <returns>The updated log file document</returns>
+		Task<ILogFile?> TryUpdateIndexAsync(ILogFile LogFileInterface, long NewIndexLength);
+
+		/// <summary>
+		/// Gets a logfile by ID
+		/// </summary>
+		/// <param name="LogFileId">Unique id of the log file</param>
+		/// <returns>The logfile document</returns>
+		Task<ILogFile?> GetLogFileAsync(ObjectId LogFileId);
+
+		/// <summary>
+		/// Gets all the log files
+		/// </summary>
+		/// <returns>List of log files</returns>
+		Task<List<ILogFile>> GetLogFilesAsync(int? Index = null, int? Count = null);
+	}
+}
diff --git a/Engine/Source/Programs/Horde/HordeServer/Logs/Impl/LogEventCollection.cs b/Engine/Source/Programs/Horde/HordeServer/Logs/Impl/LogEventCollection.cs
new file mode 100644
--- /dev/null
+++ b/Engine/Source/Programs/Horde/HordeServer/Logs/Impl/LogEventCollection.cs
@@ -0,0 +1,202 @@
+// Copyright Epic Games, Inc. All Rights Reserved.
+
+using HordeCommon;
+using HordeServer.Models;
+using HordeServer.Services;
+using HordeServer.Utilities;
+using Microsoft.Extensions.Logging;
+using MongoDB.Bson;
+using MongoDB.Bson.Serialization;
+using MongoDB.Bson.Serialization.Attributes;
+using MongoDB.Driver;
+using Serilog.Events;
+using System;
+using System.Collections.Generic;
+using System.Linq;
+using System.Threading.Tasks;
+
+namespace HordeServer.Collections.Impl
+{
+	/// <summary>
+	/// Collection of event documents
+	/// </summary>
+	public class LogEventCollection : ILogEventCollection
+	{
+		class LogEventId
+		{
+			[BsonElement("l")]
+			public ObjectId LogId { get; set; }
+
+			[BsonElement("n")]
+			public int LineIndex { get; set; }
+		}
+
+		class LogEventDocument : ILogEvent
+		{
+			[BsonId]
+			public LogEventId Id { get; set; }
+
+			[BsonElement("w"), BsonIgnoreIfDefault, BsonDefaultValue(false)]
+			public bool IsWarning { get; set; }
+
+			[BsonElement("c"), BsonIgnoreIfNull]
+			public int? LineCount { get; set; }
+
+			[BsonElement("s")]
+			public ObjectId? SpanId { get; set; }
+
+			ObjectId ILogEvent.LogId => Id.LogId;
+			int ILogEvent.LineIndex => Id.LineIndex;
+			int ILogEvent.LineCount => LineCount ?? 1;
+			EventSeverity ILogEvent.Severity => IsWarning ? EventSeverity.Warning : EventSeverity.Error;
+
+			public LogEventDocument()
+			{
+				this.Id = new LogEventId();
+			}
+
+			public LogEventDocument(ObjectId LogId, EventSeverity Severity, int LineIndex, int LineCount, ObjectId? SpanId)
+			{
+				this.Id = new LogEventId { LogId = LogId, LineIndex = LineIndex };
+				this.IsWarning = Severity == EventSeverity.Warning;
+				this.LineCount = (LineCount > 1) ? (int?)LineCount : null;
+				this.SpanId = SpanId;
+			}
+
+			public LogEventDocument(NewLogEventData Data)
+				: this(Data.LogId, Data.Severity, Data.LineIndex, Data.LineCount, Data.SpanId)
+			{
+			}
+		}
+
+		class LegacyEventDocument
+		{
+			public ObjectId Id { get; set; }
+			public DateTime Time { get; set; }
+			public EventSeverity Severity { get; set; }
+			public ObjectId LogId { get; set; }
+			public int LineIndex { get; set; }
+			public int LineCount { get; set; }
+
+			public string? Message { get; set; }
+
+			[BsonIgnoreIfNull, BsonElement("IssueId2")]
+			public int? IssueId { get; set; }
+
+			public BsonDocument? Data { get; set; }
+
+			public int UpgradeVersion { get; set; }
+		}
+
+		/// <summary>
+		/// Collection of event documents
+		/// </summary>
+		IMongoCollection<LogEventDocument> LogEvents;
+
+		/// <summary>
+		/// Collection of legacy event documents
+		/// </summary>
+		IMongoCollection<LegacyEventDocument> LegacyEvents;
+
+		/// <summary>
+		/// Logger for upgrade messages
+		/// </summary>
+		ILogger<LogEventCollection> Logger;
+
+		/// <summary>
+		/// Constructor
+		/// </summary>
+		/// <param name="DatabaseService">The database service instance</param>
+		/// <param name="Logger">The logger instance</param>
+		public LogEventCollection(DatabaseService DatabaseService, ILogger<LogEventCollection> Logger)
+		{
+			this.Logger = Logger;
+
+			LogEvents = DatabaseService.GetCollection<LogEventDocument>("LogEvents");
+			LegacyEvents = DatabaseService.GetCollection<LegacyEventDocument>("Events");
+
+			if (!DatabaseService.ReadOnlyMode)
+			{
+				LogEvents.Indexes.CreateOne(new CreateIndexModel<LogEventDocument>(Builders<LogEventDocument>.IndexKeys.Ascending(x => x.Id.LogId)));
+				LogEvents.Indexes.CreateOne(new CreateIndexModel<LogEventDocument>(Builders<LogEventDocument>.IndexKeys.Ascending(x => x.SpanId).Ascending(x => x.Id)));
+				LegacyEvents.Indexes.CreateOne(new CreateIndexModel<LegacyEventDocument>(Builders<LegacyEventDocument>.IndexKeys.Ascending(x => x.LogId)));
+			}
+		}
+
+		/// <inheritdoc/>
+		public Task AddAsync(NewLogEventData NewEvent)
+		{
+			return LogEvents.InsertOneAsync(new LogEventDocument(NewEvent));
+		}
+
+		/// <inheritdoc/>
+		public Task AddManyAsync(List<NewLogEventData> NewEvents)
+		{
+			return LogEvents.InsertManyAsync(NewEvents.ConvertAll(x => new LogEventDocument(x)));
+		}
+
+		/// <inheritdoc/>
+		public async Task<List<ILogEvent>> FindAsync(ObjectId LogId, int? Index = null, int? Count = null)
+		{
+			Logger.LogInformation("Querying for log events for log {LogId} creation time {CreateTime}", LogId, LogId.CreationTime);
+
+			FilterDefinitionBuilder<LogEventDocument> Builder = Builders<LogEventDocument>.Filter;
+
+			FilterDefinition<LogEventDocument> Filter = Builder.Eq(x => x.Id.LogId, LogId);
+
+			IFindFluent<LogEventDocument, LogEventDocument> Results = LogEvents.Find(Filter);
+			if (Index != null)
+			{
+				Results = Results.Skip(Index.Value);
+			}
+			if (Count != null)
+			{
+				Results = Results.Limit(Count.Value);
+			}
+
+			List<LogEventDocument> EventDocuments = await Results.ToListAsync();
+			return EventDocuments.ConvertAll<ILogEvent>(x => x);
+		}
+
+		/// <inheritdoc/>
+		public async Task<List<ILogEvent>> FindEventsForSpanAsync(ObjectId SpanId, ObjectId[]? LogIds, int Index, int Count)
+		{
+			FilterDefinition<LogEventDocument> Filter = Builders<LogEventDocument>.Filter.Eq(x => x.SpanId, SpanId);
+			if (LogIds != null && LogIds.Length > 0)
+			{
+				Filter &= Builders<LogEventDocument>.Filter.In(x => x.Id.LogId, LogIds);
+			}
+
+			List<LogEventDocument> EventDocuments = await LogEvents.Find(Filter).Skip(Index).Limit(Count).ToListAsync();
+			return EventDocuments.ConvertAll<ILogEvent>(x => x);
+		}
+
+		/// <inheritdoc/>
+		public async Task<List<ILogEvent>> FindEventsForSpansAsync(IEnumerable<ObjectId> SpanIds, ObjectId[]? LogIds, int Index, int Count)
+		{
+			FilterDefinition<LogEventDocument> Filter = Builders<LogEventDocument>.Filter.In(x => x.SpanId, SpanIds.Select<ObjectId, ObjectId?>(x => x));
+			if (LogIds != null && LogIds.Length > 0)
+			{
+				Filter &= Builders<LogEventDocument>.Filter.In(x => x.Id.LogId, LogIds);
+			}
+
+			List<LogEventDocument> EventDocuments = await LogEvents.Find(Filter).Skip(Index).Limit(Count).ToListAsync();
+			return EventDocuments.ConvertAll<ILogEvent>(x => x);
+		}
+
+		/// <inheritdoc/>
+		public async Task DeleteLogAsync(ObjectId LogId)
+		{
+			await LogEvents.DeleteManyAsync(x => x.Id.LogId == LogId);
+			await LegacyEvents.DeleteManyAsync(x => x.LogId == LogId);
+		}
+
+		/// <inheritdoc/>
+		public async Task AddSpanToEventsAsync(IEnumerable<ILogEvent> Events, ObjectId SpanId)
+		{
+			FilterDefinition<LogEventDocument> EventFilter = Builders<LogEventDocument>.Filter.In(x => x.Id, Events.OfType<LogEventDocument>().Select(x => x.Id));
+			UpdateDefinition<LogEventDocument> EventUpdate = Builders<LogEventDocument>.Update.Set(x => x.SpanId, SpanId);
+			await LogEvents.UpdateManyAsync(EventFilter, EventUpdate);
+		}
+	}
+}
diff --git a/Engine/Source/Programs/Horde/HordeServer/Logs/Impl/LogFileCollection.cs b/Engine/Source/Programs/Horde/HordeServer/Logs/Impl/LogFileCollection.cs
new file mode 100644
--- /dev/null
+++ b/Engine/Source/Programs/Horde/HordeServer/Logs/Impl/LogFileCollection.cs
@@ -0,0 +1,230 @@
+// Copyright Epic Games, Inc. All Rights Reserved.
+
+using HordeServer.Api;
+using HordeServer.Models;
+using HordeServer.Services;
+using MongoDB.Bson;
+using MongoDB.Bson.Serialization.Attributes;
+using MongoDB.Driver;
+using System;
+using System.Collections.Generic;
+using System.Linq;
+using System.Net;
+using System.Threading.Tasks;
+
+namespace HordeServer.Collections.Impl
+{
+	/// <summary>
+	/// Wrapper around the jobs collection in a mongo DB
+	/// </summary>
+	public class LogFileCollection : ILogFileCollection
+	{
+		class LogChunkDocument : ILogChunk
+		{
+			public long Offset { get; set; }
+			public int Length { get; set; }
+			public int LineIndex { get; set; }
+
+			[BsonIgnoreIfNull]
+			public string? Server { get; set; }
+
+			[BsonConstructor]
+			public LogChunkDocument()
+			{
+			}
+
+			public LogChunkDocument(LogChunkDocument Other)
+			{
+				Offset = Other.Offset;
+				Length = Other.Length;
+				LineIndex = Other.LineIndex;
+				Server = Other.Server;
+			}
+
+			public LogChunkDocument Clone()
+			{
+				return (LogChunkDocument)MemberwiseClone();
+			}
+		}
+
+		class LogFileDocument : ILogFile
+		{
+			[BsonRequired, BsonId]
+			public ObjectId Id { get; set; }
+
+			[BsonRequired]
+			public ObjectId JobId { get; set; }
+
+			public ObjectId? SessionId { get; set; }
+			public LogType Type { get; set; }
+
+			[BsonIgnoreIfNull]
+			public int? MaxLineIndex { get; set; }
+
+			[BsonIgnoreIfNull]
+			public long? IndexLength { get; set; }
+
+			public List<LogChunkDocument> Chunks { get; set; } = new List<LogChunkDocument>();
+
+			[BsonRequired]
+			public int UpdateIndex { get; set; }
+
+			IReadOnlyList<ILogChunk> ILogFile.Chunks => Chunks;
+
+			[BsonConstructor]
+			private LogFileDocument()
+			{
+			}
+
+			public LogFileDocument(ObjectId JobId, ObjectId? SessionId, LogType Type)
+			{
+				this.Id = ObjectId.GenerateNewId();
+				this.JobId = JobId;
+				this.SessionId = SessionId;
+				this.Type = Type;
+				this.MaxLineIndex = 0;
+			}
+
+			public LogFileDocument Clone()
+			{
+				LogFileDocument Document = (LogFileDocument)MemberwiseClone();
+				Document.Chunks = Document.Chunks.ConvertAll(x => x.Clone());
+				return Document;
+			}
+		}
+
+		/// <summary>
+		/// The jobs collection
+		/// </summary>
+		IMongoCollection<LogFileDocument> LogFiles;
+
+		/// <summary>
+		/// Hostname for the current server
+		/// </summary>
+		string HostName;
+
+		/// <summary>
+		/// Constructor
+		/// </summary>
+		/// <param name="DatabaseService">The database service singleton</param>
+		public LogFileCollection(DatabaseService DatabaseService)
+		{
+			LogFiles = DatabaseService.GetCollection<LogFileDocument>("LogFiles");
+			this.HostName = Dns.GetHostName();
+		}
+
+		/// <inheritdoc/>
+		public async Task<ILogFile> CreateLogFileAsync(ObjectId JobId, ObjectId? SessionId, LogType Type)
+		{
+			LogFileDocument NewLogFile = new LogFileDocument(JobId, SessionId, Type);
+			await LogFiles.InsertOneAsync(NewLogFile);
+			return NewLogFile;
+		}
+
+		/// <inheritdoc/>
+		public async Task<ILogFile?> TryAddChunkAsync(ILogFile LogFileInterface, long Offset, int LineIndex)
+		{
+			LogFileDocument LogFile = ((LogFileDocument)LogFileInterface).Clone();
+
+			int ChunkIdx = LogFile.Chunks.GetChunkForOffset(Offset) + 1;
+
+			LogChunkDocument Chunk = new LogChunkDocument();
+			Chunk.Offset = Offset;
+			Chunk.LineIndex = LineIndex;
+			Chunk.Server = HostName;
+			LogFile.Chunks.Insert(ChunkIdx, Chunk);
+
+			UpdateDefinition<LogFileDocument> Update = Builders<LogFileDocument>.Update.Set(x => x.Chunks, LogFile.Chunks);
+			if (ChunkIdx == LogFile.Chunks.Count - 1)
+			{
+				LogFile.MaxLineIndex = null;
+				Update = Update.Unset(x => x.MaxLineIndex);
+			}
+
+			if (!await TryUpdateLogFileAsync(LogFile, Update))
+			{
+				return null;
+			}
+
+			return LogFile;
+		}
+
+		/// <inheritdoc/>
+		public async Task<ILogFile?> TryCompleteChunksAsync(ILogFile LogFileInterface, IEnumerable<CompleteLogChunkUpdate> ChunkUpdates)
+		{
+			LogFileDocument LogFile = ((LogFileDocument)LogFileInterface).Clone();
+
+			// Update the length of any complete chunks
+			UpdateDefinitionBuilder<LogFileDocument> UpdateBuilder = Builders<LogFileDocument>.Update;
+			List<UpdateDefinition<LogFileDocument>> Updates = new List<UpdateDefinition<LogFileDocument>>();
+			foreach (CompleteLogChunkUpdate ChunkUpdate in ChunkUpdates)
+			{
+				LogChunkDocument Chunk = LogFile.Chunks[ChunkUpdate.Index];
+				Chunk.Length = ChunkUpdate.Length;
+				Updates.Add(UpdateBuilder.Set(x => x.Chunks[ChunkUpdate.Index].Length, ChunkUpdate.Length));
+
+				if (ChunkUpdate.Index == LogFile.Chunks.Count - 1)
+				{
+					LogFile.MaxLineIndex = Chunk.LineIndex + ChunkUpdate.LineCount;
+					Updates.Add(UpdateBuilder.Set(x => x.MaxLineIndex, LogFile.MaxLineIndex));
+				}
+			}
+
+			// Try to apply the updates
+			if (Updates.Count > 0 && !await TryUpdateLogFileAsync(LogFile, UpdateBuilder.Combine(Updates)))
+			{
+				return null;
+			}
+
+			return LogFile;
+		}
+
+		/// <inheritdoc/>
+		public async Task<ILogFile?> TryUpdateIndexAsync(ILogFile LogFileInterface, long NewIndexLength)
+		{
+			LogFileDocument LogFile = ((LogFileDocument)LogFileInterface).Clone();
+
+			UpdateDefinition<LogFileDocument> Update = Builders<LogFileDocument>.Update.Set(x => x.IndexLength, NewIndexLength);
+			if (!await TryUpdateLogFileAsync(LogFile, Update))
+			{
+				return null;
+			}
+
+			LogFile.IndexLength = NewIndexLength;
+			return LogFile;
+		}
+
+		/// <inheritdoc/>
+		private async Task<bool> TryUpdateLogFileAsync(LogFileDocument Current, UpdateDefinition<LogFileDocument> Update)
+		{
+			int PrevUpdateIndex = Current.UpdateIndex;
+			Current.UpdateIndex++;
+			UpdateResult Result = await LogFiles.UpdateOneAsync<LogFileDocument>(x => x.Id == Current.Id && x.UpdateIndex == PrevUpdateIndex, Update.Set(x => x.UpdateIndex, Current.UpdateIndex));
+			return Result.ModifiedCount == 1;
+		}
+
+		/// <inheritdoc/>
+		public async Task<ILogFile?> GetLogFileAsync(ObjectId LogFileId)
+		{
+			LogFileDocument LogFile = await LogFiles.Find<LogFileDocument>(x => x.Id == LogFileId).FirstOrDefaultAsync();
+			return LogFile;
+		}
+
+		/// <inheritdoc/>
+		public async Task<List<ILogFile>> GetLogFilesAsync(int? Index = null, int? Count = null)
+		{
+			IFindFluent<LogFileDocument, LogFileDocument> Query = LogFiles.Find(FilterDefinition<LogFileDocument>.Empty);
+			if(Index != null)
+			{
+				Query = Query.Skip(Index.Value);
+			}
+			if(Count != null)
+			{
+				Query = Query.Limit(Count.Value);
+			}
+
+			List<LogFileDocument> Results = await Query.ToListAsync();
+			return Results.ConvertAll<ILogFile>(x => x);
+		}
+	}
+}
diff --git a/Engine/Source/Programs/Horde/HordeServer/Models/IDevice.cs b/Engine/Source/Programs/Horde/HordeServer/Models/IDevice.cs
deleted file mode 100644
diff --git a/Engine/Source/Programs/Horde/HordeServer/Models/ILogEvent.cs b/Engine/Source/Programs/Horde/HordeServer/Models/ILogEvent.cs
deleted file mode 100644
diff --git a/Engine/Source/Programs/Horde/HordeServer/Models/ILogFile.cs b/Engine/Source/Programs/Horde/HordeServer/Models/ILogFile.cs
deleted file mode 100644
diff --git a/Engine/Source/Programs/Horde/HordeServer/Services/DeviceService.cs b/Engine/Source/Programs/Horde/HordeServer/Services/DeviceService.cs
deleted file mode 100644
diff --git a/Engine/Source/Programs/Horde/HordeServer/Utilities/AgentId.cs b/Engine/Source/Programs/Horde/HordeServer/Utilities/AgentId.cs
deleted file mode 100644
