Horde: Move jobs, graphs, issues and users into subfolders.

diff --git a/Engine/Source/Programs/Horde/HordeServer/Agents/ILeaseCollection.cs b/Engine/Source/Programs/Horde/HordeServer/Agents/ILeaseCollection.cs
new file mode 100644
--- /dev/null
+++ b/Engine/Source/Programs/Horde/HordeServer/Agents/ILeaseCollection.cs
@@ -0,0 +1,84 @@
+// Copyright Epic Games, Inc. All Rights Reserved.
+
+using HordeCommon;
+using HordeServer.Models;
+using HordeServer.Services;
+using HordeServer.Utilities;
+using MongoDB.Bson;
+using MongoDB.Bson.Serialization.Attributes;
+using MongoDB.Driver;
+using System;
+using System.Collections.Generic;
+using System.Linq;
+using System.Threading.Tasks;
+
+namespace HordeServer.Collections
+{
+	using PoolId = StringId<IPool>;
+	using StreamId = StringId<IStream>;
+
+	/// <summary>
+	/// Interface for a collection of lease documents
+	/// </summary>
+	public interface ILeaseCollection
+	{
+		/// <summary>
+		/// Adds a lease to the collection
+		/// </summary>
+		/// <param name="Id">The lease id</param>
+		/// <param name="Name">Name of the lease</param>
+		/// <param name="AgentId">The agent id</param>
+		/// <param name="SessionId">The agent session handling the lease</param>
+		/// <param name="StreamId">Stream for the payload</param>
+		/// <param name="PoolId">The pool for the lease</param>
+		/// <param name="LogId">Log id for the lease</param>
+		/// <param name="StartTime">Start time of the lease</param>
+		/// <param name="Payload">Payload for the lease</param>
+		/// <returns>Async task</returns>
+		Task<ILease> AddAsync(ObjectId Id, string Name, AgentId AgentId, ObjectId SessionId, StreamId? StreamId, PoolId? PoolId, ObjectId? LogId, DateTime StartTime, byte[] Payload);
+
+		/// <summary>
+		/// Deletes a lease from the collection
+		/// </summary>
+		/// <param name="LeaseId">Unique id of the lease</param>
+		/// <returns>Async task</returns>
+		Task DeleteAsync(ObjectId LeaseId);
+
+		/// <summary>
+		/// Gets a specific lease
+		/// </summary>
+		/// <param name="LeaseId">Unique id of the lease</param>
+		/// <returns>The lease that was found, or null if it does not exist</returns>
+		Task<ILease?> GetAsync(ObjectId LeaseId);
+
+		/// <summary>
+		/// Finds all leases matching a set of criteria
+		/// </summary>
+		/// <param name="AgentId">Unqiue id of the agent executing this lease</param>
+		/// <param name="SessionId">Unique id of the agent session</param>
+		/// <param name="MinTime">Start of the window to include leases</param>
+		/// <param name="MaxTime">End of the window to include leases</param>
+		/// <param name="Index">Index of the first result to return</param>
+		/// <param name="Count">Number of results to return</param>
+		/// <returns>List of leases matching the given criteria</returns>
+		Task<List<ILease>> FindLeasesAsync(AgentId? AgentId = null, ObjectId? SessionId = null, DateTime? MinTime = null, DateTime? MaxTime = null, int? Index = null, int? Count = null);
+
+		/// <summary>
+		/// Finds all active leases
+		/// </summary>
+		/// <param name="Index">Index of the first result to return</param>
+		/// <param name="Count">Number of results to return</param>
+		/// <returns>List of leases</returns>
+		Task<List<ILease>> FindActiveLeasesAsync(int? Index = null, int? Count = null);
+
+		/// <summary>
+		/// Sets the outcome of a lease
+		/// </summary>
+		/// <param name="LeaseId">The lease to update</param>
+		/// <param name="FinishTime">Time at which the lease finished</param>
+		/// <param name="Outcome">Outcome of the lease</param>
+		/// <param name="Output">Output data from the task</param>
+		/// <returns>True if the lease was updated, false otherwise</returns>
+		Task<bool> TrySetOutcomeAsync(ObjectId LeaseId, DateTime FinishTime, LeaseOutcome Outcome, byte[]? Output);
+	}
+}
diff --git a/Engine/Source/Programs/Horde/HordeServer/Agents/ISessionCollection.cs b/Engine/Source/Programs/Horde/HordeServer/Agents/ISessionCollection.cs
new file mode 100644
--- /dev/null
+++ b/Engine/Source/Programs/Horde/HordeServer/Agents/ISessionCollection.cs
@@ -0,0 +1,76 @@
+// Copyright Epic Games, Inc. All Rights Reserved.
+
+using HordeServer.Models;
+using HordeServer.Services;
+using HordeServer.Utilities;
+using Microsoft.Extensions.DependencyInjection;
+using MongoDB.Bson;
+using MongoDB.Bson.Serialization.Attributes;
+using MongoDB.Driver;
+using System;
+using System.Collections.Generic;
+using System.Linq;
+using System.Threading.Tasks;
+
+namespace HordeServer.Collections
+{
+	/// <summary>
+	/// Interface for a collection of session documents
+	/// </summary>
+	public interface ISessionCollection
+	{
+		/// <summary>
+		/// Adds a new session to the collection
+		/// </summary>
+		/// <param name="Id">The session id</param>
+		/// <param name="AgentId">The agent this session is for</param>
+		/// <param name="StartTime">Start time of this session</param>
+		/// <param name="Properties">Properties of this agent at the time the session started</param>
+		/// <param name="Resources">Resources which the agent has</param>
+		/// <param name="Version">Version of the agent software</param>
+		Task<ISession> AddAsync(ObjectId Id, AgentId AgentId, DateTime StartTime, IReadOnlyList<string>? Properties, IReadOnlyDictionary<string, int>? Resources, string? Version);
+
+		/// <summary>
+		/// Gets information about a particular session
+		/// </summary>
+		/// <param name="SessionId">The unique session id</param>
+		/// <returns>The session information</returns>
+		Task<ISession?> GetAsync(ObjectId SessionId);
+
+		/// <summary>
+		/// Find sessions for the given agent
+		/// </summary>
+		/// <param name="AgentId">The unique agent id</param>
+		/// <param name="StartTime">Start time to include in the search</param>
+		/// <param name="FinishTime">Finish time to include in the search</param>
+		/// <param name="Index">Index of the first result to return</param>
+		/// <param name="Count">Number of results to return</param>
+		/// <returns>List of sessions matching the given criteria</returns>
+		Task<List<ISession>> FindAsync(AgentId AgentId, DateTime? StartTime, DateTime? FinishTime, int Index, int Count);
+
+		/// <summary>
+		/// Finds any active sessions
+		/// </summary>
+		/// <param name="Index">Index of the first result to return</param>
+		/// <param name="Count">Number of results to return</param>
+		/// <returns>List of sessions</returns>
+		Task<List<ISession>> FindActiveSessionsAsync(int? Index = null, int? Count = null);
+
+		/// <summary>
+		/// Update a session from the collection
+		/// </summary>
+		/// <param name="SessionId">The session to update</param>
+		/// <param name="FinishTime">Time at which the session finished</param>
+		/// <param name="Properties">The agent properties</param>
+		/// <param name="Resources">Resources which the agent has</param>
+		/// <returns>Async task</returns>
+		Task UpdateAsync(ObjectId SessionId, DateTime FinishTime, IReadOnlyList<string> Properties, IReadOnlyDictionary<string, int> Resources);
+
+		/// <summary>
+		/// Delete a session from the collection
+		/// </summary>
+		/// <param name="SessionId">The session id</param>
+		/// <returns>Async task</returns>
+		Task DeleteAsync(ObjectId SessionId);
+	}
+}
diff --git a/Engine/Source/Programs/Horde/HordeServer/Agents/Impl/LeaseCollection.cs b/Engine/Source/Programs/Horde/HordeServer/Agents/Impl/LeaseCollection.cs
new file mode 100644
--- /dev/null
+++ b/Engine/Source/Programs/Horde/HordeServer/Agents/Impl/LeaseCollection.cs
@@ -0,0 +1,163 @@
+// Copyright Epic Games, Inc. All Rights Reserved.
+
+using HordeCommon;
+using HordeServer.Models;
+using HordeServer.Services;
+using HordeServer.Utilities;
+using MongoDB.Bson;
+using MongoDB.Bson.Serialization.Attributes;
+using MongoDB.Driver;
+using System;
+using System.Collections.Generic;
+using System.Linq;
+using System.Threading.Tasks;
+
+namespace HordeServer.Collections.Impl
+{
+	using PoolId = StringId<IPool>;
+	using StreamId = StringId<IStream>;
+
+	/// <summary>
+	/// Collection of lease documents
+	/// </summary>
+	public class LeaseCollection : ILeaseCollection
+	{
+		/// <summary>
+		/// Concrete implementation of a lease document
+		/// </summary>
+		class LeaseDocument : ILease
+		{
+			public ObjectId Id { get; set; }
+			public string Name { get; set; }
+			public AgentId AgentId { get; set; }
+			public ObjectId SessionId { get; set; }
+			public StreamId? StreamId { get; set; }
+			public PoolId? PoolId { get; set; }
+			public ObjectId? LogId { get; set; }
+			public DateTime StartTime { get; set; }
+			public DateTime? FinishTime { get; set; }
+			public byte[] Payload { get; set; }
+			public LeaseOutcome Outcome { get; set; }
+
+			[BsonIgnoreIfNull]
+			public byte[]? Output { get; set; }
+
+			ReadOnlyMemory<byte> ILease.Payload => Payload;
+			ReadOnlyMemory<byte> ILease.Output => Output ?? Array.Empty<byte>();
+
+			[BsonConstructor]
+			private LeaseDocument()
+			{
+				Name = String.Empty;
+				Payload = null!;
+			}
+
+			public LeaseDocument(ObjectId Id, string Name, AgentId AgentId, ObjectId SessionId, StreamId? StreamId, PoolId? PoolId, ObjectId? LogId, DateTime StartTime, byte[] Payload)
+			{
+				this.Id = Id;
+				this.Name = Name;
+				this.AgentId = AgentId;
+				this.SessionId = SessionId;
+				this.StreamId = StreamId;
+				this.PoolId = PoolId;
+				this.LogId = LogId;
+				this.StartTime = StartTime;
+				this.Payload = Payload;
+			}
+		}
+
+		/// <summary>
+		/// Collection of lease documents
+		/// </summary>
+		readonly IMongoCollection<LeaseDocument> Leases;
+
+		/// <summary>
+		/// Constructor
+		/// </summary>
+		/// <param name="DatabaseService">The database service instance</param>
+		public LeaseCollection(DatabaseService DatabaseService)
+		{
+			Leases = DatabaseService.GetCollection<LeaseDocument>("Leases");
+
+			if (!DatabaseService.ReadOnlyMode)
+			{
+				Leases.Indexes.CreateOne(new CreateIndexModel<LeaseDocument>(Builders<LeaseDocument>.IndexKeys.Ascending(x => x.AgentId)));
+				Leases.Indexes.CreateOne(new CreateIndexModel<LeaseDocument>(Builders<LeaseDocument>.IndexKeys.Ascending(x => x.SessionId)));
+				Leases.Indexes.CreateOne(new CreateIndexModel<LeaseDocument>(Builders<LeaseDocument>.IndexKeys.Ascending(x => x.StartTime)));
+				Leases.Indexes.CreateOne(new CreateIndexModel<LeaseDocument>(Builders<LeaseDocument>.IndexKeys.Ascending(x => x.FinishTime)));
+			}
+		}
+
+		/// <inheritdoc/>
+		public async Task<ILease> AddAsync(ObjectId Id, string Name, AgentId AgentId, ObjectId SessionId, StreamId? StreamId, PoolId? PoolId, ObjectId? LogId, DateTime StartTime, byte[] Payload)
+		{
+			LeaseDocument Lease = new LeaseDocument(Id, Name, AgentId, SessionId, StreamId, PoolId, LogId, StartTime, Payload);
+			await Leases.ReplaceOneAsync(x => x.Id == Id, Lease, new ReplaceOptions { IsUpsert = true });
+			return Lease;
+		}
+
+		/// <inheritdoc/>
+		public async Task DeleteAsync(ObjectId LeaseId)
+		{
+			await Leases.DeleteOneAsync(x => x.Id == LeaseId);
+		}
+
+		/// <inheritdoc/>
+		public async Task<ILease?> GetAsync(ObjectId LeaseId)
+		{
+			return await Leases.Find(x => x.Id == LeaseId).FirstOrDefaultAsync();
+		}
+
+		/// <inheritdoc/>
+		public async Task<List<ILease>> FindLeasesAsync(AgentId? AgentId, ObjectId? SessionId, DateTime? MinTime, DateTime? MaxTime, int? Index, int? Count)
+		{
+			FilterDefinitionBuilder<LeaseDocument> FilterBuilder = Builders<LeaseDocument>.Filter;
+
+			FilterDefinition<LeaseDocument> Filter = FilterDefinition<LeaseDocument>.Empty;
+			if (AgentId != null)
+			{
+				Filter &= FilterBuilder.Eq(x => x.AgentId, AgentId.Value);
+			}
+			if (SessionId != null)
+			{
+				Filter &= FilterBuilder.Eq(x => x.SessionId, SessionId.Value);
+			}
+			if (MinTime != null)
+			{
+				Filter &= FilterBuilder.Or(FilterBuilder.Eq(x => x.FinishTime!, null), FilterBuilder.Gt(x => x.FinishTime!, MinTime.Value));
+			}
+			if (MaxTime != null)
+			{
+				Filter &= FilterBuilder.Lt(x => x.StartTime, MaxTime.Value);
+			}
+
+			List<LeaseDocument> Results = await Leases.Find(Filter).SortByDescending(x => x.StartTime).Range(Index, Count).ToListAsync();
+			return Results.ConvertAll<ILease>(x => x);
+		}
+
+		/// <inheritdoc/>
+		public async Task<List<ILease>> FindActiveLeasesAsync(int? Index = null, int? Count = null)
+		{
+			List<LeaseDocument> Results = await Leases.Find(x => x.FinishTime == null).Range(Index, Count).ToListAsync();
+			return Results.ConvertAll<ILease>(x => x);
+		}
+
+		/// <inheritdoc/>
+		public async Task<bool> TrySetOutcomeAsync(ObjectId LeaseId, DateTime FinishTime, LeaseOutcome Outcome, byte[]? Output)
+		{
+			FilterDefinitionBuilder<LeaseDocument> FilterBuilder = Builders<LeaseDocument>.Filter;
+			FilterDefinition<LeaseDocument> Filter = FilterBuilder.Eq(x => x.Id, LeaseId) & FilterBuilder.Eq(x => x.FinishTime, null);
+
+			UpdateDefinitionBuilder<LeaseDocument> UpdateBuilder = Builders<LeaseDocument>.Update;
+			UpdateDefinition<LeaseDocument> Update = UpdateBuilder.Set(x => x.FinishTime, FinishTime).Set(x => x.Outcome, Outcome);
+
+			if(Output != null && Output.Length > 0)
+			{
+				Update = Update.Set(x => x.Output, Output);
+			}
+
+			UpdateResult Result = await Leases.UpdateOneAsync(Filter, Update);
+			return Result.ModifiedCount > 0;
+		}
+	}
+}
diff --git a/Engine/Source/Programs/Horde/HordeServer/Agents/Impl/SessionCollection.cs b/Engine/Source/Programs/Horde/HordeServer/Agents/Impl/SessionCollection.cs
new file mode 100644
--- /dev/null
+++ b/Engine/Source/Programs/Horde/HordeServer/Agents/Impl/SessionCollection.cs
@@ -0,0 +1,141 @@
+// Copyright Epic Games, Inc. All Rights Reserved.
+
+using HordeServer.Models;
+using HordeServer.Services;
+using HordeServer.Utilities;
+using Microsoft.Extensions.DependencyInjection;
+using MongoDB.Bson;
+using MongoDB.Bson.Serialization.Attributes;
+using MongoDB.Driver;
+using System;
+using System.Collections.Generic;
+using System.Linq;
+using System.Threading.Tasks;
+
+namespace HordeServer.Collections.Impl
+{
+	/// <summary>
+	/// Collection of session documents
+	/// </summary>
+	public class SessionCollection : ISessionCollection
+	{
+		/// <summary>
+		/// Concrete implementation of ISession
+		/// </summary>
+		class SessionDocument : ISession
+		{
+			[BsonRequired, BsonId]
+			public ObjectId Id { get; set; }
+
+			[BsonRequired]
+			public AgentId AgentId { get; set; }
+
+			public DateTime StartTime { get; set; }
+			public DateTime? FinishTime { get; set; }
+			public List<string>? Properties { get; set; }
+			public Dictionary<string, int>? Resources { get; set; }
+			public string? Version { get; set; }
+
+			IReadOnlyList<string>? ISession.Properties => Properties;
+
+			[BsonConstructor]
+			private SessionDocument()
+			{
+			}
+
+			public SessionDocument(ObjectId Id, AgentId AgentId, DateTime StartTime, IReadOnlyList<string>? Properties, IReadOnlyDictionary<string, int>? Resources, string? Version)
+			{
+				this.Id = Id;
+				this.AgentId = AgentId;
+				this.StartTime = StartTime;
+				if (Properties != null)
+				{
+					this.Properties = new List<string>(Properties);
+				}
+				if (Resources != null)
+				{
+					this.Resources = new Dictionary<string, int>(Resources);
+				}
+				this.Version = Version;
+			}
+		}
+
+		/// <summary>
+		/// Collection of session documents
+		/// </summary>
+		readonly IMongoCollection<SessionDocument> Sessions;
+
+		/// <summary>
+		/// Constructor
+		/// </summary>
+		/// <param name="DatabaseService">The database service</param>
+		public SessionCollection(DatabaseService DatabaseService)
+		{
+			Sessions = DatabaseService.GetCollection<SessionDocument>("Sessions");
+
+			if (!DatabaseService.ReadOnlyMode)
+			{
+				Sessions.Indexes.CreateOne(new CreateIndexModel<SessionDocument>(Builders<SessionDocument>.IndexKeys.Ascending(x => x.AgentId)));
+				Sessions.Indexes.CreateOne(new CreateIndexModel<SessionDocument>(Builders<SessionDocument>.IndexKeys.Ascending(x => x.StartTime)));
+				Sessions.Indexes.CreateOne(new CreateIndexModel<SessionDocument>(Builders<SessionDocument>.IndexKeys.Ascending(x => x.FinishTime)));
+			}
+		}
+
+		/// <inheritdoc/>
+		public async Task<ISession> AddAsync(ObjectId Id, AgentId AgentId, DateTime StartTime, IReadOnlyList<string>? Properties, IReadOnlyDictionary<string, int>? Resources, string? Version)
+		{
+			SessionDocument NewSession = new SessionDocument(Id, AgentId, StartTime, Properties, Resources, Version);
+			await Sessions.InsertOneAsync(NewSession);
+			return NewSession;
+		}
+
+		/// <inheritdoc/>
+		public async Task<ISession?> GetAsync(ObjectId SessionId)
+		{
+			return await Sessions.Find(x => x.Id == SessionId).FirstOrDefaultAsync();
+		}
+
+		/// <inheritdoc/>
+		public async Task<List<ISession>> FindAsync(AgentId AgentId, DateTime? StartTime, DateTime? FinishTime, int Index, int Count)
+		{
+			FilterDefinitionBuilder<SessionDocument> FilterBuilder = Builders<SessionDocument>.Filter;
+
+			FilterDefinition<SessionDocument> Filter = FilterBuilder.Eq(x => x.AgentId, AgentId);
+			if (StartTime != null)
+			{
+				Filter &= FilterBuilder.Gte(x => x.StartTime, StartTime.Value);
+			}
+			if (FinishTime != null)
+			{
+				Filter &= FilterBuilder.Or(FilterBuilder.Eq(x => x.FinishTime, null), FilterBuilder.Lte(x => x.FinishTime, FinishTime.Value));
+			}
+
+			List<SessionDocument> Results = await Sessions.Find(Filter).SortByDescending(x => x.StartTime).Skip(Index).Limit(Count).ToListAsync();
+			return Results.ConvertAll<ISession>(x => x);
+		}
+
+		/// <inheritdoc/>
+		public async Task<List<ISession>> FindActiveSessionsAsync(int? Index, int? Count)
+		{
+			List<SessionDocument> Results = await Sessions.Find(x => x.FinishTime == null).Range(Index, Count).ToListAsync();
+			return Results.ConvertAll<ISession>(x => x);
+		}
+
+		/// <inheritdoc/>
+		public Task UpdateAsync(ObjectId SessionId, DateTime FinishTime, IReadOnlyList<string> Properties, IReadOnlyDictionary<string, int> Resources)
+		{
+			UpdateDefinition<SessionDocument> Update = Builders<SessionDocument>.Update
+				.Set(x => x.FinishTime, FinishTime)
+				.Set(x => x.Properties, new List<string>(Properties))
+				.Set(x => x.Resources, new Dictionary<string, int>(Resources));
+
+			return Sessions.FindOneAndUpdateAsync(x => x.Id == SessionId, Update);
+		}
+
+		/// <inheritdoc/>
+		public Task DeleteAsync(ObjectId SessionId)
+		{
+			return Sessions.DeleteOneAsync(x => x.Id == SessionId);
+		}
+	}
+}
diff --git a/Engine/Source/Programs/Horde/HordeServer/Health/IIssue.cs b/Engine/Source/Programs/Horde/HordeServer/Health/IIssue.cs
new file mode 100644
--- /dev/null
+++ b/Engine/Source/Programs/Horde/HordeServer/Health/IIssue.cs
@@ -0,0 +1,244 @@
+// Copyright Epic Games, Inc. All Rights Reserved.
+
+using EpicGames.Core;
+using HordeServer.Api;
+using HordeServer.Collections;
+using HordeCommon;
+using HordeServer.Services;
+using HordeServer.Utilities;
+using MongoDB.Bson;
+using MongoDB.Bson.Serialization.Attributes;
+using MongoDB.Bson.Serialization.Serializers;
+using System;
+using System.Collections.Generic;
+using System.Diagnostics.CodeAnalysis;
+using System.Linq;
+using System.Runtime.CompilerServices;
+
+namespace HordeServer.Models
+{
+	using StreamId = StringId<IStream>;
+	using TemplateRefId = StringId<TemplateRef>;
+
+	/// <summary>
+	/// The severity of an issue
+	/// </summary>
+	public enum IssueSeverity
+	{
+		/// <summary>
+		/// Unspecified severity
+		/// </summary>
+		Unspecified,
+
+		/// <summary>
+		/// This error represents a warning
+		/// </summary>
+		Warning,
+
+		/// <summary>
+		/// This issue represents an error
+		/// </summary>
+		Error,
+	}
+
+	/// <summary>
+	/// Stores information about a build health issue
+	/// </summary>
+	public interface IIssue
+	{
+		/// <summary>
+		/// Id for resolved by system
+		/// </summary>
+		static ObjectId ResolvedByTimeoutId { get; } = ObjectId.Parse("609592712b5c90b5bcf88c48");
+
+		/// <summary>
+		/// Id for resolved by unknown user
+		/// </summary>
+		static ObjectId ResolvedByUnknownId { get; } = ObjectId.Parse("609593b83e9b0b6dde620cf3");
+
+		/// <summary>
+		/// The unique object id
+		/// </summary>
+		public int Id { get; }
+
+		/// <summary>
+		/// Summary for this issue
+		/// </summary>
+		public string Summary { get; }
+
+		/// <summary>
+		/// Summary set by users
+		/// </summary>
+		public string? UserSummary { get; }
+
+		/// <summary>
+		/// Fingerprint for this issue
+		/// </summary>
+		public IIssueFingerprint Fingerprint { get; }
+
+		/// <summary>
+		/// Severity of this issue
+		/// </summary>
+		public IssueSeverity Severity { get; }
+
+		/// <summary>
+		/// User id of the owner
+		/// </summary>
+		public ObjectId? OwnerId { get; }
+
+		/// <summary>
+		/// User id of the person that nominated the owner
+		/// </summary>
+		public ObjectId? NominatedById { get; }
+
+		/// <summary>
+		/// Time at which the issue was created
+		/// </summary>
+		public DateTime CreatedAt { get; }
+
+		/// <summary>
+		/// Time that the current owner was nominated
+		/// </summary>
+		public DateTime? NominatedAt { get; }
+
+		/// <summary>
+		/// Time at which the issue was acknowledged
+		/// </summary>
+		public DateTime? AcknowledgedAt { get; }
+
+		/// <summary>
+		/// Time at which the issue was resolved
+		/// </summary>
+		public DateTime? ResolvedAt { get; }
+
+		/// <summary>
+		/// User that resolved the issue
+		/// </summary>
+		public ObjectId? ResolvedById { get; }
+
+		/// <summary>
+		/// Time at which the issue was verified fixed
+		/// </summary>
+		public DateTime? VerifiedAt { get; }
+
+		/// <summary>
+		/// Time at which the issue was last seen.
+		/// </summary>
+		public DateTime LastSeenAt { get; }
+
+		/// <summary>
+		/// Fix changelist for this issue
+		/// </summary>
+		public int? FixChange { get; }
+
+		/// <summary>
+		/// List of streams affected by this issue
+		/// </summary>
+		public IReadOnlyList<IIssueStream> Streams { get; }
+
+		/// <summary>
+		/// Whether all suspects should be notified about this issue
+		/// </summary>
+		public bool NotifySuspects { get; }
+
+		/// <summary>
+		/// Update index for this instance
+		/// </summary>
+		public int UpdateIndex { get; }
+	}
+
+	/// <summary>
+	/// Information about a stream affected by an issue
+	/// </summary>
+	public interface IIssueStream
+	{
+		/// <summary>
+		/// The stream id
+		/// </summary>
+		StreamId StreamId { get; }
+
+		/// <summary>
+		/// Whether this stream contains the fix change
+		/// </summary>
+		bool? ContainsFix { get; }
+	}
+
+	/// <summary>
+	/// Suspect for an issue
+	/// </summary>
+	public interface IIssueSuspect
+	{
+		/// <summary>
+		/// Unique id for this suspect
+		/// </summary>
+		ObjectId Id { get; }
+
+		/// <summary>
+		/// Issue that this suspect belongs to
+		/// </summary>
+		int IssueId { get; }
+
+		/// <summary>
+		/// The user id of the change's author
+		/// </summary>
+		ObjectId AuthorId { get; }
+
+		/// <summary>
+		/// The change suspected of causing this issue (in the origin stream)
+		/// </summary>
+		int Change { get; }
+
+		/// <summary>
+		/// Time at which the author declined the issue
+		/// </summary>
+		DateTime? DeclinedAt { get; }
+	}
+
+	/// <summary>
+	/// Extension methods for issues
+	/// </summary>
+	static class IssueExtensions
+	{
+		/// <summary>
+		/// Creates a lookup from stream id to whether it's fixed
+		/// </summary>
+		/// <param name="Issue"></param>
+		/// <returns></returns>
+		public static Dictionary<StreamId, bool> GetFixStreamIds(this IIssue Issue)
+		{
+			Dictionary<StreamId, bool> FixStreamIds = new Dictionary<StreamId, bool>();
+			foreach (IIssueStream Stream in Issue.Streams)
+			{
+				if (Stream.ContainsFix.HasValue)
+				{
+					FixStreamIds[Stream.StreamId] = Stream.ContainsFix.Value;
+				}
+			}
+			return FixStreamIds;
+		}
+
+		/// <summary>
+		/// Find the first fix-failed step for this issue
+		/// </summary>
+		/// <param name="Issue"></param>
+		/// <param name="Spans"></param>
+		/// <returns></returns>
+		public static IIssueStep? FindFixFailedStep(this IIssue Issue, IEnumerable<IIssueSpan> Spans)
+		{
+			IIssueStep? FixFailedStep = null;
+			if (Issue.FixChange != null && Issue.FixChange.Value >= 0)
+			{
+				foreach (IIssueSpan Span in Spans)
+				{
+					IIssueStream? Stream = Issue.Streams.FirstOrDefault(x => x.StreamId == Span.StreamId);
+					if(Stream != null && (Stream.ContainsFix ?? false) && Span.LastFailure.Change >= Issue.FixChange.Value)
+					{
+						FixFailedStep = Span.LastFailure;
+						break;
+					}
+				}
+			}
+			return FixFailedStep;
+		}
+	}
+}
\ No newline at end of file
diff --git a/Engine/Source/Programs/Horde/HordeServer/Health/IIssueCollection.cs b/Engine/Source/Programs/Horde/HordeServer/Health/IIssueCollection.cs
new file mode 100644
--- /dev/null
+++ b/Engine/Source/Programs/Horde/HordeServer/Health/IIssueCollection.cs
@@ -0,0 +1,544 @@
+// Copyright Epic Games, Inc. All Rights Reserved.
+
+using HordeCommon;
+using HordeServer.Models;
+using HordeServer.Services;
+using HordeServer.Utilities;
+using MessagePack.Formatters;
+using Microsoft.AspNetCore.Connections.Features;
+using Microsoft.Extensions.Logging;
+using MongoDB.Bson;
+using MongoDB.Bson.Serialization;
+using MongoDB.Bson.Serialization.Attributes;
+using MongoDB.Driver;
+using System;
+using System.Collections.Generic;
+using System.Diagnostics.CodeAnalysis;
+using System.IO;
+using System.Linq;
+using System.Linq.Expressions;
+using System.Text.RegularExpressions;
+using System.Threading.Tasks;
+using System.Transactions;
+
+namespace HordeServer.Collections
+{
+	using StreamId = StringId<IStream>;
+	using TemplateRefId = StringId<TemplateRef>;
+
+	/// <summary>
+	/// Fingerprint for an issue
+	/// </summary>
+	public class NewIssueFingerprint : IIssueFingerprint
+	{
+		/// <summary>
+		/// The type of issue
+		/// </summary>
+		public string Type { get; }
+
+		/// <summary>
+		/// List of keys which identify this issue.
+		/// </summary>
+		public CaseInsensitiveStringSet Keys { get; set; }
+
+		/// <summary>
+		/// Set of keys which should trigger a negative match
+		/// </summary>
+		public CaseInsensitiveStringSet? RejectKeys { get; set; }
+
+		/// <summary>
+		/// Constructor
+		/// </summary>
+		/// <param name="Type">The type of issue</param>
+		/// <param name="Keys">Keys which uniquely identify this issue</param>
+		/// <param name="RejectKeys">Keys which should not match with this issue</param>
+		public NewIssueFingerprint(string Type, IEnumerable<string> Keys, IEnumerable<string>? RejectKeys)
+		{
+			this.Type = Type;
+			this.Keys = new CaseInsensitiveStringSet(Keys);
+
+			if (RejectKeys != null && RejectKeys.Any())
+			{
+				this.RejectKeys = new CaseInsensitiveStringSet(RejectKeys);
+			}
+		}
+
+		/// <summary>
+		/// Copy constructor
+		/// </summary>
+		/// <param name="Other">The fingerprint to copy from</param>
+		public NewIssueFingerprint(IIssueFingerprint Other)
+			: this(Other.Type, Other.Keys, Other.RejectKeys)
+		{
+		}
+
+		/// <summary>
+		/// Merges another fingerprint into this one
+		/// </summary>
+		/// <param name="Other">The other fingerprint to merge with</param>
+		public NewIssueFingerprint MergeWith(IIssueFingerprint Other)
+		{
+			NewIssueFingerprint Result = new NewIssueFingerprint(this);
+
+			Result.Keys.UnionWith(Other.Keys);
+			if (Other.RejectKeys != null)
+			{
+				if (Result.RejectKeys == null)
+				{
+					Result.RejectKeys = new CaseInsensitiveStringSet();
+				}
+				Result.RejectKeys.UnionWith(Other.RejectKeys);
+			}
+
+			return Result;
+		}
+
+		/// <inheritdoc/>
+		public override bool Equals(object? Other)
+		{
+			NewIssueFingerprint? OtherFingerprint = Other as NewIssueFingerprint;
+			if(OtherFingerprint == null || !Type.Equals(OtherFingerprint.Type, StringComparison.Ordinal) || !Keys.SetEquals(OtherFingerprint.Keys))
+			{
+				return false;
+			}
+
+			if (RejectKeys == null)
+			{
+				if (OtherFingerprint.RejectKeys != null && OtherFingerprint.RejectKeys.Count > 0)
+				{
+					return false;
+				}
+			}
+			else
+			{
+				if (OtherFingerprint.RejectKeys == null || !RejectKeys.SetEquals(OtherFingerprint.RejectKeys))
+				{
+					return false;
+				}
+			}
+
+			return true;
+		}
+
+		/// <summary>
+		/// Gets a hashcode for this fingerprint
+		/// </summary>
+		/// <returns>The hashcode value</returns>
+		public override int GetHashCode()
+		{
+			int Result = StringComparer.Ordinal.GetHashCode(Type);
+			foreach (string Key in Keys)
+			{
+				Result = HashCode.Combine(Result, Key);
+			}
+			if (RejectKeys != null)
+			{
+				foreach (string RejectKey in RejectKeys)
+				{
+					Result = HashCode.Combine(Result, RejectKey);
+				}
+			}
+			return Result;
+		}
+
+		/// <summary>
+		/// Merges two fingerprints togetherthis issue fingerprint with another
+		/// </summary>
+		/// <param name="A">The first fingerprint</param>
+		/// <param name="B">The second fingerprint to merge with</param>
+		/// <returns>Merged fingerprint</returns>
+		public static NewIssueFingerprint Merge(IIssueFingerprint A, IIssueFingerprint B)
+		{
+			NewIssueFingerprint NewFingerprint = new NewIssueFingerprint(A);
+			NewFingerprint.Keys.UnionWith(B.Keys);
+
+			if (B.RejectKeys != null)
+			{
+				if (NewFingerprint.RejectKeys == null)
+				{
+					NewFingerprint.RejectKeys = new CaseInsensitiveStringSet();
+				}
+				NewFingerprint.RejectKeys.UnionWith(B.RejectKeys);
+			}
+
+			return NewFingerprint;
+		}
+
+		/// <inheritdoc/>
+		public override string ToString()
+		{
+			return $"{Type}: \"{String.Join("\", \"", Keys)}\"";
+		}
+	}
+
+	/// <summary>
+	/// Identifies a particular changelist and job
+	/// </summary>
+	public class NewIssueStepData
+	{
+		/// <summary>
+		/// The changelist number
+		/// </summary>
+		public int Change { get; set; }
+
+		/// <summary>
+		/// Severity of the issue
+		/// </summary>
+		public IssueSeverity Severity { get; set; }
+
+		/// <summary>
+		/// Name of the job
+		/// </summary>
+		public string JobName { get; set; }
+
+		/// <summary>
+		/// The unique job id
+		/// </summary>
+		public ObjectId JobId { get; set; }
+
+		/// <summary>
+		/// Batch id for the step
+		/// </summary>
+		public SubResourceId BatchId { get; set; }
+
+		/// <summary>
+		/// Id of the step
+		/// </summary>
+		public SubResourceId StepId { get; set; }
+
+		/// <summary>
+		/// Time that the step started
+		/// </summary>
+		public DateTime StepTime { get; set; }
+
+		/// <summary>
+		/// The log id for this step
+		/// </summary>
+		public ObjectId? LogId { get; set; }
+
+		/// <summary>
+		/// Whether to notify suspects for this step
+		/// </summary>
+		public bool NotifySuspects { get; set; }
+
+		/// <summary>
+		/// Constructor
+		/// </summary>
+		/// <param name="Change">The changelist number for this job</param>
+		/// <param name="Severity">Severity of the issue in this step</param>
+		/// <param name="JobName">The job name</param>
+		/// <param name="JobId">The unique job id</param>
+		/// <param name="BatchId">The batch id</param>
+		/// <param name="StepId">The step id</param>
+		/// <param name="StepTime">Time that the step started</param>
+		/// <param name="LogId">Unique id of the log file for this step</param>
+		/// <param name="NotifySuspects">Whether to notify suspects for this step</param>
+		public NewIssueStepData(int Change, IssueSeverity Severity, string JobName, ObjectId JobId, SubResourceId BatchId, SubResourceId StepId, DateTime StepTime, ObjectId? LogId, bool NotifySuspects)
+		{
+			this.Change = Change;
+			this.Severity = Severity;
+			this.JobName = JobName;
+			this.JobId = JobId;
+			this.BatchId = BatchId;
+			this.StepId = StepId;
+			this.StepTime = StepTime;
+			this.LogId = LogId;
+			this.NotifySuspects = NotifySuspects;
+		}
+
+		/// <summary>
+		/// Constructor
+		/// </summary>
+		/// <param name="Job">The job being built</param>
+		/// <param name="Batch">Batch of the job for the step</param>
+		/// <param name="Step">The step being built</param>
+		/// <param name="Severity">Severity of the issue in this step</param>
+		/// <param name="NotifySuspects">Whether to notify suspects for this step failing</param>
+		public NewIssueStepData(IJob Job, IJobStepBatch Batch, IJobStep Step, IssueSeverity Severity, bool NotifySuspects)
+			: this(Job.Change, Severity, Job.Name, Job.Id, Batch.Id, Step.Id, Step.StartTimeUtc ?? default, Step.LogId, NotifySuspects)
+		{
+		}
+
+		/// <summary>
+		/// Constructor
+		/// </summary>
+		/// <param name="JobStepRef">The jobstep reference</param>
+		public NewIssueStepData(IJobStepRef JobStepRef)
+			: this(JobStepRef.Change, IssueSeverity.Unspecified, JobStepRef.JobName, JobStepRef.Id.JobId, JobStepRef.Id.BatchId, JobStepRef.Id.StepId, JobStepRef.StartTimeUtc, JobStepRef.LogId, false)
+		{
+		}
+	}
+
+	/// <summary>
+	/// Information about a suspect changelist that may have caused an issue
+	/// </summary>
+	public class NewIssueSuspectData
+	{
+		/// <summary>
+		/// Author of the changelist
+		/// </summary>
+		public ObjectId AuthorId { get; set; }
+
+		/// <summary>
+		/// The submitted changelist
+		/// </summary>
+		public int Change { get; set; }
+
+		/// <summary>
+		/// The time that the user declined causing this issue
+		/// </summary>
+		public DateTime? DeclinedAt { get; set; }
+
+		/// <summary>
+		/// Constructor
+		/// </summary>
+		/// <param name="AuthorId">Author of the change</param>
+		/// <param name="Change">The changelist number</param>
+		/// <param name="DeclinedAt">The time that the user declined this issue</param>
+		public NewIssueSuspectData(ObjectId AuthorId, int Change, DateTime? DeclinedAt)
+		{
+			this.AuthorId = AuthorId;
+			this.Change = Change;
+			this.DeclinedAt = DeclinedAt;
+		}
+	}
+
+	/// <summary>
+	/// Information about a suspect changelist that may have caused an issue
+	/// </summary>
+	public class NewIssueSpanSuspectData
+	{
+		/// <summary>
+		/// The submitted changelist
+		/// </summary>
+		public int Change { get; set; }
+
+		/// <summary>
+		/// Author of the changelist
+		/// </summary>
+		public ObjectId AuthorId { get; set; }
+
+		/// <summary>
+		/// The original changelist number, if merged from another branch. For changes merged between several branches, this is the originally submitted change.
+		/// </summary>
+		public int? OriginatingChange { get; set; }
+
+		/// <summary>
+		/// Constructor
+		/// </summary>
+		/// <param name="Change">The changelist number</param>
+		/// <param name="AuthorId">Author of the change</param>
+		public NewIssueSpanSuspectData(int Change, ObjectId AuthorId)
+		{
+			this.Change = Change;
+			this.AuthorId = AuthorId;
+		}
+	}
+
+	/// <summary>
+	/// Interface for a token used to indicate a position in the transaction queue. Transactions are opportunistic, and will fail if the token is longer valid at the time that it is used.
+	/// </summary>
+	[SuppressMessage("Design", "CA1040:Avoid empty interfaces")]
+	public interface IIssueSequenceToken
+	{
+		/// <summary>
+		/// Reset the token to the last operation
+		/// </summary>
+		/// <returns>Async task</returns>
+		public Task ResetAsync();
+	}
+
+	/// <summary>
+	/// Interface for an issue collection
+	/// </summary>
+	public interface IIssueCollection
+	{
+		/// <summary>
+		/// Gets a token used to sequence operations to the issue collection
+		/// </summary>
+		/// <returns>Sequence token instance</returns>
+		Task<IIssueSequenceToken> GetSequenceTokenAsync();
+
+		/// <summary>
+		/// Assign a unique issue id. For compatibility with the legacy issue system.
+		/// </summary>
+		/// <param name="Token">The current sequence token</param>
+		/// <returns>Sequence token</returns>
+		Task<int?> ReserveUniqueIdAsync(IIssueSequenceToken Token);
+
+		#region Issues
+
+		/// <summary>
+		/// Creates a new issue
+		/// </summary>
+		/// <param name="Token">Token for the transaction</param>
+		/// <param name="Summary">Summary text for the issue</param>
+		/// <param name="Span">The initial span for the issue</param>
+		/// <returns>The new issue instance</returns>
+		Task<IIssue?> AddIssueAsync(IIssueSequenceToken Token, string Summary, IIssueSpan Span);
+
+		/// <summary>
+		/// Retrieves and issue by id
+		/// </summary>
+		/// <param name="IssueId">Unique id of the issue</param>
+		/// <returns>The issue matching the given id, or null</returns>
+		Task<IIssue?> GetIssueAsync(int IssueId);
+
+		/// <summary>
+		/// Gets the suspects for an issue
+		/// </summary>
+		/// <param name="Issue">The issue to retrieve suspects for</param>
+		/// <returns>List of suspects</returns>
+		Task<List<IIssueSuspect>> GetSuspectsAsync(IIssue Issue);
+
+		/// <summary>
+		/// Searches for open issues
+		/// </summary>
+		/// <param name="Ids">Set of issue ids to find</param>
+		/// <param name="UserId">The user to find issues for</param>
+		/// <param name="StreamId">The stream affected by the issue</param>
+		/// <param name="MinChange">Minimum changelist affected by the issue</param>
+		/// <param name="MaxChange">Maximum changelist affected by the issue</param>
+		/// <param name="Resolved">Include issues that are now resolved</param>
+		/// <param name="Index">Index within the results to return</param>
+		/// <param name="Count">Number of results</param>
+		/// <returns>List of streams open in the given stream at the given changelist</returns>
+		Task<List<IIssue>> FindIssuesAsync(IEnumerable<int>? Ids = null, ObjectId? UserId = null, StreamId? StreamId = null, int? MinChange = null, int? MaxChange = null, bool? Resolved = null, int? Index = null, int? Count = null);
+
+		/// <summary>
+		/// Searches for open issues
+		/// </summary>
+		/// <param name="Changes">List of suspect changes</param>
+		/// <returns>List of issues that are affected by the given changes</returns>
+		Task<List<IIssue>> FindIssuesForSuspectsAsync(List<int> Changes);
+
+		/// <summary>
+		/// Updates the state of an issue
+		/// </summary>
+		/// <param name="Issue">The issue to update</param>
+		/// <param name="NewSeverity">New severity for the issue</param>
+		/// <param name="NewSummary">New summary for the issue</param>
+		/// <param name="NewUserSummary">New user summary for the issue</param>
+		/// <param name="NewOwnerId">New owner of the issue</param>
+		/// <param name="NewNominatedById">Person that nominated the new owner</param>
+		/// <param name="NewAcknowledged">Whether the issue has been acknowledged</param>
+		/// <param name="NewDeclinedById">Name of a user that has declined the issue</param>
+		/// <param name="NewFixChange">Fix changelist for the issue. Pass 0 to clear the fix changelist, -1 for systemic issue.</param>
+		/// <param name="NewFixStreamIds">List of stream ids with a bool indicating whether the fix changelist is in this stream</param>
+		/// <param name="NewResolvedById">User that resolved the issue (may be ObjectId.Empty to clear)</param>
+		/// <param name="NewLastSeenAt"></param>
+		/// <param name="NewNotifySuspects">Whether all suspects should be notified about this issue</param>
+		/// <returns>True if the issue was updated</returns>
+		Task<IIssue?> UpdateIssueAsync(IIssue Issue, IssueSeverity? NewSeverity = null, string? NewSummary = null, string? NewUserSummary = null, ObjectId? NewOwnerId = null, ObjectId? NewNominatedById = null, bool? NewAcknowledged = null, ObjectId? NewDeclinedById = null, int? NewFixChange = null, Dictionary<StreamId, bool>? NewFixStreamIds = null, ObjectId? NewResolvedById = null, DateTime? NewLastSeenAt = null, bool? NewNotifySuspects = null);
+
+		/// <summary>
+		/// Updates the list of suspects for an issue, and optionally sets the resolved state
+		/// </summary>
+		/// <param name="Issue">The issue to update</param>
+		/// <returns>True if the issue was updated</returns>
+		Task<IIssue?> UpdateIssueDerivedDataAsync(IIssue Issue);
+
+		#endregion
+
+		#region Spans
+
+		/// <summary>
+		/// Creates a new span from the given failure
+		/// </summary>
+		/// <param name="Token">Token to sequence the addition of new spans</param>
+		/// <param name="Stream">The stream containing the span</param>
+		/// <param name="TemplateRefId">The template being executed</param>
+		/// <param name="NodeName">Name of the step</param>
+		/// <param name="Fingerprint">Fingerprint for this issue</param>
+		/// <param name="LastSuccess">Last time the build succeeded before the span</param>
+		/// <param name="Failure">List of failing steps</param>
+		/// <param name="NextSuccess">First time the build succeeded after the span</param>
+		/// <param name="Suspects">Suspects for this span</param>
+		/// <returns>New span, or null if the sequence token is not valid</returns>
+		Task<IIssueSpan?> AddSpanAsync(IIssueSequenceToken Token, IStream Stream, TemplateRefId TemplateRefId, string NodeName, NewIssueFingerprint Fingerprint, NewIssueStepData? LastSuccess, NewIssueStepData Failure, NewIssueStepData? NextSuccess, List<NewIssueSpanSuspectData> Suspects);
+
+		/// <summary>
+		/// Gets a particular span
+		/// </summary>
+		/// <param name="SpanId">Unique id of the span</param>
+		/// <returns>New span, or null if the sequence token is not valid</returns>
+		Task<IIssueSpan?> GetSpanAsync(ObjectId SpanId);
+
+		/// <summary>
+		/// Updates the given span. Note that data in the span's issue may be derived from this, and the issue should be updated afterwards.
+		/// </summary>
+		/// <param name="Span">Span to update</param>
+		/// <param name="NewSeverity">New severity for the span</param>
+		/// <param name="NewLastSuccess">New last successful step</param>
+		/// <param name="NewFailure">New failed step</param>
+		/// <param name="NewNextSuccess">New next successful step</param>
+		/// <param name="NewSuspects">New suspects for the span</param>
+		/// <param name="NewModified">New modified flag for the span</param>
+		/// <returns>The updated span, or null on failure</returns>
+		Task<IIssueSpan?> UpdateSpanAsync(IIssueSpan Span, IssueSeverity? NewSeverity = null, NewIssueStepData? NewLastSuccess = null, NewIssueStepData? NewFailure = null, NewIssueStepData? NewNextSuccess = null, List<NewIssueSpanSuspectData>? NewSuspects = null, bool? NewModified = null);
+
+		/// <summary>
+		/// Updates the given span
+		/// </summary>
+		/// <param name="Token">The token</param>
+		/// <param name="Span">The span to update</param>
+		/// <param name="Issue">The issue to attach to</param>
+		/// <param name="NewSeverity">New severity for the span</param>
+		/// <param name="NewFingerprint">New fingerprint for the span</param>
+		/// <param name="NewSummary">New summary for the issue</param>
+		/// <returns>The updated span, or null on failure</returns>
+		Task<bool> AttachSpanToIssueAsync(IIssueSequenceToken Token, IIssueSpan Span, IIssue Issue, IssueSeverity? NewSeverity = null, NewIssueFingerprint? NewFingerprint = null, string? NewSummary = null);
+
+		/// <summary>
+		/// Gets all the spans for a particular issue
+		/// </summary>
+		/// <param name="IssueId">Issue id</param>
+		/// <returns>List of spans</returns>
+		Task<List<IIssueSpan>> FindSpansAsync(int IssueId);
+
+		/// <summary>
+		/// Retrieves multiple spans
+		/// </summary>
+		/// <param name="SpanIds">The span ids</param>
+		/// <returns>List of spans</returns>
+		Task<List<IIssueSpan>> FindSpansAsync(IEnumerable<ObjectId> SpanIds);
+
+		/// <summary>
+		/// Finds the open issues for a given stream
+		/// </summary>
+		/// <param name="StreamId">The stream id</param>
+		/// <param name="TemplateId">The template id</param>
+		/// <param name="Name">Name of the node</param>
+		/// <param name="Change">Changelist number to query</param>
+		/// <returns>List of open issues</returns>
+		Task<List<IIssueSpan>> FindOpenSpansAsync(StreamId StreamId, TemplateRefId TemplateId, string Name, int Change);
+
+		/// <summary>
+		/// Finds a span marked as modified
+		/// </summary>
+		/// <returns></returns>
+		Task<IIssueSpan?> FindModifiedSpanAsync();
+
+		#endregion
+
+		#region Steps
+			
+		/// <summary>
+		/// Find steps for the given spans
+		/// </summary>
+		/// <param name="SpanIds">Span ids</param>
+		/// <returns>List of steps</returns>
+		Task<List<IIssueStep>> FindStepsAsync(IEnumerable<ObjectId> SpanIds);
+
+		/// <summary>
+		/// Find steps for the given spans
+		/// </summary>
+		/// <param name="JobId">The job id</param>
+		/// <param name="BatchId">The batch id</param>
+		/// <param name="StepId">The step id</param>
+		/// <returns>List of steps</returns>
+		Task<List<IIssueStep>> FindStepsAsync(ObjectId JobId, SubResourceId? BatchId, SubResourceId? StepId);
+
+		#endregion
+	}
+}
diff --git a/Engine/Source/Programs/Horde/HordeServer/Health/IIssueFingerprint.cs b/Engine/Source/Programs/Horde/HordeServer/Health/IIssueFingerprint.cs
new file mode 100644
--- /dev/null
+++ b/Engine/Source/Programs/Horde/HordeServer/Health/IIssueFingerprint.cs
@@ -0,0 +1,79 @@
+// Copyright Epic Games, Inc. All Rights Reserved.
+
+using HordeServer.Utilities;
+using System;
+using System.Collections.Generic;
+using System.Linq;
+using System.Threading.Tasks;
+
+namespace HordeServer.Models
+{
+	/// <summary>
+	/// Fingerprint for an issue
+	/// </summary>
+	public interface IIssueFingerprint
+	{
+		/// <summary>
+		/// The type of issue
+		/// </summary>
+		public string Type { get; }
+
+		/// <summary>
+		/// List of keys which identify this issue.
+		/// </summary>
+		public CaseInsensitiveStringSet Keys { get; }
+
+		/// <summary>
+		/// Set of keys which should trigger a negative match
+		/// </summary>
+		public CaseInsensitiveStringSet? RejectKeys { get; }
+	}
+
+	/// <summary>
+	/// Extension methods for <see cref="IIssueFingerprint"/>
+	/// </summary>
+	public static class IssueFingerprintExtensions
+	{
+		/// <summary>
+		/// Checks if a fingerprint matches another fingerprint
+		/// </summary>
+		/// <param name="Fingerprint">The first fingerprint to compare</param>
+		/// <param name="Other">The other fingerprint to compare to</param>
+		/// <returns>True is the fingerprints match</returns>
+		public static bool IsMatch(this IIssueFingerprint Fingerprint, IIssueFingerprint Other)
+		{
+			if (!Fingerprint.Type.Equals(Other.Type, StringComparison.Ordinal))
+			{
+				return false;
+			}
+			if (Fingerprint.Keys.Count > 0 && Other.Keys.Count > 0 && !Fingerprint.Keys.Any(x => Other.Keys.Contains(x)))
+			{
+				return false;
+			}
+			if (Fingerprint.RejectKeys != null && Fingerprint.RejectKeys.Any(x => Other.Keys.Contains(x)))
+			{
+				return false;
+			}
+			return true;
+		}
+
+		/// <summary>
+		/// Checks if a fingerprint matches another fingerprint for creating a new span
+		/// </summary>
+		/// <param name="Fingerprint">The first fingerprint to compare</param>
+		/// <param name="Other">The other fingerprint to compare to</param>
+		/// <returns>True is the fingerprints match</returns>
+		public static bool IsMatchForNewSpan(this IIssueFingerprint Fingerprint, IIssueFingerprint Other)
+		{
+			if (!Fingerprint.Type.Equals(Other.Type, StringComparison.Ordinal))
+			{
+				return false;
+			}
+			if (Fingerprint.RejectKeys != null && Fingerprint.RejectKeys.Any(x => Other.Keys.Contains(x)))
+			{
+				return false;
+			}
+			return true;
+		}
+	}
+}
diff --git a/Engine/Source/Programs/Horde/HordeServer/Health/IIssueHandler.cs b/Engine/Source/Programs/Horde/HordeServer/Health/IIssueHandler.cs
new file mode 100644
--- /dev/null
+++ b/Engine/Source/Programs/Horde/HordeServer/Health/IIssueHandler.cs
@@ -0,0 +1,161 @@
+// Copyright Epic Games, Inc. All Rights Reserved.
+
+using EpicGames.Core;
+using HordeServer.Models;
+using HordeServer.Collections;
+using System;
+using System.Collections.Generic;
+using System.Linq;
+using System.Threading.Tasks;
+using System.Diagnostics.CodeAnalysis;
+
+namespace HordeServer.IssueHandlers
+{
+	/// <summary>
+	/// Information about a changelist and a value ranking the likelihood that it caused an issue
+	/// </summary>
+	class SuspectChange
+	{
+		/// <summary>
+		/// Set of extensions to treat as code
+		/// </summary>
+		static readonly HashSet<StringView> CodeExtensions = new HashSet<StringView>(StringViewComparer.OrdinalIgnoreCase)
+		{
+			".c",
+			".cc",
+			".cpp",
+			".m",
+			".mm",
+			".rc",
+			".cs",
+			".csproj",
+			".h",
+			".hpp",
+			".inl",
+			".usf",
+			".ush",
+			".uproject",
+			".uplugin",
+			".sln"
+		};
+
+		/// <summary>
+		/// Set of file extensions to treat as content
+		/// </summary>
+		static readonly HashSet<StringView> ContentExtensions = new HashSet<StringView>(StringViewComparer.OrdinalIgnoreCase)
+		{
+			".uasset",
+			".umap",
+			".ini"
+		};
+
+		/// <summary>
+		/// The change detials
+		/// </summary>
+		public ChangeDetails Details { get; set; }
+
+		/// <summary>
+		/// Rank for the change. A value less than or equal to zero indicates a lack of culpability, positive values indicate
+		/// the possibility of a change being the culprit.
+		/// </summary>
+		public int Rank { get; set; }
+
+		/// <summary>
+		/// Whether the change contains code
+		/// </summary>
+		public bool ContainsCode { get; }
+
+		/// <summary>
+		/// Whether the change modifies content
+		/// </summary>
+		public bool ContainsContent { get; }
+
+		/// <summary>
+		/// Constructor
+		/// </summary>
+		/// <param name="Details">The changelist details</param>
+		public SuspectChange(ChangeDetails Details)
+		{
+			this.Details = Details;
+
+			foreach (ChangeFile File in Details.Files)
+			{
+				int Idx = File.Path.LastIndexOf('.');
+				if (Idx != -1)
+				{
+					StringView Extension = new StringView(File.Path, Idx);
+					if (CodeExtensions.Contains(Extension))
+					{
+						ContainsCode = true;
+					}
+					if (ContentExtensions.Contains(Extension))
+					{
+						ContainsContent = true;
+					}
+					if (ContainsCode && ContainsContent)
+					{
+						break;
+					}
+				}
+			}
+		}
+
+		/// <summary>
+		/// Determines whether this change modifies the given file
+		/// </summary>
+		/// <param name="FileToCheck">The file to look for</param>
+		/// <returns>True if the change modifies the given file</returns>
+		public bool ModifiesFile(string FileToCheck)
+		{
+			foreach (ChangeFile File in Details.Files)
+			{
+				if (File.Path.EndsWith(FileToCheck, StringComparison.OrdinalIgnoreCase) && (File.Length == FileToCheck.Length || File.Path[File.Path.Length - FileToCheck.Length - 1] == '/'))
+				{
+					return true;
+				}
+			}
+			return false;
+		}
+	}
+
+	/// <summary>
+	/// Interface for issue matchers
+	/// </summary>
+	interface IIssueHandler
+	{
+		/// <summary>
+		/// Identifier for the type of issue
+		/// </summary>
+		string Type { get; }
+
+		/// <summary>
+		/// Priority of this matcher
+		/// </summary>
+		int Priority { get; }
+
+		/// <summary>
+		/// Match the given event and produce a fingerprint
+		/// </summary>
+		/// <param name="Job">The job that spawned the event</param>
+		/// <param name="Node">Node that was executed</param>
+		/// <param name="EventData">The event data</param>
+		/// <param name="Fingerprint">Receives the fingerprint on success</param>
+		/// <returns>True if the match is successful</returns>
+		bool TryGetFingerprint(IJob Job, INode Node, ILogEventData EventData, [NotNullWhen(true)] out NewIssueFingerprint? Fingerprint);
+
+		/// <summary>
+		/// Rank all the suspect changes for a given fingerprint
+		/// </summary>
+		/// <param name="Fingerprint">The issue fingerprint</param>
+		/// <param name="Suspects">Potential suspects</param>
+		void RankSuspects(IIssueFingerprint Fingerprint, List<SuspectChange> Suspects);
+
+		/// <summary>
+		/// Gets the summary text for an issue
+		/// </summary>
+		/// <param name="Fingerprint">The fingerprint</param>
+		/// <param name="Severity">Severity of the issue</param>
+		/// <returns>The summary text</returns>
+		string GetSummary(IIssueFingerprint Fingerprint, IssueSeverity Severity);
+	}
+}
diff --git a/Engine/Source/Programs/Horde/HordeServer/Health/IIssueService.cs b/Engine/Source/Programs/Horde/HordeServer/Health/IIssueService.cs
new file mode 100644
--- /dev/null
+++ b/Engine/Source/Programs/Horde/HordeServer/Health/IIssueService.cs
@@ -0,0 +1,151 @@
+// Copyright Epic Games, Inc. All Rights Reserved.
+
+using HordeServer.Models;
+using HordeServer.Utilities;
+using Microsoft.Extensions.Hosting;
+using MongoDB.Bson;
+using System;
+using System.Collections.Generic;
+using System.Linq;
+using System.Threading.Tasks;
+
+namespace HordeServer.Services
+{
+	using StreamId = StringId<IStream>;
+
+	/// <summary>
+	/// Service which manages build health issues
+	/// </summary>
+	public interface IIssueService : IHostedService
+	{
+		/// <summary>
+		/// Delegate for issue creation events
+		/// </summary>
+		public delegate void IssueUpdatedEvent(IIssue Issue);
+
+		/// <summary>
+		/// Event triggered when a label state changes
+		/// </summary>
+		event IssueUpdatedEvent? OnIssueUpdated;
+
+		/// <summary>
+		/// Gets a list of cached open issues
+		/// </summary>
+		/// <returns>List of cached issues</returns>
+		IReadOnlyList<IIssueDetails> CachedOpenIssues { get; }
+
+		/// <summary>
+		/// Find a cached issue, or update it
+		/// </summary>
+		/// <param name="IssueId">The issue id</param>
+		/// <returns>The cached issue data</returns>
+		Task<IIssueDetails?> GetCachedIssueDetailsAsync(int IssueId);
+
+		/// <summary>
+		/// Gets the details for a particular issue
+		/// </summary>
+		/// <param name="Issue">The issue id</param>
+		/// <returns>The issue details</returns>
+		Task<IIssueDetails> GetIssueDetailsAsync(IIssue Issue);
+
+		/// <summary>
+		/// Determines whether to show alerts for a given issue
+		/// </summary>
+		/// <param name="Issue">The issue to query</param>
+		/// <param name="Spans">Spans for the issue</param>
+		/// <returns>True if we should show alerts for the given issue</returns>
+		bool ShowDesktopAlertsForIssue(IIssue Issue, IReadOnlyList<IIssueSpan> Spans);
+
+		/// <summary>
+		/// Gets an issue with the given id
+		/// </summary>
+		/// <param name="IssueId"></param>
+		/// <returns></returns>
+		Task<IIssue?> GetIssueAsync(int IssueId);
+
+		/// <summary>
+		/// Gets the spans for the given issue
+		/// </summary>
+		/// <param name="Issue">The issue id</param>
+		/// <returns>List of spans for the issue</returns>
+		Task<List<IIssueSpan>> GetIssueSpansAsync(IIssue Issue);
+
+		/// <summary>
+		/// Gets the spans for the given issue
+		/// </summary>
+		/// <param name="Span">The issue span</param>
+		/// <returns>List of steps for the issue</returns>
+		Task<List<IIssueStep>> GetIssueStepsAsync(IIssueSpan Span);
+
+		/// <summary>
+		/// Gets the suspects for an issue
+		/// </summary>
+		/// <param name="Issue"></param>
+		/// <returns></returns>
+		Task<List<IIssueSuspect>> GetIssueSuspectsAsync(IIssue Issue);
+
+		/// <summary>
+		/// Searches for open issues
+		/// </summary>
+		/// <param name="Ids">List of issue ids to include</param>
+		/// <param name="UserId">User to filter issues returned by</param>
+		/// <param name="StreamId">The stream to return issues for</param>
+		/// <param name="MinChange">The minimum changelist number to include</param>
+		/// <param name="MaxChange">The maximum changelist number to include</param>
+		/// <param name="Resolved">Whether to include results that are resolved</param>
+		/// <param name="Index">Index within the results to return</param>
+		/// <param name="Count">Number of results</param>
+		/// <returns>List of streams open in the given stream at the given changelist</returns>
+		Task<List<IIssue>> FindIssuesAsync(IEnumerable<int>? Ids = null, ObjectId? UserId = null, StreamId? StreamId = null, int? MinChange = null, int? MaxChange = null, bool? Resolved = null, int? Index = null, int? Count = null);
+
+		/// <summary>
+		/// Searches for open issues affecting a job
+		/// </summary>
+		/// <param name="Ids"></param>
+		/// <param name="Job"></param>
+		/// <param name="Graph"></param>
+		/// <param name="StepId"></param>
+		/// <param name="BatchId"></param>
+		/// <param name="LabelIdx"></param>
+		/// <param name="UserId"></param>
+		/// <param name="Resolved">Whether to include results that are resolved</param>
+		/// <param name="Index">Index within the results to return</param>
+		/// <param name="Count">Number of results</param>
+		/// <returns></returns>
+		Task<List<IIssue>> FindIssuesForJobAsync(int[]? Ids, IJob Job, IGraph Graph, SubResourceId? StepId = null, SubResourceId? BatchId = null, int? LabelIdx = null, ObjectId? UserId = null, bool? Resolved = null, int? Index = null, int? Count = null);
+
+		/// <summary>
+		/// Find the events for a particular issue
+		/// </summary>
+		/// <param name="IssueId">The issue id</param>
+		/// <param name="LogIds">Log ids to include</param>
+		/// <param name="Index">Index of the first event to return</param>
+		/// <param name="Count">Number of issues to return</param>
+		/// <returns>List of log events</returns>
+		Task<List<ILogEvent>> FindEventsForIssueAsync(int IssueId, ObjectId[]? LogIds = null, int Index = 0, int Count = 10);
+
+		/// <summary>
+		/// Updates the state of an issue
+		/// </summary>
+		/// <param name="Id">The current issue id</param>
+		/// <param name="Summary">New summary for the issue</param>
+		/// <param name="OwnerId">New owner of the issue</param>
+		/// <param name="NominatedById">Person that nominated the new owner</param>
+		/// <param name="Acknowledged">Whether the issue has been acknowledged</param>
+		/// <param name="DeclinedById">Name of the user that has declined this issue</param>
+		/// <param name="FixChange">Fix changelist for the issue</param>
+		/// <param name="ResolvedById">Whether the issue has been resolved</param>
+		/// <returns>True if the issue was updated</returns>
+		Task<bool> UpdateIssueAsync(int Id, string? Summary = null, ObjectId? OwnerId = null, ObjectId? NominatedById = null, bool? Acknowledged = null, ObjectId? DeclinedById = null, int? FixChange = null, ObjectId? ResolvedById = null);
+
+		/// <summary>
+		/// Marks a step as complete
+		/// </summary>
+		/// <param name="Job">The job to update</param>
+		/// <param name="Graph">Graph for the job</param>
+		/// <param name="BatchId">Unique id of the batch</param>
+		/// <param name="StepId">Unique id of the step</param>
+		/// <returns>Async task</returns>
+		Task UpdateCompleteStep(IJob Job, IGraph Graph, SubResourceId BatchId, SubResourceId StepId);
+	}
+}
diff --git a/Engine/Source/Programs/Horde/HordeServer/Health/IIssueStep.cs b/Engine/Source/Programs/Horde/HordeServer/Health/IIssueStep.cs
new file mode 100644
--- /dev/null
+++ b/Engine/Source/Programs/Horde/HordeServer/Health/IIssueStep.cs
@@ -0,0 +1,68 @@
+// Copyright Epic Games, Inc. All Rights Reserved.
+
+using MongoDB.Bson;
+using System;
+using System.Collections.Generic;
+using System.Linq;
+using System.Threading.Tasks;
+
+using SubResourceId = HordeServer.Models.SubResourceId;
+
+namespace HordeServer.Models
+{
+	/// <summary>
+	/// Identifies a particular changelist and job that contributes to a span
+	/// </summary>
+	public interface IIssueStep
+	{
+		/// <summary>
+		/// The span that this step belongs to
+		/// </summary>
+		public ObjectId SpanId { get; }
+
+		/// <summary>
+		/// The changelist number
+		/// </summary>
+		public int Change { get; }
+
+		/// <summary>
+		/// Severity of the issue in this step
+		/// </summary>
+		public IssueSeverity Severity { get; }
+
+		/// <summary>
+		/// Name of the job
+		/// </summary>
+		public string JobName { get; }
+
+		/// <summary>
+		/// The unique job id
+		/// </summary>
+		public ObjectId JobId { get; }
+
+		/// <summary>
+		/// Unique id of the batch within the job
+		/// </summary>
+		public SubResourceId BatchId { get; }
+
+		/// <summary>
+		/// Unique id of the step within the job
+		/// </summary>
+		public SubResourceId StepId { get; }
+
+		/// <summary>
+		/// Time that the step started
+		/// </summary>
+		public DateTime StepTime { get; }
+
+		/// <summary>
+		/// Whether to notify suspects for this issue
+		/// </summary>
+		public bool NotifySuspects { get; }
+
+		/// <summary>
+		/// The log id for this step
+		/// </summary>
+		public ObjectId? LogId { get; }
+	}
+}
diff --git a/Engine/Source/Programs/Horde/HordeServer/Health/Impl/Handlers/CompileIssueHandler.cs b/Engine/Source/Programs/Horde/HordeServer/Health/Impl/Handlers/CompileIssueHandler.cs
new file mode 100644
--- /dev/null
+++ b/Engine/Source/Programs/Horde/HordeServer/Health/Impl/Handlers/CompileIssueHandler.cs
@@ -0,0 +1,132 @@
+// Copyright Epic Games, Inc. All Rights Reserved.
+
+using HordeCommon;
+using HordeServer.Collections;
+using HordeCommon.Rpc;
+using HordeServer.Services;
+using HordeServer.Utilities;
+using Microsoft.Extensions.Logging;
+using MongoDB.Bson;
+using MongoDB.Bson.Serialization;
+using MongoDB.Driver;
+using System;
+using System.Collections.Generic;
+using System.Diagnostics.CodeAnalysis;
+using System.IO;
+using System.Linq;
+using System.Security.Policy;
+using System.Text;
+using System.Threading.Tasks;
+using System.Text.Json;
+using HordeServer.Models;
+using EpicGames.Core;
+
+namespace HordeServer.IssueHandlers.Impl
+{
+	/// <summary>
+	/// Instance of a particular compile error
+	/// </summary>
+	class CompileIssueHandler : IIssueHandler
+	{
+		/// <inheritdoc/>
+		public string Type => "Compile";
+
+		/// <inheritdoc/>
+		public int Priority => 10;
+
+		/// <summary>
+		/// Determines if the given event id matches
+		/// </summary>
+		/// <param name="EventId">The event id to compare</param>
+		/// <returns>True if the given event id matches</returns>
+		public static bool IsMatchingEventId(EventId? EventId)
+		{
+			return EventId == KnownLogEvents.Compiler || EventId == KnownLogEvents.AutomationTool_SourceFileLine || EventId == KnownLogEvents.MSBuild;
+		}
+
+		/// <summary>
+		/// Extracts a list of source files from an event
+		/// </summary>
+		/// <param name="Event">The event data</param>
+		/// <param name="SourceFiles">List of source files</param>
+		static void GetSourceFiles(ILogEventData Event, List<string> SourceFiles)
+		{
+			foreach (ILogEventLine Line in Event.Lines)
+			{
+				JsonElement Properties;
+				if (Line.Data.TryGetProperty("properties", out Properties) && Properties.ValueKind == JsonValueKind.Object)
+				{
+					foreach (JsonProperty Property in Properties.EnumerateObject())
+					{
+						if (Property.NameEquals("file") && Property.Value.ValueKind == JsonValueKind.String)
+						{
+							AddSourceFile(SourceFiles, Property.Value.GetString()!);
+						}
+						if (Property.Value.HasStringProperty("type", "SourceFile") && Property.Value.TryGetStringProperty("relativePath", out string? Value))
+						{
+							AddSourceFile(SourceFiles, Value);
+						}
+					}
+				}
+			}
+		}
+
+		/// <summary>
+		/// Add a new source file to a list of unique source files
+		/// </summary>
+		/// <param name="SourceFiles">List of source files</param>
+		/// <param name="RelativePath">File to add</param>
+		static void AddSourceFile(List<string> SourceFiles, string RelativePath)
+		{
+			int EndIdx = RelativePath.LastIndexOfAny(new char[] { '/', '\\' }) + 1;
+
+			string FileName = RelativePath.Substring(EndIdx);
+			if (!SourceFiles.Any(x => x.Equals(FileName, StringComparison.OrdinalIgnoreCase)))
+			{
+				SourceFiles.Add(FileName);
+			}
+		}
+
+		/// <inheritdoc/>
+		public bool TryGetFingerprint(IJob Job, INode Node, ILogEventData EventData, [NotNullWhen(true)] out NewIssueFingerprint? Fingerprint)
+		{
+			if (!IsMatchingEventId(EventData.EventId))
+			{
+				Fingerprint = null;
+				return false;
+			}
+
+			List<string> NewFileNames = new List<string>();
+			GetSourceFiles(EventData, NewFileNames);
+			Fingerprint = new NewIssueFingerprint(Type, NewFileNames, null);
+			return true;
+		}
+
+		/// <inheritdoc/>
+		public string GetSummary(IIssueFingerprint Fingerprint, IssueSeverity Severity)
+		{
+			string Type = (Severity == IssueSeverity.Warning) ? "Compile warnings" : "Compile errors";
+			string List = StringUtils.FormatList(Fingerprint.Keys.ToArray(), 2);
+			return $"{Type} in {List}";
+		}
+
+		/// <inheritdoc/>
+		public void RankSuspects(IIssueFingerprint Fingerprint, List<SuspectChange> Suspects)
+		{
+			foreach (SuspectChange Change in Suspects)
+			{
+				if (Change.ContainsCode)
+				{
+					if (Fingerprint.Keys.Any(x => Change.ModifiesFile(x)))
+					{
+						Change.Rank += 20;
+					}
+					else
+					{
+						Change.Rank += 10;
+					}
+				}
+			}
+		}
+	}
+}
diff --git a/Engine/Source/Programs/Horde/HordeServer/Health/Impl/Handlers/ContentIssueHandler.cs b/Engine/Source/Programs/Horde/HordeServer/Health/Impl/Handlers/ContentIssueHandler.cs
new file mode 100644
--- /dev/null
+++ b/Engine/Source/Programs/Horde/HordeServer/Health/Impl/Handlers/ContentIssueHandler.cs
@@ -0,0 +1,105 @@
+// Copyright Epic Games, Inc. All Rights Reserved.
+
+using HordeCommon;
+using HordeServer.Collections;
+using HordeServer.Models;
+using HordeServer.Utilities;
+using Microsoft.Extensions.Logging;
+using MongoDB.Bson;
+using MongoDB.Bson.Serialization;
+using MongoDB.Driver;
+using System;
+using System.Collections.Generic;
+using System.Diagnostics.CodeAnalysis;
+using System.IO;
+using System.Linq;
+using System.Security.Policy;
+using System.Text;
+using System.Threading.Tasks;
+using EpicGames.Core;
+
+namespace HordeServer.IssueHandlers.Impl
+{
+	/// <summary>
+	/// Instance of a particular compile error
+	/// </summary>
+	class ContentIssueHandler : IIssueHandler
+	{
+		/// <inheritdoc/>
+		public string Type => "Content";
+
+		/// <inheritdoc/>
+		public int Priority => 10;
+
+		/// <summary>
+		/// Determines if the given event id matches
+		/// </summary>
+		/// <param name="EventId">The event id to compare</param>
+		/// <returns>True if the given event id matches</returns>
+		public static bool IsMatchingEventId(EventId? EventId)
+		{
+			return EventId == KnownLogEvents.Engine_AssetLog;
+		}
+
+		/// <summary>
+		/// Adds all the assets from the given log event
+		/// </summary>
+		/// <param name="Event">The log event to parse</param>
+		/// <param name="AssetNames">Receives the referenced asset names</param>
+		public static void GetAssetNames(ILogEventData Event, HashSet<string> AssetNames)
+		{
+			foreach (ILogEventLine Line in Event.Lines)
+			{
+				string? RelativePath;
+				if (Line.Data.TryGetNestedProperty("properties.asset.relativePath", out RelativePath))
+				{
+					int EndIdx = RelativePath.LastIndexOfAny(new char[] { '/', '\\' }) + 1;
+					string FileName = RelativePath.Substring(EndIdx);
+					AssetNames.Add(FileName);
+				}
+			}
+		}
+
+		/// <inheritdoc/>
+		public bool TryGetFingerprint(IJob Job, INode Node, ILogEventData EventData, [NotNullWhen(true)] out NewIssueFingerprint? Fingerprint)
+		{
+			if (!IsMatchingEventId(EventData.EventId))
+			{
+				Fingerprint = null;
+				return false;
+			}
+
+			HashSet<string> NewAssetNames = new HashSet<string>();
+			GetAssetNames(EventData, NewAssetNames);
+			Fingerprint = new NewIssueFingerprint(Type, NewAssetNames, null);
+			return true;
+		}
+
+		/// <inheritdoc/>
+		public string GetSummary(IIssueFingerprint Fingerprint, IssueSeverity Severity)
+		{
+			string Type = (Severity == IssueSeverity.Warning) ? "Warnings" : "Errors";
+			string List = StringUtils.FormatList(Fingerprint.Keys.ToArray(), 2);
+			return $"{Type} in {List}";
+		}
+
+		/// <inheritdoc/>
+		public void RankSuspects(IIssueFingerprint Fingerprint, List<SuspectChange> Suspects)
+		{
+			foreach (SuspectChange Suspect in Suspects)
+			{
+				if (Suspect.ContainsContent)
+				{
+					if (Suspect.Details.Files.Any(x => Fingerprint.Keys.Any(y => x.Path.Contains(y, StringComparison.OrdinalIgnoreCase))))
+					{
+						Suspect.Rank += 20;
+					}
+					else
+					{
+						Suspect.Rank += 10;
+					}
+				}
+			}
+		}
+	}
+}
diff --git a/Engine/Source/Programs/Horde/HordeServer/Health/Impl/Handlers/CopyrightIssueHandler.cs b/Engine/Source/Programs/Horde/HordeServer/Health/Impl/Handlers/CopyrightIssueHandler.cs
new file mode 100644
--- /dev/null
+++ b/Engine/Source/Programs/Horde/HordeServer/Health/Impl/Handlers/CopyrightIssueHandler.cs
@@ -0,0 +1,109 @@
+// Copyright Epic Games, Inc. All Rights Reserved.
+
+using HordeCommon;
+using HordeServer.Collections;
+using HordeCommon.Rpc;
+using HordeServer.Models;
+using HordeServer.Utilities;
+using Microsoft.Extensions.Logging;
+using MongoDB.Bson;
+using MongoDB.Bson.Serialization;
+using MongoDB.Driver;
+using System;
+using System.Collections.Generic;
+using System.Diagnostics.CodeAnalysis;
+using System.IO;
+using System.Linq;
+using System.Security.Policy;
+using System.Text;
+using System.Threading.Tasks;
+using EpicGames.Core;
+
+namespace HordeServer.IssueHandlers.Impl
+{
+	/// <summary>
+	/// Instance of a particular compile error
+	/// </summary>
+	class CopyrightIssueHandler : IIssueHandler
+	{
+		/// <inheritdoc/>
+		public string Type => "Copyright";
+
+		/// <inheritdoc/>
+		public int Priority => 10;
+
+		/// <summary>
+		/// Filenames containing errors or warnings
+		/// </summary>
+		public HashSet<string> FileNames { get; set; } = new HashSet<string>();
+
+		/// <summary>
+		/// Determines if the given event id matches
+		/// </summary>
+		/// <param name="EventId">The event id to compare</param>
+		/// <returns>True if the given event id matches</returns>
+		public static bool IsMatchingEventId(EventId? EventId)
+		{
+			return EventId == KnownLogEvents.AutomationTool_MissingCopyright;
+		}
+
+		/// <summary>
+		/// Extracts a list of source files from an event
+		/// </summary>
+		/// <param name="Event">The event data</param>
+		/// <param name="SourceFiles">List of source files</param>
+		public static void GetSourceFiles(ILogEventData Event, HashSet<string> SourceFiles)
+		{
+			foreach (ILogEventLine Line in Event.Lines)
+			{
+				string? RelativePath;
+				if (Line.Data.TryGetNestedProperty("properties.file.relativePath", out RelativePath) || Line.Data.TryGetNestedProperty("properties.file", out RelativePath))
+				{
+					int EndIdx = RelativePath.LastIndexOfAny(new char[] { '/', '\\' }) + 1;
+					string FileName = RelativePath.Substring(EndIdx);
+					SourceFiles.Add(FileName);
+				}
+			}
+		}
+
+		/// <inheritdoc/>
+		public bool TryGetFingerprint(IJob Job, INode Node, ILogEventData EventData, [NotNullWhen(true)] out NewIssueFingerprint? Fingerprint)
+		{
+			if (!IsMatchingEventId(EventData.EventId))
+			{
+				Fingerprint = null;
+				return false;
+			}
+
+			HashSet<string> NewFileNames = new HashSet<string>(StringComparer.OrdinalIgnoreCase);
+			GetSourceFiles(EventData, NewFileNames);
+			Fingerprint = new NewIssueFingerprint(Type, NewFileNames, null);
+			return true;
+		}
+
+		/// <inheritdoc/>
+		public void RankSuspects(IIssueFingerprint Fingerprint, List<SuspectChange> Suspects)
+		{
+			foreach (SuspectChange Suspect in Suspects)
+			{
+				if (Suspect.ContainsCode)
+				{
+					if (FileNames.Any(x => Suspect.ModifiesFile(x)))
+					{
+						Suspect.Rank += 20;
+					}
+					else
+					{
+						Suspect.Rank += 10;
+					}
+				}
+			}
+		}
+
+		/// <inheritdoc/>
+		public string GetSummary(IIssueFingerprint Fingerprint, IssueSeverity Severity)
+		{
+			return $"Missing copyright notice in {StringUtils.FormatList(Fingerprint.Keys.ToArray(), 2)}";
+		}
+	}
+}
diff --git a/Engine/Source/Programs/Horde/HordeServer/Health/Impl/Handlers/DefaultIssueHandler.cs b/Engine/Source/Programs/Horde/HordeServer/Health/Impl/Handlers/DefaultIssueHandler.cs
new file mode 100644
--- /dev/null
+++ b/Engine/Source/Programs/Horde/HordeServer/Health/Impl/Handlers/DefaultIssueHandler.cs
@@ -0,0 +1,55 @@
+// Copyright Epic Games, Inc. All Rights Reserved.
+
+using HordeServer.Collections;
+using HordeServer.Models;
+using System;
+using System.Collections.Generic;
+using System.Diagnostics.CodeAnalysis;
+using System.Linq;
+using System.Threading.Tasks;
+
+namespace HordeServer.IssueHandlers.Impl
+{
+	/// <summary>
+	/// Instance of a particular compile error
+	/// </summary>
+	class DefaultIssueHandler : IIssueHandler
+	{
+		/// <summary>
+		/// Name of the handler
+		/// </summary>
+		public const string Type = "Default";
+
+		/// <inheritdoc/>
+		string IIssueHandler.Type => Type;
+
+		/// <inheritdoc/>
+		public int Priority => 0;
+
+		/// <inheritdoc/>
+		public bool TryGetFingerprint(IJob Job, INode Node, ILogEventData EventData, [NotNullWhen(true)] out NewIssueFingerprint? Fingerprint)
+		{
+			Fingerprint = new NewIssueFingerprint(Type, new[] { Node.Name }, null);
+			return true;
+		}
+
+		/// <inheritdoc/>
+		public void RankSuspects(IIssueFingerprint Fingerprint, List<SuspectChange> Suspects)
+		{
+		}
+
+		/// <inheritdoc/>
+		public string GetSummary(IIssueFingerprint Fingerprint, IssueSeverity Severity)
+		{
+			string NodeName = Fingerprint.Keys.FirstOrDefault() ?? "(unknown)";
+			if(Severity == IssueSeverity.Warning)
+			{
+				return $"Warnings in {NodeName}";
+			}
+			else
+			{
+				return $"Errors in {NodeName}";
+			}
+		}
+	}
+}
diff --git a/Engine/Source/Programs/Horde/HordeServer/Health/Impl/Handlers/GauntletIssueHandler.cs b/Engine/Source/Programs/Horde/HordeServer/Health/Impl/Handlers/GauntletIssueHandler.cs
new file mode 100644
--- /dev/null
+++ b/Engine/Source/Programs/Horde/HordeServer/Health/Impl/Handlers/GauntletIssueHandler.cs
@@ -0,0 +1,177 @@
+// Copyright Epic Games, Inc. All Rights Reserved.
+
+using HordeCommon;
+using HordeServer.Models;
+using HordeServer.Utilities;
+using Microsoft.AspNetCore.Authentication;
+using Microsoft.AspNetCore.Connections.Features;
+using Microsoft.Extensions.Logging;
+using Serilog.Formatting.Json;
+using System;
+using System.Collections.Generic;
+using System.Linq;
+using System.Security.Cryptography.X509Certificates;
+using System.Text.Json;
+using System.Threading.Tasks;
+using System.Diagnostics.CodeAnalysis;
+using HordeServer.Collections;
+using EpicGames.Core;
+
+namespace HordeServer.IssueHandlers.Impl
+{
+	/// <summary>
+	/// Instance of a particular Gauntlet error
+	/// </summary>
+	class GauntletIssueHandler : IIssueHandler
+	{
+		/// <summary>
+		/// Prefix for unit test keys
+		/// </summary>
+		const string UnitTestPrefix = "UnitTest:";
+
+		/// <summary>
+		/// Prefix for screenshot test keys
+		/// </summary>
+		const string ScreenshotTestPrefix = "Screenshot:";
+
+		/// <inheritdoc/>
+		public string Type => "Gauntlet";
+
+		/// <inheritdoc/>
+		public int Priority => 10;
+
+		/// <summary>
+		/// Determines if the given event id matches
+		/// </summary>
+		/// <param name="EventId">The event id to compare</param>
+		/// <returns>True if the given event id matches</returns>
+		public static bool IsMatchingEventId(EventId? EventId)
+		{
+			return EventId == KnownLogEvents.Gauntlet_UnitTest || EventId == KnownLogEvents.Gauntlet_ScreenshotTest;
+		}
+
+		/// <summary>
+		/// Parses symbol names from a log event
+		/// </summary>
+		/// <param name="EventData">The log event data</param>
+		/// <param name="UnitTestNames">Receives a set of the unit test names</param>
+		private static void GetUnitTestNames(ILogEventData EventData, HashSet<string> UnitTestNames)
+		{
+			string? Group = null;
+			string? Name = null;
+
+			foreach (ILogEventLine Line in EventData.Lines)
+			{
+				string? Value;
+				if (Line.Data.TryGetNestedProperty("properties.group", out Value))
+				{
+					Group = Value;
+				}
+				if (Line.Data.TryGetNestedProperty("properties.name", out Value))
+				{
+					Name = Value;
+				}
+			}
+
+			if(Group != null && Name != null)
+			{
+				UnitTestNames.Add($"{UnitTestPrefix}:{Group}/{Name}");
+			}
+		}
+
+		/// <summary>
+		/// Parses screenshot test names from a log event
+		/// </summary>
+		/// <param name="EventData">The event data</param>
+		/// <param name="ScreenshotTestNames">Receives the parsed screenshot test names</param>
+		private static void GetScreenshotTestNames(ILogEventData EventData, HashSet<string> ScreenshotTestNames)
+		{
+			foreach (ILogEventLine Line in EventData.Lines)
+			{
+				if (Line.Data.TryGetProperty("properties", JsonValueKind.Object, out JsonElement Properties))
+				{
+					foreach (JsonProperty Property in Properties.EnumerateObject())
+					{
+						JsonElement Value = Property.Value;
+						if (Value.ValueKind == JsonValueKind.Object && Value.HasStringProperty("type", "Screenshot") && Value.TryGetStringProperty("name", out string? Name))
+						{
+							ScreenshotTestNames.Add($"{ScreenshotTestPrefix}:{Name}");
+						}
+					}
+				}
+			}
+		}
+
+		/// <inheritdoc/>
+		public bool TryGetFingerprint(IJob Job, INode Node, ILogEventData EventData, [NotNullWhen(true)] out NewIssueFingerprint? Fingerprint)
+		{
+			if(!IsMatchingEventId(EventData.EventId))
+			{
+				Fingerprint = null;
+				return false;
+			}
+
+			HashSet<string> Keys = new HashSet<string>();
+			GetUnitTestNames(EventData, Keys);
+			GetScreenshotTestNames(EventData, Keys);
+
+			if (Keys.Count == 0)
+			{
+				Fingerprint = null;
+				return false;
+			}
+
+			Fingerprint = new NewIssueFingerprint(Type, Keys, null);
+			return true;
+		}
+
+		/// <inheritdoc/>
+		public void RankSuspects(IIssueFingerprint Fingerprint, List<SuspectChange> Changes)
+		{
+			foreach (SuspectChange Change in Changes)
+			{
+				if (Change.ContainsCode)
+				{
+					Change.Rank += 10;
+				}
+			}
+		}
+
+		/// <inheritdoc/>
+		public string GetSummary(IIssueFingerprint Fingerprint, IssueSeverity Severity)
+		{
+			List<string> UnitTestNames = Fingerprint.Keys.Where(x => x.StartsWith(UnitTestPrefix, StringComparison.Ordinal)).Select(x => x.Substring(UnitTestPrefix.Length + 1)).ToList();
+			if (UnitTestNames.Count > 0)
+			{
+				HashSet<string> GroupNames = new HashSet<string>(StringComparer.OrdinalIgnoreCase);
+				HashSet<string> TestNames = new HashSet<string>(StringComparer.OrdinalIgnoreCase);
+				foreach (string UnitTestName in UnitTestNames)
+				{
+					int Idx = UnitTestName.IndexOf('/', StringComparison.OrdinalIgnoreCase);
+					if (Idx != -1)
+					{
+						GroupNames.Add(UnitTestName.Substring(0, Idx));
+					}
+					TestNames.Add(UnitTestName.Substring(Idx + 1));
+				}
+
+				if (GroupNames.Count == 1)
+				{
+					return $"{GroupNames.First()} test failures: {StringUtils.FormatList(TestNames.ToArray(), 3)}";
+				}
+				else
+				{
+					return $"{StringUtils.FormatList(GroupNames.OrderBy(x => x).ToArray(), 100)} test failures";
+				}
+			}
+
+			List<string> ScreenshotTestNames = Fingerprint.Keys.Where(x => x.StartsWith(ScreenshotTestPrefix, StringComparison.Ordinal)).Select(x => x.Substring(ScreenshotTestPrefix.Length + 1)).ToList();
+			if (ScreenshotTestNames.Count > 0)
+			{
+				return $"Screenshot test failures: {StringUtils.FormatList(ScreenshotTestNames.ToArray(), 3)}";
+			}
+
+			return "Test failures";
+		}
+	}
+}
diff --git a/Engine/Source/Programs/Horde/HordeServer/Health/Impl/Handlers/SymbolIssueHandler.cs b/Engine/Source/Programs/Horde/HordeServer/Health/Impl/Handlers/SymbolIssueHandler.cs
new file mode 100644
--- /dev/null
+++ b/Engine/Source/Programs/Horde/HordeServer/Health/Impl/Handlers/SymbolIssueHandler.cs
@@ -0,0 +1,105 @@
+// Copyright Epic Games, Inc. All Rights Reserved.
+
+using HordeCommon;
+using HordeServer.Collections;
+using HordeServer.Models;
+using HordeServer.Utilities;
+using Microsoft.Extensions.Logging;
+using MongoDB.Bson;
+using MongoDB.Bson.Serialization;
+using System;
+using System.Collections.Generic;
+using System.Diagnostics.CodeAnalysis;
+using System.Linq;
+using System.Text.RegularExpressions;
+using System.Threading.Tasks;
+using EpicGames.Core;
+
+namespace HordeServer.IssueHandlers.Impl
+{
+	/// <summary>
+	/// Instance of a particular compile error
+	/// </summary>
+	class SymbolIssueHandler : IIssueHandler
+	{
+		/// <inheritdoc/>
+		public string Type => "Symbol";
+
+		/// <inheritdoc/>
+		public int Priority => 10;
+
+		/// <summary>
+		/// Determines if the given event id matches
+		/// </summary>
+		/// <param name="EventId">The event id to compare</param>
+		/// <returns>True if the given event id matches</returns>
+		public static bool IsMatchingEventId(EventId? EventId)
+		{
+			return EventId == KnownLogEvents.Linker_UndefinedSymbol || EventId == KnownLogEvents.Linker_DuplicateSymbol || EventId == KnownLogEvents.Linker;
+		}
+
+		/// <summary>
+		/// Parses symbol names from a log event
+		/// </summary>
+		/// <param name="EventData">The log event data</param>
+		/// <param name="SymbolNames">Receives the list of symbol names</param>
+		public static void GetSymbolNames(ILogEventData EventData, SortedSet<string> SymbolNames)
+		{
+			foreach (ILogEventLine Line in EventData.Lines)
+			{
+				string? Identifier;
+				if (Line.Data.TryGetNestedProperty("properties.symbol.identifier", out Identifier))
+				{
+					SymbolNames.Add(Identifier);
+				}
+			}
+		}
+
+		/// <inheritdoc/>
+		public void RankSuspects(IIssueFingerprint Fingerprint, List<SuspectChange> Changes)
+		{
+			HashSet<string> Names = new HashSet<string>();
+			foreach (string Name in Fingerprint.Keys)
+			{
+				Names.UnionWith(Name.Split("::", StringSplitOptions.RemoveEmptyEntries));
+			}
+
+			foreach (SuspectChange Change in Changes)
+			{
+				if (Change.ContainsCode)
+				{
+					int Matches = Names.Count(x => Change.Details.Files.Any(y => y.Path.Contains(x, StringComparison.OrdinalIgnoreCase)));
+					Change.Rank += 10 + (10 * Matches);
+				}
+			}
+		}
+
+		/// <inheritdoc/>
+		public string GetSummary(IIssueFingerprint Fingerprint, IssueSeverity Severity)
+		{
+			HashSet<string> Symbols = Fingerprint.Keys;
+			if (Symbols.Count == 1)
+			{
+				return $"Undefined symbol '{Symbols.First()}'";
+			}
+			else
+			{
+				return $"Undefined symbols: {StringUtils.FormatList(Symbols.ToArray(), 3)}";
+			}
+		}
+
+		public bool TryGetFingerprint(IJob Job, INode Node, ILogEventData EventData, [NotNullWhen(true)] out NewIssueFingerprint? Fingerprint)
+		{
+			if (!IsMatchingEventId(EventData.EventId))
+			{
+				Fingerprint = null;
+				return false;
+			}
+
+			SortedSet<string> SymbolNames = new SortedSet<string>();
+			GetSymbolNames(EventData, SymbolNames);
+			Fingerprint = new NewIssueFingerprint(Type, SymbolNames, null);
+			return true;
+		}
+	}
+}
