Horde: Move jobs, graphs, issues and users into subfolders.

diff --git a/Engine/Source/Programs/Horde/HordeServer/Jobs/IJobCollection.cs b/Engine/Source/Programs/Horde/HordeServer/Jobs/IJobCollection.cs
new file mode 100644
--- /dev/null
+++ b/Engine/Source/Programs/Horde/HordeServer/Jobs/IJobCollection.cs
@@ -0,0 +1,242 @@
+// Copyright Epic Games, Inc. All Rights Reserved.
+
+using EpicGames.Core;
+using HordeServer.Api;
+using HordeCommon;
+using HordeServer.Models;
+using HordeServer.Services;
+using HordeServer.Utilities;
+using Microsoft.Extensions.Logging;
+using MongoDB.Bson;
+using MongoDB.Bson.Serialization.Attributes;
+using MongoDB.Driver;
+using System;
+using System.Collections.Generic;
+using System.Diagnostics.CodeAnalysis;
+using System.IO;
+using System.Linq;
+using System.Threading.Tasks;
+
+namespace HordeServer.Collections
+{
+	using PoolId = StringId<IPool>;
+	using StreamId = StringId<IStream>;
+	using TemplateRefId = StringId<TemplateRef>;
+
+	/// <summary>
+	/// Interface for a collection of job documents
+	/// </summary>
+	[SuppressMessage("Compiler", "CA1054:URI properties should not be strings")]
+	public interface IJobCollection
+	{
+		/// <summary>
+		/// Creates a new job
+		/// </summary>
+		/// <param name="JobId">A requested job id</param>
+		/// <param name="StreamId">Unique id of the stream that this job belongs to</param>
+		/// <param name="TemplateRefId">Name of the template ref</param>
+		/// <param name="TemplateHash">Template for this job</param>
+		/// <param name="Graph">The graph for the new job</param>
+		/// <param name="Name">Name of the job</param>
+		/// <param name="Change">The change to build</param>
+		/// <param name="CodeChange">The corresponding code changelist number</param>
+		/// <param name="PreflightChange">Optional changelist to preflight</param>
+		/// <param name="ClonedPreflightChange">Optional cloned preflight changelist</param>
+		/// <param name="StartedByUserId">User that started the job</param>
+		/// <param name="StartedByUserName">User that started the job</param>
+		/// <param name="Priority">Priority of the job</param>
+		/// <param name="AutoSubmit">Whether to automatically submit the preflighted change on completion</param>
+		/// <param name="UpdateIssues">Whether to update issues based on the outcome of this job</param>
+		/// <param name="JobTriggers">List of downstream job triggers</param>
+		/// <param name="ShowUgsBadges">Whether to show badges in UGS for this job</param>
+		/// <param name="ShowUgsAlerts">Whether to show alerts in UGS for this job</param>
+		/// <param name="NotificationChannel">Notification channel for this job</param>
+		/// <param name="NotificationChannelFilter">Notification channel filter for this job</param>
+		/// <param name="HelixSwarmCallbackUrl">Callback URL to a Helix Server review, if any</param>
+		/// <param name="Arguments">Arguments for the job</param>
+		/// <returns>The new job document</returns>
+		Task<IJob> AddAsync(ObjectId JobId, StreamId StreamId, TemplateRefId TemplateRefId, ContentHash TemplateHash, IGraph Graph, string Name, int Change, int CodeChange, int? PreflightChange, int? ClonedPreflightChange, ObjectId? StartedByUserId, string? StartedByUserName, Priority? Priority, bool? AutoSubmit, bool? UpdateIssues, List<ChainedJobTemplate>? JobTriggers, bool ShowUgsBadges, bool ShowUgsAlerts, string? NotificationChannel, string? NotificationChannelFilter, string? HelixSwarmCallbackUrl, List<string>? Arguments);
+
+		/// <summary>
+		/// Gets a job with the given unique id
+		/// </summary>
+		/// <param name="JobId">Job id to search for</param>
+		/// <returns>Information about the given job</returns>
+		Task<IJob?> GetAsync(ObjectId JobId);
+
+		/// <summary>
+		/// Deletes a job
+		/// </summary>
+		/// <param name="Job">The job to remove</param>
+		Task<bool> RemoveAsync(IJob Job);
+
+		/// <summary>
+		/// Delete all the jobs for a stream
+		/// </summary>
+		/// <param name="StreamId">Unique id of the stream</param>
+		/// <returns>Async task</returns>
+		Task RemoveStreamAsync(StreamId StreamId);
+
+		/// <summary>
+		/// Gets a job's permissions info by ID
+		/// </summary>
+		/// <param name="JobId">Unique id of the job</param>
+		/// <returns>The job document</returns>
+		Task<IJobPermissions?> GetPermissionsAsync(ObjectId JobId);
+
+		/// <summary>
+		/// Searches for jobs matching the given criteria
+		/// </summary>
+		/// <param name="JobIds">List of job ids to return</param>
+		/// <param name="StreamId">The stream containing the job</param>
+		/// <param name="Name">Name of the job</param>
+		/// <param name="Templates">Templates to look for</param>
+		/// <param name="MinChange">The minimum changelist number</param>
+		/// <param name="MaxChange">The maximum changelist number</param>
+		/// <param name="PreflightChange">Preflight change to find</param>
+		/// <param name="StartedByUser">User id for which to include jobs</param>
+		/// <param name="PreflightStartedByUser">User for which to include preflight jobs</param>
+		/// <param name="MinCreateTime">The minimum creation time</param>
+		/// <param name="MaxCreateTime">The maximum creation time</param>
+		/// <param name="ModifiedBefore">Filter the results by modified time</param>
+		/// <param name="ModifiedAfter">Filter the results by modified time</param>
+		/// <param name="Index">Index of the first result to return</param>
+		/// <param name="Count">Number of results to return</param>
+		/// <returns>List of jobs matching the given criteria</returns>
+		Task<List<IJob>> FindAsync(ObjectId[]? JobIds = null, StreamId? StreamId = null, string? Name = null, TemplateRefId[]? Templates = null, int? MinChange = null, int? MaxChange = null, int? PreflightChange = null, ObjectId? PreflightStartedByUser = null, ObjectId? StartedByUser = null, DateTimeOffset ? MinCreateTime = null, DateTimeOffset? MaxCreateTime = null, DateTimeOffset? ModifiedBefore = null, DateTimeOffset? ModifiedAfter = null, int? Index = null, int? Count = null);
+
+		/// <summary>
+		/// Updates a new job
+		/// </summary>
+		/// <param name="Job">The job document to update</param>
+		/// <param name="Graph">Graph for the job</param>
+		/// <param name="Name">Name of the job</param>
+		/// <param name="Priority">Priority of the job</param>
+		/// <param name="AutoSubmit">Automatically submit the job on completion</param>
+		/// <param name="AutoSubmitChange">Changelist that was automatically submitted</param>
+		/// <param name="AutoSubmitMessage"></param>
+		/// <param name="AbortedByUser">Name of the user that aborted the job</param>
+		/// <param name="NotificationTriggerId">Id for a notification trigger</param>
+		/// <param name="Reports">New reports</param>
+		/// <param name="Arguments">New arguments for the job</param>
+		/// <param name="LabelIdxToTriggerId">New trigger ID for a label in the job</param>
+		/// <param name="JobTrigger">New downstream job id</param>
+		Task<bool> TryUpdateJobAsync(IJob Job, IGraph Graph, string? Name = null, Priority? Priority = null, bool? AutoSubmit = null, int? AutoSubmitChange = null, string? AutoSubmitMessage = null, string? AbortedByUser = null, ObjectId? NotificationTriggerId = null, List<Report>? Reports = null, List<string>? Arguments = null, KeyValuePair<int, ObjectId>? LabelIdxToTriggerId = null, KeyValuePair<TemplateRefId, ObjectId>? JobTrigger = null);
+
+		/// <summary>
+		/// Updates the state of a batch
+		/// </summary>
+		/// <param name="Job">Job to update</param>
+		/// <param name="Graph">The graph for this job</param>
+		/// <param name="BatchId">Unique id of the batch to update</param>
+		/// <param name="NewLogId">The new log file id</param>
+		/// <param name="NewState">New state of the jobstep</param>
+		/// <param name="NewError">Error code for the batch</param>
+		/// <returns>True if the job was updated, false if it was deleted</returns>
+		Task<bool> TryUpdateBatchAsync(IJob Job, IGraph Graph, SubResourceId BatchId, ObjectId? NewLogId, JobStepBatchState? NewState, JobStepBatchError? NewError);
+
+		/// <summary>
+		/// Update a jobstep state
+		/// </summary>
+		/// <param name="Job">Job to update</param>
+		/// <param name="Graph">The graph for this job</param>
+		/// <param name="BatchId">Unique id of the batch containing the step</param>
+		/// <param name="StepId">Unique id of the step to update</param>
+		/// <param name="NewState">New state of the jobstep</param>
+		/// <param name="NewOutcome">New outcome of the jobstep</param>
+		/// <param name="NewAbortRequested">New state of request abort</param>
+		/// <param name="NewAbortByUser">New name of user that requested the abort</param>
+		/// <param name="NewLogId">New log id for the jobstep</param>
+		/// <param name="NewNotificationTriggerId">New id for a notification trigger</param>
+		/// <param name="NewRetryByUser">Whether the step should be retried</param>
+		/// <param name="NewPriority">New priority for this step</param>
+		/// <param name="NewReports">New report documents</param>
+		/// <param name="NewProperties">Property changes. Any properties with a null value will be removed.</param>
+		/// <returns>True if the job was updated, false if it was deleted in the meantime</returns>
+		Task<bool> TryUpdateStepAsync(IJob Job, IGraph Graph, SubResourceId BatchId, SubResourceId StepId, JobStepState NewState = default, JobStepOutcome NewOutcome = default, bool? NewAbortRequested = null, string? NewAbortByUser = null, ObjectId? NewLogId = null, ObjectId? NewNotificationTriggerId = null, string? NewRetryByUser = null, Priority? NewPriority = null, List<Report>? NewReports = null, Dictionary<string, string?>? NewProperties = null);
+
+		/// <summary>
+		/// Attempts to update the node groups to be executed for a job. Fails if another write happens in the meantime.
+		/// </summary>
+		/// <param name="Job">The job to update</param>
+		/// <param name="NewGraph">New graph for this job</param>
+		/// <returns>True if the groups were updated to the given list. False if another write happened first.</returns>
+		Task<bool> TryUpdateGraphAsync(IJob Job, IGraph NewGraph);
+
+		/// <summary>
+		/// Removes a job from the dispatch queue. Ignores the state of any batches still remaining to execute. Should only be used to correct for inconsistent state.
+		/// </summary>
+		/// <param name="Job">Job</param>
+		/// <returns></returns>
+		Task<bool> TryRemoveFromDispatchQueueAsync(IJob Job);
+
+		/// <summary>
+		/// Adds an issue to a job
+		/// </summary>
+		/// <param name="JobId">The job id</param>
+		/// <param name="IssueId">The issue to add</param>
+		/// <returns>Async task</returns>
+		Task AddIssueToJobAsync(ObjectId JobId, int IssueId);
+
+		/// <summary>
+		/// Gets a queue of jobs to consider for execution
+		/// </summary>
+		/// <returns>Sorted list of jobs to execute</returns>
+		Task<List<IJob>> GetDispatchQueueAsync();
+
+		/// <summary>
+		/// Marks a job as skipped
+		/// </summary>
+		/// <param name="Job">The job to update</param>
+		/// <param name="Graph">Graph for the job</param>
+		/// <param name="Reason">Reason for this batch being failed</param>
+		/// <returns>Updated version of the job</returns>
+		Task<IJob?> SkipAllBatchesAsync(IJob? Job, IGraph Graph, JobStepBatchError Reason);
+
+		/// <summary>
+		/// Marks a batch as skipped
+		/// </summary>
+		/// <param name="Job">The job to update</param>
+		/// <param name="BatchIdx">The batch to mark as skipped</param>
+		/// <param name="Graph">Graph for the job</param>
+		/// <param name="Reason">Reason for this batch being failed</param>
+		/// <returns>Updated version of the job</returns>
+		Task<IJob?> SkipBatchAsync(IJob? Job, int BatchIdx, IGraph Graph, JobStepBatchError Reason);
+
+		/// <summary>
+		/// Abort an agent's lease, and update the payload accordingly
+		/// </summary>
+		/// <param name="Job">The job to update</param>
+		/// <param name="BatchIdx">Index of the batch to cancel</param>
+		/// <param name="Graph">Graph for the job</param>
+		/// <returns>True if the job is updated</returns>
+		Task<bool> TryFailBatchAsync(IJob Job, int BatchIdx, IGraph Graph);
+
+		/// <summary>
+		/// Attempt to assign a lease to execute a batch
+		/// </summary>
+		/// <param name="Job">The job containing the batch</param>
+		/// <param name="BatchIdx">Index of the batch</param>
+		/// <param name="PoolId">The pool id</param>
+		/// <param name="AgentId">New agent to execute the batch</param>
+		/// <param name="SessionId">Session of the agent that is to execute the batch</param>
+		/// <param name="LeaseId">The lease unique id</param>
+		/// <param name="LogId">Unique id of the log for the batch</param>
+		/// <returns>True if the batch is updated</returns>
+		Task<bool> TryAssignLeaseAsync(IJob Job, int BatchIdx, PoolId PoolId, AgentId AgentId, ObjectId SessionId, ObjectId LeaseId, ObjectId LogId);
+
+		/// <summary>
+		/// Cancel a lease reservation on a batch (before it has started)
+		/// </summary>
+		/// <param name="Job">The job containing the lease</param>
+		/// <param name="BatchIdx">Index of the batch to cancel</param>
+		/// <returns>True if the job is updated</returns>
+		Task<bool> TryCancelLeaseAsync(IJob Job, int BatchIdx);
+
+		/// <summary>
+		/// Upgrade all documents in the collection
+		/// </summary>
+		/// <returns>Async task</returns>
+		Task UpgradeDocumentsAsync();
+	}
+}
diff --git a/Engine/Source/Programs/Horde/HordeServer/Jobs/IJobStepRefCollection.cs b/Engine/Source/Programs/Horde/HordeServer/Jobs/IJobStepRefCollection.cs
new file mode 100644
--- /dev/null
+++ b/Engine/Source/Programs/Horde/HordeServer/Jobs/IJobStepRefCollection.cs
@@ -0,0 +1,98 @@
+// Copyright Epic Games, Inc. All Rights Reserved.
+
+using HordeCommon;
+using HordeServer.Models;
+using HordeServer.Services;
+using MongoDB.Bson;
+using MongoDB.Bson.Serialization.Attributes;
+using MongoDB.Driver;
+using System;
+using System.Collections.Generic;
+using System.Linq;
+using System.Threading.Tasks;
+using HordeServer.Utilities;
+
+namespace HordeServer.Collections
+{
+	using PoolId = StringId<IPool>;
+	using StreamId = StringId<IStream>;
+	using TemplateRefId = StringId<TemplateRef>;
+
+	/// <summary>
+	/// Interface for a collection of JobStepRef documents
+	/// </summary>
+	public interface IJobStepRefCollection
+	{
+		/// <summary>
+		/// Constructor
+		/// </summary>
+		/// <param name="Id">Unique id for the step being referenced</param>
+		/// <param name="JobName">Name of the job</param>
+		/// <param name="StepName">Name of the step</param>
+		/// <param name="StreamId">Unique id for the stream containing the job</param>
+		/// <param name="TemplateId"></param>
+		/// <param name="Change">The change number being built</param>
+		/// <param name="LogId">The log file id</param>
+		/// <param name="PoolId">The pool id</param>
+		/// <param name="AgentId">The agent id</param>
+		/// <param name="Outcome">Outcome of this step, if known</param>
+		/// <param name="WaitTime">Time taken for the batch containing this step to start</param>
+		/// <param name="InitTime">Time taken for the batch containing this step to initializer</param>
+		/// <param name="StartTimeUtc">Start time</param>
+		/// <param name="FinishTimeUtc">Finish time for the step, if known</param>
+		Task<IJobStepRef> InsertOrReplaceAsync(JobStepRefId Id, string JobName, string StepName, StreamId StreamId, TemplateRefId TemplateId, int Change, ObjectId? LogId, PoolId? PoolId, AgentId? AgentId, JobStepOutcome? Outcome, float WaitTime, float InitTime, DateTime StartTimeUtc, DateTime? FinishTimeUtc);
+
+		/// <summary>
+		/// Gets the history of a given node
+		/// </summary>
+		/// <param name="StreamId">Unique id for a stream</param>
+		/// <param name="TemplateId"></param>
+		/// <param name="NodeName">Name of the node</param>
+		/// <param name="Change">The current change</param>
+		/// <param name="IncludeFailed">Whether to include failed nodes</param>
+		/// <param name="Count">Number of results to return</param>
+		/// <returns>List of step references</returns>
+		Task<List<IJobStepRef>> GetStepsForNodeAsync(StreamId StreamId, TemplateRefId TemplateId, string NodeName, int? Change, bool IncludeFailed, int Count);
+
+		/// <summary>
+		/// Gets the previous job that ran a given step
+		/// </summary>
+		/// <param name="StreamId">Id of the stream to search</param>
+		/// <param name="TemplateId">The template id</param>
+		/// <param name="NodeName">Name of the step to find</param>
+		/// <param name="Change">The current changelist number</param>
+		/// <returns>The previous job, or null.</returns>
+		Task<IJobStepRef?> GetPrevStepForNodeAsync(StreamId StreamId, TemplateRefId TemplateId, string NodeName, int Change);
+
+		/// <summary>
+		/// Gets the next job that ran a given step
+		/// </summary>
+		/// <param name="StreamId">Id of the stream to search</param>
+		/// <param name="TemplateId">The template id</param>
+		/// <param name="NodeName">Name of the step to find</param>
+		/// <param name="Change">The current changelist number</param>
+		/// <returns>The previous job, or null.</returns>
+		Task<IJobStepRef?> GetNextStepForNodeAsync(StreamId StreamId, TemplateRefId TemplateId, string NodeName, int Change);
+	}
+
+	static class JobStepRefCollectionExtensions
+	{
+		public static Task UpdateAsync(this IJobStepRefCollection JobStepRefs, IJob Job, IJobStepBatch Batch, IJobStep Step, IGraph Graph)
+		{
+			if (Job.PreflightChange == 0)
+			{
+				float WaitTime = (float)(Batch.GetWaitTime() ?? TimeSpan.Zero).TotalSeconds;
+				float InitTime = (float)(Batch.GetInitTime() ?? TimeSpan.Zero).TotalSeconds;
+
+				string NodeName = Graph.Groups[Batch.GroupIdx].Nodes[Step.NodeIdx].Name;
+				JobStepOutcome? Outcome = Step.IsPending() ? (JobStepOutcome?)null : Step.Outcome;
+
+				return JobStepRefs.InsertOrReplaceAsync(new JobStepRefId(Job.Id, Batch.Id, Step.Id), Job.Name, NodeName, Job.StreamId, Job.TemplateId, Job.Change, Step.LogId, Batch.PoolId, Batch.AgentId, Outcome, WaitTime, InitTime, Step.StartTimeUtc ?? DateTime.UtcNow, Step.FinishTimeUtc);
+			}
+			else
+			{
+				return Task.CompletedTask;
+			}
+		}
+	}
+}
diff --git a/Engine/Source/Programs/Horde/HordeServer/Jobs/IJobTimingCollection.cs b/Engine/Source/Programs/Horde/HordeServer/Jobs/IJobTimingCollection.cs
new file mode 100644
--- /dev/null
+++ b/Engine/Source/Programs/Horde/HordeServer/Jobs/IJobTimingCollection.cs
@@ -0,0 +1,100 @@
+// Copyright Epic Games, Inc. All Rights Reserved.
+
+using Amazon.Runtime.Internal.Util;
+using Amazon.S3.Model.Internal.MarshallTransformations;
+using HordeServer.Models;
+using HordeServer.Services;
+using HordeServer.Utilities;
+using MongoDB.Bson;
+using MongoDB.Bson.Serialization.Attributes;
+using MongoDB.Driver;
+using System;
+using System.Collections.Generic;
+using System.Diagnostics.CodeAnalysis;
+using System.Linq;
+using System.Threading.Tasks;
+
+namespace HordeServer.Collections
+{
+	/// <summary>
+	/// Average timing information for a node
+	/// </summary>
+	public class JobStepTimingData
+	{
+		/// <summary>
+		/// Name of the node
+		/// </summary>
+		public string Name { get; set; }
+
+		/// <summary>
+		/// Wait time before executing the group containing this node
+		/// </summary>
+		[BsonIgnoreIfNull]
+		public float? AverageWaitTime { get; set; }
+
+		/// <summary>
+		/// Time taken for the group containing this node to initialize
+		/// </summary>
+		[BsonIgnoreIfNull]
+		public float? AverageInitTime { get; set; }
+
+		/// <summary>
+		/// Time spent executing this node
+		/// </summary>
+		[BsonIgnoreIfNull]
+		public float? AverageDuration { get; set; }
+
+		/// <summary>
+		/// Constructor for serialization
+		/// </summary>
+		[BsonConstructor]
+		private JobStepTimingData()
+		{
+			this.Name = null!;
+		}
+
+		/// <summary>
+		/// Constructor
+		/// </summary>
+		/// <param name="Name">Name of the node</param>
+		/// <param name="AverageWaitTime">Wait time before executing the group containing this node</param>
+		/// <param name="AverageInitTime">Time taken for the group containing this node to initialize</param>
+		/// <param name="AverageDuration">Time spent executing this node</param>
+		public JobStepTimingData(string Name, float? AverageWaitTime, float? AverageInitTime, float? AverageDuration)
+		{
+			this.Name = Name;
+			this.AverageWaitTime = AverageWaitTime;
+			this.AverageInitTime = AverageInitTime;
+			this.AverageDuration = AverageDuration;
+		}
+	}
+
+	/// <summary>
+	/// Interface for a collection of JobTiming documents
+	/// </summary>
+	public interface IJobTimingCollection
+	{
+		/// <summary>
+		/// Add timing information for the given job
+		/// </summary>
+		/// <param name="JobId">The job id</param>
+		/// <param name="Steps">List of timing info for each step</param>
+		/// <returns>New timing document</returns>
+		Task<IJobTiming?> TryAddAsync(ObjectId JobId, List<JobStepTimingData> Steps);
+
+		/// <summary>
+		/// Attempts to get the timing information for a particular job
+		/// </summary>
+		/// <param name="JobId">The unique job id</param>
+		/// <returns>Timing info for the requested jbo</returns>
+		Task<IJobTiming?> TryGetAsync(ObjectId JobId);
+
+		/// <summary>
+		/// Adds timing information for the particular job
+		/// </summary>
+		/// <param name="JobTiming">The current timing info for the job</param>
+		/// <param name="Steps">List of steps to add</param>
+		/// <returns>New timing document. Null if the document could not be updated.</returns>
+		Task<IJobTiming?> TryAddStepsAsync(IJobTiming JobTiming, List<JobStepTimingData> Steps);
+	}
+}
diff --git a/Engine/Source/Programs/Horde/HordeServer/Jobs/Impl/GraphCollection.cs b/Engine/Source/Programs/Horde/HordeServer/Jobs/Impl/GraphCollection.cs
new file mode 100644
--- /dev/null
+++ b/Engine/Source/Programs/Horde/HordeServer/Jobs/Impl/GraphCollection.cs
@@ -0,0 +1,469 @@
+// Copyright Epic Games, Inc. All Rights Reserved.
+
+using EpicGames.Core;
+using HordeServer.Api;
+using HordeCommon;
+using HordeServer.Models;
+using HordeServer.Services;
+using MongoDB.Bson;
+using MongoDB.Bson.Serialization.Attributes;
+using MongoDB.Driver;
+using System;
+using System.Collections.Concurrent;
+using System.Collections.Generic;
+using System.Diagnostics;
+using System.Linq;
+using System.Threading;
+using System.Threading.Tasks;
+
+namespace HordeServer.Collections.Impl
+{
+	/// <summary>
+	/// Collection of graph documents
+	/// </summary>
+	public class GraphCollection : IGraphCollection
+	{
+		/// <summary>
+		/// Represents a node in the graph
+		/// </summary>
+		[DebuggerDisplay("{Name}")]
+		class Node : INode
+		{
+			[BsonRequired]
+			public string Name { get; set; }
+
+			[BsonIgnoreIfNull]
+			public NodeRef[] InputDependencies { get; set; }
+
+			[BsonIgnoreIfNull]
+			public NodeRef[] OrderDependencies { get; set; }
+
+			public Priority Priority { get; set; }
+			public bool AllowRetry { get; set; } = true;
+			public bool RunEarly { get; set; }
+			public bool Warnings { get; set; } = true;
+
+			[BsonIgnoreIfNull]
+			public Dictionary<string, string>? Credentials { get; set; }
+
+			[BsonIgnoreIfNull]
+			public Dictionary<string, string>? Properties { get; set; }
+
+			IReadOnlyDictionary<string, string>? INode.Credentials => Credentials;
+			IReadOnlyDictionary<string, string>? INode.Properties => Properties;
+
+			[BsonConstructor]
+			private Node()
+			{
+				Name = null!;
+				InputDependencies = null!;
+				OrderDependencies = null!;
+			}
+
+			public Node(string Name, NodeRef[] InputDependencies, NodeRef[] OrderDependencies, Priority Priority, bool AllowRetry, bool RunEarly, bool Warnings, Dictionary<string, string>? Credentials, Dictionary<string, string>? Properties)
+			{
+				this.Name = Name;
+				this.InputDependencies = InputDependencies;
+				this.OrderDependencies = OrderDependencies;
+				this.Priority = Priority;
+				this.AllowRetry = AllowRetry;
+				this.RunEarly = RunEarly;
+				this.Warnings = Warnings;
+				this.Credentials = Credentials;
+				this.Properties = Properties;
+			}
+		}
+
+		class NodeGroup : INodeGroup
+		{
+			public string AgentType { get; set; }
+			public List<Node> Nodes { get; set; }
+
+			IReadOnlyList<INode> INodeGroup.Nodes => Nodes;
+
+			/// <summary>
+			/// Private constructor for BSON serializer
+			/// </summary>
+			[BsonConstructor]
+			private NodeGroup()
+			{
+				AgentType = null!;
+				Nodes = null!;
+			}
+
+			/// <summary>
+			/// Constructor
+			/// </summary>
+			/// <param name="AgentType">The type of agent to execute this group</param>
+			/// <param name="Nodes">Nodes to execute</param>
+			public NodeGroup(string AgentType, List<Node> Nodes)
+			{
+				this.AgentType = AgentType;
+				this.Nodes = Nodes;
+			}
+		}
+
+		class Aggregate : IAggregate
+		{
+			public string Name { get; set; }
+			public List<NodeRef> Nodes { get; set; }
+
+			IReadOnlyList<NodeRef> IAggregate.Nodes => Nodes;
+
+			private Aggregate()
+			{
+				Name = null!;
+				Nodes = new List<NodeRef>();
+			}
+
+			public Aggregate(string Name, List<NodeRef> Nodes)
+			{
+				this.Name = Name;
+				this.Nodes = Nodes;
+			}
+		}
+
+		class Label : ILabel
+		{
+			[BsonIgnoreIfNull]
+			public string? Name { get; set; }
+
+			[BsonIgnoreIfNull]
+			public string? Category { get; set; }
+
+			[BsonIgnoreIfNull]
+			public string? DashboardName { get; set; }
+
+			[BsonIgnoreIfNull]
+			public string? DashboardCategory { get; set; }
+
+			[BsonIgnoreIfNull]
+			public string? UgsName { get; set; }
+
+			[BsonIgnoreIfNull]
+			public string? UgsProject { get; set; }
+
+			[BsonIgnoreIfNull]
+			public LabelChange Change { get; set; }
+
+			public List<NodeRef> RequiredNodes { get; set; }
+			public List<NodeRef> IncludedNodes { get; set; }
+
+			string? ILabel.DashboardName => DashboardName ?? Name;
+			string? ILabel.DashboardCategory => DashboardCategory ?? Category;
+
+			private Label()
+			{
+				RequiredNodes = new List<NodeRef>();
+				IncludedNodes = new List<NodeRef>();
+			}
+
+			public Label(string? DashboardName, string? DashboardCategory, string? UgsName, string? UgsProject, LabelChange Change, List<NodeRef> RequiredNodes, List<NodeRef> IncludedNodes)
+			{
+				this.DashboardName = DashboardName;
+				this.DashboardCategory = DashboardCategory;
+				this.UgsName = UgsName;
+				this.UgsProject = UgsProject;
+				this.Change = Change;
+				this.RequiredNodes = RequiredNodes;
+				this.IncludedNodes = IncludedNodes;
+			}
+		}
+
+		class GraphDocument : IGraph
+		{
+			public static GraphDocument Empty { get; } = new GraphDocument(new List<NodeGroup>(), new List<Aggregate>(), new List<Label>());
+
+			[BsonRequired, BsonId]
+			public ContentHash Id { get; private set; } = ContentHash.Empty;
+
+			public int Schema { get; set; }
+			public List<NodeGroup> Groups { get; private set; } = new List<NodeGroup>();
+			public List<Aggregate> Aggregates { get; private set; } = new List<Aggregate>();
+			public List<Label> Labels { get; private set; } = new List<Label>();
+
+			[BsonIgnore]
+			IReadOnlyDictionary<string, NodeRef>? CachedNodeNameToRef;
+
+			IReadOnlyList<INodeGroup> IGraph.Groups => Groups;
+			IReadOnlyList<IAggregate> IGraph.Aggregates => Aggregates;
+			IReadOnlyList<ILabel> IGraph.Labels => Labels;
+
+			[BsonConstructor]
+			private GraphDocument()
+			{
+			}
+
+			public GraphDocument(List<NodeGroup> Groups, List<Aggregate> Aggregates, List<Label> Labels)
+			{
+				this.Groups = Groups;
+				this.Aggregates = Aggregates;
+				this.Labels = Labels;
+				this.Id = ContentHash.SHA1(BsonExtensionMethods.ToBson(this));
+			}
+
+			public GraphDocument(GraphDocument BaseGraph, List<NewGroup>? NewGroupRequests, List<NewAggregate>? NewAggregateRequests, List<NewLabel>? NewLabelRequests)
+			{
+				Dictionary<string, NodeRef> NodeNameToRef = new Dictionary<string, NodeRef>(BaseGraph.GetNodeNameToRef(), StringComparer.OrdinalIgnoreCase);
+
+				// Update the new list of groups
+				List<NodeGroup> NewGroups = new List<NodeGroup>(BaseGraph.Groups);
+				if (NewGroupRequests != null)
+				{
+					foreach (NewGroup NewGroupRequest in NewGroupRequests)
+					{
+						List<Node> Nodes = new List<Node>();
+						foreach (NewNode NewNodeRequest in NewGroupRequest.Nodes)
+						{
+							int NodeIdx = Nodes.Count;
+
+							Priority Priority = NewNodeRequest.Priority ?? Priority.Normal;
+							bool bAllowRetry = NewNodeRequest.AllowRetry ?? true;
+							bool bRunEarly = NewNodeRequest.RunEarly ?? false;
+							bool bWarnings = NewNodeRequest.Warnings ?? true;
+
+							NodeRef[] InputDependencies = (NewNodeRequest.InputDependencies == null) ? Array.Empty<NodeRef>() : NewNodeRequest.InputDependencies.Select(x => NodeNameToRef[x]).ToArray();
+							NodeRef[] OrderDependencies = (NewNodeRequest.OrderDependencies == null) ? Array.Empty<NodeRef>() : NewNodeRequest.OrderDependencies.Select(x => NodeNameToRef[x]).ToArray();
+							OrderDependencies = OrderDependencies.Union(InputDependencies).ToArray();
+							Nodes.Add(new Node(NewNodeRequest.Name, InputDependencies, OrderDependencies, Priority, bAllowRetry, bRunEarly, bWarnings, NewNodeRequest.Credentials, NewNodeRequest.Properties));
+
+							NodeNameToRef.Add(NewNodeRequest.Name, new NodeRef(NewGroups.Count, NodeIdx));
+						}
+						NewGroups.Add(new NodeGroup(NewGroupRequest.AgentType, Nodes));
+					}
+				}
+
+				// Update the list of aggregates
+				List<Aggregate> NewAggregates = new List<Aggregate>(BaseGraph.Aggregates);
+				if (NewAggregateRequests != null)
+				{
+					foreach (NewAggregate NewAggregateRequest in NewAggregateRequests)
+					{
+						List<NodeRef> Nodes = NewAggregateRequest.Nodes.ConvertAll(x => NodeNameToRef[x]);
+						NewAggregates.Add(new Aggregate(NewAggregateRequest.Name, Nodes));
+					}
+				}
+
+				// Update the list of labels
+				List<Label> NewLabels = new List<Label>(BaseGraph.Labels);
+				if (NewLabelRequests != null)
+				{
+					foreach (NewLabel NewLabelRequest in NewLabelRequests)
+					{
+						List<NodeRef> RequiredNodes = NewLabelRequest.RequiredNodes.ConvertAll(x => NodeNameToRef[x]);
+						List<NodeRef> IncludedNodes = NewLabelRequest.IncludedNodes.ConvertAll(x => NodeNameToRef[x]);
+						NewLabels.Add(new Label(NewLabelRequest.DashboardName, NewLabelRequest.DashboardCategory, NewLabelRequest.UgsName, NewLabelRequest.UgsProject, NewLabelRequest.Change, RequiredNodes, IncludedNodes));
+					}
+				}
+
+				// Create the new arrays
+				Groups = NewGroups;
+				Aggregates = NewAggregates;
+				Labels = NewLabels;
+
+				// Create the new graph, and save the generated node lookup into it
+				CachedNodeNameToRef = NodeNameToRef;
+
+				// Compute the hash
+				Id = ContentHash.SHA1(BsonExtensionMethods.ToBson(this));
+			}
+
+			public IReadOnlyDictionary<string, NodeRef> GetNodeNameToRef()
+			{
+				if (CachedNodeNameToRef == null)
+				{
+					Dictionary<string, NodeRef> NodeNameToRef = new Dictionary<string, NodeRef>(StringComparer.OrdinalIgnoreCase);
+					for (int GroupIdx = 0; GroupIdx < Groups.Count; GroupIdx++)
+					{
+						List<Node> Nodes = Groups[GroupIdx].Nodes;
+						for (int NodeIdx = 0; NodeIdx < Nodes.Count; NodeIdx++)
+						{
+							Node Node = Nodes[NodeIdx];
+							NodeNameToRef[Node.Name] = new NodeRef(GroupIdx, NodeIdx);
+						}
+					}
+					CachedNodeNameToRef = NodeNameToRef;
+				}
+				return CachedNodeNameToRef;
+			}
+		}
+
+		/// <summary>
+		/// Stores information about a cached graph
+		/// </summary>
+		class CachedGraph
+		{
+			/// <summary>
+			/// Time at which the graph was last accessed
+			/// </summary>
+			public long LastAccessTime
+			{
+				get { return LastAccessTimePrivate; }
+			}
+
+			/// <summary>
+			/// Backing value for <see cref="LastAccessTime"/>
+			/// </summary>
+			private long LastAccessTimePrivate;
+
+			/// <summary>
+			/// The graph instance
+			/// </summary>
+			public IGraph Graph { get; }
+
+			/// <summary>
+			/// Constructor
+			/// </summary>
+			/// <param name="Graph">The graph to store</param>
+			public CachedGraph(IGraph Graph)
+			{
+				this.LastAccessTimePrivate = Stopwatch.GetTimestamp();
+				this.Graph = Graph;
+			}
+
+			/// <summary>
+			/// Update the last access time
+			/// </summary>
+			public void Touch()
+			{
+				for (; ; )
+				{
+					long Time = Stopwatch.GetTimestamp();
+					long LastAccessTimeCopy = LastAccessTimePrivate;
+					if (Time < LastAccessTimeCopy || Interlocked.CompareExchange(ref LastAccessTimePrivate, Time, LastAccessTimeCopy) == LastAccessTimeCopy)
+					{
+						break;
+					}
+				}
+			}
+		}
+
+		/// <summary>
+		/// The jobs collection
+		/// </summary>
+		IMongoCollection<GraphDocument> Graphs;
+
+		/// <summary>
+		/// Maximum number of graphs to keep in the cache
+		/// </summary>
+		const int MaxGraphs = 1000;
+
+		/// <summary>
+		/// Cache of recently accessed graphs
+		/// </summary>
+		ConcurrentDictionary<ContentHash, CachedGraph> CachedGraphs = new ConcurrentDictionary<ContentHash, CachedGraph>();
+
+		/// <summary>
+		/// Constructor
+		/// </summary>
+		/// <param name="DatabaseService">The database service singleton</param>
+		public GraphCollection(DatabaseService DatabaseService)
+		{
+			Graphs = DatabaseService.GetCollection<GraphDocument>("Graphs");
+		}
+
+		/// <summary>
+		/// Adds a new graph document
+		/// </summary>
+		/// <param name="Graph">The graph to add</param>
+		/// <returns>Async task</returns>
+		async Task AddAsync(GraphDocument Graph)
+		{
+			if (!await Graphs.Find(x => x.Id == Graph.Id).AnyAsync())
+			{
+				try
+				{
+					await Graphs.InsertOneAsync(Graph);
+				}
+				catch (MongoWriteException Ex)
+				{
+					if (Ex.WriteError.Category != ServerErrorCategory.DuplicateKey)
+					{
+						throw;
+					}
+				}
+			}
+		}
+
+		/// <inheritdoc/>
+		public async Task<IGraph> AddAsync(ITemplate Template)
+		{
+			Node Node = new Node(IJob.SetupNodeName, Array.Empty<NodeRef>(), Array.Empty<NodeRef>(), Priority.High, true, false, true, null, null);
+			NodeGroup Group = new NodeGroup(Template.InitialAgentType ?? "Win64", new List<Node> { Node });
+
+			GraphDocument Graph = new GraphDocument(new List<NodeGroup> { Group }, new List<Aggregate>(), new List<Label>());
+			await AddAsync(Graph);
+			return Graph;
+		}
+
+		/// <inheritdoc/>
+		public async Task<IGraph> AppendAsync(IGraph? BaseGraph, List<NewGroup>? NewGroupRequests, List<NewAggregate>? NewAggregateRequests, List<NewLabel>? NewLabelRequests)
+		{
+			GraphDocument Graph = new GraphDocument((GraphDocument?)BaseGraph ?? GraphDocument.Empty, NewGroupRequests, NewAggregateRequests, NewLabelRequests);
+			await AddAsync(Graph);
+			return Graph;
+		}
+
+		/// <inheritdoc/>
+		public async Task<IGraph> GetAsync(ContentHash? Hash)
+		{
+			// Special case for an empty graph request
+			if (Hash == null || Hash == ContentHash.Empty || Hash == GraphDocument.Empty.Id)
+			{
+				return GraphDocument.Empty;
+			}
+
+			// Try to read the graph from the cache
+			CachedGraph? CachedGraph;
+			if (CachedGraphs.TryGetValue(Hash, out CachedGraph))
+			{
+				// Update the last access time
+				CachedGraph.Touch();
+			}
+			else
+			{
+				// Trim the cache
+				while (CachedGraphs.Count > MaxGraphs)
+				{
+					ContentHash? RemoveHash = CachedGraphs.OrderBy(x => x.Value.LastAccessTime).Select(x => x.Key).FirstOrDefault();
+					if (RemoveHash == null || RemoveHash == ContentHash.Empty)
+					{
+						break;
+					}
+					CachedGraphs.TryRemove(RemoveHash, out _);
+				}
+
+				// Create the new entry
+				CachedGraph = new CachedGraph(await Graphs.Find<GraphDocument>(x => x.Id == Hash).FirstAsync());
+				CachedGraphs.TryAdd(Hash, CachedGraph);
+			}
+			return CachedGraph.Graph;
+		}
+
+		/// <inheritdoc/>
+		public async Task<List<IGraph>> FindAllAsync(ContentHash[]? Hashes, int? Index, int? Count)
+		{
+			FilterDefinitionBuilder<GraphDocument> FilterBuilder = Builders<GraphDocument>.Filter;
+
+			FilterDefinition<GraphDocument> Filter = FilterBuilder.Empty;
+			if (Hashes != null)
+			{
+				Filter &= FilterBuilder.In(x => x.Id, Hashes);
+			}
+
+			List<GraphDocument> Results;
+			IFindFluent<GraphDocument, GraphDocument> Search = Graphs.Find(Filter);
+			if (Index != null)
+			{
+				Search = Search.Skip(Index.Value);
+			}
+			if (Count != null)
+			{
+				Search = Search.Limit(Count.Value);
+			}
+
+			Results = await Search.ToListAsync();
+			return Results.ConvertAll<IGraph>(x => x);
+		}
+	}
+}
diff --git a/Engine/Source/Programs/Horde/HordeServer/Jobs/Impl/JobCollection.cs b/Engine/Source/Programs/Horde/HordeServer/Jobs/Impl/JobCollection.cs
new file mode 100644
--- /dev/null
+++ b/Engine/Source/Programs/Horde/HordeServer/Jobs/Impl/JobCollection.cs
@@ -0,0 +1,1649 @@
+// Copyright Epic Games, Inc. All Rights Reserved.
+
+using EpicGames.Core;
+using HordeServer.Api;
+using HordeCommon;
+using HordeServer.Models;
+using HordeServer.Services;
+using HordeServer.Utilities;
+using Microsoft.Extensions.Logging;
+using MongoDB.Bson;
+using MongoDB.Bson.Serialization.Attributes;
+using MongoDB.Driver;
+using System;
+using System.Collections.Generic;
+using System.Diagnostics.CodeAnalysis;
+using System.IO;
+using System.Linq;
+using System.Threading.Tasks;
+using OpenTracing.Util;
+using OpenTracing;
+
+namespace HordeServer.Collections.Impl
+{
+	using StreamId = StringId<IStream>;
+	using TemplateRefId = StringId<TemplateRef>;
+	using PoolId = StringId<IPool>;
+
+	/// <summary>
+	/// Wrapper around the jobs collection in a mongo DB
+	/// </summary>
+	public class JobCollection : IJobCollection
+	{
+		/// <summary>
+		/// Embedded jobstep document
+		/// </summary>
+		[BsonIgnoreExtraElements]
+		class JobStepDocument : IJobStep
+		{
+			[BsonRequired]
+			public SubResourceId Id { get; set; }
+
+			[BsonRequired]
+			public int NodeIdx { get; set; }
+
+			[BsonRequired]
+			public JobStepState State { get; set; } = JobStepState.Waiting;
+
+			public JobStepOutcome Outcome { get; set; } = JobStepOutcome.Success;
+
+			[BsonIgnoreIfNull]
+			public ObjectId? LogId { get; set; }
+
+			[BsonIgnoreIfNull]
+			public ObjectId? NotificationTriggerId { get; set; }
+
+			[BsonIgnoreIfNull]
+			public DateTimeOffset? ReadyTime { get; set; }
+
+			[BsonIgnoreIfNull]
+			public DateTimeOffset? StartTime { get; set; }
+
+			[BsonIgnoreIfNull]
+			public DateTimeOffset? FinishTime { get; set; }
+
+			[BsonIgnoreIfNull]
+			public Priority? Priority { get; set; }
+
+			[BsonIgnoreIfDefault, BsonDefaultValue(false)]
+			public bool Retry { get; set; }
+
+			public string? RetryByUser { get; set; }
+
+			public bool AbortRequested { get; set; } = false;
+
+			public string? AbortByUser { get; set; }
+
+			[BsonIgnoreIfNull]
+			public List<Report>? Reports { get; set; }
+			IReadOnlyList<IReport>? IJobStep.Reports => Reports;
+
+			[BsonIgnoreIfNull]
+			public Dictionary<string, string>? Properties { get; set; }
+
+			string? IJobStep.RetryByUser => RetryByUser ?? (Retry ? (string?)"Unknown" : null);
+			DateTime? IJobStep.ReadyTimeUtc => ReadyTime?.UtcDateTime;
+			DateTime? IJobStep.StartTimeUtc => StartTime?.UtcDateTime;
+			DateTime? IJobStep.FinishTimeUtc => FinishTime?.UtcDateTime;
+
+			[BsonConstructor]
+			private JobStepDocument()
+			{
+			}
+
+			public JobStepDocument(SubResourceId Id, int NodeIdx)
+			{
+				this.Id = Id;
+				this.NodeIdx = NodeIdx;
+			}
+		}
+
+		class JobStepBatchDocument : IJobStepBatch
+		{
+			[BsonRequired]
+			public SubResourceId Id { get; set; }
+
+			public ObjectId? LogId { get; set; }
+
+			[BsonRequired]
+			public int GroupIdx { get; set; }
+
+			[BsonRequired]
+			public JobStepBatchState State { get; set; }
+
+			[BsonIgnoreIfDefault, BsonDefaultValue(JobStepBatchError.None)]
+			public JobStepBatchError Error { get; set; }
+
+			public List<JobStepDocument> Steps { get; set; } = new List<JobStepDocument>();
+
+			[BsonIgnoreIfNull]
+			public PoolId? PoolId { get; set; }
+
+			[BsonIgnoreIfNull]
+			public AgentId? AgentId { get; set; }
+
+			[BsonIgnoreIfNull]
+			public ObjectId? SessionId { get; set; }
+
+			[BsonIgnoreIfNull]
+			public ObjectId? LeaseId { get; set; }
+
+			public int SchedulePriority { get; set; }
+
+			[BsonIgnoreIfNull]
+			public DateTimeOffset? ReadyTime { get; set; }
+
+			[BsonIgnoreIfNull]
+			public DateTimeOffset? StartTime { get; set; }
+
+			[BsonIgnoreIfNull]
+			public DateTimeOffset? FinishTime { get; set; }
+
+			IReadOnlyList<IJobStep> IJobStepBatch.Steps => Steps;
+			DateTime? IJobStepBatch.ReadyTimeUtc => ReadyTime?.UtcDateTime;
+			DateTime? IJobStepBatch.StartTimeUtc => StartTime?.UtcDateTime;
+			DateTime? IJobStepBatch.FinishTimeUtc => FinishTime?.UtcDateTime;
+
+			[BsonConstructor]
+			private JobStepBatchDocument()
+			{
+			}
+
+			public JobStepBatchDocument(SubResourceId Id, int GroupIdx)
+			{
+				this.Id = Id;
+				this.GroupIdx = GroupIdx;
+			}
+		}
+
+		class ChainedJobDocument : IChainedJob
+		{
+			public string Target { get; set; }
+			public TemplateRefId TemplateRefId { get; set; }
+			public ObjectId? JobId { get; set; }
+
+			[BsonConstructor]
+			private ChainedJobDocument()
+			{
+				Target = String.Empty;
+			}
+
+			public ChainedJobDocument(ChainedJobTemplate Trigger)
+			{
+				this.Target = Trigger.Trigger;
+				this.TemplateRefId = Trigger.TemplateRefId;
+			}
+		}
+
+		class LabelNotificationDocument
+		{
+			public int LabelIdx;
+			public ObjectId TriggerId;
+		}
+
+		class JobDocument : IJob
+		{
+			[BsonRequired, BsonId]
+			public ObjectId Id { get; set; }
+
+			public StreamId StreamId { get; set; }
+			public TemplateRefId TemplateId { get; set; }
+			public ContentHash? TemplateHash { get; set; }
+			public ContentHash GraphHash { get; set; }
+
+			[BsonIgnoreIfNull]
+			public ObjectId? StartedByUserId { get; set; }
+
+			[BsonIgnoreIfNull]
+			public string? StartedByUser { get; set; }
+
+			[BsonIgnoreIfNull]
+			public string? AbortedByUser { get; set; }
+
+			[BsonRequired]
+			public string Name { get; set; }
+
+			public int Change { get; set; }
+			public int CodeChange { get; set; }
+			public int PreflightChange { get; set; }
+			public int ClonedPreflightChange { get; set; }
+			public Priority Priority { get; set; }
+
+			[BsonIgnoreIfDefault]
+			public bool AutoSubmit { get; set; }
+
+			[BsonIgnoreIfNull]
+			public int? AutoSubmitChange { get; set; }
+
+			[BsonIgnoreIfNull]
+			public string? AutoSubmitMessage { get; set; }
+
+			public bool UpdateIssues { get; set; }
+
+			[BsonIgnoreIfNull]
+			public DateTimeOffset? CreateTime { get; set; }
+
+			[BsonIgnoreIfNull]
+			public DateTime? CreateTimeUtc { get; set; }
+
+			public int SchedulePriority { get; set; }
+			public List<JobStepBatchDocument> Batches { get; set; } = new List<JobStepBatchDocument>();
+			public List<Report>? Reports { get; set; }
+			public List<string> Arguments { get; set; } = new List<string>();
+			public List<int> ReferencedByIssues { get; set; } = new List<int>();
+			public ObjectId? NotificationTriggerId { get; set; }
+			public bool ShowUgsBadges { get; set; }
+			public bool ShowUgsAlerts { get; set; }
+			public string? NotificationChannel { get; set; }
+			public string? NotificationChannelFilter { get; set; }
+			public string? HelixSwarmCallbackUrl { get; set; }
+			public List<LabelNotificationDocument> LabelNotifications = new List<LabelNotificationDocument>();
+			public List<ChainedJobDocument> ChainedJobs { get; set; } = new List<ChainedJobDocument>();
+
+			[BsonIgnoreIfNull]
+			public List<NodeRef>? RetriedNodes { get; set; }
+			public SubResourceId NextSubResourceId { get; set; }
+
+			[BsonIgnoreIfNull]
+			public DateTimeOffset? UpdateTime { get; set; }
+
+			[BsonIgnoreIfNull]
+			public DateTime? UpdateTimeUtc { get; set; }
+
+			public Acl? Acl { get; set; }
+
+			[BsonRequired]
+			public int UpdateIndex { get; set; }
+
+			DateTime IJob.CreateTimeUtc => CreateTimeUtc ?? CreateTime?.UtcDateTime ?? DateTime.UnixEpoch;
+			DateTime IJob.UpdateTimeUtc => UpdateTimeUtc ?? UpdateTime?.UtcDateTime ?? DateTime.UnixEpoch;
+			IReadOnlyList<IJobStepBatch> IJob.Batches => (IReadOnlyList<JobStepBatchDocument>)Batches;
+			IReadOnlyList<IReport>? IJob.Reports => Reports;
+			IReadOnlyList<string> IJob.Arguments => Arguments;
+			IReadOnlyList<int> IJob.Issues => ReferencedByIssues;
+			IReadOnlyDictionary<int, ObjectId> IJob.LabelIdxToTriggerId => LabelNotifications.ToDictionary(x => x.LabelIdx, x => x.TriggerId);
+			IReadOnlyList<IChainedJob> IJob.ChainedJobs => ChainedJobs;
+
+			[BsonConstructor]
+			private JobDocument()
+			{
+				Name = null!;
+				GraphHash = null!;
+			}
+
+			public JobDocument(ObjectId Id, StreamId StreamId, TemplateRefId TemplateId, ContentHash TemplateHash, ContentHash GraphHash, string Name, int Change, int CodeChange, int PreflightChange, int ClonedPreflightChange, ObjectId? StartedByUserId, string? StartedByUserName, Priority? Priority, bool? AutoSubmit, bool? UpdateIssues, DateTime CreateTimeUtc, List<ChainedJobDocument> ChainedJobs, bool ShowUgsBadges, bool ShowUgsAlerts, string? NotificationChannel, string? NotificationChannelFilter, string? HelixSwarmCallbackUrl, List<string>? Arguments)
+			{
+				this.Id = Id;
+				this.StreamId = StreamId;
+				this.TemplateId = TemplateId;
+				this.TemplateHash = TemplateHash;
+				this.GraphHash = GraphHash;
+				this.Name = Name;
+				this.Change = Change;
+				this.CodeChange = CodeChange;
+				this.PreflightChange = PreflightChange;
+				this.ClonedPreflightChange = ClonedPreflightChange;
+				this.StartedByUserId = StartedByUserId;
+				this.StartedByUser = StartedByUserName;
+				this.Priority = Priority ?? HordeCommon.Priority.Normal;
+				this.AutoSubmit = AutoSubmit ?? false;
+				this.UpdateIssues = UpdateIssues ?? (StartedByUserId == null && PreflightChange == 0);
+				this.CreateTimeUtc = CreateTimeUtc;
+				this.ChainedJobs = ChainedJobs;
+				this.ShowUgsBadges = ShowUgsBadges;
+				this.ShowUgsAlerts = ShowUgsAlerts;
+				this.NotificationChannel = NotificationChannel;
+				this.NotificationChannelFilter = NotificationChannelFilter;
+				this.HelixSwarmCallbackUrl = HelixSwarmCallbackUrl;
+				this.Arguments = Arguments ?? this.Arguments;
+				this.NextSubResourceId = SubResourceId.Random();
+				this.UpdateTimeUtc = CreateTimeUtc;
+			}
+
+			public void PostLoad()
+			{
+				if (GraphHash == ContentHash.Empty)
+				{
+					Batches.Clear();
+				}
+			}
+		}
+
+		/// <summary>
+		/// Projection of a job definition to just include permissions info
+		/// </summary>
+		[SuppressMessage("Design", "CA1812: Class is never instantiated")]
+		class JobPermissions : IJobPermissions
+		{
+			public static ProjectionDefinition<JobDocument> Projection { get; } = Builders<JobDocument>.Projection.Include(x => x.Acl).Include(x => x.StreamId);
+
+			public Acl? Acl { get; set; }
+			public StreamId StreamId { get; set; }
+		}
+
+		/// <summary>
+		/// Maximum number of times a step can be retried (after the original run)
+		/// </summary>
+		const int MaxRetries = 2;
+
+		/// <summary>
+		/// The jobs collection
+		/// </summary>
+		IMongoCollection<JobDocument> Jobs;
+
+		/// <summary>
+		/// Logger for output
+		/// </summary>
+		ILogger<JobCollection> Logger;
+
+		/// <summary>
+		/// Constructor
+		/// </summary>
+		/// <param name="DatabaseService">The database service singleton</param>
+		/// <param name="Logger">The logger instance</param>
+		public JobCollection(DatabaseService DatabaseService, ILogger<JobCollection> Logger)
+		{
+			this.Logger = Logger;
+
+			Jobs = DatabaseService.GetCollection<JobDocument>("Jobs");
+
+			if (!DatabaseService.ReadOnlyMode)
+			{
+				Jobs.Indexes.CreateOne(new CreateIndexModel<JobDocument>(Builders<JobDocument>.IndexKeys.Ascending(x => x.StreamId)));
+				Jobs.Indexes.CreateOne(new CreateIndexModel<JobDocument>(Builders<JobDocument>.IndexKeys.Ascending(x => x.Change)));
+				Jobs.Indexes.CreateOne(new CreateIndexModel<JobDocument>(Builders<JobDocument>.IndexKeys.Ascending(x => x.PreflightChange)));
+				Jobs.Indexes.CreateOne(new CreateIndexModel<JobDocument>(Builders<JobDocument>.IndexKeys.Ascending(x => x.CreateTimeUtc)));
+				Jobs.Indexes.CreateOne(new CreateIndexModel<JobDocument>(Builders<JobDocument>.IndexKeys.Ascending(x => x.UpdateTimeUtc)));
+				Jobs.Indexes.CreateOne(new CreateIndexModel<JobDocument>(Builders<JobDocument>.IndexKeys.Ascending(x => x.Name)));
+				Jobs.Indexes.CreateOne(new CreateIndexModel<JobDocument>(Builders<JobDocument>.IndexKeys.Ascending(x => x.StartedByUser)));
+				Jobs.Indexes.CreateOne(new CreateIndexModel<JobDocument>(Builders<JobDocument>.IndexKeys.Ascending(x => x.TemplateId)));
+				Jobs.Indexes.CreateOne(new CreateIndexModel<JobDocument>(Builders<JobDocument>.IndexKeys.Descending(x => x.SchedulePriority)));
+			}
+		}
+
+		/// <inheritdoc/>
+		[SuppressMessage("Compiler", "CA1054:URI parameters should not be strings")]
+		public async Task<IJob> AddAsync(ObjectId JobId, StreamId StreamId, TemplateRefId TemplateRefId, ContentHash TemplateHash, IGraph Graph, string Name, int Change, int CodeChange, int? PreflightChange, int? ClonedPreflightChange, ObjectId? StartedByUserId, string? StartedByUserName, Priority? Priority, bool? AutoSubmit, bool? UpdateIssues, List<ChainedJobTemplate>? ChainedJobs, bool ShowUgsBadges, bool ShowUgsAlerts, string? NotificationChannel, string? NotificationChannelFilter, string? HelixSwarmCallbackUrl, List<string>? Arguments)
+		{
+			List<ChainedJobDocument> JobTriggers = new List<ChainedJobDocument>();
+			if (ChainedJobs == null)
+			{
+				JobTriggers = new List<ChainedJobDocument>();
+			}
+			else
+			{
+				JobTriggers = ChainedJobs.ConvertAll(x => new ChainedJobDocument(x));
+			}
+
+			JobDocument NewJob = new JobDocument(JobId, StreamId, TemplateRefId, TemplateHash, Graph.Id, Name, Change, CodeChange, PreflightChange ?? 0, ClonedPreflightChange ?? 0, StartedByUserId, StartedByUserName, Priority, AutoSubmit, UpdateIssues, DateTime.UtcNow, JobTriggers, ShowUgsBadges, ShowUgsAlerts, NotificationChannel, NotificationChannelFilter, HelixSwarmCallbackUrl, Arguments);
+			CreateBatches(NewJob, Graph, Logger);
+
+			await Jobs.InsertOneAsync(NewJob);
+
+			return NewJob;
+		}
+
+		/// <inheritdoc/>
+		public async Task<IJob?> GetAsync(ObjectId JobId)
+		{
+			JobDocument? Job = await Jobs.Find<JobDocument>(x => x.Id == JobId).FirstOrDefaultAsync();
+			if (Job != null)
+			{
+				Job.PostLoad();
+			}
+			return Job;
+		}
+
+		/// <inheritdoc/>
+		public async Task<bool> RemoveAsync(IJob Job)
+		{
+			DeleteResult Result = await Jobs.DeleteOneAsync(x => x.Id == Job.Id && x.UpdateIndex == Job.UpdateIndex);
+			return Result.DeletedCount > 0;
+		}
+
+		/// <inheritdoc/>
+		public async Task RemoveStreamAsync(StreamId StreamId)
+		{
+			await Jobs.DeleteManyAsync(x => x.StreamId == StreamId);
+		}
+
+		/// <inheritdoc/>
+		public async Task<IJobPermissions?> GetPermissionsAsync(ObjectId JobId)
+		{
+			return await Jobs.Find<JobDocument>(x => x.Id == JobId).Project<JobPermissions>(JobPermissions.Projection).FirstOrDefaultAsync();
+		}
+
+		/// <inheritdoc/>
+		public async Task<List<IJob>> FindAsync(ObjectId[]? JobIds, StreamId? StreamId, string? Name, TemplateRefId[]? Templates, int? MinChange, int? MaxChange, int? PreflightChange, ObjectId? PreflightStartedByUser, ObjectId? StartedByUser, DateTimeOffset? MinCreateTime, DateTimeOffset? MaxCreateTime, DateTimeOffset? ModifiedBefore, DateTimeOffset? ModifiedAfter, int? Index, int? Count)
+		{
+			FilterDefinitionBuilder<JobDocument> FilterBuilder = Builders<JobDocument>.Filter;
+
+			FilterDefinition<JobDocument> Filter = FilterBuilder.Empty;
+			if (JobIds != null && JobIds.Length > 0)
+			{
+				Filter &= FilterBuilder.In(x => x.Id, JobIds);
+			}
+			if (StreamId != null)
+			{
+				Filter &= FilterBuilder.Eq(x => x.StreamId, StreamId.Value);
+			}
+			if (Name != null)
+			{
+				Filter &= FilterBuilder.Eq(x => x.Name, Name);
+			}
+			if (Templates != null)
+			{
+				Filter &= FilterBuilder.In(x => x.TemplateId, Templates);
+			}
+			if (MinChange != null)
+			{
+				Filter &= FilterBuilder.Gte(x => x.Change, MinChange);
+			}
+			if (MaxChange != null)
+			{
+				Filter &= FilterBuilder.Lte(x => x.Change, MaxChange);
+			}
+			if (PreflightChange != null)
+			{
+				Filter &= FilterBuilder.Eq(x => x.PreflightChange, PreflightChange);
+			}
+			if (PreflightStartedByUser != null)
+			{
+				Filter &= FilterBuilder.Or(FilterBuilder.Eq(x => x.PreflightChange, 0), FilterBuilder.Eq(x => x.StartedByUserId, PreflightStartedByUser));
+			}
+			if (StartedByUser != null)
+			{
+				Filter &= FilterBuilder.Eq(x => x.StartedByUserId, StartedByUser);
+			}
+			if (MinCreateTime != null)
+			{
+				Filter &= FilterBuilder.Gte(x => x.CreateTimeUtc!, MinCreateTime.Value.UtcDateTime);
+			}
+			if (MaxCreateTime != null)
+			{
+				Filter &= FilterBuilder.Lte(x => x.CreateTimeUtc!, MaxCreateTime.Value.UtcDateTime);
+			}
+			if (ModifiedBefore != null)
+			{
+				Filter &= FilterBuilder.Lte(x => x.UpdateTimeUtc!, ModifiedBefore.Value.UtcDateTime);
+			}
+			if (ModifiedAfter != null)
+			{
+				Filter &= FilterBuilder.Gte(x => x.UpdateTimeUtc!, ModifiedAfter.Value.UtcDateTime);
+			}
+
+			List<JobDocument> Results;
+			using (IScope Scope = GlobalTracer.Instance.BuildSpan("Jobs.Find").StartActive())
+			{
+				IFindFluent<JobDocument, JobDocument> Query = Jobs.Find<JobDocument>(Filter).SortByDescending(x => x.CreateTimeUtc!);
+				if (Index != null)
+				{
+					Query = Query.Skip(Index.Value);
+				}
+				if (Count != null)
+				{
+					Query = Query.Limit(Count.Value);
+				}
+				Results = await Query.ToListAsync();
+			}
+			foreach (JobDocument Result in Results)
+			{
+				Result.PostLoad();
+			}
+			return Results.ConvertAll<JobDocument, IJob>(x => x);
+		}
+
+		/// <inheritdoc/>
+		public async Task<bool> TryUpdateJobAsync(IJob Job, IGraph Graph, string? Name, Priority? Priority, bool? AutoSubmit, int? AutoSubmitChange, string? AutoSubmitMessage, string? AbortedByUser, ObjectId? NotificationTriggerId, List<Report>? Reports, List<string>? Arguments, KeyValuePair<int, ObjectId>? LabelIdxToTriggerId, KeyValuePair<TemplateRefId, ObjectId>? JobTrigger)
+		{
+			// Create the update 
+			UpdateDefinitionBuilder<JobDocument> UpdateBuilder = Builders<JobDocument>.Update;
+			List<UpdateDefinition<JobDocument>> Updates = new List<UpdateDefinition<JobDocument>>();
+
+			// Flag for whether to update batches
+			bool bUpdateBatches = false;
+
+			// Build the update list
+			JobDocument JobDocument = (JobDocument)Job;
+			if (Name != null)
+			{
+				JobDocument.Name = Name;
+				Updates.Add(UpdateBuilder.Set(x => x.Name, Job.Name));
+			}
+			if (Priority != null)
+			{
+				JobDocument.Priority = Priority.Value;
+				Updates.Add(UpdateBuilder.Set(x => x.Priority, Job.Priority));
+			}
+			if (AutoSubmit != null)
+			{
+				JobDocument.AutoSubmit = AutoSubmit.Value;
+				Updates.Add(UpdateBuilder.Set(x => x.AutoSubmit, Job.AutoSubmit));
+			}
+			if (AutoSubmitChange != null)
+			{
+				JobDocument.AutoSubmitChange = AutoSubmitChange.Value;
+				Updates.Add(UpdateBuilder.Set(x => x.AutoSubmitChange, Job.AutoSubmitChange));
+			}
+			if (AutoSubmitMessage != null)
+			{
+				JobDocument.AutoSubmitMessage = (AutoSubmitMessage.Length == 0)? null : AutoSubmitMessage;
+				Updates.Add(UpdateBuilder.SetOrUnsetNullRef(x => x.AutoSubmitMessage, Job.AutoSubmitMessage));
+			}
+			if (AbortedByUser != null && Job.AbortedByUser == null)
+			{
+				JobDocument.AbortedByUser = AbortedByUser;
+				Updates.Add(UpdateBuilder.Set(x => x.AbortedByUser, Job.AbortedByUser));
+				bUpdateBatches = true;
+			}
+			if (NotificationTriggerId != null)
+			{
+				JobDocument.NotificationTriggerId = NotificationTriggerId.Value;
+				Updates.Add(UpdateBuilder.Set(x => x.NotificationTriggerId, NotificationTriggerId));
+			}
+			if (LabelIdxToTriggerId != null)
+			{
+				if (JobDocument.LabelNotifications.Any(x => x.LabelIdx == LabelIdxToTriggerId.Value.Key))
+				{
+					throw new ArgumentException("Cannot update label trigger that already exists");
+				}
+				JobDocument.LabelNotifications.Add(new LabelNotificationDocument { LabelIdx = LabelIdxToTriggerId.Value.Key, TriggerId = LabelIdxToTriggerId.Value.Value });
+				Updates.Add(UpdateBuilder.Set(x => x.LabelNotifications, JobDocument.LabelNotifications));
+			}
+			if (JobTrigger != null)
+			{
+				for (int Idx = 0; Idx < JobDocument.ChainedJobs.Count; Idx++)
+				{
+					ChainedJobDocument JobTriggerDocument = JobDocument.ChainedJobs[Idx];
+					if (JobTriggerDocument.TemplateRefId == JobTrigger.Value.Key)
+					{
+						int LocalIdx = Idx;
+						JobTriggerDocument.JobId = JobTrigger.Value.Value;
+						Updates.Add(UpdateBuilder.Set(x => x.ChainedJobs[LocalIdx].JobId, JobTrigger.Value.Value));
+					}
+				}
+			}
+			if (Reports != null)
+			{
+				JobDocument.Reports ??= new List<Report>();
+				JobDocument.Reports.RemoveAll(x => Reports.Any(y => y.Name == x.Name));
+				JobDocument.Reports.AddRange(Reports);
+				Updates.Add(UpdateBuilder.Set(x => x.Reports, JobDocument.Reports));
+			}
+			if (Arguments != null)
+			{
+				HashSet<string> ModifiedArguments = new HashSet<string>(Job.Arguments);
+				ModifiedArguments.SymmetricExceptWith(Arguments);
+
+				foreach (string ModifiedArgument in ModifiedArguments)
+				{
+					if (ModifiedArgument.StartsWith(IJob.TargetArgumentPrefix, StringComparison.OrdinalIgnoreCase))
+					{
+						bUpdateBatches = true;
+					}
+				}
+
+				JobDocument.Arguments = Arguments.ToList();
+				Updates.Add(UpdateBuilder.Set(x => x.Arguments, JobDocument.Arguments));
+			}
+
+			// Update the batches
+			if (bUpdateBatches)
+			{
+				UpdateBatches(JobDocument, Graph, Updates, Logger);
+			}
+
+			// Update the new list of job steps
+			return (Updates.Count == 0 || await TryUpdateAsync(JobDocument, UpdateBuilder.Combine(Updates)));
+		}
+
+		/// <inheritdoc/>
+		public async Task<bool> TryUpdateBatchAsync(IJob Job, IGraph Graph, SubResourceId BatchId, ObjectId? NewLogId, JobStepBatchState? NewState, JobStepBatchError? NewError)
+		{
+			JobDocument JobDocument = (JobDocument)Job;
+
+			// Find the index of the appropriate batch
+			int BatchIdx = JobDocument.Batches.FindIndex(x => x.Id == BatchId);
+			if (BatchIdx == -1)
+			{
+				return false;
+			}
+
+			// If we're marking the batch as complete and there are still steps to run (eg. because the agent crashed), we need to mark all the steps as complete first
+			JobStepBatchDocument Batch = JobDocument.Batches[BatchIdx];
+
+			// Create the update 
+			UpdateDefinitionBuilder<JobDocument> UpdateBuilder = Builders<JobDocument>.Update;
+			List<UpdateDefinition<JobDocument>> Updates = new List<UpdateDefinition<JobDocument>>();
+
+			// Update the batch
+			if (NewLogId != null)
+			{
+				Batch.LogId = NewLogId.Value;
+				Updates.Add(UpdateBuilder.Set(x => x.Batches[BatchIdx].LogId, Batch.LogId));
+			}
+			if (NewState != null)
+			{
+				Batch.State = NewState.Value;
+				Updates.Add(UpdateBuilder.Set(x => x.Batches[BatchIdx].State, Batch.State));
+
+				if (Batch.StartTime == null && NewState >= JobStepBatchState.Starting)
+				{
+					Batch.StartTime = DateTimeOffset.Now;
+					Updates.Add(UpdateBuilder.Set(x => x.Batches[BatchIdx].StartTime, Batch.StartTime));
+				}
+				if (NewState == JobStepBatchState.Complete)
+				{
+					Batch.FinishTime = DateTimeOffset.Now;
+					Updates.Add(UpdateBuilder.Set(x => x.Batches[BatchIdx].FinishTime, Batch.FinishTime));
+				}
+			}
+			if (NewError != null)
+			{
+				Batch.Error = NewError.Value;
+				Updates.Add(UpdateBuilder.Set(x => x.Batches[BatchIdx].Error, Batch.Error));
+			}
+
+			// If the batch is being marked as incomplete, see if we can reschedule any of the work
+			if (NewError == JobStepBatchError.Incomplete)
+			{
+				// Get the new list of retried nodes
+				List<NodeRef> RetriedNodes = JobDocument.RetriedNodes ?? new List<NodeRef>();
+
+				// Check if there are any steps that need to be run again
+				bool bUpdateState = false;
+				foreach (JobStepDocument Step in Batch.Steps)
+				{
+					if (Step.State == JobStepState.Ready || Step.State == JobStepState.Waiting)
+					{
+						if (CanRetryNode(JobDocument, Batch.GroupIdx, Step.NodeIdx))
+						{
+							RetriedNodes.Add(new NodeRef(Batch.GroupIdx, Step.NodeIdx));
+						}
+						else
+						{
+							Step.State = JobStepState.Skipped;
+						}
+						bUpdateState = true;
+					}
+				}
+
+				// Update the steps
+				if (bUpdateState)
+				{
+					Updates.Clear();
+					UpdateBatches(JobDocument, Graph, Updates, Logger);
+
+					JobDocument.RetriedNodes = RetriedNodes;
+					Updates.Add(UpdateBuilder.Set(x => x.RetriedNodes, JobDocument.RetriedNodes));
+				}
+			}
+
+			// Update the new list of job steps
+			return (Updates.Count == 0 || await TryUpdateAsync(JobDocument, UpdateBuilder.Combine(Updates)));
+		}
+
+		/// <inheritdoc/>
+		public async Task<bool> TryUpdateStepAsync(IJob Job, IGraph Graph, SubResourceId BatchId, SubResourceId StepId, JobStepState NewState, JobStepOutcome NewOutcome, bool? NewAbortRequested, string? NewAbortByUser, ObjectId? NewLogId, ObjectId? NewNotificationTriggerId, string? NewRetryByUser, Priority? NewPriority, List<Report>? NewReports, Dictionary<string, string?>? NewProperties)
+		{
+			JobDocument JobDocument = (JobDocument)Job;
+
+			// Create the update 
+			UpdateDefinitionBuilder<JobDocument> UpdateBuilder = Builders<JobDocument>.Update;
+			List<UpdateDefinition<JobDocument>> Updates = new List<UpdateDefinition<JobDocument>>();
+
+			// Update the appropriate batch
+			bool bRefreshBatches = false;
+			bool bRefreshDependentJobSteps = false;
+			for (int LoopBatchIdx = 0; LoopBatchIdx < JobDocument.Batches.Count; LoopBatchIdx++)
+			{
+				int BatchIdx = LoopBatchIdx; // For lambda capture
+				JobStepBatchDocument Batch = JobDocument.Batches[BatchIdx];
+				if (Batch.Id == BatchId)
+				{
+					for (int LoopStepIdx = 0; LoopStepIdx < Batch.Steps.Count; LoopStepIdx++)
+					{
+						int StepIdx = LoopStepIdx; // For lambda capture
+						JobStepDocument Step = Batch.Steps[StepIdx];
+						if (Step.Id == StepId)
+						{
+							// Update the state
+							if (NewState != JobStepState.Unspecified && Step.State != NewState)
+							{
+								Step.State = NewState;
+								Updates.Add(UpdateBuilder.Set(x => x.Batches[BatchIdx].Steps[StepIdx].State, Step.State));
+
+								if (Step.State == JobStepState.Running)
+								{
+									Step.StartTime = DateTimeOffset.Now;
+									Updates.Add(UpdateBuilder.Set(x => x.Batches[BatchIdx].Steps[StepIdx].StartTime, Step.StartTime));
+								}
+								else if (Step.State == JobStepState.Completed || Step.State == JobStepState.Aborted)
+								{
+									Step.FinishTime = DateTimeOffset.Now;
+									Updates.Add(UpdateBuilder.Set(x => x.Batches[BatchIdx].Steps[StepIdx].FinishTime, Step.FinishTime));
+								}
+
+								bRefreshDependentJobSteps = true;
+							}
+
+							// Update the job outcome
+							if (NewOutcome != JobStepOutcome.Unspecified && Step.Outcome != NewOutcome)
+							{
+								Step.Outcome = NewOutcome;
+								Updates.Add(UpdateBuilder.Set(x => x.Batches[BatchIdx].Steps[StepIdx].Outcome, Step.Outcome));
+
+								bRefreshDependentJobSteps = true;
+							}
+
+							// Update the request abort status
+							if (NewAbortRequested != null && Step.AbortRequested == false)
+							{
+								Step.AbortRequested = NewAbortRequested.Value;
+								Updates.Add(UpdateBuilder.Set(x => x.Batches[BatchIdx].Steps[StepIdx].AbortRequested, Step.AbortRequested));
+
+								bRefreshDependentJobSteps = true;
+							}
+
+							// Update the user that requested the abort
+							if (NewAbortByUser != null && Step.AbortByUser == null)
+							{
+								Step.AbortByUser = NewAbortByUser;
+								Updates.Add(UpdateBuilder.Set(x => x.Batches[BatchIdx].Steps[StepIdx].AbortByUser, Step.AbortByUser));
+
+								bRefreshDependentJobSteps = true;
+							}
+
+							// Update the log id
+							if (NewLogId != null && Step.LogId != NewLogId.Value)
+							{
+								Step.LogId = NewLogId.Value;
+								Updates.Add(UpdateBuilder.Set(x => x.Batches[BatchIdx].Steps[StepIdx].LogId, Step.LogId));
+
+								bRefreshDependentJobSteps = true;
+							}
+
+							// Update the notification trigger id
+							if (NewNotificationTriggerId != null && Step.NotificationTriggerId == null)
+							{
+								Step.NotificationTriggerId = NewNotificationTriggerId.Value;
+								Updates.Add(UpdateBuilder.Set(x => x.Batches[BatchIdx].Steps[StepIdx].NotificationTriggerId, Step.NotificationTriggerId));
+							}
+
+							// Update the retry flag
+							if (NewRetryByUser != null && Step.RetryByUser == null)
+							{
+								Step.RetryByUser = NewRetryByUser;
+								Updates.Add(UpdateBuilder.Set(x => x.Batches[BatchIdx].Steps[StepIdx].RetryByUser, Step.RetryByUser));
+
+								bRefreshBatches = true;
+							}
+
+							// Update the priority
+							if (NewPriority != null && NewPriority.Value != Step.Priority)
+							{
+								Step.Priority = NewPriority.Value;
+								Updates.Add(UpdateBuilder.Set(x => x.Batches[BatchIdx].Steps[StepIdx].Priority, Step.Priority));
+
+								bRefreshBatches = true;
+							}
+
+							// Add any new reports
+							if (NewReports != null)
+							{
+								Step.Reports ??= new List<Report>();
+								Step.Reports.RemoveAll(x => NewReports.Any(y => y.Name == x.Name));
+								Step.Reports.AddRange(NewReports);
+								Updates.Add(UpdateBuilder.Set(x => x.Batches[BatchIdx].Steps[StepIdx].Reports, Step.Reports));
+							}
+
+							// Apply any property updates
+							if (NewProperties != null)
+							{
+								if (Step.Properties == null)
+								{
+									Step.Properties = new Dictionary<string, string>(StringComparer.Ordinal);
+								}
+
+								foreach (KeyValuePair<string, string?> Pair in NewProperties)
+								{
+									if (Pair.Value == null)
+									{
+										Step.Properties.Remove(Pair.Key);
+									}
+									else
+									{
+										Step.Properties[Pair.Key] = Pair.Value;
+									}
+								}
+
+								if (Step.Properties.Count == 0)
+								{
+									Step.Properties = null;
+									Updates.Add(UpdateBuilder.Unset(x => x.Batches[BatchIdx].Steps[StepIdx].Properties));
+								}
+								else
+								{
+									foreach (KeyValuePair<string, string?> Pair in NewProperties)
+									{
+										if (Pair.Value == null)
+										{
+											Updates.Add(UpdateBuilder.Unset(x => x.Batches[BatchIdx].Steps[StepIdx].Properties![Pair.Key]));
+										}
+										else
+										{
+											Updates.Add(UpdateBuilder.Set(x => x.Batches[BatchIdx].Steps[StepIdx].Properties![Pair.Key], Pair.Value));
+										}
+									}
+								}
+							}
+							break;
+						}
+					}
+					break;
+				}
+			}
+
+			// Update the batches
+			if (bRefreshBatches)
+			{
+				Updates.Clear(); // UpdateBatches will update the entire batches list. We need to remove all the individual step updates to avoid an exception.
+				UpdateBatches(JobDocument, Graph, Updates, Logger);
+			}
+
+			// Update the state of dependent jobsteps
+			if (bRefreshDependentJobSteps)
+			{
+				RefreshDependentJobSteps(JobDocument, Graph, Updates, Logger);
+				RefreshJobPriority(JobDocument, Updates);
+			}
+
+			// Update the new list of job steps
+			return (Updates.Count == 0 || await TryUpdateAsync(JobDocument, UpdateBuilder.Combine(Updates)));
+		}
+
+		/// <inheritdoc/>
+		async Task<bool> TryUpdateAsync(JobDocument Job, UpdateDefinition<JobDocument> Update)
+		{
+			int NewUpdateIndex = Job.UpdateIndex + 1;
+			Update = Update.Set(x => x.UpdateIndex, NewUpdateIndex);
+
+			DateTime NewUpdateTimeUtc = DateTime.UtcNow;
+			Update = Update.Set(x => x.UpdateTimeUtc, NewUpdateTimeUtc);
+
+			UpdateResult Result = await Jobs.UpdateOneAsync<JobDocument>(x => x.Id == Job.Id && x.UpdateIndex == Job.UpdateIndex, Update);
+			if (Result.ModifiedCount > 0)
+			{
+				Job.UpdateIndex = NewUpdateIndex;
+				Job.UpdateTimeUtc = NewUpdateTimeUtc;
+				return true;
+			}
+			return false;
+		}
+
+		/// <inheritdoc/>
+		public Task<bool> TryRemoveFromDispatchQueueAsync(IJob Job)
+		{
+			JobDocument JobDocument = (JobDocument)Job;
+
+			// Create the update 
+			UpdateDefinitionBuilder<JobDocument> UpdateBuilder = Builders<JobDocument>.Update;
+			List<UpdateDefinition<JobDocument>> Updates = new List<UpdateDefinition<JobDocument>>();
+
+			JobDocument.SchedulePriority = 0;
+			Updates.Add(UpdateBuilder.Set(x => x.SchedulePriority, JobDocument.SchedulePriority));
+
+			// Update the new list of job steps
+			return TryUpdateAsync(JobDocument, UpdateBuilder.Combine(Updates));
+		}
+
+		/// <inheritdoc/>
+		public Task<bool> TryUpdateGraphAsync(IJob Job, IGraph NewGraph)
+		{
+			JobDocument JobDocument = (JobDocument)Job;
+
+			// Create the update 
+			UpdateDefinitionBuilder<JobDocument> UpdateBuilder = Builders<JobDocument>.Update;
+			List<UpdateDefinition<JobDocument>> Updates = new List<UpdateDefinition<JobDocument>>();
+
+			JobDocument.GraphHash = NewGraph.Id;
+			Updates.Add(UpdateBuilder.Set(x => x.GraphHash, Job.GraphHash));
+
+			UpdateBatches(JobDocument, NewGraph, Updates, Logger);
+
+			// Update the new list of job steps
+			return TryUpdateAsync(JobDocument, UpdateBuilder.Combine(Updates));
+		}
+
+		/// <inheritdoc/>
+		public async Task AddIssueToJobAsync(ObjectId JobId, int IssueId)
+		{
+			FilterDefinition<JobDocument> JobFilter = Builders<JobDocument>.Filter.Eq(x => x.Id, JobId);
+			UpdateDefinition<JobDocument> JobUpdate = Builders<JobDocument>.Update.AddToSet(x => x.ReferencedByIssues, IssueId).Inc(x => x.UpdateIndex, 1).Max(x => x.UpdateTimeUtc, DateTime.UtcNow);
+			await Jobs.UpdateOneAsync(JobFilter, JobUpdate);
+		}
+
+		/// <inheritdoc/>
+		public async Task<List<IJob>> GetDispatchQueueAsync()
+		{
+			List<JobDocument> NewJobs = await Jobs.Find(x => x.SchedulePriority > 0).SortByDescending(x => x.SchedulePriority).ThenBy(x => x.CreateTimeUtc).ToListAsync();
+			foreach (JobDocument Result in NewJobs)
+			{
+				Result.PostLoad();
+			}
+			return NewJobs.ConvertAll<JobDocument, IJob>(x => x);
+		}
+
+		/// <summary>
+		/// Marks a job as skipped
+		/// </summary>
+		/// <param name="Job">The job to update</param>
+		/// <param name="Graph">Graph for the job</param>
+		/// <param name="Reason">Reason for this batch being failed</param>
+		/// <returns>Updated version of the job</returns>
+		public async Task<IJob?> SkipAllBatchesAsync(IJob? Job, IGraph Graph, JobStepBatchError Reason)
+		{
+			while (Job != null)
+			{
+				JobDocument JobDocument = (JobDocument)Job;
+
+				for (int BatchIdx = 0; BatchIdx < JobDocument.Batches.Count; BatchIdx++)
+				{
+					JobStepBatchDocument Batch = JobDocument.Batches[BatchIdx];
+					if (Batch.State == JobStepBatchState.Ready || Batch.State == JobStepBatchState.Waiting)
+					{
+						Batch.State = JobStepBatchState.Complete;
+						Batch.Error = Reason;
+						Batch.FinishTime = DateTimeOffset.UtcNow;
+
+						for (int StepIdx = 0; StepIdx < Batch.Steps.Count; StepIdx++)
+						{
+							JobStepDocument Step = Batch.Steps[StepIdx];
+							if (Step.State == JobStepState.Ready || Step.State == JobStepState.Waiting)
+							{
+								Step.State = JobStepState.Completed;
+								Step.Outcome = JobStepOutcome.Failure;
+							}
+						}
+					}
+				}
+
+				List<UpdateDefinition<JobDocument>> Updates = new List<UpdateDefinition<JobDocument>>();
+				UpdateBatches(JobDocument, Graph, Updates, Logger);
+
+				if (await TryUpdateAsync(JobDocument, Builders<JobDocument>.Update.Combine(Updates)))
+				{
+					break;
+				}
+
+				Job = await GetAsync(Job.Id);
+			}
+			return Job;
+		}
+
+		/// <inheritdoc/>
+		public async Task<IJob?> SkipBatchAsync(IJob? Job, int BatchIdx, IGraph Graph, JobStepBatchError Reason)
+		{
+			while (Job != null)
+			{
+				JobDocument JobDocument = (JobDocument)Job;
+				JobStepBatchDocument Batch = JobDocument.Batches[BatchIdx];
+
+				Batch.State = JobStepBatchState.Complete;
+				Batch.Error = Reason;
+				Batch.FinishTime = DateTimeOffset.UtcNow;
+
+				foreach (JobStepDocument Step in Batch.Steps)
+				{
+					if (Step.State != JobStepState.Skipped)
+					{
+						Step.State = JobStepState.Skipped;
+						Step.Outcome = JobStepOutcome.Failure;
+					}
+				}
+
+				List<UpdateDefinition<JobDocument>> Updates = new List<UpdateDefinition<JobDocument>>();
+				UpdateBatches(JobDocument, Graph, Updates, Logger);
+
+				if (await TryUpdateAsync(JobDocument, Builders<JobDocument>.Update.Combine(Updates)))
+				{
+					break;
+				}
+
+				Job = await GetAsync(Job.Id);
+			}
+			return Job;
+		}
+
+		/// <inheritdoc/>
+		public async Task<bool> TryAssignLeaseAsync(IJob Job, int BatchIdx, PoolId PoolId, AgentId AgentId, ObjectId SessionId, ObjectId LeaseId, ObjectId LogId)
+		{
+			// Try to update the job with this agent id
+			UpdateDefinitionBuilder<JobDocument> UpdateBuilder = Builders<JobDocument>.Update;
+			List<UpdateDefinition<JobDocument>> Updates = new List<UpdateDefinition<JobDocument>>();
+
+			Updates.Add(UpdateBuilder.Set(x => x.Batches[BatchIdx].PoolId, PoolId));
+			Updates.Add(UpdateBuilder.Set(x => x.Batches[BatchIdx].AgentId, AgentId));
+			Updates.Add(UpdateBuilder.Set(x => x.Batches[BatchIdx].SessionId, SessionId));
+			Updates.Add(UpdateBuilder.Set(x => x.Batches[BatchIdx].LeaseId, LeaseId));
+			Updates.Add(UpdateBuilder.Set(x => x.Batches[BatchIdx].LogId, LogId));
+
+			// Extra logs for catching why session IDs sometimes doesn't match in prod. Resulting in PermissionDenied errors.
+			if (BatchIdx < Job.Batches.Count && Job.Batches[BatchIdx].SessionId != null)
+			{
+				string CurrentSessionId = Job.Batches[BatchIdx].SessionId!.Value.ToString();
+				Logger.LogError("Attempt to replace current session ID {CurrSessionId} with {NewSessionId} for batch {JobId}:{BatchId}", CurrentSessionId, SessionId.ToString(), Job.Id.ToString(), Job.Batches[BatchIdx].Id);
+				return false;
+			}
+
+			JobDocument JobDocument = (JobDocument)Job;
+			if (!await TryUpdateAsync(JobDocument, UpdateBuilder.Combine(Updates)))
+			{
+				return false;
+			}
+
+			JobStepBatchDocument Batch = JobDocument.Batches[BatchIdx];
+			Batch.AgentId = AgentId;
+			Batch.SessionId = SessionId;
+			Batch.LeaseId = LeaseId;
+			Batch.LogId = LogId;
+
+			return true;
+		}
+
+		/// <inheritdoc/>
+		public async Task<bool> TryCancelLeaseAsync(IJob Job, int BatchIdx)
+		{
+			JobDocument JobDocument = (JobDocument)Job;
+
+			UpdateDefinition<JobDocument> Update = Builders<JobDocument>.Update.Unset(x => x.Batches[BatchIdx].AgentId).Unset(x => x.Batches[BatchIdx].SessionId).Unset(x => x.Batches[BatchIdx].LeaseId);
+			if (await TryUpdateAsync(JobDocument, Update))
+			{
+				JobStepBatchDocument Batch = JobDocument.Batches[BatchIdx];
+				Batch.AgentId = null;
+				Batch.SessionId = null;
+				Batch.LeaseId = null;
+				return true;
+			}
+			return false;
+		}
+
+		/// <inheritdoc/>
+		public async Task<bool> TryFailBatchAsync(IJob Job, int BatchIdx, IGraph Graph)
+		{
+			JobDocument JobDocument = (JobDocument)Job;
+			JobStepBatchDocument Batch = JobDocument.Batches[BatchIdx];
+
+			UpdateDefinitionBuilder<JobDocument> UpdateBuilder = new UpdateDefinitionBuilder<JobDocument>();
+			List<UpdateDefinition<JobDocument>> Updates = new List<UpdateDefinition<JobDocument>>();
+
+			if (Batch.State != JobStepBatchState.Complete)
+			{
+				Batch.State = JobStepBatchState.Complete;
+				Updates.Add(UpdateBuilder.Set(x => x.Batches[BatchIdx].State, Batch.State));
+
+				Batch.Error = JobStepBatchError.LostConnection;
+				Updates.Add(UpdateBuilder.Set(x => x.Batches[BatchIdx].Error, Batch.Error));
+
+				Batch.FinishTime = DateTimeOffset.UtcNow;
+				Updates.Add(UpdateBuilder.Set(x => x.Batches[BatchIdx].FinishTime, Batch.FinishTime));
+			}
+
+			for (int StepIdx = 0; StepIdx < Batch.Steps.Count; StepIdx++)
+			{
+				JobStepDocument Step = Batch.Steps[StepIdx];
+				if (Step.State == JobStepState.Running)
+				{
+					int StepIdxCopy = StepIdx;
+
+					Step.State = JobStepState.Aborted;
+					Updates.Add(UpdateBuilder.Set(x => x.Batches[BatchIdx].Steps[StepIdxCopy].State, Step.State));
+
+					Step.Outcome = JobStepOutcome.Failure;
+					Updates.Add(UpdateBuilder.Set(x => x.Batches[BatchIdx].Steps[StepIdxCopy].Outcome, Step.Outcome));
+
+					Step.FinishTime = DateTimeOffset.Now;
+					Updates.Add(UpdateBuilder.Set(x => x.Batches[BatchIdx].Steps[StepIdxCopy].FinishTime, Step.FinishTime));
+				}
+				else if (Step.State == JobStepState.Ready)
+				{
+					int StepIdxCopy = StepIdx;
+
+					Step.State = JobStepState.Skipped;
+					Updates.Add(UpdateBuilder.Set(x => x.Batches[BatchIdx].Steps[StepIdxCopy].State, Step.State));
+
+					Step.Outcome = JobStepOutcome.Failure;
+					Updates.Add(UpdateBuilder.Set(x => x.Batches[BatchIdx].Steps[StepIdxCopy].Outcome, Step.Outcome));
+				}
+			}
+
+			RefreshDependentJobSteps(JobDocument, Graph, Updates, Logger);
+			RefreshJobPriority(JobDocument, Updates);
+
+			return (Updates.Count == 0 || await TryUpdateAsync(JobDocument, UpdateBuilder.Combine(Updates)));
+		}
+
+		/// <summary>
+		/// Populate the list of batches for a job
+		/// </summary>
+		/// <param name="Job">The job to update</param>
+		/// <param name="Graph">The graph definition for this job</param>
+		/// <param name="Logger">Logger for output messages</param>
+		private static void CreateBatches(JobDocument Job, IGraph Graph, ILogger Logger)
+		{
+			UpdateBatches(Job, Graph, new List<UpdateDefinition<JobDocument>>(), Logger);
+		}
+
+		/// <summary>
+		/// Updates the list of batches for a job
+		/// </summary>
+		/// <param name="Job">Job to update</param>
+		/// <param name="Graph">Graph definition for this job</param>
+		/// <param name="Updates">List of updates for the job</param>
+		/// <param name="Logger">Logger for updates</param>
+		static void UpdateBatches(JobDocument Job, IGraph Graph, List<UpdateDefinition<JobDocument>> Updates, ILogger Logger)
+		{
+			UpdateDefinitionBuilder<JobDocument> UpdateBuilder = Builders<JobDocument>.Update;
+
+			// Update the list of batches
+			CreateOrUpdateBatches(Job, Graph);
+			Updates.Add(UpdateBuilder.Set(x => x.Batches, Job.Batches));
+			Updates.Add(UpdateBuilder.Set(x => x.NextSubResourceId, Job.NextSubResourceId));
+
+			// Update all the dependencies
+			RefreshDependentJobSteps(Job, Graph, new List<UpdateDefinition<JobDocument>>(), Logger);
+			RefreshJobPriority(Job, Updates);
+		}
+
+		/// <summary>
+		/// Update the jobsteps for the given node graph 
+		/// </summary>
+		/// <param name="Job">The job to update</param>
+		/// <param name="Graph">The graph for this job</param>
+		private static void CreateOrUpdateBatches(JobDocument Job, IGraph Graph)
+		{
+			// Find the priorities of each node, incorporating all the per-step overrides
+			Dictionary<INode, Priority> NodePriorities = new Dictionary<INode, Priority>();
+			foreach (INodeGroup Group in Graph.Groups)
+			{
+				foreach (INode Node in Group.Nodes)
+				{
+					NodePriorities[Node] = Node.Priority;
+				}
+			}
+			foreach (JobStepBatchDocument Batch in Job.Batches)
+			{
+				INodeGroup Group = Graph.Groups[Batch.GroupIdx];
+				foreach (JobStepDocument Step in Batch.Steps)
+				{
+					if (Step.Priority != null)
+					{
+						INode Node = Group.Nodes[Step.NodeIdx];
+						NodePriorities[Node] = Step.Priority.Value;
+					}
+				}
+			}
+
+			// Remove any steps and batches that haven't started yet
+			foreach (JobStepBatchDocument Batch in Job.Batches)
+			{
+				Batch.Steps.RemoveAll(x => x.State == JobStepState.Waiting || x.State == JobStepState.Ready);
+			}
+
+			// Remove any skipped nodes whose skipped state is no longer valid
+			HashSet<INode> FailedNodes = new HashSet<INode>();
+			foreach (JobStepBatchDocument Batch in Job.Batches)
+			{
+				INodeGroup Group = Graph.Groups[Batch.GroupIdx];
+				foreach (JobStepDocument Step in Batch.Steps)
+				{
+					INode Node = Group.Nodes[Step.NodeIdx];
+					if (Step.Retry)
+					{
+						FailedNodes.Remove(Node);
+					}
+					else if (Step.State == JobStepState.Skipped && (Node.InputDependencies.Any(x => FailedNodes.Contains(Graph.GetNode(x))) || !CanRetryNode(Job, Batch.GroupIdx, Step.NodeIdx)))
+					{
+						FailedNodes.Add(Node);
+					}
+					else if (Step.Outcome == JobStepOutcome.Failure)
+					{
+						FailedNodes.Add(Node);
+					}
+					else
+					{
+						FailedNodes.Remove(Node);
+					}
+				}
+				Batch.Steps.RemoveAll(x => x.State == JobStepState.Skipped && !FailedNodes.Contains(Group.Nodes[x.NodeIdx]));
+			}
+
+			// Remove any batches which are now empty
+			Job.Batches.RemoveAll(x => x.Steps.Count == 0 && x.Error == JobStepBatchError.None);
+
+			// Find all the targets in this job
+			HashSet<string> Targets = new HashSet<string>(StringComparer.OrdinalIgnoreCase);
+			if (Job.AbortedByUser == null)
+			{
+				foreach (string Argument in Job.Arguments)
+				{
+					if (Argument.StartsWith(IJob.TargetArgumentPrefix, StringComparison.OrdinalIgnoreCase))
+					{
+						Targets.UnionWith(Argument.Substring(IJob.TargetArgumentPrefix.Length).Split(';'));
+					}
+				}
+				Targets.Add(IJob.SetupNodeName);
+			}
+
+			// Add all the referenced aggregates
+			HashSet<INode> NewNodesToExecute = new HashSet<INode>();
+			foreach (IAggregate Aggregate in Graph.Aggregates)
+			{
+				if (Targets.Contains(Aggregate.Name))
+				{
+					NewNodesToExecute.UnionWith(Aggregate.Nodes.Select(x => Graph.GetNode(x)));
+				}
+			}
+
+			// Add any individual nodes
+			foreach (INode Node in Graph.Groups.SelectMany(x => x.Nodes))
+			{
+				if (Targets.Contains(Node.Name))
+				{
+					NewNodesToExecute.Add(Node);
+				}
+			}
+
+			// Also add any dependencies of these nodes
+			for (int GroupIdx = Graph.Groups.Count - 1; GroupIdx >= 0; GroupIdx--)
+			{
+				INodeGroup Group = Graph.Groups[GroupIdx];
+				for (int NodeIdx = Group.Nodes.Count - 1; NodeIdx >= 0; NodeIdx--)
+				{
+					INode Node = Group.Nodes[NodeIdx];
+					if (NewNodesToExecute.Contains(Node))
+					{
+						foreach (NodeRef Dependency in Node.InputDependencies)
+						{
+							NewNodesToExecute.Add(Graph.GetNode(Dependency));
+						}
+					}
+				}
+			}
+
+			// Cancel any batches which are still running but are no longer required
+			foreach (JobStepBatchDocument Batch in Job.Batches)
+			{
+				if (Batch.State == JobStepBatchState.Starting || Batch.State == JobStepBatchState.Running)
+				{
+					INodeGroup Group = Graph.Groups[Batch.GroupIdx];
+					if (!Batch.Steps.Any(x => NewNodesToExecute.Contains(Group.Nodes[x.NodeIdx])))
+					{
+						Batch.Error = JobStepBatchError.Cancelled;
+					}
+				}
+			}
+
+			// Remove all the nodes which have already succeeded
+			foreach (JobStepBatchDocument Batch in Job.Batches)
+			{
+				foreach (IJobStep Step in Batch.Steps)
+				{
+					if ((Step.State == JobStepState.Running && Step.RetryByUser == null) || (Step.State == JobStepState.Completed && Step.RetryByUser == null) || (Step.State == JobStepState.Aborted && Step.RetryByUser == null) || Step.State == JobStepState.Skipped)
+					{
+						NewNodesToExecute.Remove(Graph.Groups[Batch.GroupIdx].Nodes[Step.NodeIdx]);
+					}
+				}
+			}
+
+			// Re-add all the nodes that have input dependencies in the same group.
+			for (int GroupIdx = Graph.Groups.Count - 1; GroupIdx >= 0; GroupIdx--)
+			{
+				INodeGroup Group = Graph.Groups[GroupIdx];
+				for (int NodeIdx = Group.Nodes.Count - 1; NodeIdx >= 0; NodeIdx--)
+				{
+					INode Node = Group.Nodes[NodeIdx];
+					if (NewNodesToExecute.Contains(Node))
+					{
+						foreach (NodeRef Dependency in Node.InputDependencies)
+						{
+							if (Dependency.GroupIdx == GroupIdx)
+							{
+								NewNodesToExecute.Add(Group.Nodes[Dependency.NodeIdx]);
+							}
+						}
+					}
+				}
+			}
+
+			// Build a list of nodes which are currently set to be executed
+			HashSet<INode> ExistingNodesToExecute = new HashSet<INode>();
+			foreach (JobStepBatchDocument Batch in Job.Batches)
+			{
+				INodeGroup Group = Graph.Groups[Batch.GroupIdx];
+				foreach (IJobStep Step in Batch.Steps)
+				{
+					INode Node = Group.Nodes[Step.NodeIdx];
+					ExistingNodesToExecute.Add(Node);
+				}
+			}
+
+			// Figure out the existing batch for each group
+			JobStepBatchDocument?[] AppendToBatches = new JobStepBatchDocument?[Graph.Groups.Count];
+			foreach (JobStepBatchDocument Batch in Job.Batches)
+			{
+				if (Batch.CanBeAppendedTo())
+				{
+					INodeGroup Group = Graph.Groups[Batch.GroupIdx];
+					INode FirstNode = Group.Nodes[Batch.Steps[0].NodeIdx];
+					AppendToBatches[Batch.GroupIdx] = Batch;
+				}
+			}
+
+			// Invalidate all the entries for groups where we're too late to append new entries (ie. we need to execute an earlier node that wasn't executed previously)
+			for (int GroupIdx = 0; GroupIdx < Graph.Groups.Count; GroupIdx++)
+			{
+				INodeGroup Group = Graph.Groups[GroupIdx];
+				for (int NodeIdx = 0; NodeIdx < Group.Nodes.Count; NodeIdx++)
+				{
+					INode Node = Group.Nodes[NodeIdx];
+					if (NewNodesToExecute.Contains(Node) && !ExistingNodesToExecute.Contains(Node))
+					{
+						IJobStepBatch? Batch = AppendToBatches[GroupIdx];
+						if (Batch != null)
+						{
+							IJobStep LastStep = Batch.Steps[Batch.Steps.Count - 1];
+							if (NodeIdx < LastStep.NodeIdx)
+							{
+								AppendToBatches[GroupIdx] = null;
+							}
+						}
+					}
+				}
+			}
+
+			// Create all the new jobsteps
+			for (int GroupIdx = 0; GroupIdx < Graph.Groups.Count; GroupIdx++)
+			{
+				INodeGroup Group = Graph.Groups[GroupIdx];
+				for (int NodeIdx = 0; NodeIdx < Group.Nodes.Count; NodeIdx++)
+				{
+					INode Node = Group.Nodes[NodeIdx];
+					if (NewNodesToExecute.Contains(Node))
+					{
+						JobStepBatchDocument? Batch = AppendToBatches[GroupIdx];
+						if (Batch == null)
+						{
+							Job.NextSubResourceId = Job.NextSubResourceId.Next();
+
+							Batch = new JobStepBatchDocument(Job.NextSubResourceId, GroupIdx);
+							Job.Batches.Add(Batch);
+
+							AppendToBatches[GroupIdx] = Batch;
+						}
+						if (Batch.Steps.Count == 0 || NodeIdx > Batch.Steps[Batch.Steps.Count - 1].NodeIdx)
+						{
+							Job.NextSubResourceId = Job.NextSubResourceId.Next();
+
+							JobStepDocument Step = new JobStepDocument(Job.NextSubResourceId, NodeIdx);
+							Batch.Steps.Add(Step);
+						}
+					}
+				}
+			}
+
+			// Find the priority of each node, propagating dependencies from dependent nodes
+			for (int GroupIdx = 0; GroupIdx < Graph.Groups.Count; GroupIdx++)
+			{
+				INodeGroup Group = Graph.Groups[GroupIdx];
+				for (int NodeIdx = 0; NodeIdx < Group.Nodes.Count; NodeIdx++)
+				{
+					INode Node = Group.Nodes[NodeIdx];
+					Priority NodePriority = NodePriorities[Node];
+
+					foreach (NodeRef DependencyRef in Node.OrderDependencies)
+					{
+						INode Dependency = Graph.Groups[DependencyRef.GroupIdx].Nodes[DependencyRef.NodeIdx];
+						if (NodePriorities[Node] > NodePriority)
+						{
+							NodePriorities[Dependency] = NodePriority;
+						}
+					}
+				}
+			}
+			foreach (JobStepBatchDocument Batch in Job.Batches)
+			{
+				if (Batch.Steps.Count > 0)
+				{
+					INodeGroup Group = Graph.Groups[Batch.GroupIdx];
+					Priority NodePriority = Batch.Steps.Max(x => NodePriorities[Group.Nodes[x.NodeIdx]]);
+					Batch.SchedulePriority = ((int)Job.Priority * 10) + (int)NodePriority + 1; // Reserve '0' for none.
+				}
+			}
+
+			// Check we're not running a node which doesn't allow retries more than once
+			Dictionary<INode, int> NodeExecutionCount = new Dictionary<INode, int>();
+			foreach (IJobStepBatch Batch in Job.Batches)
+			{
+				INodeGroup Group = Graph.Groups[Batch.GroupIdx];
+				foreach (IJobStep Step in Batch.Steps)
+				{
+					INode Node = Group.Nodes[Step.NodeIdx];
+
+					int Count;
+					NodeExecutionCount.TryGetValue(Node, out Count);
+
+					if (!Node.AllowRetry && Count > 0)
+					{
+						throw new RetryNotAllowedException(Node.Name);
+					}
+
+					NodeExecutionCount[Node] = Count + 1;
+				}
+			}
+		}
+
+		/// <summary>
+		/// Tests whether a node can be retried again
+		/// </summary>
+		static bool CanRetryNode(JobDocument Job, int GroupIdx, int NodeIdx)
+		{
+			return Job.RetriedNodes == null || Job.RetriedNodes.Count(x => x.GroupIdx == GroupIdx && x.NodeIdx == NodeIdx) < MaxRetries;
+		}
+
+		/// <summary>
+		/// Gets the scheduling priority of this job
+		/// </summary>
+		/// <param name="Job">Job to consider</param>
+		public static int GetSchedulePriority(IJob Job)
+		{
+			int NewSchedulePriority = 0;
+			foreach (IJobStepBatch Batch in Job.Batches)
+			{
+				if (Batch.State == JobStepBatchState.Ready)
+				{
+					NewSchedulePriority = Math.Max(Batch.SchedulePriority, NewSchedulePriority);
+				}
+			}
+			return NewSchedulePriority;
+		}
+
+		/// <summary>
+		/// Update the state of any jobsteps that are dependent on other jobsteps (eg. transition them from waiting to ready based on other steps completing)
+		/// </summary>
+		/// <param name="Job">The job to update</param>
+		/// <param name="Graph">The graph for this job</param>
+		/// <param name="Updates">List of updates to the job</param>
+		/// <param name="Logger">Logger instance</param>
+		static void RefreshDependentJobSteps(JobDocument Job, IGraph Graph, List<UpdateDefinition<JobDocument>> Updates, ILogger Logger)
+		{
+			// Update the batches
+			UpdateDefinitionBuilder<JobDocument> UpdateBuilder = Builders<JobDocument>.Update;
+			if (Job.Batches != null)
+			{
+				Dictionary<INode, IJobStep> StepForNode = new Dictionary<INode, IJobStep>();
+				for (int LoopBatchIdx = 0; LoopBatchIdx < Job.Batches.Count; LoopBatchIdx++)
+				{
+					int BatchIdx = LoopBatchIdx; // For lambda capture
+					JobStepBatchDocument Batch = Job.Batches[BatchIdx];
+
+					for (int LoopStepIdx = 0; LoopStepIdx < Batch.Steps.Count; LoopStepIdx++)
+					{
+						int StepIdx = LoopStepIdx; // For lambda capture
+						JobStepDocument Step = Batch.Steps[StepIdx];
+
+						JobStepState NewState = Step.State;
+						JobStepOutcome NewOutcome = Step.Outcome;
+
+						INode Node = Graph.Groups[Batch.GroupIdx].Nodes[Step.NodeIdx];
+						if (NewState == JobStepState.Waiting)
+						{
+							List<IJobStep> Steps = GetDependentSteps(Graph, Node, StepForNode);
+							if (Steps.Any(x => x.IsFailedOrSkipped()))
+							{
+								NewState = JobStepState.Skipped;
+								NewOutcome = JobStepOutcome.Failure;
+							}
+							else if (!Steps.Any(x => x.IsPending()))
+							{
+								Logger.LogDebug("Transitioning job {JobId}, batch {BatchId}, step {StepId} to ready state ({Dependencies})", Job.Id, Batch.Id, Step.Id, String.Join(", ", Steps.Select(x => x.Id.ToString())));
+								NewState = JobStepState.Ready;
+							}
+						}
+
+						if (NewState != Step.State)
+						{
+							Step.State = NewState;
+							Updates.Add(UpdateBuilder.Set(x => x.Batches[BatchIdx].Steps[StepIdx].State, NewState));
+						}
+
+						if (NewOutcome != Step.Outcome)
+						{
+							Step.Outcome = NewOutcome;
+							Updates.Add(UpdateBuilder.Set(x => x.Batches[BatchIdx].Steps[StepIdx].Outcome, NewOutcome));
+						}
+
+						StepForNode[Node] = Step;
+					}
+
+					if (Batch.State == JobStepBatchState.Waiting || Batch.State == JobStepBatchState.Ready)
+					{
+						DateTime? NewReadyTime;
+						JobStepBatchState NewState = GetBatchState(Job, Graph, Batch, StepForNode, out NewReadyTime);
+						if (Batch.State != NewState)
+						{
+							Batch.State = NewState;
+							Updates.Add(UpdateBuilder.Set(x => x.Batches[BatchIdx].State, Batch.State));
+						}
+						if (Batch.ReadyTime != NewReadyTime)
+						{
+							Batch.ReadyTime = NewReadyTime;
+							Updates.Add(UpdateBuilder.Set(x => x.Batches[BatchIdx].ReadyTime, Batch.ReadyTime));
+						}
+					}
+				}
+			}
+		}
+
+		/// <summary>
+		/// Updates the schedule priority of a job
+		/// </summary>
+		/// <param name="Job"></param>
+		/// <param name="Updates"></param>
+		static void RefreshJobPriority(JobDocument Job, List<UpdateDefinition<JobDocument>> Updates)
+		{
+			// Update the weighted priority for the job
+			int NewSchedulePriority = GetSchedulePriority(Job);
+			if (Job.SchedulePriority != NewSchedulePriority)
+			{
+				Job.SchedulePriority = NewSchedulePriority;
+				Updates.Add(Builders<JobDocument>.Update.Set(x => x.SchedulePriority, NewSchedulePriority));
+			}
+		}
+
+		/// <summary>
+		/// Gets the steps that a node depends on
+		/// </summary>
+		/// <param name="Graph">The graph for this job</param>
+		/// <param name="Node">The node to test</param>
+		/// <param name="StepForNode">Map of node to step</param>
+		/// <returns></returns>
+		static List<IJobStep> GetDependentSteps(IGraph Graph, INode Node, Dictionary<INode, IJobStep> StepForNode)
+		{
+			List<IJobStep> Steps = new List<IJobStep>();
+			foreach (NodeRef OrderDependencyRef in Node.OrderDependencies)
+			{
+				IJobStep? Step;
+				if (StepForNode.TryGetValue(Graph.GetNode(OrderDependencyRef), out Step))
+				{
+					Steps.Add(Step);
+				}
+			}
+			return Steps;
+		}
+
+		/// <summary>
+		/// Gets the new state for a batch
+		/// </summary>
+		/// <param name="Job">The job being executed</param>
+		/// <param name="Graph">Graph for the job</param>
+		/// <param name="Batch">List of nodes in the job</param>
+		/// <param name="StepForNode">Array mapping each index to the appropriate step for that node</param>
+		/// <param name="OutReadyTimeUtc">Receives the time at which the batch was ready to execute</param>
+		/// <returns>True if the batch is ready, false otherwise</returns>
+		static JobStepBatchState GetBatchState(IJob Job, IGraph Graph, IJobStepBatch Batch, Dictionary<INode, IJobStep> StepForNode, out DateTime? OutReadyTimeUtc)
+		{
+			// Check if the batch is already complete
+			if (Batch.Steps.All(x => x.State == JobStepState.Skipped || x.State == JobStepState.Completed || x.State == JobStepState.Aborted))
+			{
+				OutReadyTimeUtc = Batch.ReadyTimeUtc;
+				return JobStepBatchState.Complete;
+			}
+
+			// Get the dependencies for this batch to start. Some steps may be "after" dependencies that are optional parts of the graph.
+			List<INode> NodeDependencies = Batch.GetStartDependencies(Graph.Groups).ToList();
+
+			// Check if we're still waiting on anything
+			DateTime ReadyTimeUtc = Job.CreateTimeUtc;
+			foreach (INode NodeDependency in NodeDependencies)
+			{
+				IJobStep? StepDependency;
+				if (StepForNode.TryGetValue(NodeDependency, out StepDependency))
+				{
+					if (StepDependency.State != JobStepState.Completed && StepDependency.State != JobStepState.Skipped && StepDependency.State != JobStepState.Aborted)
+					{
+						OutReadyTimeUtc = null;
+						return JobStepBatchState.Waiting;
+					}
+
+					if (StepDependency.FinishTimeUtc != null && StepDependency.FinishTimeUtc.Value > ReadyTimeUtc)
+					{
+						ReadyTimeUtc = StepDependency.FinishTimeUtc.Value;
+					}
+				}
+			}
+
+			// Otherwise return the ready state
+			OutReadyTimeUtc = ReadyTimeUtc;
+			return JobStepBatchState.Ready;
+		}
+
+		/// <inheritdoc/>
+		public async Task UpgradeDocumentsAsync()
+		{
+			IAsyncCursor<JobDocument> Cursor = await Jobs.Find(Builders<JobDocument>.Filter.Eq(x => x.UpdateTimeUtc, null)).ToCursorAsync();
+
+			int NumUpdated = 0;
+			while (await Cursor.MoveNextAsync())
+			{
+				foreach (JobDocument Document in Cursor.Current)
+				{
+					UpdateDefinition<JobDocument> Update = Builders<JobDocument>.Update.Set(x => x.CreateTimeUtc, ((IJob)Document).CreateTimeUtc).Set(x => x.UpdateTimeUtc, ((IJob)Document).UpdateTimeUtc);
+					await Jobs.UpdateOneAsync(Builders<JobDocument>.Filter.Eq(x => x.Id, Document.Id), Update);
+					NumUpdated++;
+				}
+				Logger.LogInformation("Updated {NumDocuments} documents", NumUpdated);
+			}
+		}
+	}
+}
