Horde: Move jobs, graphs, issues and users into subfolders.

diff --git a/Engine/Source/Programs/Horde/HordeServer/Jobs/Impl/JobService.cs b/Engine/Source/Programs/Horde/HordeServer/Jobs/Impl/JobService.cs
new file mode 100644
--- /dev/null
+++ b/Engine/Source/Programs/Horde/HordeServer/Jobs/Impl/JobService.cs
@@ -0,0 +1,1260 @@
+// Copyright Epic Games, Inc. All Rights Reserved.
+
+using EpicGames.Core;
+using HordeServer.Api;
+using HordeServer.Collections;
+using HordeCommon;
+using HordeServer.Models;
+using HordeCommon.Rpc;
+using HordeServer.Utilities;
+using Microsoft.AspNetCore.Http;
+using Microsoft.Extensions.Logging;
+using Microsoft.Extensions.Options;
+using MongoDB.Bson;
+using MongoDB.Bson.Serialization.Serializers;
+using MongoDB.Driver;
+using System;
+using System.Collections.Generic;
+using System.ComponentModel;
+using System.Diagnostics;
+using System.Diagnostics.CodeAnalysis;
+using System.Linq;
+using System.Security.Claims;
+using System.Text.RegularExpressions;
+using System.Threading;
+using System.Threading.Tasks;
+using System.Globalization;
+using HordeServer.Tasks.Impl;
+using OpenTracing.Util;
+using OpenTracing;
+
+namespace HordeServer.Services
+{
+	using StreamId = StringId<IStream>;
+	using TemplateRefId = StringId<TemplateRef>;
+
+	/// <summary>
+	/// Exception thrown when attempting to retry executing a node that does not allow retries
+	/// </summary>
+	public class RetryNotAllowedException : Exception
+	{
+		/// <summary>
+		/// The node name
+		/// </summary>
+		public string NodeName { get; private set; }
+
+		/// <summary>
+		/// Constructor
+		/// </summary>
+		/// <param name="NodeName">Name of the node that does not permit retries</param>
+		public RetryNotAllowedException(string NodeName)
+			: base($"Node '{NodeName}' does not permit retries")
+		{
+			this.NodeName = NodeName;
+		}
+	}
+
+	/// <summary>
+	/// Cache of information about job ACLs
+	/// </summary>
+	public class JobPermissionsCache : StreamPermissionsCache
+	{
+		/// <summary>
+		/// Map of job id to permissions for that job
+		/// </summary>
+		public Dictionary<ObjectId, IJobPermissions?> Jobs { get; set; } = new Dictionary<ObjectId, IJobPermissions?>();
+	}
+
+	/// <summary>
+	/// Wraps funtionality for manipulating jobs, jobsteps, and jobstep runs
+	/// </summary>
+	[SuppressMessage("Compiler", "CA1054:URI parameters should not be strings")]
+	public class JobService
+	{
+		/// <summary>
+		/// Collection of job documents
+		/// </summary>
+		IJobCollection Jobs;
+
+		/// <summary>
+		/// Collection of graph documents
+		/// </summary>
+		IGraphCollection Graphs;
+
+		/// <summary>
+		/// Collection of agent documents
+		/// </summary>
+		IAgentCollection Agents;
+
+		/// <summary>
+		/// Collection of jobstepref documents
+		/// </summary>
+		IJobStepRefCollection JobStepRefs;
+
+		/// <summary>
+		/// Collection of job timing documents
+		/// </summary>
+		IJobTimingCollection JobTimings;
+
+		/// <summary>
+		/// Collection of users
+		/// </summary>
+		IUserCollection UserCollection;
+
+		/// <summary>
+		/// Collection of counter documents
+		/// </summary>
+		ICounterCollection CounterCollection;
+
+		/// <summary>
+		/// Collection of notification triggers
+		/// </summary>
+		INotificationTriggerCollection TriggerCollection;
+
+		/// <summary>
+		/// The singleton instance of the queue service
+		/// </summary>
+		JobTaskSource JobTaskSource;
+
+		/// <summary>
+		/// Singleton instance of the stream service
+		/// </summary>
+		StreamService StreamService;
+
+		/// <summary>
+		/// Singleton instance of the template service
+		/// </summary>
+		TemplateService TemplateService;
+
+		/// <summary>
+		/// Singleton instance of the issue service
+		/// </summary>
+		IIssueService? IssueService;
+
+		/// <summary>
+		/// The perforce service
+		/// </summary>
+		IPerforceService PerforceService;
+
+		/// <summary>
+		/// The server settings
+		/// </summary>
+		IOptionsMonitor<ServerSettings> Settings;
+
+		/// <summary>
+		/// Log output device
+		/// </summary>
+		ILogger Logger;
+
+		/// <summary>
+		/// Delegate for label update events
+		/// </summary>
+		public delegate void LabelUpdateEvent(IJob Job, IReadOnlyList<(LabelState, LabelOutcome)> OldLabelStates, IReadOnlyList<(LabelState, LabelOutcome)> NewLabelStates);
+
+		/// <summary>
+		/// Event triggered when a label state changes
+		/// </summary>
+		public event LabelUpdateEvent? OnLabelUpdate;
+
+		/// <summary>
+		/// Delegate for job step complete events
+		/// </summary>
+		public delegate void JobStepCompleteEvent(IJob Job, IGraph Graph, SubResourceId BatchId, SubResourceId StepId);
+
+		/// <summary>
+		/// Event triggered when a job step completes
+		/// </summary>
+		public event JobStepCompleteEvent? OnJobStepComplete;
+
+		/// <summary>
+		/// Constructor
+		/// </summary>
+		/// <param name="Jobs">The jobs collection</param>
+		/// <param name="Graphs">The graphs collection</param>
+		/// <param name="Agents">The agents collection</param>
+		/// <param name="JobStepRefs">The jobsteprefs collection</param>
+		/// <param name="JobTimings">The job timing document collection</param>
+		/// <param name="UserCollection">User profiles</param>
+		/// <param name="CounterCollection">Counter collection</param>
+		/// <param name="TriggerCollection">Trigger collection</param>
+		/// <param name="JobTaskSource">The queue service</param>
+		/// <param name="StreamService">The stream service</param>
+		/// <param name="TemplateService">The template service</param>
+		/// <param name="IssueService">The issue service</param>
+		/// <param name="PerforceService">The perforce service</param>
+		/// <param name="Settings">Settings instance</param>
+		/// <param name="Logger">Log output</param>
+		public JobService(IJobCollection Jobs, IGraphCollection Graphs, IAgentCollection Agents, IJobStepRefCollection JobStepRefs, IJobTimingCollection JobTimings, IUserCollection UserCollection, ICounterCollection CounterCollection, INotificationTriggerCollection TriggerCollection, JobTaskSource JobTaskSource, StreamService StreamService, TemplateService TemplateService, IIssueService IssueService, IPerforceService PerforceService, IOptionsMonitor<ServerSettings> Settings, ILogger<JobService> Logger)
+		{
+			this.Jobs = Jobs;
+			this.Graphs = Graphs;
+			this.Agents = Agents;
+			this.JobStepRefs = JobStepRefs;
+			this.JobTimings = JobTimings;
+			this.UserCollection = UserCollection;
+			this.CounterCollection = CounterCollection;
+			this.TriggerCollection = TriggerCollection;
+			this.JobTaskSource = JobTaskSource;
+			this.StreamService = StreamService;
+			this.TemplateService = TemplateService;
+			this.IssueService = IssueService;
+			this.PerforceService = PerforceService;
+			this.Settings = Settings;
+			this.Logger = Logger;
+		}
+
+#pragma warning disable CA1801
+		static bool ShouldClonePreflightChange(StreamId StreamId)
+		{
+			return false; //			return StreamId == new StreamId("ue5-main");
+		}
+#pragma warning restore CA1801
+
+		/// <summary>
+		/// Creates a new job
+		/// </summary>
+		/// <param name="JobId">A requested job id</param>
+		/// <param name="Stream">The stream that this job belongs to</param>
+		/// <param name="TemplateRefId">Name of the template ref</param>
+		/// <param name="TemplateHash">Template for this job</param>
+		/// <param name="Graph">The graph for the new job</param>
+		/// <param name="Name">Name of the job</param>
+		/// <param name="Change">The change to build</param>
+		/// <param name="CodeChange">The corresponding code changelist</param>
+		/// <param name="PreflightChange">Optional changelist to preflight</param>
+		/// <param name="ClonedPreflightChange">Duplicated preflight change</param>
+		/// <param name="StartedByUserId">Id of the user that started the job</param>
+		/// <param name="StartedByUserName">User that started the job</param>
+		/// <param name="Priority">Priority of the job</param>
+		/// <param name="AutoSubmit">Whether to automatically submit the preflighted change on completion</param>
+		/// <param name="UpdateIssues">Whether to update issues when this job completes</param>
+		/// <param name="JobTriggers">List of downstream job triggers</param>
+		/// <param name="ShowUgsBadges">Whether to show badges in UGS for this job</param>
+		/// <param name="ShowUgsAlerts">Whether to show alerts in UGS for this job</param>
+		/// <param name="NotificationChannel">Notification Channel for this job</param>
+		/// <param name="NotificationChannelFilter">Notification Channel filter for this job</param>
+		/// <param name="HelixSwarmCallbackUrl">Helix Swarm callback URL for review, if any</param>
+		/// <param name="Counters">Counters for the job</param>
+		/// <param name="Arguments">Arguments for the job</param>
+		/// <returns>Unique id representing the job</returns>
+		public async Task<IJob> CreateJobAsync(ObjectId? JobId, IStream Stream, TemplateRefId TemplateRefId, ContentHash TemplateHash, IGraph Graph, string Name, int Change, int CodeChange, int? PreflightChange, int? ClonedPreflightChange, ObjectId? StartedByUserId, string? StartedByUserName, Priority? Priority, bool? AutoSubmit, bool? UpdateIssues, List<ChainedJobTemplate>? JobTriggers, bool ShowUgsBadges, bool ShowUgsAlerts, string? NotificationChannel, string? NotificationChannelFilter, string? HelixSwarmCallbackUrl, IReadOnlyList<ITemplateCounter> Counters, IReadOnlyList<string> Arguments)
+		{
+			ObjectId JobIdValue = JobId ?? ObjectId.GenerateNewId();
+			using IDisposable Scope = Logger.BeginScope("CreateJobAsync({JobId})", JobIdValue);
+
+			if (PreflightChange != null && ShouldClonePreflightChange(Stream.Id))
+			{
+				ClonedPreflightChange = await CloneShelvedChangeAsync(Stream.ClusterName, ClonedPreflightChange ?? PreflightChange.Value);
+			}
+
+			Logger.LogInformation("Creating job at CL {Change}, code CL {CodeChange}, preflight CL {PreflightChange}, cloned CL {ClonedPreflightChange}", Change, CodeChange, PreflightChange, ClonedPreflightChange);
+
+			Dictionary<string, string> Properties = new Dictionary<string, string>();
+			Properties["Change"] = Change.ToString(CultureInfo.InvariantCulture);
+			Properties["CodeChange"] = CodeChange.ToString(CultureInfo.InvariantCulture);
+			Properties["PreflightChange"] = PreflightChange?.ToString(CultureInfo.InvariantCulture) ?? String.Empty;
+			Properties["ClonedPreflightChange"] = ClonedPreflightChange?.ToString(CultureInfo.InvariantCulture) ?? String.Empty;
+			Properties["StreamId"] = Stream.Id.ToString();
+			Properties["TemplateId"] = TemplateRefId.ToString();
+			Properties["JobId"] = JobIdValue.ToString();
+
+			foreach (ITemplateCounter Counter in Counters)
+			{
+				string CounterName = StringUtils.ExpandProperties(Counter.CounterName, Properties);
+				int Value = await CounterCollection.IncrementAsync(CounterName);
+				Properties[Counter.PropertyName] = Value.ToString(CultureInfo.InvariantCulture);
+				Properties[$"{Counter.PropertyName}.Suffix"] = (Value == 1) ? String.Empty : $" ({Value})";
+			}
+
+			List<string> ExpandedArguments = new List<string>();
+			if (Arguments != null)
+			{
+				foreach (string Argument in Arguments)
+				{
+					string ExpandedArgument = StringUtils.ExpandProperties(Argument, Properties);
+					ExpandedArguments.Add(ExpandedArgument);
+				}
+			}
+
+			Name = StringUtils.ExpandProperties(Name, Properties);
+
+			IJob NewJob = await Jobs.AddAsync(JobIdValue, Stream.Id, TemplateRefId, TemplateHash, Graph, Name, Change, CodeChange, PreflightChange, ClonedPreflightChange, StartedByUserId, StartedByUserName, Priority, AutoSubmit, UpdateIssues, JobTriggers, ShowUgsBadges, ShowUgsAlerts, NotificationChannel, NotificationChannelFilter, HelixSwarmCallbackUrl, ExpandedArguments);
+			JobTaskSource.UpdateQueuedJob(NewJob, Graph);
+
+			await JobTaskSource.UpdateUgsBadges(NewJob, Graph, new List<(LabelState, LabelOutcome)>());
+
+			if (StartedByUserId != null)
+			{
+				await UserCollection.UpdateSettingsAsync(StartedByUserId.Value, AddPinnedJobIds: new[] { NewJob.Id });
+			}
+
+			await AbortAnyDuplicateJobs(NewJob);
+
+			return NewJob;
+		}
+
+		private async Task AbortAnyDuplicateJobs(IJob NewJob)
+		{
+			List<IJob> JobsToAbort = new List<IJob>();
+			if (NewJob.PreflightChange > 0)
+			{
+				JobsToAbort = await Jobs.FindAsync(PreflightChange: NewJob.PreflightChange);
+			}
+
+			foreach (IJob Job in JobsToAbort)
+			{
+				if (Job.GetState() == JobState.Complete) continue;
+				if (Job.Id == NewJob.Id) continue; // Don't remove the new job
+				if (Job.TemplateId != NewJob.TemplateId) continue;
+				if (Job.TemplateHash != NewJob.TemplateHash) continue;
+				if (Job.StartedByUser != NewJob.StartedByUser) continue;
+				if (string.Join(",", Job.Arguments) != string.Join(",", NewJob.Arguments)) continue;
+
+				bool bUpdateSuccessful = await UpdateJobAsync(Job, null, null, null, "horde.duplicated.by.newer.CL", null, null, null);
+				if (!bUpdateSuccessful)
+				{
+					Logger.LogError("Failed marking duplicate job as aborted! Job ID: {JobId}", Job.Id);
+				}
+			}
+		}
+
+		/// <summary>
+		/// Deletes a job
+		/// </summary>
+		/// <param name="Job">The job to delete</param>
+		public async Task<bool> DeleteJobAsync(IJob Job)
+		{
+			using IDisposable Scope = Logger.BeginScope("DeleteJobAsync({JobId})", Job.Id);
+
+			// Delete the job
+			while (!await Jobs.RemoveAsync(Job))
+			{
+				IJob? NewJob = await Jobs.GetAsync(Job.Id);
+				if (NewJob == null)
+				{
+					return false;
+				}
+				Job = NewJob;
+			}
+
+			// Remove all the triggers from it
+			List<ObjectId> TriggerIds = new List<ObjectId>();
+			if (Job.NotificationTriggerId != null)
+			{
+				TriggerIds.Add(Job.NotificationTriggerId.Value);
+			}
+			foreach (IJobStep Step in Job.Batches.SelectMany(x => x.Steps))
+			{
+				if (Step.NotificationTriggerId != null)
+				{
+					TriggerIds.Add(Step.NotificationTriggerId.Value);
+				}
+			}
+			await TriggerCollection.DeleteAsync(TriggerIds);
+
+			return true;
+		}
+
+		/// <summary>
+		/// Delete all the jobs for a stream
+		/// </summary>
+		/// <param name="StreamId">Unique id of the stream</param>
+		/// <returns>Async task</returns>
+		public async Task DeleteJobsForStreamAsync(StreamId StreamId)
+		{
+			await Jobs.RemoveStreamAsync(StreamId);
+		}
+
+		/// <summary>
+		/// Updates a new job
+		/// </summary>
+		/// <param name="Job">The job document to update</param>
+		/// <param name="Name">Name of the job</param>
+		/// <param name="Priority">Priority of the job</param>
+		/// <param name="AutoSubmit">Whether to automatically submit the preflighted change on completion</param>
+		/// <param name="AbortedByUser">Name of the user that aborted this job</param>
+		/// <param name="OnCompleteTriggerId">Object id for a notification trigger</param>
+		/// <param name="Reports">New reports to add</param>
+		/// <param name="Arguments">New arguments for the job</param>
+		/// <param name="LabelIdxToTriggerId">New trigger ID for a label in the job</param>
+		public async Task<bool> UpdateJobAsync(IJob Job, string? Name = null, Priority? Priority = null, bool? AutoSubmit = null, string? AbortedByUser = null, ObjectId? OnCompleteTriggerId = null, List<Report>? Reports = null, List<string>? Arguments = null, KeyValuePair<int, ObjectId>? LabelIdxToTriggerId = null)
+		{
+			using IDisposable Scope = Logger.BeginScope("UpdateJobAsync({JobId})", Job.Id);
+			for (; ; )
+			{
+				IGraph Graph = await GetGraphAsync(Job);
+
+				// Capture the previous label states
+				IReadOnlyList<(LabelState, LabelOutcome)> OldLabelStates = Job.GetLabelStates(Graph);
+
+				// Update the new list of job steps
+				if (await Jobs.TryUpdateJobAsync(Job, Graph, Name, Priority, AutoSubmit, null, null, AbortedByUser, OnCompleteTriggerId, Reports, Arguments, LabelIdxToTriggerId))
+				{
+					// Update any badges that have been modified
+					await JobTaskSource.UpdateUgsBadges(Job, Graph, OldLabelStates);
+
+					// Cancel any leases which are no longer required
+					foreach (IJobStepBatch Batch in Job.Batches)
+					{
+						if (Batch.Error == JobStepBatchError.Cancelled && (Batch.State == JobStepBatchState.Starting || Batch.State == JobStepBatchState.Running) && Batch.AgentId != null && Batch.LeaseId != null)
+						{
+							await CancelLeaseAsync(Batch.AgentId.Value, Batch.LeaseId.Value);
+						}
+					}
+					return true;
+				}
+
+				// Get the new job state
+				IJob? NewJob = await GetJobAsync(Job.Id);
+				if (NewJob == null)
+				{
+					return false;
+				}
+				Job = NewJob;
+			}
+		}
+
+		/// <summary>
+		/// Cancel an active lease
+		/// </summary>
+		/// <param name="AgentId">The agent to retreive</param>
+		/// <param name="LeaseId">The lease id to update</param>
+		/// <returns></returns>
+		async Task CancelLeaseAsync(AgentId AgentId, ObjectId LeaseId)
+		{
+			for (; ; )
+			{
+				IAgent? Agent = await Agents.GetAsync(AgentId);
+				if (Agent == null)
+				{
+					break;
+				}
+
+				int Index = 0;
+				while (Index < Agent.Leases.Count && Agent.Leases[Index].Id != LeaseId)
+				{
+					Index++;
+				}
+				if (Index == Agent.Leases.Count)
+				{
+					break;
+				}
+
+				AgentLease Lease = Agent.Leases[Index];
+				if (Lease.State != LeaseState.Active)
+				{
+					break;
+				}
+
+				IAgent? NewAgent = await Agents.TryCancelLeaseAsync(Agent, Index);
+				if (NewAgent != null)
+				{
+					JobTaskSource.CancelLongPollForAgent(Agent.Id);
+					break;
+				}
+			}
+		}
+
+		/// <summary>
+		/// Gets a job with the given unique id
+		/// </summary>
+		/// <param name="JobId">Job id to search for</param>
+		/// <returns>Information about the given job</returns>
+		public Task<IJob?> GetJobAsync(ObjectId JobId)
+		{
+			return Jobs.GetAsync(JobId);
+		}
+
+		/// <summary>
+		/// Gets the graph for a job
+		/// </summary>
+		/// <param name="Job">Job to retrieve the graph for</param>
+		/// <returns>The graph for this job</returns>
+		public Task<IGraph> GetGraphAsync(IJob Job)
+		{
+			return Graphs.GetAsync(Job.GraphHash);
+		}
+
+		/// <summary>
+		/// Gets a job's permissions info by ID
+		/// </summary>
+		/// <param name="JobId">Unique id of the job</param>
+		/// <returns>The job document</returns>
+		public Task<IJobPermissions?> GetJobPermissionsAsync(ObjectId JobId)
+		{
+			return Jobs.GetPermissionsAsync(JobId);
+		}
+
+		/// <summary>
+		/// Searches for jobs matching the given criteria
+		/// </summary>
+		/// <param name="JobIds">List of job ids to return</param>
+		/// <param name="StreamId">The stream containing the job</param>
+		/// <param name="Name">Name of the job</param>
+		/// <param name="Templates">Templates to look for</param>
+		/// <param name="MinChange">The minimum changelist number</param>
+		/// <param name="MaxChange">The maximum changelist number</param>		
+		/// <param name="PreflightChange">The preflight change to look for</param>
+		/// <param name="PreflightStartedByUser">User for which to include preflight jobs</param>		
+		/// <param name="StartedByUser">User for which to include jobs</param>
+		/// <param name="MinCreateTime">The minimum creation time</param>
+		/// <param name="MaxCreateTime">The maximum creation time</param>
+		/// <param name="Target">The target to query</param>
+		/// <param name="State">State to query</param>
+		/// <param name="Outcome">Outcomes to return</param>
+		/// <param name="ModifiedBefore">Filter the results by last modified time</param>
+		/// <param name="ModifiedAfter">Filter the results by last modified time</param>
+		/// <param name="Index">Index of the first result to return</param>
+		/// <param name="Count">Number of results to return</param>
+		/// <returns>List of jobs matching the given criteria</returns>
+		public async Task<List<IJob>> FindJobsAsync(ObjectId[]? JobIds = null, StreamId? StreamId = null, string? Name = null, TemplateRefId[]? Templates = null, int? MinChange = null, int? MaxChange = null, int? PreflightChange = null, ObjectId ? PreflightStartedByUser = null, ObjectId? StartedByUser = null, DateTimeOffset ? MinCreateTime = null, DateTimeOffset? MaxCreateTime = null, string? Target = null, JobStepState[]? State = null, JobStepOutcome[]? Outcome = null, DateTimeOffset? ModifiedBefore = null, DateTimeOffset? ModifiedAfter = null, int? Index = null, int? Count = null)
+		{
+			if (Target == null && (State == null || State.Length == 0) && (Outcome == null || Outcome.Length == 0))
+			{
+				return await Jobs.FindAsync(JobIds, StreamId, Name, Templates, MinChange, MaxChange, PreflightChange, PreflightStartedByUser, StartedByUser, MinCreateTime, MaxCreateTime, ModifiedBefore, ModifiedAfter, Index, Count);
+			}
+			else
+			{
+				List<IJob> Results = new List<IJob>();
+				Logger.LogInformation("Performing scan for job with ");
+
+				int MaxCount = (Count ?? 1);
+				while (Results.Count < MaxCount)
+				{
+					List<IJob> ScanJobs = await Jobs.FindAsync(JobIds, StreamId, Name, Templates, MinChange, MaxChange, PreflightChange, PreflightStartedByUser, StartedByUser, MinCreateTime, MaxCreateTime, ModifiedBefore, ModifiedAfter, 0, 5);
+					if (ScanJobs.Count == 0)
+					{
+						break;
+					}
+
+					foreach (IJob Job in ScanJobs.OrderByDescending(x => x.Change))
+					{
+						(JobStepState, JobStepOutcome)? Result;
+						if (Target == null)
+						{
+							Result = Job.GetTargetState();
+						}
+						else
+						{
+							Result = Job.GetTargetState(await GetGraphAsync(Job), Target);
+						}
+
+						if (Result != null)
+						{
+							(JobStepState JobState, JobStepOutcome JobOutcome) = Result.Value;
+							if ((State == null || State.Length == 0 || State.Contains(JobState)) && (Outcome == null || Outcome.Length == 0 || Outcome.Contains(JobOutcome)))
+							{
+								Results.Add(Job);
+								if (Results.Count == MaxCount)
+								{
+									break;
+								}
+							}
+						}
+					}
+
+					MaxChange = ScanJobs.Min(x => x.Change) - 1;
+				}
+
+				return Results;
+			}
+		}
+
+		/// <summary>
+		/// Attempts to update the node groups to be executed for a job. Fails if another write happens in the meantime.
+		/// </summary>
+		/// <param name="Job">The job to update</param>
+		/// <param name="NewGraph">New graph for this job</param>
+		/// <returns>True if the groups were updated to the given list. False if another write happened first.</returns>
+		public async Task<bool> TryUpdateGraphAsync(IJob Job, IGraph NewGraph)
+		{
+			using IDisposable Scope = Logger.BeginScope("TryUpdateGraphAsync({JobId})", Job.Id);
+
+			IReadOnlyList<(LabelState, LabelOutcome)> OldLabelStates = Job.GetLabelStates(NewGraph);
+			if (await Jobs.TryUpdateGraphAsync(Job, NewGraph))
+			{
+				await JobTaskSource.UpdateUgsBadges(Job, NewGraph, OldLabelStates);
+
+				JobTaskSource.UpdateQueuedJob(Job, NewGraph);
+				return true;
+			}
+			return false;
+		}
+
+		/// <summary>
+		/// Gets the timing info for a particular job.
+		/// </summary>
+		/// <param name="Job">The job to get timing info for</param>
+		/// <returns>Timing info for the given job</returns>
+		public async Task<IJobTiming> GetJobTimingAsync(IJob Job)
+		{
+			using IDisposable Scope = Logger.BeginScope("GetJobTimingAsync({JobId})", Job.Id);
+
+			IGraph Graph = await Graphs.GetAsync(Job.GraphHash);
+
+			Dictionary<string, JobStepTimingData> CachedNewSteps = new Dictionary<string, JobStepTimingData>();
+			for (; ; )
+			{
+				// Try to get the current timing document
+				IJobTiming? JobTiming = await JobTimings.TryGetAsync(Job.Id);
+
+				// Calculate timing info for any missing steps
+				List<JobStepTimingData> NewSteps = new List<JobStepTimingData>();
+				foreach (IJobStepBatch Batch in Job.Batches)
+				{
+					INodeGroup Group = Graph.Groups[Batch.GroupIdx];
+					foreach (IJobStep Step in Batch.Steps)
+					{
+						INode Node = Group.Nodes[Step.NodeIdx];
+						if (JobTiming == null || !JobTiming.TryGetStepTiming(Node.Name, out _))
+						{
+							JobStepTimingData? StepTimingData;
+							if (!CachedNewSteps.TryGetValue(Node.Name, out StepTimingData))
+							{
+								StepTimingData = await GetStepTimingInfo(Job.StreamId, Job.TemplateId, Node.Name, Job.Change);
+							}
+							NewSteps.Add(StepTimingData);
+						}
+					}
+				}
+
+				// Try to add or update the timing document
+				if (JobTiming == null)
+				{
+					IJobTiming? NewJobTiming = await JobTimings.TryAddAsync(Job.Id, NewSteps);
+					if (NewJobTiming != null)
+					{
+						return NewJobTiming;
+					}
+				}
+				else if (NewSteps.Count == 0)
+				{
+					return JobTiming;
+				}
+				else
+				{
+					IJobTiming? NewJobTiming = await JobTimings.TryAddStepsAsync(JobTiming, NewSteps);
+					if (NewJobTiming != null)
+					{
+						return NewJobTiming;
+					}
+				}
+
+				// Update the cached list of steps for the next iteration
+				foreach (JobStepTimingData NewStep in NewSteps)
+				{
+					CachedNewSteps[NewStep.Name] = NewStep;
+				}
+			}
+		}
+
+		/// <summary>
+		/// Finds the average duration for the given step
+		/// </summary>
+		/// <param name="StreamId">The stream to search</param>
+		/// <param name="TemplateId">The template id</param>
+		/// <param name="NodeName">Name of the node</param>
+		/// <param name="Change">Maximum changelist to consider</param>
+		/// <returns>Expected duration for the given step</returns>
+		async Task<JobStepTimingData> GetStepTimingInfo(StreamId StreamId, TemplateRefId TemplateId, string NodeName, int? Change)
+		{
+			// Find all the steps matching the given criteria
+			List<IJobStepRef> Steps = await JobStepRefs.GetStepsForNodeAsync(StreamId, TemplateId, NodeName, Change, false, 10);
+
+			// Sum up all the durations and wait times
+			int Count = 0;
+			float WaitTimeSum = 0;
+			float InitTimeSum = 0;
+			float DurationSum = 0;
+
+			foreach (IJobStepRef Step in Steps)
+			{
+				if (Step.FinishTimeUtc != null)
+				{
+					WaitTimeSum += Step.BatchWaitTime;
+					InitTimeSum += Step.BatchInitTime;
+					DurationSum += (float)(Step.FinishTimeUtc.Value - Step.StartTimeUtc).TotalSeconds;
+					Count++;
+				}
+			}
+
+			// Compute the averages
+			if (Count == 0)
+			{
+				return new JobStepTimingData(NodeName, null, null, null);
+			}
+			else
+			{
+				return new JobStepTimingData(NodeName, WaitTimeSum / Count, InitTimeSum / Count, DurationSum / Count);
+			}
+		}
+
+		/// <summary>
+		/// Updates the state of a batch
+		/// </summary>
+		/// <param name="Job">Job to update</param>
+		/// <param name="BatchId">Unique id of the batch to update</param>
+		/// <param name="NewLogId">The new log file id</param>
+		/// <param name="NewState">New state of the jobstep</param>
+		/// <returns>True if the job was updated, false if it was deleted</returns>
+		public async Task<bool> UpdateBatchAsync(IJob Job, SubResourceId BatchId, ObjectId? NewLogId = null, JobStepBatchState? NewState = null)
+		{
+			using IDisposable Scope = Logger.BeginScope("UpdateBatchAsync({JobId})", Job.Id);
+
+			bool bCheckForBadAgent = true;
+			for (; ; )
+			{
+				// Find the index of the appropriate batch
+				int BatchIdx = Job.Batches.FindIndex(x => x.Id == BatchId);
+				if (BatchIdx == -1)
+				{
+					break;
+				}
+
+				IJobStepBatch Batch = Job.Batches[BatchIdx];
+
+				// If the batch has already been marked as complete, error out
+				if (Batch.State == JobStepBatchState.Complete)
+				{
+					break;
+				}
+
+				// If we're marking the batch as complete before the agent has run everything, mark it to conform
+				JobStepBatchError? NewError = null;
+				if (NewState == JobStepBatchState.Complete && Batch.AgentId != null)
+				{
+					if (Batch.Steps.Any(x => x.State == JobStepState.Waiting || x.State == JobStepState.Ready))
+					{
+						// Mark the batch as incomplete
+						NewError = JobStepBatchError.Incomplete;
+
+						// Find the agent and set the conform flag
+						if (bCheckForBadAgent)
+						{
+							for (; ; )
+							{
+								IAgent? Agent = await Agents.GetAsync(Batch.AgentId.Value);
+								if (Agent == null || Agent.RequestConform)
+								{
+									break;
+								}
+								if (await Agents.TryUpdateSettingsAsync(Agent, bRequestConform: true) != null)
+								{
+									Logger.LogError("Agent {AgentId} did not complete lease; marking for conform", Agent.Id);
+									break;
+								}
+							}
+							bCheckForBadAgent = false;
+						}
+					}
+				}
+
+				// Update the batch state
+				if (await TryUpdateBatchAsync(Job, BatchId, NewLogId, NewState, NewError))
+				{
+					return true;
+				}
+
+				// Update the job
+				IJob? NewJob = await GetJobAsync(Job.Id);
+				if (NewJob == null)
+				{
+					break;
+				}
+				Job = NewJob;
+			}
+			return false;
+		}
+
+		/// <summary>
+		/// Attempts to update the state of a batch
+		/// </summary>
+		/// <param name="Job">Job to update</param>
+		/// <param name="BatchId">Unique id of the batch to update</param>
+		/// <param name="NewLogId">The new log file id</param>
+		/// <param name="NewState">New state of the jobstep</param>
+		/// <param name="NewError">New error state</param>
+		/// <returns>The updated job, otherwise null</returns>
+		public async Task<bool> TryUpdateBatchAsync(IJob Job, SubResourceId BatchId, ObjectId? NewLogId = null, JobStepBatchState? NewState = null, JobStepBatchError? NewError = null)
+		{
+			IGraph Graph = await GetGraphAsync(Job);
+			if (!await Jobs.TryUpdateBatchAsync(Job, Graph, BatchId, NewLogId, NewState, NewError))
+			{
+				return false;
+			}
+			return true;
+		}
+
+		/// <summary>
+		/// Update a jobstep state
+		/// </summary>
+		/// <param name="Job">Job to update</param>
+		/// <param name="BatchId">Unique id of the batch containing the step</param>
+		/// <param name="StepId">Unique id of the step to update</param>
+		/// <param name="NewState">New state of the jobstep</param>
+		/// <param name="NewOutcome">New outcome of the jobstep</param>
+		/// <param name="NewAbortRequested">New state of abort request</param>
+		/// <param name="NewAbortByUser">New user that requested the abort</param>
+		/// <param name="NewLogId">New log id for the jobstep</param>
+		/// <param name="NewNotificationTriggerId">New notification trigger id for the jobstep</param>
+		/// <param name="NewRetryByUser">Whether the step should be retried</param>
+		/// <param name="NewPriority">New priority for this step</param>
+		/// <param name="NewReports">New list of reports</param>
+		/// <param name="NewProperties">Property changes. Any properties with a null value will be removed.</param>
+		/// <returns>True if the job was updated, false if it was deleted in the meantime</returns>
+		public async Task<IJob?> UpdateStepAsync(IJob Job, SubResourceId BatchId, SubResourceId StepId, JobStepState NewState = JobStepState.Unspecified, JobStepOutcome NewOutcome = JobStepOutcome.Unspecified, bool? NewAbortRequested = null, string? NewAbortByUser = null, ObjectId? NewLogId = null, ObjectId? NewNotificationTriggerId = null, string? NewRetryByUser = null, Priority? NewPriority = null, List<Report>? NewReports = null, Dictionary<string, string?>? NewProperties = null)
+		{
+			using IDisposable Scope = Logger.BeginScope("UpdateStepAsync({JobId})", Job.Id);
+			for (; ;)
+			{
+				if (await TryUpdateStepAsync(Job, BatchId, StepId, NewState, NewOutcome, NewAbortRequested, NewAbortByUser, NewLogId, NewNotificationTriggerId, NewRetryByUser, NewPriority, NewReports, NewProperties))
+				{
+					return Job;
+				}
+
+				IJob? NewJob = await GetJobAsync(Job.Id);
+				if(NewJob == null)
+				{
+					return null;
+				}
+
+				Job = NewJob;
+			}
+		}
+
+		/// <summary>
+		/// Update a jobstep state
+		/// </summary>
+		/// <param name="Job">Job to update</param>
+		/// <param name="BatchId">Unique id of the batch containing the step</param>
+		/// <param name="StepId">Unique id of the step to update</param>
+		/// <param name="NewState">New state of the jobstep</param>
+		/// <param name="NewOutcome">New outcome of the jobstep</param>
+		/// <param name="NewAbortRequested">New state for request abort</param>
+		/// <param name="NewAbortByUser">New name of user that requested the abort</param>
+		/// <param name="NewLogId">New log id for the jobstep</param>
+		/// <param name="NewTriggerId">New trigger id for the jobstep</param>
+		/// <param name="NewRetryByUser">Whether the step should be retried</param>
+		/// <param name="NewPriority">New priority for this step</param>
+		/// <param name="NewReports">New reports</param>
+		/// <param name="NewProperties">Property changes. Any properties with a null value will be removed.</param>
+		/// <returns>True if the job was updated, false if it was deleted in the meantime</returns>
+		public async Task<bool> TryUpdateStepAsync(IJob Job, SubResourceId BatchId, SubResourceId StepId, JobStepState NewState = JobStepState.Unspecified, JobStepOutcome NewOutcome = JobStepOutcome.Unspecified, bool? NewAbortRequested = null, string? NewAbortByUser = null, ObjectId? NewLogId = null, ObjectId? NewTriggerId = null, string? NewRetryByUser = null, Priority? NewPriority = null, List<Report>? NewReports = null, Dictionary<string, string?>? NewProperties = null)
+		{
+			using IDisposable Scope = Logger.BeginScope("TryUpdateStepAsync({JobId})", Job.Id);
+
+			// Get the graph for this job
+			IGraph Graph = await GetGraphAsync(Job);
+
+			// Make sure the job state is set to unspecified
+			JobStepState OldState = JobStepState.Unspecified;
+			JobStepOutcome OldOutcome = JobStepOutcome.Unspecified;
+			if (NewState != JobStepState.Unspecified)
+			{
+				if (Job.TryGetStep(BatchId, StepId, out IJobStep? Step))
+				{
+					OldState = Step.State;
+					OldOutcome = Step.Outcome;
+				}
+			}
+
+			// If we're changing the state of a step, capture the label states beforehand so we can update any that change
+			IReadOnlyList<(LabelState, LabelOutcome)>? OldLabelStates = null;
+			if (OldState != NewState || OldOutcome != NewOutcome)
+			{
+				OldLabelStates = Job.GetLabelStates(Graph);
+			}
+
+			// Update the step
+			if (await Jobs.TryUpdateStepAsync(Job, Graph, BatchId, StepId, NewState, NewOutcome, NewAbortRequested, NewAbortByUser, NewLogId, NewTriggerId, NewRetryByUser, NewPriority, NewReports, NewProperties))
+			{
+				using IScope DdScope = GlobalTracer.Instance.BuildSpan("TryUpdateStepAsync").StartActive();
+				DdScope.Span.SetTag("JobId", Job.Id.ToString());
+				DdScope.Span.SetTag("BatchId", BatchId.ToString());
+				DdScope.Span.SetTag("StepId", StepId.ToString());
+				
+				if (OldState != NewState || OldOutcome != NewOutcome)
+				{
+					Logger.LogDebug("Transitioned job {JobId}, batch {BatchId}, step {StepId} from {OldState} to {NewState}", Job.Id, BatchId, StepId, OldState, NewState);
+
+					// Send any updates for modified badges
+					if (OldLabelStates != null)
+					{
+						using IScope _ = GlobalTracer.Instance.BuildSpan("Send badge updates").StartActive();
+						IReadOnlyList<(LabelState, LabelOutcome)> NewLabelStates = Job.GetLabelStates(Graph);
+						OnLabelUpdate?.Invoke(Job, OldLabelStates, NewLabelStates);
+						await JobTaskSource.UpdateUgsBadges(Job, Graph, OldLabelStates, NewLabelStates);
+					}
+
+					// Submit the change if auto-submit is enabled
+					if (Job.PreflightChange != 0 && NewState == JobStepState.Completed)
+					{
+						(JobStepState State, JobStepOutcome Outcome) = Job.GetTargetState();
+						if (State == JobStepState.Completed)
+						{
+							if (Job.AutoSubmit && Outcome == JobStepOutcome.Success && Job.AbortedByUser == null)
+							{
+								Job = await AutoSubmitChangeAsync(Job, Graph);
+							}
+							else if (Job.ClonedPreflightChange != 0)
+							{
+								IStream? Stream = await StreamService.GetCachedStream(Job.StreamId);
+								if (Stream != null)
+								{
+									await DeleteShelvedChangeAsync(Stream.ClusterName, Job.ClonedPreflightChange);
+								}
+							}
+						}
+					}
+
+					// Notify subscribers
+					if (NewState == JobStepState.Completed)
+					{
+						using IScope _ = GlobalTracer.Instance.BuildSpan("Notify subscribers").StartActive();
+						OnJobStepComplete?.Invoke(Job, Graph, BatchId, StepId);
+					}
+
+					// Create any downstream jobs
+					if (NewState == JobStepState.Completed && NewOutcome != JobStepOutcome.Failure)
+					{
+						using IScope _ = GlobalTracer.Instance.BuildSpan("Create downstream jobs").StartActive();
+						for (int Idx = 0; Idx < Job.ChainedJobs.Count; Idx++)
+						{
+							IChainedJob JobTrigger = Job.ChainedJobs[Idx];
+							if (JobTrigger.JobId == null)
+							{
+								JobStepOutcome JobTriggerOutcome = Job.GetTargetOutcome(Graph, JobTrigger.Target);
+								if (JobTriggerOutcome == JobStepOutcome.Success || JobTriggerOutcome == JobStepOutcome.Warnings)
+								{
+									Job = await FireJobTriggerAsync(Job, Graph, JobTrigger);
+								}
+							}
+						}
+					}
+
+					// Update the jobstep ref if it completed
+					if (NewState == JobStepState.Running || NewState == JobStepState.Completed || NewState == JobStepState.Aborted)
+					{
+						using IScope _ = GlobalTracer.Instance.BuildSpan("Update job step ref").StartActive();
+						if (Job.TryGetBatch(BatchId, out IJobStepBatch? Batch) && Batch.TryGetStep(StepId, out IJobStep? Step) && Step.StartTimeUtc != null)
+						{
+							await JobStepRefs.UpdateAsync(Job, Batch, Step, Graph);
+						}
+					}
+
+					// Update any issues that depend on this step
+					if (NewState == JobStepState.Completed && Job.UpdateIssues)
+					{
+						if (IssueService != null)
+						{
+							using IScope _ = GlobalTracer.Instance.BuildSpan("Update issues (V2)").StartActive();
+							try
+							{
+								await IssueService.UpdateCompleteStep(Job, Graph, BatchId, StepId);
+							}
+							catch(Exception Ex)
+							{
+								Logger.LogError(Ex, "Exception while updating issue service for job {JobId}:{BatchId}:{StepId}: {Message}", Job.Id, BatchId, StepId, Ex.Message);
+							}
+						}
+					}
+				}
+
+				using (IScope DispatchScope = GlobalTracer.Instance.BuildSpan("Update queued jobs").StartActive())
+				{
+					// Notify the dispatch service that the job has changed
+					JobTaskSource.UpdateQueuedJob(Job, Graph);
+				}
+
+				return true;
+			}
+
+			return false;
+		}
+
+		/// <summary>
+		/// Submit the given change for a preflight
+		/// </summary>
+		/// <param name="Job">The job being run</param>
+		/// <param name="Graph">Graph for the job</param>
+		/// <returns></returns>
+		[SuppressMessage("Design", "CA1031:Do not catch general exception types", Justification = "<Pending>")]
+		private async Task<IJob> AutoSubmitChangeAsync(IJob Job, IGraph Graph)
+		{
+			int? Change;
+			string Message;
+			try
+			{
+				IStream? Stream = await StreamService.GetCachedStream(Job.StreamId);
+				if (Stream != null)
+				{
+					int ClonedPreflightChange = Job.ClonedPreflightChange;
+					if (ClonedPreflightChange == 0)
+					{
+						if (ShouldClonePreflightChange(Job.StreamId))
+						{
+							ClonedPreflightChange = await CloneShelvedChangeAsync(Stream.ClusterName, Job.PreflightChange);
+						}
+						else
+						{
+							ClonedPreflightChange = Job.PreflightChange;
+						}
+					}
+
+					Logger.LogInformation("Updating description for {ClonedPreflightChange}", ClonedPreflightChange);
+
+					ChangeDetails Details = await PerforceService.GetChangeDetailsAsync(Stream.ClusterName, Stream.Name, ClonedPreflightChange, null);
+					await PerforceService.UpdateChangelistDescription(Stream.ClusterName, ClonedPreflightChange, Details.Description.TrimEnd() + $"\n#preflight {Job.Id}");
+
+					Logger.LogInformation("Submitting change {Change} (through {ChangeCopy}) after successful completion of {JobId}", Job.PreflightChange, ClonedPreflightChange, Job.Id);
+					(Change, Message) = await PerforceService.SubmitShelvedChangeAsync(Stream.ClusterName, ClonedPreflightChange, Job.PreflightChange);
+
+					Logger.LogInformation("Attempt to submit {Change} (through {ChangeCopy}): {Message}", Job.PreflightChange, ClonedPreflightChange, Message);
+
+					if (ShouldClonePreflightChange(Job.StreamId))
+					{
+						if (Change != null && Job.ClonedPreflightChange != 0)
+						{
+							await DeleteShelvedChangeAsync(Stream.ClusterName, Job.PreflightChange);
+						}
+					}
+					else
+					{
+						if (Change != null && Job.PreflightChange != 0)
+						{
+							await DeleteShelvedChangeAsync(Stream.ClusterName, Job.PreflightChange);
+						}
+					}
+				}
+				else
+				{
+					(Change, Message) = ((int?)null, "Stream not found");
+				}
+			}
+			catch (Exception Ex)
+			{
+				(Change, Message) = ((int?)null, "Internal error");
+				Logger.LogError(Ex, "Unable to submit shelved change");
+			}
+
+			for (; ; )
+			{
+				if (await Jobs.TryUpdateJobAsync(Job, Graph, AutoSubmitChange: Change, AutoSubmitMessage: Message))
+				{
+					return Job;
+				}
+
+				IJob? NewJob = await GetJobAsync(Job.Id);
+				if (NewJob == null)
+				{
+					return Job;
+				}
+
+				Job = NewJob;
+			}
+		}
+
+		/// <summary>
+		/// Clone a shelved changelist for running a preflight
+		/// </summary>
+		/// <param name="ClusterName">Name of the Perforce cluster</param>
+		/// <param name="Change">The changelist to clone</param>
+		/// <returns></returns>
+		private async Task<int> CloneShelvedChangeAsync(string ClusterName, int Change)
+		{
+			int ClonedChange;
+			try
+			{
+				ClonedChange = await PerforceService.DuplicateShelvedChangeAsync(ClusterName, Change);
+				Logger.LogInformation("CL {Change} was duplicated into {ClonedChange}", Change, ClonedChange);
+			}
+			catch (Exception Ex)
+			{
+				Logger.LogError(Ex, "Unable to CL {Change} for preflight: {Message}", Change, Ex.Message);
+				throw;
+			}
+			return ClonedChange;
+		}
+
+		/// <summary>
+		/// Deletes a shelved changelist, catching and logging any exceptions
+		/// </summary>
+		/// <param name="ClusterName"></param>
+		/// <param name="Change">The changelist to delete</param>
+		/// <returns>True if the change was deleted successfully, false otherwise</returns>
+		[SuppressMessage("Design", "CA1031:Do not catch general exception types", Justification = "<Pending>")]
+		private async Task<bool> DeleteShelvedChangeAsync(string ClusterName, int Change)
+		{
+			Logger.LogInformation("Removing shelf {Change}", Change);
+			try
+			{
+				await PerforceService.DeleteShelvedChangeAsync(ClusterName, Change);
+				return true;
+			}
+			catch (Exception Ex)
+			{
+				Logger.LogError(Ex, "Unable to delete shelved change {Change}", Change);
+				return false;
+			}
+		}
+
+		/// <summary>
+		/// Fires a trigger for a chained job
+		/// </summary>
+		/// <param name="Job">The job containing the trigger</param>
+		/// <param name="Graph">Graph for the job containing the trigger</param>
+		/// <param name="JobTrigger">The trigger object to fire</param>
+		/// <returns>New job instance</returns>
+		private async Task<IJob> FireJobTriggerAsync(IJob Job, IGraph Graph, IChainedJob JobTrigger)
+		{
+			for (; ; )
+			{
+				// Update the job
+				ObjectId ChainedJobId = ObjectId.GenerateNewId();
+				if (await Jobs.TryUpdateJobAsync(Job, Graph, JobTrigger: new KeyValuePair<TemplateRefId, ObjectId>(JobTrigger.TemplateRefId, ChainedJobId)))
+				{
+					IStream? Stream = await StreamService.GetStreamAsync(Job.StreamId);
+					if(Stream == null)
+					{
+						Logger.LogWarning("Cannot find stream {StreamId} for downstream job", Job.StreamId);
+						break;
+					}
+
+					TemplateRef? TemplateRef;
+					if (!Stream.Templates.TryGetValue(JobTrigger.TemplateRefId, out TemplateRef))
+					{
+						Logger.LogWarning("Cannot find template {TemplateRefId} in stream {StreamId}", JobTrigger.TemplateRefId, Job.StreamId);
+						break;
+					}
+
+					ITemplate? Template = await TemplateService.GetTemplateAsync(TemplateRef.Hash);
+					if (Template == null)
+					{
+						Logger.LogWarning("Cannot find template {TemplateHash}", TemplateRef.Hash);
+						break;
+					}
+
+					IGraph TriggerGraph = await Graphs.AddAsync(Template);
+					Logger.LogInformation("Creating downstream job {ChainedJobId} from job {JobId}", ChainedJobId, Job.Id);
+
+					await CreateJobAsync(ChainedJobId, Stream, JobTrigger.TemplateRefId, TemplateRef.Hash, TriggerGraph, TemplateRef.Name, Job.Change, Job.CodeChange, Job.PreflightChange, Job.ClonedPreflightChange, Job.StartedByUserId, Job.StartedByUser, Template.Priority, null, Job.UpdateIssues, TemplateRef.ChainedJobs, false, false, TemplateRef.NotificationChannel, TemplateRef.NotificationChannelFilter, null, Template.Counters, Template.Arguments);
+					break;
+				}
+
+				// Fetch the job again
+				IJob? NewJob = await Jobs.GetAsync(Job.Id);
+				if(NewJob == null)
+				{
+					break;
+				}
+				Job = NewJob;
+			}
+			return Job;
+		}
+
+		/// <summary>
+		/// Determines if the user is authorized to perform an action on a particular stream
+		/// </summary>
+		/// <param name="Acl">The ACL to check</param>
+		/// <param name="StreamId">The stream containing this job</param>
+		/// <param name="Action">The action being performed</param>
+		/// <param name="User">The principal to authorize</param>
+		/// <param name="Cache">Cache of stream permissions</param>
+		/// <returns>True if the action is authorized</returns>
+		private Task<bool> AuthorizeAsync(Acl? Acl, StreamId StreamId, AclAction Action, ClaimsPrincipal User, StreamPermissionsCache? Cache)
+		{
+			// Do the regular authorization
+			bool? Result = Acl?.Authorize(Action, User);
+			if (Result == null)
+			{
+				return StreamService.AuthorizeAsync(StreamId, Action, User, Cache);
+			}
+			else
+			{
+				return Task.FromResult(Result.Value);
+			}
+		}
+
+		/// <summary>
+		/// Determines if the user is authorized to perform an action on a particular stream
+		/// </summary>
+		/// <param name="Job">The job to check</param>
+		/// <param name="Action">The action being performed</param>
+		/// <param name="User">The principal to authorize</param>
+		/// <param name="Cache">Cache of stream permissions</param>
+		/// <returns>True if the action is authorized</returns>
+		public Task<bool> AuthorizeAsync(IJob Job, AclAction Action, ClaimsPrincipal User, StreamPermissionsCache? Cache)
+		{
+			return AuthorizeAsync(Job.Acl, Job.StreamId, Action, User, Cache);
+		}
+
+		/// <summary>
+		/// Determines if the user is authorized to perform an action on a particular stream
+		/// </summary>
+		/// <param name="JobId">The job to check</param>
+		/// <param name="Action">The action being performed</param>
+		/// <param name="User">The principal to authorize</param>
+		/// <param name="Cache">Cache of job permissions</param>
+		/// <returns>True if the action is authorized</returns>
+		public async Task<bool> AuthorizeAsync(ObjectId JobId, AclAction Action, ClaimsPrincipal User, JobPermissionsCache? Cache)
+		{
+			IJobPermissions? Permissions;
+			if (Cache == null)
+			{
+				Permissions = await GetJobPermissionsAsync(JobId);
+			}
+			else if (!Cache.Jobs.TryGetValue(JobId, out Permissions))
+			{
+				Permissions = await GetJobPermissionsAsync(JobId);
+				Cache.Jobs.Add(JobId, Permissions);
+			}
+			return Permissions != null && await AuthorizeAsync(Permissions.Acl, Permissions.StreamId, Action, User, Cache);
+		}
+
+		/// <summary>
+		/// Determines if the user is authorized to perform an action on a particular stream
+		/// </summary>
+		/// <param name="Job">The job to check</param>
+		/// <param name="User">The principal to authorize</param>
+		/// <returns>True if the action is authorized</returns>
+		public static bool AuthorizeSession(IJob Job, ClaimsPrincipal User)
+		{
+			return Job.Batches.Any(x => AuthorizeSession(x, User));
+		}
+
+		/// <summary>
+		/// Determines if the user is authorized to perform an action on a particular stream
+		/// </summary>
+		/// <param name="Batch">The batch being executed</param>
+		/// <param name="User">The principal to authorize</param>
+		/// <returns>True if the action is authorized</returns>
+		public static bool AuthorizeSession(IJobStepBatch Batch, ClaimsPrincipal User)
+		{
+			if(Batch.SessionId != null)
+			{
+				foreach (Claim Claim in User.Claims)
+				{
+					if(Claim.Type == HordeClaimTypes.AgentSessionId)
+					{
+						ObjectId SessionId;
+						if (ObjectId.TryParse(Claim.Value, out SessionId) && SessionId == Batch.SessionId.Value)
+						{
+							return true;
+						}
+					}
+				}
+			}
+			return false;
+		}
+		
+		/// <summary>
+		/// Get a list of each batch's session ID formatted as a string for debugging purposes
+		/// </summary>
+		/// <param name="Job">The job to list</param>
+		/// <returns>List of session IDs</returns>
+		public static string GetAllBatchSessionIds(IJob Job)
+		{
+			return string.Join(",", Job.Batches.Select(b => b.SessionId));
+		}
+	}
+}
diff --git a/Engine/Source/Programs/Horde/HordeServer/Jobs/Impl/JobStepRefCollection.cs b/Engine/Source/Programs/Horde/HordeServer/Jobs/Impl/JobStepRefCollection.cs
new file mode 100644
--- /dev/null
+++ b/Engine/Source/Programs/Horde/HordeServer/Jobs/Impl/JobStepRefCollection.cs
@@ -0,0 +1,141 @@
+// Copyright Epic Games, Inc. All Rights Reserved.
+
+using HordeCommon;
+using HordeServer.Models;
+using HordeServer.Services;
+using MongoDB.Bson;
+using MongoDB.Bson.Serialization.Attributes;
+using MongoDB.Driver;
+using System;
+using System.Collections.Generic;
+using System.Linq;
+using System.Threading.Tasks;
+using HordeServer.Utilities;
+
+namespace HordeServer.Collections.Impl
+{
+	using StreamId = StringId<IStream>;
+	using TemplateRefId = StringId<TemplateRef>;
+	using PoolId = StringId<IPool>;
+
+	/// <summary>
+	/// Collection of JobStepRef documents
+	/// </summary>
+	public class JobStepRefCollection : IJobStepRefCollection
+	{
+		class JobStepRef : IJobStepRef
+		{
+			[BsonId]
+			public JobStepRefId Id { get; set; }
+
+			public string JobName { get; set; } = "Unknown";
+
+			public string Name { get; set; }
+
+			public StreamId StreamId { get; set; }
+			public TemplateRefId TemplateId { get; set; }
+			public int Change { get; set; }
+			public ObjectId? LogId { get; set; }
+			public PoolId? PoolId { get; set; }
+			public AgentId? AgentId { get; set; }
+			public JobStepOutcome? Outcome { get; set; }
+			public float BatchWaitTime { get; set; }
+			public float BatchInitTime { get; set; }
+
+			[BsonIgnoreIfNull]
+			public DateTimeOffset? StartTime { get; set; }
+
+			[BsonIgnoreIfNull]
+			public DateTime? StartTimeUtc { get; set; }
+
+			[BsonIgnoreIfNull]
+			public DateTimeOffset? FinishTime { get; set; }
+
+			[BsonIgnoreIfNull]
+			public DateTime? FinishTimeUtc { get; set; }
+
+			DateTime IJobStepRef.StartTimeUtc => StartTimeUtc ?? StartTime?.UtcDateTime ?? default;
+			DateTime? IJobStepRef.FinishTimeUtc => FinishTimeUtc ?? FinishTime?.UtcDateTime;
+			string IJobStepRef.NodeName => Name;
+
+			public JobStepRef(JobStepRefId Id, string JobName, string NodeName, StreamId StreamId, TemplateRefId TemplateId, int Change, ObjectId? LogId, PoolId? PoolId, AgentId? AgentId, JobStepOutcome? Outcome, float BatchWaitTime, float BatchInitTime, DateTime StartTimeUtc, DateTime? FinishTimeUtc)
+			{
+				this.Id = Id;
+				this.JobName = JobName;
+				this.Name = NodeName;
+				this.StreamId = StreamId;
+				this.TemplateId = TemplateId;
+				this.Change = Change;
+				this.LogId = LogId;
+				this.PoolId = PoolId;
+				this.AgentId = AgentId;
+				this.Outcome = Outcome;
+				this.BatchWaitTime = BatchWaitTime;
+				this.BatchInitTime = BatchInitTime;
+				this.StartTime = StartTimeUtc;
+				this.StartTimeUtc = StartTimeUtc;
+				this.FinishTime = FinishTimeUtc;
+				this.FinishTimeUtc = FinishTimeUtc;
+			}
+		}
+
+		IMongoCollection<JobStepRef> JobStepRefs;
+
+		/// <summary>
+		/// Constructor
+		/// </summary>
+		/// <param name="DatabaseService">The database service instance</param>
+		public JobStepRefCollection(DatabaseService DatabaseService)
+		{
+			JobStepRefs = DatabaseService.GetCollection<JobStepRef>("JobStepRefs");
+
+			if (!DatabaseService.ReadOnlyMode)
+			{
+				JobStepRefs.Indexes.CreateOne(new CreateIndexModel<JobStepRef>(Builders<JobStepRef>.IndexKeys.Ascending(x => x.StreamId).Ascending(x => x.TemplateId).Ascending(x => x.Name).Descending(x => x.Change)));
+			}
+		}
+
+		/// <inheritdoc/>
+		public async Task<IJobStepRef> InsertOrReplaceAsync(JobStepRefId Id, string JobName, string StepName, StreamId StreamId, TemplateRefId TemplateId, int Change, ObjectId? LogId, PoolId? PoolId, AgentId? AgentId, JobStepOutcome? Outcome, float WaitTime, float InitTime, DateTime StartTimeUtc, DateTime? FinishTimeUtc)
+		{
+			JobStepRef NewJobStepRef = new JobStepRef(Id, JobName, StepName, StreamId, TemplateId, Change, LogId, PoolId, AgentId, Outcome, WaitTime, InitTime, StartTimeUtc, FinishTimeUtc);
+			await JobStepRefs.ReplaceOneAsync(Builders<JobStepRef>.Filter.Eq(x => x.Id, NewJobStepRef.Id), NewJobStepRef, new ReplaceOptions { IsUpsert = true });
+			return NewJobStepRef;
+		}
+
+		/// <inheritdoc/>
+		public async Task<List<IJobStepRef>> GetStepsForNodeAsync(StreamId StreamId, TemplateRefId TemplateId, string NodeName, int? Change, bool IncludeFailed, int Count)
+		{
+			// Find all the steps matching the given criteria
+			FilterDefinitionBuilder<JobStepRef> FilterBuilder = Builders<JobStepRef>.Filter;
+
+			FilterDefinition<JobStepRef> Filter = FilterDefinition<JobStepRef>.Empty;
+			Filter &= FilterBuilder.Eq(x => x.StreamId, StreamId);
+			Filter &= FilterBuilder.Eq(x => x.TemplateId, TemplateId);
+			Filter &= FilterBuilder.Eq(x => x.Name, NodeName);
+			if (Change != null)
+			{
+				Filter &= FilterBuilder.Lte(x => x.Change, Change.Value);
+			}
+			if (!IncludeFailed)
+			{
+				Filter &= FilterBuilder.Ne(x => x.Outcome, JobStepOutcome.Failure);
+			}
+
+			List<JobStepRef> Steps = await JobStepRefs.Find(Filter).SortByDescending(x => x.Change).ThenByDescending(x => x.StartTime).Limit(Count).ToListAsync();
+			return Steps.ConvertAll<IJobStepRef>(x => x);
+		}
+
+		/// <inheritdoc/>
+		public async Task<IJobStepRef?> GetPrevStepForNodeAsync(StreamId StreamId, TemplateRefId TemplateId, string NodeName, int Change)
+		{
+			return await JobStepRefs.Find(x => x.StreamId == StreamId && x.TemplateId == TemplateId && x.Name == NodeName && x.Change < Change && x.Outcome != null).SortByDescending(x => x.Change).FirstOrDefaultAsync();
+		}
+
+		/// <inheritdoc/>
+		public async Task<IJobStepRef?> GetNextStepForNodeAsync(StreamId StreamId, TemplateRefId TemplateId, string NodeName, int Change)
+		{
+			return await JobStepRefs.Find(x => x.StreamId == StreamId && x.TemplateId == TemplateId && x.Name == NodeName && x.Change > Change && x.Outcome != null).SortBy(x => x.Change).FirstOrDefaultAsync();
+		}
+	}
+}
diff --git a/Engine/Source/Programs/Horde/HordeServer/Jobs/Impl/JobTimingCollection.cs b/Engine/Source/Programs/Horde/HordeServer/Jobs/Impl/JobTimingCollection.cs
new file mode 100644
--- /dev/null
+++ b/Engine/Source/Programs/Horde/HordeServer/Jobs/Impl/JobTimingCollection.cs
@@ -0,0 +1,145 @@
+// Copyright Epic Games, Inc. All Rights Reserved.
+
+using Amazon.Runtime.Internal.Util;
+using Amazon.S3.Model.Internal.MarshallTransformations;
+using HordeServer.Models;
+using HordeServer.Services;
+using HordeServer.Utilities;
+using MongoDB.Bson;
+using MongoDB.Bson.Serialization.Attributes;
+using MongoDB.Driver;
+using System;
+using System.Collections.Generic;
+using System.Diagnostics.CodeAnalysis;
+using System.Linq;
+using System.Threading.Tasks;
+
+namespace HordeServer.Collections.Impl
+{
+	/// <summary>
+	/// Concrete implementation of IJobTimingCollection
+	/// </summary>
+	public class JobTimingCollection : IJobTimingCollection
+	{
+		/// <summary>
+		/// Timing information for a particular step
+		/// </summary>
+		class JobStepTimingDocument : IJobStepTiming
+		{
+			public string Name { get; set; }
+			public float? AverageWaitTime { get; set; }
+			public float? AverageInitTime { get; set; }
+			public float? AverageDuration { get; set; }
+
+			[BsonConstructor]
+			private JobStepTimingDocument()
+			{
+				this.Name = null!;
+			}
+
+			public JobStepTimingDocument(string Name, float? AverageWaitTime, float? AverageInitTime, float? AverageDuration)
+			{
+				this.Name = Name;
+				this.AverageWaitTime = AverageWaitTime;
+				this.AverageInitTime = AverageInitTime;
+				this.AverageDuration = AverageDuration;
+			}
+		}
+
+		/// <summary>
+		/// Timing information for a particular job
+		/// </summary>
+		class JobTimingDocument : IJobTiming
+		{
+			public ObjectId Id { get; set; }
+			public List<JobStepTimingDocument> Steps { get; set; } = new List<JobStepTimingDocument>();
+			public int UpdateIndex { get; set; }
+
+			[BsonIgnore]
+			public Dictionary<string, IJobStepTiming>? NameToStep;
+
+			public bool TryGetStepTiming(string Name, [NotNullWhen(true)] out IJobStepTiming? Timing)
+			{
+				if (NameToStep == null)
+				{
+					NameToStep = new Dictionary<string, IJobStepTiming>();
+					foreach (JobStepTimingDocument Step in Steps)
+					{
+						if (NameToStep.ContainsKey(Step.Name))
+						{
+							Serilog.Log.Logger.Warning("Step {Name} appears twice in job timing document {Id}", Step.Name, Id);
+						}
+						NameToStep[Step.Name] = Step;
+					}
+				}
+				return NameToStep.TryGetValue(Name, out Timing);
+			}
+		}
+
+		/// <summary>
+		/// Collection of timing documents
+		/// </summary>
+		IMongoCollection<JobTimingDocument> Collection;
+
+		/// <summary>
+		/// Constructor
+		/// </summary>
+		/// <param name="DatabaseService">The database service singleton</param>
+		public JobTimingCollection(DatabaseService DatabaseService)
+		{
+			Collection = DatabaseService.GetCollection<JobTimingDocument>("JobTiming");
+		}
+
+		/// <inheritdoc/>
+		public async Task<IJobTiming?> TryAddAsync(ObjectId JobId, List<JobStepTimingData> Steps)
+		{
+			JobTimingDocument JobTiming = new JobTimingDocument();
+			JobTiming.Id = JobId;
+			JobTiming.Steps.AddRange(Steps.Select(x => new JobStepTimingDocument(x.Name, x.AverageWaitTime, x.AverageInitTime, x.AverageDuration)));
+			JobTiming.UpdateIndex = 1;
+
+			try
+			{
+				await Collection.InsertOneAsync(JobTiming);
+				return JobTiming;
+			}
+			catch (MongoWriteException Ex)
+			{
+				if (Ex.WriteError.Category == ServerErrorCategory.DuplicateKey)
+				{
+					return null;
+				}
+				else
+				{
+					throw;
+				}
+			}
+		}
+
+		/// <inheritdoc/>
+		public async Task<IJobTiming?> TryGetAsync(ObjectId JobId)
+		{
+			return await Collection.Find(x => x.Id == JobId).FirstOrDefaultAsync();
+		}
+
+		/// <inheritdoc/>
+		public async Task<IJobTiming?> TryAddStepsAsync(IJobTiming JobTiming, List<JobStepTimingData> Steps)
+		{
+			JobTimingDocument JobTimingDocument = (JobTimingDocument)JobTiming;
+			List<JobStepTimingDocument> StepDocuments = Steps.ConvertAll(x => new JobStepTimingDocument(x.Name, x.AverageWaitTime, x.AverageInitTime, x.AverageDuration));
+
+			FilterDefinition<JobTimingDocument> Filter = Builders<JobTimingDocument>.Filter.Expr(x => x.Id == JobTimingDocument.Id && x.UpdateIndex == JobTimingDocument.UpdateIndex);
+			UpdateDefinition<JobTimingDocument> Update = Builders<JobTimingDocument>.Update.Set(x => x.UpdateIndex, JobTimingDocument.UpdateIndex + 1).PushEach(x => x.Steps, StepDocuments);
+
+			UpdateResult Result = await Collection.UpdateOneAsync(Filter, Update);
+			if (Result.ModifiedCount > 0)
+			{
+				JobTimingDocument.NameToStep = null;
+				JobTimingDocument.Steps.AddRange(StepDocuments);
+				JobTimingDocument.UpdateIndex++;
+				return JobTimingDocument;
+			}
+			return null;
+		}
+	}
+}
diff --git a/Engine/Source/Programs/Horde/HordeServer/Users/IUserClaim.cs b/Engine/Source/Programs/Horde/HordeServer/Users/IUserClaim.cs
new file mode 100644
--- /dev/null
+++ b/Engine/Source/Programs/Horde/HordeServer/Users/IUserClaim.cs
@@ -0,0 +1,74 @@
+// Copyright Epic Games, Inc. All Rights Reserved.
+
+using System;
+using System.Collections.Generic;
+using System.Linq;
+using System.Security.Claims;
+using System.Threading.Tasks;
+
+namespace HordeServer.Models
+{
+	/// <summary>
+	/// Claim for a user
+	/// </summary>
+	public interface IUserClaim
+	{
+		/// <summary>
+		/// Name of the claim
+		/// </summary>
+		public string Type { get; }
+
+		/// <summary>
+		/// Value of the claim
+		/// </summary>
+		public string Value { get; }
+	}
+
+	/// <summary>
+	/// New claim document
+	/// </summary>
+	public class UserClaim : IUserClaim
+	{
+		/// <inheritdoc/>
+		public string Type { get; set; }
+
+		/// <inheritdoc/>
+		public string Value { get; set; }
+
+		/// <summary>
+		/// Constructor
+		/// </summary>
+		/// <param name="Type"></param>
+		/// <param name="Value"></param>
+		public UserClaim(string Type, string Value)
+		{
+			this.Type = Type;
+			this.Value = Value;
+		}
+
+		/// <summary>
+		/// Constructor
+		/// </summary>
+		/// <param name="Other">Claim to construct from</param>
+		public UserClaim(IUserClaim Other)
+			: this(Other.Type, Other.Value)
+		{
+		}
+
+		/// <summary>
+		/// Constructs a UserClaim from a Claim object
+		/// </summary>
+		/// <param name="Claim">Claim object</param>
+		/// <returns></returns>
+		public static UserClaim FromClaim(Claim Claim)
+		{
+			return new UserClaim(Claim.Type, Claim.Value);
+		}
+
+		/// <summary>
+		/// Conversion operator from NET claims
+		/// </summary>
+		/// <param name="Claim"></param>
+		public static implicit operator UserClaim(Claim Claim) => FromClaim(Claim);
+	}
+}
diff --git a/Engine/Source/Programs/Horde/HordeServer/Users/IUserCollection.cs b/Engine/Source/Programs/Horde/HordeServer/Users/IUserCollection.cs
new file mode 100644
--- /dev/null
+++ b/Engine/Source/Programs/Horde/HordeServer/Users/IUserCollection.cs
@@ -0,0 +1,119 @@
+// Copyright Epic Games, Inc. All Rights Reserved.
+
+using HordeServer.Models;
+using HordeServer.Services;
+using HordeServer.Utilities;
+using MongoDB.Bson;
+using System;
+using System.Collections.Generic;
+using System.Linq;
+using System.Security.Claims;
+using System.Threading.Tasks;
+
+namespace HordeServer.Collections
+{
+	/// <summary>
+	/// Manages user documents
+	/// </summary>
+	public interface IUserCollection
+	{
+		/// <summary>
+		/// Gets a user by unique id
+		/// </summary>
+		/// <param name="Id">Id of the user</param>
+		/// <returns>The user information</returns>
+		Task<IUser?> GetUserAsync(ObjectId Id);
+
+		/// <summary>
+		/// Gets a cached user by unique id
+		/// </summary>
+		/// <param name="Id">Id of the user</param>
+		/// <returns>The user information</returns>
+		Task<IUser?> GetCachedUserAsync(ObjectId Id);
+
+		/// <summary>
+		/// Gets a user by unique id
+		/// </summary>
+		/// <param name="Ids">Ids of the users</param>
+		/// <param name="NameRegex">Name regex to match for the users</param>
+		/// <param name="Index">Maximum number of results</param>
+		/// <param name="Count">Number of results to return</param>
+		/// <returns>The user information</returns>
+		Task<List<IUser>> FindUsersAsync(IEnumerable<ObjectId>? Ids = null, string? NameRegex = null, int? Index = null, int? Count = null);
+
+		/// <summary>
+		/// Gets a user by login
+		/// </summary>
+		/// <param name="Login">Login for the user</param>
+		/// <returns>The user information</returns>
+		Task<IUser?> FindUserByLoginAsync(string Login);
+
+		/// <summary>
+		/// Find or add a user with the given claims. Claims will be updated if the user exists.
+		/// </summary>
+		/// <param name="Login">Login id of the user</param>
+		/// <param name="Name">Full name of the user</param>
+		/// <param name="Email">Email address of the user</param>
+		/// <returns>The user document</returns>
+		Task<IUser> FindOrAddUserByLoginAsync(string Login, string? Name = null, string? Email = null);
+
+		/// <summary>
+		/// Gets the claims for a user
+		/// </summary>
+		/// <param name="UserId"></param>
+		/// <returns></returns>
+		Task<IUserClaims> GetClaimsAsync(ObjectId UserId);
+
+		/// <summary>
+		/// Update the claims for a user
+		/// </summary>
+		/// <param name="UserId"></param>
+		/// <param name="Claims"></param>
+		/// <returns></returns>
+		Task UpdateClaimsAsync(ObjectId UserId, IEnumerable<IUserClaim> Claims);
+
+		/// <summary>
+		/// Get settings for a user
+		/// </summary>
+		/// <param name="UserId"></param>
+		/// <returns></returns>
+		Task<IUserSettings> GetSettingsAsync(ObjectId UserId);
+
+		/// <summary>
+		/// Update a user
+		/// </summary>
+		/// <param name="UserId">The user to update</param>
+		/// <param name="EnableExperimentalFeatures"></param>
+		/// <param name="EnableIssueNotifications"></param>
+		/// <param name="DashboardSettings">Opaque settings object for the dashboard</param>
+		/// <param name="AddPinnedJobIds"></param>
+		/// <param name="RemovePinnedJobIds"></param>
+		/// <returns>Updated user object</returns>
+		Task UpdateSettingsAsync(ObjectId UserId, bool? EnableExperimentalFeatures = null, bool? EnableIssueNotifications = null, BsonValue? DashboardSettings = null, IEnumerable<ObjectId>? AddPinnedJobIds = null, IEnumerable<ObjectId>? RemovePinnedJobIds = null);
+	}
+
+	/// <summary>
+	/// Extension methods for the user collect
+	/// </summary>
+	static class UserCollectionExtensions
+	{
+		/// <summary>
+		/// Gets a particular user info from the collection
+		/// </summary>
+		/// <param name="UserCollection"></param>
+		/// <param name="Principal"></param>
+		/// <returns></returns>
+		public static Task<IUser?> GetUserAsync(this IUserCollection UserCollection, ClaimsPrincipal Principal)
+		{
+			ObjectId? UserId = Principal.GetUserId();
+			if (UserId != null)
+			{
+				return UserCollection.GetUserAsync(UserId.Value);
+			}
+			else
+			{
+				return Task.FromResult<IUser?>(null);
+			}
+		}
+	}
+}
diff --git a/Engine/Source/Programs/Horde/HordeServer/Users/Impl/UserCollection.cs b/Engine/Source/Programs/Horde/HordeServer/Users/Impl/UserCollection.cs
new file mode 100644
--- /dev/null
+++ b/Engine/Source/Programs/Horde/HordeServer/Users/Impl/UserCollection.cs
@@ -0,0 +1,209 @@
+// Copyright Epic Games, Inc. All Rights Reserved.
+
+using HordeServer.Models;
+using HordeServer.Services;
+using HordeServer.Utilities;
+using MongoDB.Bson;
+using MongoDB.Bson.Serialization;
+using MongoDB.Bson.Serialization.Attributes;
+using MongoDB.Driver;
+using System;
+using System.Collections.Generic;
+using System.Linq;
+using System.Security.Claims;
+using System.Threading.Tasks;
+
+namespace HordeServer.Collections.Impl
+{
+	/// <summary>
+	/// Manages user documents
+	/// </summary>
+	class UserCollectionV1 : IUserCollection
+	{
+		class UserDocument : IUser, IUserClaims, IUserSettings
+		{
+			public ObjectId Id { get; set; }
+
+			public ClaimDocument PrimaryClaim { get; set; } = null!;
+			public List<ClaimDocument> Claims { get; set; } = new List<ClaimDocument>();
+
+			[BsonDefaultValue(false), BsonIgnoreIfDefault]
+			public bool EnableExperimentalFeatures { get; set; }
+
+			[BsonDefaultValue(false), BsonIgnoreIfDefault]
+			public bool EnableIssueNotifications { get; set; }
+
+			public BsonValue DashboardSettings { get; set; } = BsonNull.Value;
+			public List<ObjectId> PinnedJobIds { get; set; } = new List<ObjectId>();
+
+			string IUser.Name => Claims.FirstOrDefault(x => String.Equals(x.Type, "name", StringComparison.Ordinal))?.Value ?? PrimaryClaim.Value;
+			string IUser.Login => Claims.FirstOrDefault(x => String.Equals(x.Type, ClaimTypes.Name, StringComparison.Ordinal))?.Value ?? PrimaryClaim.Value;
+			string? IUser.Email => Claims.FirstOrDefault(x => String.Equals(x.Type, ClaimTypes.Email, StringComparison.Ordinal))?.Value;
+
+			ObjectId IUserClaims.UserId => Id;
+			IReadOnlyList<IUserClaim> IUserClaims.Claims => Claims;
+
+			ObjectId IUserSettings.UserId => Id;
+			IReadOnlyList<ObjectId> IUserSettings.PinnedJobIds => PinnedJobIds;
+		}
+
+		class ClaimDocument : IUserClaim
+		{
+			public string Type { get; set; }
+			public string Value { get; set; }
+
+			private ClaimDocument()
+			{
+				this.Type = null!;
+				this.Value = null!;
+			}
+
+			public ClaimDocument(string Type, string Value)
+			{
+				this.Type = Type;
+				this.Value = Value;
+			}
+
+			public ClaimDocument(IUserClaim Other)
+			{
+				this.Type = Other.Type;
+				this.Value = Other.Value;
+			}
+		}
+
+		IMongoCollection<UserDocument> Users;
+
+		/// <summary>
+		/// Static constructor
+		/// </summary>
+		static UserCollectionV1()
+		{
+			BsonSerializer.RegisterDiscriminatorConvention(typeof(UserDocument), NullDiscriminatorConvention.Instance);
+			BsonSerializer.RegisterDiscriminatorConvention(typeof(ClaimDocument), NullDiscriminatorConvention.Instance);
+		}
+
+		/// <summary>
+		/// Constructor
+		/// </summary>
+		/// <param name="DatabaseService"></param>
+		public UserCollectionV1(DatabaseService DatabaseService)
+		{
+			Users = DatabaseService.GetCollection<UserDocument>("Users");
+
+			if (!DatabaseService.ReadOnlyMode)
+			{
+				Users.Indexes.CreateOne(new CreateIndexModel<UserDocument>(Builders<UserDocument>.IndexKeys.Ascending(x => x.PrimaryClaim), new CreateIndexOptions { Unique = true }));
+			}
+		}
+
+		/// <inheritdoc/>
+		public async Task<IUser?> GetUserAsync(ObjectId Id)
+		{
+			return await Users.Find(x => x.Id == Id).FirstOrDefaultAsync();
+		}
+
+		/// <inheritdoc/>
+		public Task<IUser?> GetCachedUserAsync(ObjectId Id)
+		{
+			return GetUserAsync(Id);
+		}
+
+		/// <inheritdoc/>
+		public async Task<List<IUser>> FindUsersAsync(IEnumerable<ObjectId>? Ids, string? NameRegex, int? Index, int? Count)
+		{
+			FilterDefinition<UserDocument> Filter = Builders<UserDocument>.Filter.In(x => x.Id, Ids);
+			return await Users.Find(Filter).Range(Index, Count).ToListAsync<UserDocument, IUser>();
+		}
+
+		/// <inheritdoc/>
+		public async Task<IUser?> FindUserByLoginAsync(string Login)
+		{
+			ClaimDocument PrimaryClaim = new ClaimDocument(ClaimTypes.Name, Login);
+			return await Users.Find(x => x.PrimaryClaim == PrimaryClaim).FirstOrDefaultAsync();
+		}
+
+		/// <inheritdoc/>
+		public async Task<IUser> FindOrAddUserByLoginAsync(string Login, string? Name, string? Email)
+		{
+			ClaimDocument NewPrimaryClaim = new ClaimDocument(ClaimTypes.Name, Login);
+			UpdateDefinition<UserDocument> Update = Builders<UserDocument>.Update.SetOnInsert(x => x.Id, ObjectId.GenerateNewId());
+			return await Users.FindOneAndUpdateAsync<UserDocument>(x => x.PrimaryClaim == NewPrimaryClaim, Update, new FindOneAndUpdateOptions<UserDocument> { IsUpsert = true, ReturnDocument = ReturnDocument.After });
+		}
+
+		/// <inheritdoc/>
+		public async Task<IUserClaims> GetClaimsAsync(ObjectId UserId)
+		{
+			return await Users.Find(x => x.Id == UserId).FirstOrDefaultAsync() ?? new UserDocument { Id = UserId };
+		}
+
+		/// <inheritdoc/>
+		public async Task UpdateClaimsAsync(ObjectId UserId, IEnumerable<IUserClaim> Claims)
+		{
+			List<ClaimDocument> NewClaims = Claims.Select(x => new ClaimDocument(x)).ToList();
+			await Users.FindOneAndUpdateAsync(x => x.Id == UserId, Builders<UserDocument>.Update.Set(x => x.Claims, NewClaims));
+		}
+
+		/// <inheritdoc/>
+		public async Task<IUserSettings> GetSettingsAsync(ObjectId UserId)
+		{
+			return await Users.Find(x => x.Id == UserId).FirstOrDefaultAsync() ?? new UserDocument { Id = UserId };
+		}
+
+		/// <inheritdoc/>
+		public async Task UpdateSettingsAsync(ObjectId UserId, bool? EnableExperimentalFeatures, bool? EnableIssueNotifications, BsonValue? DashboardSettings = null, IEnumerable<ObjectId>? AddPinnedJobIds = null, IEnumerable<ObjectId>? RemovePinnedJobIds = null)
+		{
+			if (AddPinnedJobIds != null)
+			{
+				foreach (ObjectId PinnedJobId in AddPinnedJobIds)
+				{
+					FilterDefinition<UserDocument> Filter = Builders<UserDocument>.Filter.Eq(x => x.Id, UserId) & Builders<UserDocument>.Filter.AnyNin<ObjectId>(x => x.PinnedJobIds, new[] { PinnedJobId });
+					UpdateDefinition<UserDocument> Update = Builders<UserDocument>.Update.PushEach(x => x.PinnedJobIds, new[] { PinnedJobId }, -50);
+					await Users.UpdateOneAsync(Filter, Update);
+				}
+			}
+
+			List<UpdateDefinition<UserDocument>> Updates = new List<UpdateDefinition<UserDocument>>();
+			if (EnableExperimentalFeatures != null)
+			{
+				Updates.Add(Builders<UserDocument>.Update.SetOrUnsetNull(x => x.EnableExperimentalFeatures, EnableExperimentalFeatures));
+			}
+			if (EnableIssueNotifications != null)
+			{
+				Updates.Add(Builders<UserDocument>.Update.SetOrUnsetNull(x => x.EnableIssueNotifications, EnableIssueNotifications));
+			}
+			if (DashboardSettings != null)
+			{
+				Updates.Add(Builders<UserDocument>.Update.Set(x => x.DashboardSettings, DashboardSettings));
+			}
+			if (RemovePinnedJobIds != null && RemovePinnedJobIds.Any())
+			{
+				Updates.Add(Builders<UserDocument>.Update.PullAll(x => x.PinnedJobIds, RemovePinnedJobIds));
+			}
+			if (Updates.Count > 0)
+			{
+				await Users.UpdateOneAsync<UserDocument>(x => x.Id == UserId, Builders<UserDocument>.Update.Combine(Updates));
+			}
+		}
+
+		/// <summary>
+		/// Enumerate all the documents in this collection
+		/// </summary>
+		/// <returns></returns>
+		public async IAsyncEnumerable<(IUser, IUserClaims, IUserSettings)> EnumerateDocumentsAsync()
+		{
+			using (IAsyncCursor<UserDocument> Cursor = await Users.Find(FilterDefinition<UserDocument>.Empty).ToCursorAsync())
+			{
+				while (await Cursor.MoveNextAsync())
+				{
+					foreach (UserDocument Document in Cursor.Current)
+					{
+						if (Document.Claims.Count > 0)
+						{
+							yield return (Document, Document, Document);
+						}
+					}
+				}
+			}
+		}
+	}
+}
diff --git a/Engine/Source/Programs/Horde/HordeServer/Users/Impl/UserCollectionV2.cs b/Engine/Source/Programs/Horde/HordeServer/Users/Impl/UserCollectionV2.cs
new file mode 100644
--- /dev/null
+++ b/Engine/Source/Programs/Horde/HordeServer/Users/Impl/UserCollectionV2.cs
@@ -0,0 +1,367 @@
+// Copyright Epic Games, Inc. All Rights Reserved.
+
+using HordeServer.Models;
+using HordeServer.Services;
+using HordeServer.Utilities;
+using Microsoft.Extensions.Caching.Memory;
+using Microsoft.Extensions.Logging;
+using MongoDB.Bson;
+using MongoDB.Bson.Serialization;
+using MongoDB.Bson.Serialization.Attributes;
+using MongoDB.Driver;
+using System;
+using System.Collections.Generic;
+using System.Linq;
+using System.Security.Claims;
+using System.Threading.Tasks;
+
+namespace HordeServer.Collections.Impl
+{
+	/// <summary>
+	/// Manages user documents
+	/// </summary>
+	class UserCollectionV2 : IUserCollection, IDisposable
+	{
+		class UserDocument : IUser
+		{
+			public ObjectId Id { get; set; }
+
+			public string Name { get; set; }
+			public string Login { get; set; }
+			public string LoginUpper { get; set; }
+			public string? Email { get; set; }
+			public string? EmailUpper { get; set; }
+
+			[BsonIgnoreIfDefault]
+			public bool? Hidden { get; set; }
+
+			[BsonConstructor]
+			private UserDocument()
+			{
+				Name = null!;
+				Login = null!;
+				LoginUpper = null!;
+			}
+
+			public UserDocument(IUser Other)
+				: this(Other.Id, Other.Name, Other.Login, Other.Email)
+			{
+			}
+
+			public UserDocument(ObjectId Id, string Name, string Login, string? Email)
+			{
+				this.Id = Id;
+				this.Name = Name;
+				this.Login = Login;
+				this.LoginUpper = Login.ToUpperInvariant();
+				this.Email = Email;
+				this.EmailUpper = Email?.ToUpperInvariant();
+			}
+		}
+
+		class UserClaimsDocument : IUserClaims
+		{
+			public ObjectId Id { get; set; }
+			public List<UserClaim> Claims { get; set; } = new List<UserClaim>();
+
+			ObjectId IUserClaims.UserId => Id;
+			IReadOnlyList<IUserClaim> IUserClaims.Claims => Claims;
+
+			[BsonConstructor]
+			private UserClaimsDocument()
+			{
+			}
+
+			public UserClaimsDocument(ObjectId Id)
+			{
+				this.Id = Id;
+			}
+
+			public UserClaimsDocument(IUserClaims Other)
+				: this(Other.UserId)
+			{
+				this.Claims.AddRange(Other.Claims.Select(x => new UserClaim(x)));
+			}
+		}
+
+		class UserSettingsDocument : IUserSettings
+		{
+			public ObjectId Id { get; set; }
+
+			[BsonDefaultValue(false), BsonIgnoreIfDefault]
+			public bool EnableExperimentalFeatures { get; set; }
+
+			[BsonDefaultValue(true), BsonIgnoreIfDefault]
+			public bool EnableIssueNotifications { get; set; } = true;
+
+			public BsonValue DashboardSettings { get; set; } = BsonNull.Value;
+			public List<ObjectId> PinnedJobIds { get; set; } = new List<ObjectId>();
+
+			ObjectId IUserSettings.UserId => Id;
+			IReadOnlyList<ObjectId> IUserSettings.PinnedJobIds => PinnedJobIds;
+
+			[BsonConstructor]
+			private UserSettingsDocument()
+			{
+			}
+
+			public UserSettingsDocument(ObjectId Id)
+			{
+				this.Id = Id;
+			}
+
+			public UserSettingsDocument(IUserSettings Other)
+				: this(Other.UserId)
+			{
+				this.EnableExperimentalFeatures = Other.EnableExperimentalFeatures;
+				this.EnableIssueNotifications = Other.EnableIssueNotifications;
+				this.DashboardSettings = Other.DashboardSettings;
+				this.PinnedJobIds = new List<ObjectId>(Other.PinnedJobIds);
+			}
+		}
+
+		class ClaimDocument : IUserClaim
+		{
+			public string Type { get; set; }
+			public string Value { get; set; }
+
+			private ClaimDocument()
+			{
+				this.Type = null!;
+				this.Value = null!;
+			}
+
+			public ClaimDocument(string Type, string Value)
+			{
+				this.Type = Type;
+				this.Value = Value;
+			}
+
+			public ClaimDocument(IUserClaim Other)
+			{
+				this.Type = Other.Type;
+				this.Value = Other.Value;
+			}
+		}
+
+		IMongoCollection<UserDocument> Users;
+		IMongoCollection<UserClaimsDocument> UserClaims;
+		IMongoCollection<UserSettingsDocument> UserSettings;
+
+		ILogger<UserCollectionV2> Logger;
+
+		MemoryCache UserCache;
+
+		/// <summary>
+		/// Static constructor
+		/// </summary>
+		static UserCollectionV2()
+		{
+			BsonSerializer.RegisterDiscriminatorConvention(typeof(UserDocument), NullDiscriminatorConvention.Instance);
+			BsonSerializer.RegisterDiscriminatorConvention(typeof(UserClaimsDocument), NullDiscriminatorConvention.Instance);
+			BsonSerializer.RegisterDiscriminatorConvention(typeof(UserSettingsDocument), NullDiscriminatorConvention.Instance);
+
+			BsonSerializer.RegisterDiscriminatorConvention(typeof(ClaimDocument), NullDiscriminatorConvention.Instance);
+		}
+
+		/// <summary>
+		/// Constructor
+		/// </summary>
+		/// <param name="DatabaseService"></param>
+		/// <param name="Logger"></param>
+		public UserCollectionV2(DatabaseService DatabaseService, ILogger<UserCollectionV2> Logger)
+		{
+			this.Logger = Logger;
+
+			Users = DatabaseService.GetCollection<UserDocument>("UsersV2");
+			UserClaims = DatabaseService.GetCollection<UserClaimsDocument>("UserClaimsV2");
+			UserSettings = DatabaseService.GetCollection<UserSettingsDocument>("UserSettingsV2");
+
+			if (!DatabaseService.ReadOnlyMode)
+			{
+				Users.Indexes.CreateOne(new CreateIndexModel<UserDocument>(Builders<UserDocument>.IndexKeys.Ascending(x => x.LoginUpper), new CreateIndexOptions { Unique = true }));
+				Users.Indexes.CreateOne(new CreateIndexModel<UserDocument>(Builders<UserDocument>.IndexKeys.Ascending(x => x.EmailUpper)));
+			}
+
+			MemoryCacheOptions Options = new MemoryCacheOptions();
+			UserCache = new MemoryCache(Options);
+		}
+
+		/// <inheritdoc/>
+		public void Dispose()
+		{
+			UserCache.Dispose();
+		}
+
+		/// <inheritdoc/>
+		public async Task<IUser?> GetUserAsync(ObjectId Id)
+		{
+			IUser? User = await Users.Find(x => x.Id == Id).FirstOrDefaultAsync();
+			using (ICacheEntry Entry = UserCache.CreateEntry(Id))
+			{
+				Entry.SetValue(User);
+				Entry.SetSlidingExpiration(TimeSpan.FromMinutes(30.0));
+			}
+			return User;
+		}
+
+		/// <inheritdoc/>
+		public Task<IUser?> GetCachedUserAsync(ObjectId Id)
+		{
+			IUser? User;
+			if (UserCache.TryGetValue(Id, out User))
+			{
+				return Task.FromResult(User);
+			}
+			else
+			{
+				return GetUserAsync(Id);
+			}
+		}
+
+		/// <inheritdoc/>
+		public async Task<List<IUser>> FindUsersAsync(IEnumerable<ObjectId>? Ids, string? NameRegex = null, int? Index = null, int? Count = null)
+		{
+			FilterDefinition<UserDocument> Filter = FilterDefinition<UserDocument>.Empty;
+			if (Ids != null)
+			{
+				Filter &= Builders<UserDocument>.Filter.In(x => x.Id, Ids);
+			}
+			
+			if (NameRegex != null)
+			{
+				BsonRegularExpression Regex = new BsonRegularExpression(NameRegex, "i");
+				Filter &= Builders<UserDocument>.Filter.Regex(x => x.Name, Regex);
+			}
+
+			Filter &= Builders<UserDocument>.Filter.Ne(x => x.Hidden, true);
+
+			return await Users.Find(Filter).Range(Index, Count ?? 100).ToListAsync<UserDocument, IUser>();
+		}
+
+		/// <inheritdoc/>
+		public async Task<IUser?> FindUserByLoginAsync(string Login)
+		{
+			string LoginUpper = Login.ToUpperInvariant();
+			return await Users.Find(x => x.LoginUpper == LoginUpper).FirstOrDefaultAsync();
+		}
+
+		/// <inheritdoc/>
+		public async Task<IUser> FindOrAddUserByLoginAsync(string Login, string? Name, string? Email)
+		{
+			UpdateDefinition<UserDocument> Update = Builders<UserDocument>.Update.SetOnInsert(x => x.Id, ObjectId.GenerateNewId()).SetOnInsert(x => x.Login, Login).Unset(x => x.Hidden);
+
+			if (Name == null)
+			{
+				Update = Update.SetOnInsert(x => x.Name, Name ?? Login);
+			}
+			else
+			{
+				Update = Update.Set(x => x.Name, Name);
+			}
+
+			if (Email != null)
+			{
+				Update = Update.Set(x => x.Email, Email).Set(x => x.EmailUpper, Email.ToUpperInvariant());
+			}
+
+			string LoginUpper = Login.ToUpperInvariant();
+			return await Users.FindOneAndUpdateAsync<UserDocument>(x => x.LoginUpper == LoginUpper, Update, new FindOneAndUpdateOptions<UserDocument, UserDocument> { IsUpsert = true, ReturnDocument = ReturnDocument.After });
+		}
+
+		/// <inheritdoc/>
+		public async Task<IUserClaims> GetClaimsAsync(ObjectId UserId)
+		{
+			IUserClaims? Claims = await UserClaims.Find(x => x.Id == UserId).FirstOrDefaultAsync();
+			if(Claims == null)
+			{
+				Claims = new UserClaimsDocument(UserId);
+			}
+			return Claims;
+		}
+
+		/// <inheritdoc/>
+		public async Task UpdateClaimsAsync(ObjectId UserId, IEnumerable<IUserClaim> Claims)
+		{
+			UserClaimsDocument NewDocument = new UserClaimsDocument(UserId);
+			NewDocument.Claims.AddRange(Claims.Select(x => new UserClaim(x)));
+			await UserClaims.ReplaceOneAsync(x => x.Id == UserId, NewDocument, new ReplaceOptions { IsUpsert = true });
+		}
+
+		/// <inheritdoc/>
+		public async Task<IUserSettings> GetSettingsAsync(ObjectId UserId)
+		{
+			IUserSettings? Settings = await UserSettings.Find(x => x.Id == UserId).FirstOrDefaultAsync();
+			if (Settings == null)
+			{
+				Settings = new UserSettingsDocument(UserId);
+			}
+			return Settings;
+		}
+
+		/// <inheritdoc/>
+		public async Task UpdateSettingsAsync(ObjectId UserId, bool? EnableExperimentalFeatures = null, bool? EnableIssueNotifications = null, BsonValue? DashboardSettings = null, IEnumerable<ObjectId>? AddPinnedJobIds = null, IEnumerable<ObjectId>? RemovePinnedJobIds = null)
+		{
+			List<UpdateDefinition<UserSettingsDocument>> Updates = new List<UpdateDefinition<UserSettingsDocument>>();
+			if (EnableExperimentalFeatures != null)
+			{
+				Updates.Add(Builders<UserSettingsDocument>.Update.SetOrUnsetNull(x => x.EnableExperimentalFeatures, EnableExperimentalFeatures));
+			}
+			if (EnableIssueNotifications != null)
+			{
+				if (EnableIssueNotifications.Value)
+				{
+					Updates.Add(Builders<UserSettingsDocument>.Update.Unset(x => x.EnableIssueNotifications));
+				}
+				else
+				{
+					Updates.Add(Builders<UserSettingsDocument>.Update.Set(x => x.EnableIssueNotifications, false));
+				}
+			}
+			if (DashboardSettings != null)
+			{
+				Updates.Add(Builders<UserSettingsDocument>.Update.Set(x => x.DashboardSettings, DashboardSettings));
+			}
+			if (AddPinnedJobIds != null && AddPinnedJobIds.Any())
+			{
+				Updates.Add(Builders<UserSettingsDocument>.Update.AddToSetEach(x => x.PinnedJobIds, AddPinnedJobIds));
+			}
+			if (RemovePinnedJobIds != null && RemovePinnedJobIds.Any())
+			{
+				Updates.Add(Builders<UserSettingsDocument>.Update.PullAll(x => x.PinnedJobIds, RemovePinnedJobIds));
+			}
+			if (Updates.Count > 0)
+			{
+				await UserSettings.UpdateOneAsync<UserSettingsDocument>(x => x.Id == UserId, Builders<UserSettingsDocument>.Update.Combine(Updates), new UpdateOptions { IsUpsert = true });
+			}
+		}
+
+		/// <summary>
+		/// Upgrade from V1 collection
+		/// </summary>
+		/// <param name="UserCollectionV1"></param>
+		/// <returns></returns>
+		public async Task ResaveDocumentsAsync(UserCollectionV1 UserCollectionV1)
+		{
+			await foreach ((IUser User, IUserClaims Claims, IUserSettings Settings) in UserCollectionV1.EnumerateDocumentsAsync())
+			{
+				try
+				{
+					await Users.ReplaceOneAsync(x => x.Id == User.Id, new UserDocument(User), new ReplaceOptions { IsUpsert = true });
+					await UserClaims.ReplaceOneAsync(x => x.Id == User.Id, new UserClaimsDocument(Claims), new ReplaceOptions { IsUpsert = true });
+					await UserSettings.ReplaceOneAsync(x => x.Id == User.Id, new UserSettingsDocument(Settings), new ReplaceOptions { IsUpsert = true });
+					Logger.LogDebug("Updated user {UserId}", User.Id);
+				}
+				catch (MongoWriteException Ex)
+				{
+					Logger.LogWarning(Ex, "Unable to resave user {UserId}", User.Id);
+				}
+
+				if(Settings.PinnedJobIds.Count > 0)
+				{
+					await UpdateSettingsAsync(User.Id, AddPinnedJobIds: Settings.PinnedJobIds);
+				}
+			}
+		}
+	}
+}
