Horde: More renames.

diff --git a/Engine/Source/Programs/Horde/HordeServer/Agents/IFleetManager.cs b/Engine/Source/Programs/Horde/HordeServer/Agents/IFleetManager.cs
new file mode 100644
--- /dev/null
+++ b/Engine/Source/Programs/Horde/HordeServer/Agents/IFleetManager.cs
@@ -0,0 +1,35 @@
+// Copyright Epic Games, Inc. All Rights Reserved.
+
+using HordeServer.Models;
+using System;
+using System.Collections.Generic;
+using System.Linq;
+using System.Threading.Tasks;
+
+namespace HordeServer.Services
+{
+	/// <summary>
+	/// Service to manage a fleet of machines
+	/// </summary>
+	public interface IFleetManager
+	{
+		/// <summary>
+		/// Expand the given pool
+		/// </summary>
+		/// <param name="Pool">Pool to resize</param>
+		/// <param name="Agents">Current list of agents in the pool</param>
+		/// <param name="Count">Number of agents to add</param>
+		/// <returns>Async task</returns>
+		Task ExpandPool(IPool Pool, IReadOnlyList<IAgent> Agents, int Count);
+
+		/// <summary>
+		/// Shrink the given pool
+		/// </summary>
+		/// <param name="Pool">Pool to resize</param>
+		/// <param name="Agents">Current list of agents in the pool</param>
+		/// <param name="Count">Number of agents to remove</param>
+		/// <returns>Async task</returns>
+		Task ShrinkPool(IPool Pool, IReadOnlyList<IAgent> Agents, int Count);
+	}
+}
+
diff --git a/Engine/Source/Programs/Horde/HordeServer/Agents/Impl/Fleet/AwsFleetManager.cs b/Engine/Source/Programs/Horde/HordeServer/Agents/Impl/Fleet/AwsFleetManager.cs
new file mode 100644
--- /dev/null
+++ b/Engine/Source/Programs/Horde/HordeServer/Agents/Impl/Fleet/AwsFleetManager.cs
@@ -0,0 +1,104 @@
+// Copyright Epic Games, Inc. All Rights Reserved.
+
+using Amazon;
+using Amazon.EC2;
+using Amazon.EC2.Model;
+using HordeServer.Collections;
+using HordeServer.Models;
+using Microsoft.Extensions.Logging;
+using System;
+using System.Collections.Generic;
+using System.Linq;
+using System.Threading.Tasks;
+
+namespace HordeServer.Services.Impl
+{
+	/// <summary>
+	/// Fleet manager for handling AWS EC2 instances
+	/// </summary>
+	public sealed class AwsFleetManager : IFleetManager, IDisposable
+	{
+		const string AwsTagPropertyName = "aws-tag";
+		const string PoolTagName = "Horde_Autoscale_Pool";
+
+		AmazonEC2Client Client;
+		IAgentCollection AgentCollection;
+		ILogger Logger;
+
+		/// <summary>
+		/// Constructor
+		/// </summary>
+		public AwsFleetManager(IAgentCollection AgentCollection, ILogger<AwsFleetManager> Logger)
+		{
+			this.AgentCollection = AgentCollection;
+			this.Logger = Logger;
+
+			AmazonEC2Config Config = new AmazonEC2Config();
+			Config.RegionEndpoint = RegionEndpoint.USEast1;
+
+			Logger.LogInformation("Initializing AWS fleet manager for region {Region}", Config.RegionEndpoint);
+
+			Client = new AmazonEC2Client(Config);
+		}
+
+		/// <inheritdoc/>
+		public void Dispose()
+		{
+			Client.Dispose();
+		}
+
+		/// <inheritdoc/>
+		public async Task ExpandPool(IPool Pool, IReadOnlyList<IAgent> Agents, int Count)
+		{
+			// Find stopped instances in the correct pool
+			DescribeInstancesRequest DescribeRequest = new DescribeInstancesRequest();
+			DescribeRequest.Filters = new List<Filter>();
+			DescribeRequest.Filters.Add(new Filter("instance-state-name", new List<string> { InstanceStateName.Stopped.Value }));
+			DescribeRequest.Filters.Add(new Filter("tag:" + PoolTagName, new List<string> { Pool.Name }));
+			DescribeInstancesResponse DescribeResponse = await Client.DescribeInstancesAsync(DescribeRequest);
+
+			// Try to start the given instances
+			StartInstancesRequest StartRequest = new StartInstancesRequest();
+			StartRequest.InstanceIds.AddRange(DescribeResponse.Reservations.SelectMany(x => x.Instances).Select(x => x.InstanceId).Take(Count));
+
+			if (StartRequest.InstanceIds.Count > 0)
+			{
+				StartInstancesResponse StartResponse = await Client.StartInstancesAsync(StartRequest);
+				if ((int)StartResponse.HttpStatusCode >= 200 && (int)StartResponse.HttpStatusCode <= 299)
+				{
+					foreach (string InstanceId in StartRequest.InstanceIds)
+					{
+						Logger.LogInformation("Starting instance {InstanceId} for pool {PoolId}", InstanceId, Pool.Id);
+					}
+				}
+			}
+		}
+
+		/// <inheritdoc/>
+		public async Task ShrinkPool(IPool Pool, IReadOnlyList<IAgent> Agents, int Count)
+		{
+			string AwsTagProperty = $"{AwsTagPropertyName}={PoolTagName}:{Pool.Name}";
+
+			// Sort the agents by number of active leases. It's better to shutdown agents currently doing nothing.
+			List<IAgent> SortedAgents = Agents.OrderBy(x => x.Leases.Count).ToList();
+
+			// Try to remove the agents
+			for (int Idx = 0; Idx < SortedAgents.Count && Count > 0; Idx++)
+			{
+				IAgent Agent = SortedAgents[Idx];
+				if (Agent.HasProperty(AwsTagProperty))
+				{
+					for (IAgent? NewAgent = Agent; NewAgent != null; NewAgent = await AgentCollection.GetAsync(Agent.Id))
+					{
+						if (await AgentCollection.TryUpdateSettingsAsync(NewAgent, bRequestShutdown: true) != null)
+						{
+							AgentCollection.GetLogger(Agent.Id).LogInformation("Shutting down due to autoscaler");
+							Count--;
+							break;
+						}
+					}
+				}
+			}
+		}
+	}
+}
diff --git a/Engine/Source/Programs/Horde/HordeServer/Agents/Impl/Fleet/DefaultFleetManager.cs b/Engine/Source/Programs/Horde/HordeServer/Agents/Impl/Fleet/DefaultFleetManager.cs
new file mode 100644
--- /dev/null
+++ b/Engine/Source/Programs/Horde/HordeServer/Agents/Impl/Fleet/DefaultFleetManager.cs
@@ -0,0 +1,42 @@
+// Copyright Epic Games, Inc. All Rights Reserved.
+
+using HordeServer.Models;
+using Microsoft.Extensions.Logging;
+using System;
+using System.Collections.Generic;
+using System.Linq;
+using System.Threading.Tasks;
+
+namespace HordeServer.Services.Impl
+{
+	/// <summary>
+	/// Default implementation of <see cref="IFleetManager"/>
+	/// </summary>
+	public class DefaultFleetManager : IFleetManager
+	{
+		ILogger Logger;
+
+		/// <summary>
+		/// Constructor
+		/// </summary>
+		/// <param name="Logger">Logging device</param>
+		public DefaultFleetManager(ILogger<DefaultFleetManager> Logger)
+		{
+			this.Logger = Logger;
+		}
+
+		/// <inheritdoc/>
+		public Task ExpandPool(IPool Pool, IReadOnlyList<IAgent> Agents, int Count)
+		{
+			Logger.LogInformation("Expand pool {PoolId} by {Count} agents", Pool.Id, Count);
+			return Task.CompletedTask;
+		}
+
+		/// <inheritdoc/>
+		public Task ShrinkPool(IPool Pool, IReadOnlyList<IAgent> Agents, int Count)
+		{
+			Logger.LogInformation("Shrink pool {PoolId} by {Count} agents", Pool.Id, Count);
+			return Task.CompletedTask;
+		}
+	}
+}
diff --git a/Engine/Source/Programs/Horde/HordeServer/Collections/INotificationTriggerCollection.cs b/Engine/Source/Programs/Horde/HordeServer/Collections/INotificationTriggerCollection.cs
deleted file mode 100644
diff --git a/Engine/Source/Programs/Horde/HordeServer/Collections/IProjectCollection.cs b/Engine/Source/Programs/Horde/HordeServer/Collections/IProjectCollection.cs
deleted file mode 100644
diff --git a/Engine/Source/Programs/Horde/HordeServer/Collections/IServiceAccountCollection.cs b/Engine/Source/Programs/Horde/HordeServer/Collections/IServiceAccountCollection.cs
deleted file mode 100644
diff --git a/Engine/Source/Programs/Horde/HordeServer/Collections/IStreamCollection.cs b/Engine/Source/Programs/Horde/HordeServer/Collections/IStreamCollection.cs
deleted file mode 100644
diff --git a/Engine/Source/Programs/Horde/HordeServer/Collections/ISubscriptionCollection.cs b/Engine/Source/Programs/Horde/HordeServer/Collections/ISubscriptionCollection.cs
deleted file mode 100644
diff --git a/Engine/Source/Programs/Horde/HordeServer/Collections/ITemplateCollection.cs b/Engine/Source/Programs/Horde/HordeServer/Collections/ITemplateCollection.cs
deleted file mode 100644
diff --git a/Engine/Source/Programs/Horde/HordeServer/Collections/Impl/NotificationTriggerCollection.cs b/Engine/Source/Programs/Horde/HordeServer/Collections/Impl/NotificationTriggerCollection.cs
deleted file mode 100644
diff --git a/Engine/Source/Programs/Horde/HordeServer/Collections/Impl/ProjectCollection.cs b/Engine/Source/Programs/Horde/HordeServer/Collections/Impl/ProjectCollection.cs
deleted file mode 100644
diff --git a/Engine/Source/Programs/Horde/HordeServer/Collections/Impl/StreamCollection.cs b/Engine/Source/Programs/Horde/HordeServer/Collections/Impl/StreamCollection.cs
deleted file mode 100644
diff --git a/Engine/Source/Programs/Horde/HordeServer/Collections/Impl/SubscriptionCollection.cs b/Engine/Source/Programs/Horde/HordeServer/Collections/Impl/SubscriptionCollection.cs
deleted file mode 100644
diff --git a/Engine/Source/Programs/Horde/HordeServer/Collections/Impl/TemplateCollection.cs b/Engine/Source/Programs/Horde/HordeServer/Collections/Impl/TemplateCollection.cs
deleted file mode 100644
diff --git a/Engine/Source/Programs/Horde/HordeServer/Jobs/ITemplate.cs b/Engine/Source/Programs/Horde/HordeServer/Jobs/ITemplate.cs
new file mode 100644
--- /dev/null
+++ b/Engine/Source/Programs/Horde/HordeServer/Jobs/ITemplate.cs
@@ -0,0 +1,505 @@
+// Copyright Epic Games, Inc. All Rights Reserved.
+
+using EpicGames.Core;
+using HordeServer.Api;
+using HordeCommon;
+using HordeServer.Utilities;
+using MongoDB.Bson;
+using MongoDB.Bson.Serialization.Attributes;
+using System;
+using System.Collections.Generic;
+using System.Linq;
+using System.Runtime.CompilerServices;
+using System.Text.Json.Serialization;
+using System.Threading.Tasks;
+
+namespace HordeServer.Models
+{
+	/// <summary>
+	/// Base class for parameters used to configure templates via the new build dialog
+	/// </summary>
+	[BsonDiscriminator(RootClass = true)]
+	[BsonKnownTypes(typeof(GroupParameter), typeof(TextParameter), typeof(ListParameter), typeof(BoolParameter))]
+	public abstract class Parameter
+	{
+		/// <summary>
+		/// Gets the default arguments for this parameter and its children
+		/// </summary>
+		/// <param name="DefaultArguments">List of default arguments</param>
+		public abstract void GetDefaultArguments(List<string> DefaultArguments);
+
+		/// <summary>
+		/// Convert this parameter to data for serialization
+		/// </summary>
+		/// <returns>Serializable parameter data</returns>
+		public abstract ParameterData ToData();
+	}
+
+	/// <summary>
+	/// Used to group a number of other parameters
+	/// </summary>
+	public class GroupParameter : Parameter
+	{
+		/// <summary>
+		/// Label to display next to this parameter
+		/// </summary>
+		public string Label { get; set; }
+
+		/// <summary>
+		/// How to display this group
+		/// </summary>
+		public GroupParameterStyle Style { get; set; }
+
+		/// <summary>
+		/// List of child parameters
+		/// </summary>
+		public List<Parameter> Children { get; set; }
+
+		/// <summary>
+		/// Private constructor for serialization
+		/// </summary>
+		private GroupParameter()
+		{
+			this.Label = null!;
+			this.Children = null!;
+		}
+
+		/// <summary>
+		/// Constructor
+		/// </summary>
+		/// <param name="Label">Name of the group</param>
+		/// <param name="Style">How to display this group</param>
+		/// <param name="Children">List of child parameters</param>
+		public GroupParameter(string Label, GroupParameterStyle Style, List<Parameter> Children)
+		{
+			this.Label = Label;
+			this.Style = Style;
+			this.Children = Children;
+		}
+
+		/// <summary>
+		/// Gets the default arguments for this parameter and its children
+		/// </summary>
+		/// <param name="DefaultArguments">List of default arguments</param>
+		public override void GetDefaultArguments(List<string> DefaultArguments)
+		{
+			foreach (Parameter Child in Children)
+			{
+				Child.GetDefaultArguments(DefaultArguments);
+			}
+		}
+
+		/// <summary>
+		/// Convert this parameter to data for serialization
+		/// </summary>
+		/// <returns>Serializable parameter data</returns>
+		public override ParameterData ToData()
+		{
+			return new GroupParameterData(Label, Style, Children.ConvertAll(x => x.ToData()));
+		}
+	}
+
+	/// <summary>
+	/// Free-form text entry parameter
+	/// </summary>
+	public sealed class TextParameter : Parameter
+	{
+		/// <summary>
+		/// Label to display next to this parameter. Should default to the parameter name.
+		/// </summary>
+		public string Label { get; set; }
+
+		/// <summary>
+		/// Argument to add (will have the value of this field appended)
+		/// </summary>
+		public string Argument { get; set; }
+
+		/// <summary>
+		/// Default value for this argument
+		/// </summary>
+		public string Default { get; set; }
+
+		/// <summary>
+		/// Hint text to display when the field is empty
+		/// </summary>
+		public string? Hint { get; set; }
+
+		/// <summary>
+		/// Regex used to validate values entered into this text field.
+		/// </summary>
+		[BsonIgnoreIfNull]
+		public string? Validation { get; set; }
+
+		/// <summary>
+		/// Message displayed to explain valid values if validation fails.
+		/// </summary>
+		[BsonIgnoreIfNull]
+		public string? Description { get; set; }
+
+		/// <summary>
+		/// Tool tip text to display
+		/// </summary>
+		public string? ToolTip { get; set; }
+
+		/// <summary>
+		/// Private constructor for serialization
+		/// </summary>
+		private TextParameter()
+		{
+			Label = null!;
+			Argument = null!;
+			Default = String.Empty;
+		}
+
+		/// <summary>
+		/// Constructor
+		/// </summary>
+		/// <param name="Label">Parameter to pass this value to the BuildGraph script with</param>
+		/// <param name="Argument">Label to show next to the parameter</param>
+		/// <param name="Default">Default value for this argument</param>
+		/// <param name="Hint">Hint text to display</param>
+		/// <param name="Validation">Regex used to validate entries</param>
+		/// <param name="Description">Message displayed for invalid values</param>
+		/// <param name="ToolTip">Tool tip text to display</param>
+		public TextParameter(string Label, string Argument, string Default, string? Hint, string? Validation, string? Description, string? ToolTip)
+		{
+			this.Label = Label;
+			this.Argument = Argument;
+			this.Default = Default;
+			this.Hint = Hint;
+			this.Validation = Validation;
+			this.Description = Description;
+			this.ToolTip = ToolTip;
+		}
+
+		/// <summary>
+		/// Gets the default arguments for this parameter and its children
+		/// </summary>
+		/// <param name="DefaultArguments">List of default arguments</param>
+		public override void GetDefaultArguments(List<string> DefaultArguments)
+		{
+			DefaultArguments.Add(Argument + Default);
+		}
+
+		/// <summary>
+		/// Convert this parameter to data for serialization
+		/// </summary>
+		/// <returns>Serializable parameter data</returns>
+		public override ParameterData ToData()
+		{
+			return new TextParameterData(Label, Argument, Default, Hint, Validation, Description, ToolTip);
+		}
+	}
+
+	/// <summary>
+	/// Possible option for a list parameter
+	/// </summary>
+	public class ListParameterItem
+	{
+		/// <summary>
+		/// Group to display this entry in
+		/// </summary>
+		[BsonIgnoreIfNull]
+		public string? Group { get; set; }
+
+		/// <summary>
+		/// Text to display for this option.
+		/// </summary>
+		public string Text { get; set; }
+
+		/// <summary>
+		/// Argument to add with this parameter.
+		/// </summary>
+		[BsonIgnoreIfNull]
+		public string? ArgumentIfEnabled { get; set; }
+
+		/// <summary>
+		/// Argument to add with this parameter.
+		/// </summary>
+		[BsonIgnoreIfNull]
+		public string? ArgumentIfDisabled { get; set; }
+
+		/// <summary>
+		/// Whether this item is selected by default
+		/// </summary>
+		[BsonIgnoreIfDefault, BsonDefaultValue(false)]
+		public bool Default { get; set; }
+
+		/// <summary>
+		/// Constructor
+		/// </summary>
+		/// <param name="Group">The group to put this parameter in</param>
+		/// <param name="Text">Text to display for this option</param>
+		/// <param name="ArgumentIfEnabled">Argument to add if this option is enabled</param>
+		/// <param name="ArgumentIfDisabled">Argument to add if this option is disabled</param>
+		/// <param name="Default">Whether this item is selected by default</param>
+		public ListParameterItem(string? Group, string Text, string? ArgumentIfEnabled, string? ArgumentIfDisabled, bool Default)
+		{
+			this.Group = Group;
+			this.Text = Text;
+			this.ArgumentIfEnabled = ArgumentIfEnabled;
+			this.ArgumentIfDisabled = ArgumentIfDisabled;
+			this.Default = Default;
+		}
+
+		/// <summary>
+		/// Convert this parameter to data for serialization
+		/// </summary>
+		/// <returns>Serializable parameter data</returns>
+		public ListParameterItemData ToData()
+		{
+			return new ListParameterItemData(Group, Text, ArgumentIfEnabled, ArgumentIfDisabled, Default);
+		}
+	}
+
+	/// <summary>
+	/// Allows the user to select a value from a constrained list of choices
+	/// </summary>
+	public class ListParameter : Parameter
+	{
+		/// <summary>
+		/// Label to display next to this parameter.
+		/// </summary>
+		public string Label { get; set; }
+
+		/// <summary>
+		/// Style of picker parameter to use
+		/// </summary>
+		public ListParameterStyle Style { get; set; }
+
+		/// <summary>
+		/// List of values to display in the list
+		/// </summary>
+		public List<ListParameterItem> Items { get; set; }
+
+		/// <summary>
+		/// Tool tip text to display
+		/// </summary>
+		public string? ToolTip { get; set; }
+
+		/// <summary>
+		/// Private constructor for serialization
+		/// </summary>
+		private ListParameter()
+		{
+			this.Label = null!;
+			this.Items = null!;
+		}
+
+		/// <summary>
+		/// List of possible values
+		/// </summary>
+		/// <param name="Label">Label to show next to this parameter</param>
+		/// <param name="Style">Style of list parameter to use</param>
+		/// <param name="Entries">Entries for this list</param>
+		/// <param name="ToolTip">Tool tip text to display</param>
+		public ListParameter(string Label, ListParameterStyle Style, List<ListParameterItem> Entries, string? ToolTip)
+		{
+			this.Label = Label;
+			this.Style = Style;
+			this.Items = Entries;
+			this.ToolTip = ToolTip;
+		}
+
+		/// <summary>
+		/// Gets the default arguments for this parameter and its children
+		/// </summary>
+		/// <param name="DefaultArguments">List of default arguments</param>
+		public override void GetDefaultArguments(List<string> DefaultArguments)
+		{
+			foreach(ListParameterItem Item in Items)
+			{
+				if (Item.Default)
+				{
+					if (Item.ArgumentIfEnabled != null)
+					{
+						DefaultArguments.Add(Item.ArgumentIfEnabled);
+					}
+				}
+				else
+				{
+					if (Item.ArgumentIfDisabled != null)
+					{
+						DefaultArguments.Add(Item.ArgumentIfDisabled);
+					}
+				}
+			}
+		}
+
+		/// <summary>
+		/// Convert this parameter to data for serialization
+		/// </summary>
+		/// <returns>Serializable parameter data</returns>
+		public override ParameterData ToData()
+		{
+			return new ListParameterData(Label, Style, Items.ConvertAll(x => x.ToData()), ToolTip);
+		}
+	}
+
+	/// <summary>
+	/// Allows the user to toggle an option on or off
+	/// </summary>
+	public class BoolParameter : Parameter
+	{
+		/// <summary>
+		/// Label to display next to this parameter.
+		/// </summary>
+		public string Label { get; set; }
+
+		/// <summary>
+		/// Value if enabled
+		/// </summary>
+		[BsonIgnoreIfNull]
+		public string? ArgumentIfEnabled { get; set; }
+
+		/// <summary>
+		/// Value if disabled
+		/// </summary>
+		[BsonIgnoreIfNull]
+		public string? ArgumentIfDisabled { get; set; }
+
+		/// <summary>
+		/// Whether this option should be enabled by default
+		/// </summary>
+		[BsonIgnoreIfDefault, BsonDefaultValue(false)]
+		public bool Default { get; set; }
+
+		/// <summary>
+		/// Tool tip text to display
+		/// </summary>
+		public string? ToolTip { get; set; }
+
+		/// <summary>
+		/// Private constructor for serialization
+		/// </summary>
+		private BoolParameter()
+		{
+			Label = null!;
+		}
+
+		/// <summary>
+		/// Constructor
+		/// </summary>
+		/// <param name="Label">Label to display for this parameter</param>
+		/// <param name="ArgumentIfEnabled">Value if enabled</param>
+		/// <param name="ArgumentIfDisabled">Value if disabled</param>
+		/// <param name="Default">Default value for this argument</param>
+		/// <param name="ToolTip">Tool tip text to display</param>
+		public BoolParameter(string Label, string? ArgumentIfEnabled, string? ArgumentIfDisabled, bool Default, string? ToolTip)
+		{
+			this.Label = Label;
+			this.ArgumentIfEnabled = ArgumentIfEnabled;
+			this.ArgumentIfDisabled = ArgumentIfDisabled;
+			this.Default = Default;
+			this.ToolTip = ToolTip;
+		}
+
+		/// <summary>
+		/// Gets the default arguments for this parameter and its children
+		/// </summary>
+		/// <param name="DefaultArguments">List of default arguments</param>
+		public override void GetDefaultArguments(List<string> DefaultArguments)
+		{
+			string? DefaultArgument = Default ? ArgumentIfEnabled : ArgumentIfDisabled;
+			if (!String.IsNullOrEmpty(DefaultArgument))
+			{
+				DefaultArguments.Add(DefaultArgument);
+			}
+		}
+		/// <summary>
+		/// Convert this parameter to data for serialization
+		/// </summary>
+		/// <returns><see cref="BoolParameterData"/> instance</returns>
+		public override ParameterData ToData()
+		{
+			return new BoolParameterData(Label, ArgumentIfEnabled, ArgumentIfDisabled, Default, ToolTip);
+		}
+	}
+
+	/// <summary>
+	/// Counter to be updated when executing a template
+	/// </summary>
+	public interface ITemplateCounter
+	{
+		/// <summary>
+		/// Name of the property to set to the counter value
+		/// </summary>
+		string PropertyName { get; }
+
+		/// <summary>
+		/// Name of the counter
+		/// </summary>
+		string CounterName { get; }
+	}
+
+	/// <summary>
+	/// Document describing a job template. These objects are considered immutable once created and uniquely referenced by hash, in order to de-duplicate across all job runs.
+	/// </summary>
+	public interface ITemplate
+	{
+		/// <summary>
+		/// Hash of this template
+		/// </summary>
+		public ContentHash Id { get; }
+
+		/// <summary>
+		/// Name of the template.
+		/// </summary>
+		public string Name { get; }
+
+		/// <summary>
+		/// Priority of this job
+		/// </summary>
+		public Priority? Priority { get; }
+
+		/// <summary>
+		/// Whether to allow preflights for this job type
+		/// </summary>
+		public bool AllowPreflights { get; }
+
+		/// <summary>
+		/// Agent type to use for parsing the job state
+		/// </summary>
+		public string? InitialAgentType { get; }
+
+		/// <summary>
+		/// Path to a file within the stream to submit to generate a new changelist for jobs
+		/// </summary>
+		public string? SubmitNewChange { get; }
+
+		/// <summary>
+		/// List of property counters
+		/// </summary>
+		public IReadOnlyList<ITemplateCounter> Counters { get; }
+
+		/// <summary>
+		/// Optional predefined user-defined properties for this job
+		/// </summary>
+		public IReadOnlyList<string> Arguments { get; }
+
+		/// <summary>
+		/// Parameters for this template
+		/// </summary>
+		public IReadOnlyList<Parameter> Parameters { get; }
+	}
+
+	/// <summary>
+	/// Extension methods for templates
+	/// </summary>
+	public static class TemplateExtensions
+	{
+		/// <summary>
+		/// Gets the arguments for default options in this template. Does not include the standard template arguments.
+		/// </summary>
+		/// <returns>List of default arguments</returns>
+		public static List<string> GetDefaultArguments(this ITemplate Template)
+		{
+			List<string> DefaultArguments = new List<string>(Template.Arguments);
+			foreach (Parameter Parameter in Template.Parameters)
+			{
+				Parameter.GetDefaultArguments(DefaultArguments);
+			}
+			return DefaultArguments;
+		}
+	}
+}
diff --git a/Engine/Source/Programs/Horde/HordeServer/Jobs/ITemplateCollection.cs b/Engine/Source/Programs/Horde/HordeServer/Jobs/ITemplateCollection.cs
new file mode 100644
--- /dev/null
+++ b/Engine/Source/Programs/Horde/HordeServer/Jobs/ITemplateCollection.cs
@@ -0,0 +1,51 @@
+// Copyright Epic Games, Inc. All Rights Reserved.
+
+using EpicGames.Core;
+using HordeCommon;
+using HordeServer.Api;
+using HordeServer.Models;
+using HordeServer.Services;
+using HordeServer.Utilities;
+using MongoDB.Bson;
+using MongoDB.Bson.Serialization;
+using MongoDB.Bson.Serialization.Attributes;
+using MongoDB.Driver;
+using System;
+using System.Collections.Generic;
+using System.Linq;
+using System.Threading.Tasks;
+
+namespace HordeServer.Collections
+{
+	/// <summary>
+	/// Interface for a collection of template documents
+	/// </summary>
+	public interface ITemplateCollection
+	{
+		/// <summary>
+		/// Public constructor
+		/// </summary>
+		/// <param name="Name">Name of the template</param>
+		/// <param name="Priority">Priority of this template</param>
+		/// <param name="bAllowPreflights">Whether to allow preflights of this job</param>
+		/// <param name="InitialAgentType">The agent type to parse the buildgraph script</param>
+		/// <param name="SubmitNewChange">Path to a file within the stream to submit to generate a new changelist for jobs</param>
+		/// <param name="Counters">List of counters for new jobs</param>
+		/// <param name="Arguments">List of arguments which are always specified</param>
+		/// <param name="Parameters">List of template parameters</param>
+		Task<ITemplate> AddAsync(string Name, Priority? Priority = null, bool bAllowPreflights = true, string? InitialAgentType = null, string? SubmitNewChange = null, List<TemplateCounter>? Counters = null, List<string>? Arguments = null, List<Parameter>? Parameters = null);
+
+		/// <summary>
+		/// Gets all the available templates
+		/// </summary>
+		/// <returns>List of template documents</returns>
+		Task<List<ITemplate>> FindAllAsync();
+
+		/// <summary>
+		/// Gets a template by ID
+		/// </summary>
+		/// <param name="TemplateId">Unique id of the template</param>
+		/// <returns>The template document</returns>
+		Task<ITemplate?> GetAsync(ContentHash TemplateId);
+	}
+}
diff --git a/Engine/Source/Programs/Horde/HordeServer/Jobs/Impl/TemplateCollection.cs b/Engine/Source/Programs/Horde/HordeServer/Jobs/Impl/TemplateCollection.cs
new file mode 100644
--- /dev/null
+++ b/Engine/Source/Programs/Horde/HordeServer/Jobs/Impl/TemplateCollection.cs
@@ -0,0 +1,169 @@
+// Copyright Epic Games, Inc. All Rights Reserved.
+
+using EpicGames.Core;
+using HordeCommon;
+using HordeServer.Api;
+using HordeServer.Models;
+using HordeServer.Services;
+using HordeServer.Utilities;
+using Microsoft.Extensions.Caching.Memory;
+using MongoDB.Bson;
+using MongoDB.Bson.Serialization;
+using MongoDB.Bson.Serialization.Attributes;
+using MongoDB.Driver;
+using System;
+using System.Collections.Generic;
+using System.Linq;
+using System.Threading.Tasks;
+
+namespace HordeServer.Collections.Impl
+{
+	/// <summary>
+	/// Collection of template documents
+	/// </summary>
+	public sealed class TemplateCollection : ITemplateCollection, IDisposable
+	{
+		/// <summary>
+		/// Counter instance
+		/// </summary>
+		class TemplateCounterDocument : ITemplateCounter
+		{
+			public string PropertyName { get; set; }
+			public string CounterName { get; set; }
+
+			private TemplateCounterDocument()
+			{
+				PropertyName = null!;
+				CounterName = null!;
+			}
+
+			public TemplateCounterDocument(string PropertyName, string CounterName)
+			{
+				this.PropertyName = PropertyName;
+				this.CounterName = CounterName;
+			}
+		}
+
+		/// <summary>
+		/// Document describing a job template. These objects are considered immutable once created and uniquely referenced by hash, in order to de-duplicate across all job runs.
+		/// </summary>
+		class TemplateDocument : ITemplate
+		{
+			[BsonRequired, BsonId]
+			public ContentHash Id { get; private set; } = ContentHash.Empty;
+
+			[BsonRequired]
+			public string Name { get; private set; }
+
+			public Priority? Priority { get; private set; }
+			public bool AllowPreflights { get; set; } = true;
+			public string? InitialAgentType { get; set; }
+
+			[BsonIgnoreIfNull]
+			public string? SubmitNewChange { get; set; }
+
+			public List<TemplateCounterDocument> Counters { get; set; } = new List<TemplateCounterDocument>();
+			public List<string> Arguments { get; private set; } = new List<string>();
+			public List<Parameter> Parameters { get; private set; } = new List<Parameter>();
+
+			IReadOnlyList<ITemplateCounter> ITemplate.Counters => Counters;
+			IReadOnlyList<string> ITemplate.Arguments => Arguments;
+			IReadOnlyList<Parameter> ITemplate.Parameters => Parameters;
+
+			[BsonConstructor]
+			private TemplateDocument()
+			{
+				Name = null!;
+			}
+
+			public TemplateDocument(string Name, Priority? Priority, bool bAllowPreflights, string? InitialAgentType, string? SubmitNewChange, List<TemplateCounter>? Counters, List<string>? Arguments, List<Parameter>? Parameters)
+			{
+				this.Name = Name;
+				this.Priority = Priority;
+				this.AllowPreflights = bAllowPreflights;
+				this.InitialAgentType = InitialAgentType;
+				this.SubmitNewChange = SubmitNewChange;
+				this.Counters = (Counters ?? new List<TemplateCounter>()).ConvertAll(x => new TemplateCounterDocument(x.PropertyName, x.CounterName));
+				this.Arguments = Arguments ?? new List<string>();
+				this.Parameters = Parameters ?? new List<Parameter>();
+
+				// Compute the hash once all other fields have been set
+				this.Id = ContentHash.SHA1(BsonExtensionMethods.ToBson(this));
+			}
+		}
+
+		/// <summary>
+		/// Template documents
+		/// </summary>
+		IMongoCollection<TemplateDocument> Templates;
+
+		/// <summary>
+		/// Cache of template documents
+		/// </summary>
+		MemoryCache TemplateCache;
+
+		/// <summary>
+		/// Constructor
+		/// </summary>
+		/// <param name="DatabaseService">The database service singleton</param>
+		public TemplateCollection(DatabaseService DatabaseService)
+		{
+			// Ensure discriminator cannot be registered twice (throws exception). Can otherwise happen during unit tests.
+			if (BsonSerializer.LookupDiscriminatorConvention(typeof(JobsTabColumn)) == null)
+			{
+				BsonSerializer.RegisterDiscriminatorConvention(typeof(JobsTabColumn), new DefaultDiscriminatorConvention(typeof(JobsTabColumn), typeof(JobsTabLabelColumn)));	
+			}
+			
+			Templates = DatabaseService.GetCollection<TemplateDocument>("Templates");
+
+			MemoryCacheOptions Options = new MemoryCacheOptions();
+			TemplateCache = new MemoryCache(Options);
+		}
+
+		/// <inheritdoc/>
+		public void Dispose()
+		{
+			TemplateCache.Dispose();
+		}
+
+		/// <inheritdoc/>
+		public async Task<ITemplate> AddAsync(string Name, Priority? Priority, bool bAllowPreflights, string? InitialAgentType, string? SubmitNewChange, List<TemplateCounter>? Counters, List<string>? Arguments, List<Parameter>? Parameters)
+		{
+			TemplateDocument Template = new TemplateDocument(Name, Priority, bAllowPreflights, InitialAgentType, SubmitNewChange, Counters, Arguments, Parameters);
+			if (await GetAsync(Template.Id) == null)
+			{
+				await Templates.ReplaceOneAsync(x => x.Id == Template.Id, Template, new ReplaceOptions { IsUpsert = true });
+			}
+			return Template;
+		}
+
+		/// <inheritdoc/>
+		public async Task<List<ITemplate>> FindAllAsync()
+		{
+			List<TemplateDocument> Results = await Templates.Find(FilterDefinition<TemplateDocument>.Empty).ToListAsync();
+			return Results.ConvertAll<ITemplate>(x => x);
+		}
+
+		/// <inheritdoc/>
+		public async Task<ITemplate?> GetAsync(ContentHash TemplateId)
+		{
+			object? Result;
+			if (TemplateCache.TryGetValue(TemplateId, out Result))
+			{
+				return (ITemplate?)Result;
+			}
+
+			ITemplate? Template = await Templates.Find<TemplateDocument>(x => x.Id == TemplateId).FirstOrDefaultAsync();
+			if (Template != null)
+			{
+				using (ICacheEntry Entry = TemplateCache.CreateEntry(TemplateId))
+				{
+					Entry.SetSlidingExpiration(TimeSpan.FromMinutes(30.0));
+					Entry.SetValue(Template);
+				}
+			}
+
+			return Template;
+		}
+	}
+}
diff --git a/Engine/Source/Programs/Horde/HordeServer/Models/IEvent.cs b/Engine/Source/Programs/Horde/HordeServer/Models/IEvent.cs
deleted file mode 100644
diff --git a/Engine/Source/Programs/Horde/HordeServer/Models/INotificationTrigger.cs b/Engine/Source/Programs/Horde/HordeServer/Models/INotificationTrigger.cs
deleted file mode 100644
diff --git a/Engine/Source/Programs/Horde/HordeServer/Models/IProject.cs b/Engine/Source/Programs/Horde/HordeServer/Models/IProject.cs
deleted file mode 100644
diff --git a/Engine/Source/Programs/Horde/HordeServer/Models/IServiceAccount.cs b/Engine/Source/Programs/Horde/HordeServer/Models/IServiceAccount.cs
deleted file mode 100644
diff --git a/Engine/Source/Programs/Horde/HordeServer/Models/IStream.cs b/Engine/Source/Programs/Horde/HordeServer/Models/IStream.cs
deleted file mode 100644
diff --git a/Engine/Source/Programs/Horde/HordeServer/Models/ISubscription.cs b/Engine/Source/Programs/Horde/HordeServer/Models/ISubscription.cs
deleted file mode 100644
diff --git a/Engine/Source/Programs/Horde/HordeServer/Models/Template.cs b/Engine/Source/Programs/Horde/HordeServer/Models/Template.cs
deleted file mode 100644
diff --git a/Engine/Source/Programs/Horde/HordeServer/Notifications/IEvent.cs b/Engine/Source/Programs/Horde/HordeServer/Notifications/IEvent.cs
new file mode 100644
--- /dev/null
+++ b/Engine/Source/Programs/Horde/HordeServer/Notifications/IEvent.cs
@@ -0,0 +1,85 @@
+// Copyright Epic Games, Inc. All Rights Reserved.
+
+using HordeServer.Api;
+using HordeServer.Utilities;
+using System;
+using System.Collections.Generic;
+using System.Diagnostics.CodeAnalysis;
+using System.Linq;
+using System.Threading.Tasks;
+
+namespace HordeServer.Models
+{
+	using StreamId = StringId<IStream>;
+	using TemplateRefId = StringId<TemplateRef>;
+
+	/// <summary>
+	/// Base interface for events that users can subscribe for notifications to
+	/// </summary>
+	public interface IEvent
+	{
+		/// <summary>
+		/// Converts this event into an API event record
+		/// </summary>
+		/// <returns>The API event record</returns>
+		public EventRecord ToRecord();
+	}
+
+	/// <summary>
+	/// Event that occurs when a job completes
+	/// </summary>
+	public interface IJobCompleteEvent : IEvent
+	{
+		/// <summary>
+		/// The stream id
+		/// </summary>
+		StreamId StreamId { get; }
+
+		/// <summary>
+		/// The template id
+		/// </summary>
+		TemplateRefId TemplateId { get; }
+	}
+
+	/// <summary>
+	/// Event that occurs when a label completes
+	/// </summary>
+	public interface ILabelCompleteEvent : IEvent
+	{
+		/// <summary>
+		/// The stream id
+		/// </summary>
+		StreamId StreamId { get; }
+
+		/// <summary>
+		/// The template id
+		/// </summary>
+		TemplateRefId TemplateId { get; }
+
+		/// <summary>
+		/// Name of the label to monitor
+		/// </summary>
+		string LabelName { get; }
+	}
+
+	/// <summary>
+	/// Event that occurs when a job step completes
+	/// </summary>
+	public interface IStepCompleteEvent : IEvent
+	{
+		/// <summary>
+		/// The stream id
+		/// </summary>
+		StreamId StreamId { get; }
+
+		/// <summary>
+		/// The template id
+		/// </summary>
+		TemplateRefId TemplateId { get; }
+
+		/// <summary>
+		/// Name of the step to monitor
+		/// </summary>
+		string StepName { get; }
+	}
+}
diff --git a/Engine/Source/Programs/Horde/HordeServer/Notifications/INotificationTrigger.cs b/Engine/Source/Programs/Horde/HordeServer/Notifications/INotificationTrigger.cs
new file mode 100644
--- /dev/null
+++ b/Engine/Source/Programs/Horde/HordeServer/Notifications/INotificationTrigger.cs
@@ -0,0 +1,51 @@
+// Copyright Epic Games, Inc. All Rights Reserved.
+
+using HordeServer.Api;
+using MongoDB.Bson;
+using MongoDB.Bson.Serialization.Attributes;
+using System.Collections.Generic;
+
+namespace HordeServer.Models
+{
+	/// <summary>
+	/// An individual subscription
+	/// </summary>
+	public interface INotificationSubscription
+	{
+		/// <summary>
+		/// Name of the user
+		/// </summary>
+		public ObjectId UserId { get; }
+
+		/// <summary>
+		/// Whether to receive email notifications
+		/// </summary>
+		public bool Email { get; }
+
+		/// <summary>
+		/// Whether to receive Slack notifications
+		/// </summary>
+		public bool Slack { get; }
+	}
+
+	/// <summary>
+	/// Trigger for notifications to be sent
+	/// </summary>
+	public interface INotificationTrigger
+	{
+		/// <summary>
+		/// Unique id for this subscription list
+		/// </summary>
+		ObjectId Id { get; }
+
+		/// <summary>
+		/// Whether this trigger has been fired
+		/// </summary>
+		bool Fired { get; }
+
+		/// <summary>
+		/// List of subscriptions to this event
+		/// </summary>
+		public IReadOnlyList<INotificationSubscription> Subscriptions { get; }
+	}
+}
diff --git a/Engine/Source/Programs/Horde/HordeServer/Notifications/INotificationTriggerCollection.cs b/Engine/Source/Programs/Horde/HordeServer/Notifications/INotificationTriggerCollection.cs
new file mode 100644
--- /dev/null
+++ b/Engine/Source/Programs/Horde/HordeServer/Notifications/INotificationTriggerCollection.cs
@@ -0,0 +1,66 @@
+// Copyright Epic Games, Inc. All Rights Reserved.
+
+using HordeServer.Models;
+using HordeServer.Services;
+using HordeServer.Utilities;
+using MongoDB.Bson;
+using MongoDB.Bson.Serialization.Attributes;
+using MongoDB.Driver;
+using System;
+using System.Collections.Generic;
+using System.Linq;
+using System.Threading.Tasks;
+
+namespace HordeServer.Collections
+{
+	/// <summary>
+	/// Collection of notification triggers
+	/// </summary>
+	public interface INotificationTriggerCollection
+	{
+		/// <summary>
+		/// Finds or adds a trigger
+		/// </summary>
+		/// <param name="TriggerId">The trigger id</param>
+		/// <returns>New trigger document</returns>
+		Task<INotificationTrigger> FindOrAddAsync(ObjectId TriggerId);
+
+		/// <summary>
+		/// Finds an existing trigger with the given id, or adds one if it does not exist
+		/// </summary>
+		/// <param name="TriggerId"></param>
+		/// <returns></returns>
+		Task<INotificationTrigger?> GetAsync(ObjectId TriggerId);
+
+		/// <summary>
+		/// Deletes a trigger
+		/// </summary>
+		/// <param name="TriggerId">The unique trigger id</param>
+		/// <returns>Async task</returns>
+		Task DeleteAsync(ObjectId TriggerId);
+
+		/// <summary>
+		/// Deletes a trigger
+		/// </summary>
+		/// <param name="TriggerIds">The unique trigger id</param>
+		/// <returns>Async task</returns>
+		Task DeleteAsync(List<ObjectId> TriggerIds);
+
+		/// <summary>
+		/// Fires the trigger, and marks notifications has having been sent
+		/// </summary>
+		/// <param name="Trigger">The trigger to fire</param>
+		/// <returns>The trigger document</returns>
+		Task<INotificationTrigger?> FireAsync(INotificationTrigger Trigger);
+
+		/// <summary>
+		/// Adds a subscriber to a particular trigger
+		/// </summary>
+		/// <param name="Trigger">The trigger to subscribe to</param>
+		/// <param name="UserId">The user name</param>
+		/// <param name="Email">Whether to receive email notifications</param>
+		/// <param name="Slack">Whether to receive Slack notifications</param>
+		/// <returns>The new trigger state</returns>
+		Task<INotificationTrigger?> UpdateSubscriptionsAsync(INotificationTrigger Trigger, ObjectId UserId, bool? Email, bool? Slack);
+	}
+}
diff --git a/Engine/Source/Programs/Horde/HordeServer/Notifications/ISubscription.cs b/Engine/Source/Programs/Horde/HordeServer/Notifications/ISubscription.cs
new file mode 100644
--- /dev/null
+++ b/Engine/Source/Programs/Horde/HordeServer/Notifications/ISubscription.cs
@@ -0,0 +1,49 @@
+// Copyright Epic Games, Inc. All Rights Reserved.
+
+using HordeCommon;
+using HordeServer.Api;
+using MongoDB.Bson;
+using System;
+using System.Collections.Generic;
+using System.Linq;
+using System.Threading.Tasks;
+
+namespace HordeServer.Models
+{
+	/// <summary>
+	/// The type of notification to send
+	/// </summary>
+	public enum NotificationType
+	{
+		/// <summary>
+		/// Send a DM on Slack
+		/// </summary>
+		Slack,
+	}
+
+	/// <summary>
+	/// Subscription to a type of event
+	/// </summary>
+	public interface ISubscription
+	{
+		/// <summary>
+		/// Unique id for this subscription
+		/// </summary>
+		public string Id { get; }
+
+		/// <summary>
+		/// Name of the event to subscribe to
+		/// </summary>
+		public IEvent Event { get; }
+
+		/// <summary>
+		/// User to notify
+		/// </summary>
+		public ObjectId UserId { get; }
+
+		/// <summary>
+		/// Type of notification to receive
+		/// </summary>
+		public NotificationType NotificationType { get; }
+	}
+}
diff --git a/Engine/Source/Programs/Horde/HordeServer/Notifications/ISubscriptionCollection.cs b/Engine/Source/Programs/Horde/HordeServer/Notifications/ISubscriptionCollection.cs
new file mode 100644
--- /dev/null
+++ b/Engine/Source/Programs/Horde/HordeServer/Notifications/ISubscriptionCollection.cs
@@ -0,0 +1,99 @@
+// Copyright Epic Games, Inc. All Rights Reserved.
+
+using EpicGames.Core;
+using HordeCommon;
+using HordeServer.Api;
+using HordeServer.Models;
+using HordeServer.Services;
+using HordeServer.Utilities;
+using MongoDB.Bson;
+using MongoDB.Bson.Serialization.Attributes;
+using MongoDB.Driver;
+using Polly;
+using System;
+using System.Collections.Generic;
+using System.Linq;
+using System.Reflection.Metadata.Ecma335;
+using System.Text;
+using System.Threading.Tasks;
+
+namespace HordeServer.Collections
+{
+	using StreamId = StringId<IStream>;
+	using TemplateRefId = StringId<TemplateRef>;
+
+	/// <summary>
+	/// Information for creating a new subscription
+	/// </summary>
+	public class NewSubscription
+	{
+		/// <summary>
+		/// Name of the event
+		/// </summary>
+		public EventRecord Event { get; set; }
+
+		/// <summary>
+		/// The user name
+		/// </summary>
+		public ObjectId UserId { get; set; }
+
+		/// <summary>
+		/// Type of notification to send
+		/// </summary>
+		public NotificationType NotificationType { get; set; }
+
+		/// <summary>
+		/// Constructor
+		/// </summary>
+		/// <param name="Event">Name of the event</param>
+		/// <param name="UserId">User name</param>
+		/// <param name="NotificationType">Notification type</param>
+		public NewSubscription(EventRecord Event, ObjectId UserId, NotificationType NotificationType)
+		{
+			this.Event = Event;
+			this.UserId = UserId;
+			this.NotificationType = NotificationType;
+		}
+	}
+
+	/// <summary>
+	/// Interface for a collection of subscriptions
+	/// </summary>
+	public interface ISubscriptionCollection
+	{
+		/// <summary>
+		/// Add new subscription documents
+		/// </summary>
+		/// <param name="Subscriptions">The new subscriptions to add</param>
+		/// <returns>The subscriptions that were added</returns>
+		public Task<List<ISubscription>> AddAsync(IEnumerable<NewSubscription> Subscriptions);
+
+		/// <summary>
+		/// Remove a set of existing subscriptions
+		/// </summary>
+		/// <param name="Subscriptions">Subscriptions to remove</param>
+		/// <returns>Async task</returns>
+		public Task RemoveAsync(IEnumerable<ISubscription> Subscriptions);
+
+		/// <summary>
+		/// Gets a subscription by id
+		/// </summary>
+		/// <param name="SubscriptionId">Subscription to remove</param>
+		/// <returns>Async task</returns>
+		public Task<ISubscription?> GetAsync(string SubscriptionId);
+
+		/// <summary>
+		/// Find all subscribers of a certain event
+		/// </summary>
+		/// <param name="Event">Name of the event</param>
+		/// <returns>Name of the event to find subscribers for</returns>
+		public Task<List<ISubscription>> FindSubscribersAsync(EventRecord Event);
+
+		/// <summary>
+		/// Find subscriptions for a particular user
+		/// </summary>
+		/// <param name="UserId">The user to search for</param>
+		/// <returns>List of subscriptions</returns>
+		public Task<List<ISubscription>> FindSubscriptionsAsync(ObjectId UserId);
+	}
+}
diff --git a/Engine/Source/Programs/Horde/HordeServer/Notifications/Impl/NotificationTriggerCollection.cs b/Engine/Source/Programs/Horde/HordeServer/Notifications/Impl/NotificationTriggerCollection.cs
new file mode 100644
--- /dev/null
+++ b/Engine/Source/Programs/Horde/HordeServer/Notifications/Impl/NotificationTriggerCollection.cs
@@ -0,0 +1,242 @@
+// Copyright Epic Games, Inc. All Rights Reserved.
+
+using HordeServer.Models;
+using HordeServer.Services;
+using HordeServer.Utilities;
+using MongoDB.Bson;
+using MongoDB.Bson.Serialization.Attributes;
+using MongoDB.Driver;
+using System;
+using System.Collections.Generic;
+using System.Linq;
+using System.Threading.Tasks;
+
+namespace HordeServer.Collections.Impl
+{
+	/// <summary>
+	/// Collection of notification triggers
+	/// </summary>
+	public class NotificationTriggerCollection : INotificationTriggerCollection
+	{
+		class SubscriptionDocument : INotificationSubscription
+		{
+			[BsonIgnoreIfNull]
+			public string? User { get; set; }
+
+			public ObjectId UserId { get; set; }
+			public bool Email { get; set; }
+			public bool Slack { get; set; }
+
+			[BsonConstructor]
+			private SubscriptionDocument()
+			{
+				User = null!;
+			}
+
+			public SubscriptionDocument(ObjectId UserId, bool Email, bool Slack)
+			{
+				this.UserId = UserId;
+				this.Email = Email;
+				this.Slack = Slack;
+			}
+
+			public SubscriptionDocument(INotificationSubscription Subscription)
+			{
+				this.UserId = Subscription.UserId;
+				this.Email = Subscription.Email;
+				this.Slack = Subscription.Slack;
+			}
+		}
+
+		class TriggerDocument : INotificationTrigger
+		{
+			public ObjectId Id { get; set; }
+			public bool Fired { get; set; }
+			public List<SubscriptionDocument> Subscriptions { get; set; } = new List<SubscriptionDocument>();
+			public int UpdateIndex { get; set; }
+
+			IReadOnlyList<INotificationSubscription> INotificationTrigger.Subscriptions => Subscriptions;
+		}
+
+		IMongoCollection<TriggerDocument> Triggers;
+		IUserCollection UserCollection;
+
+		/// <summary>
+		/// Constructor
+		/// </summary>
+		/// <param name="DatabaseService">The database singleton</param>
+		/// <param name="UserCollection"></param>
+		public NotificationTriggerCollection(DatabaseService DatabaseService, IUserCollection UserCollection)
+		{
+			this.Triggers = DatabaseService.GetCollection<TriggerDocument>("NotificationTriggers");
+			this.UserCollection = UserCollection;
+		}
+
+		/// <inheritdoc/>
+		public async Task<INotificationTrigger?> GetAsync(ObjectId TriggerId)
+		{
+			TriggerDocument? Trigger = await Triggers.Find(x => x.Id == TriggerId).FirstOrDefaultAsync();
+			if (Trigger != null)
+			{
+				for (int Idx = 0; Idx < Trigger.Subscriptions.Count; Idx++)
+				{
+					SubscriptionDocument Subscription = Trigger.Subscriptions[Idx];
+					if (Subscription.User != null)
+					{
+						IUser? User = await UserCollection.FindUserByLoginAsync(Subscription.User);
+						if (User == null)
+						{
+							Trigger.Subscriptions.RemoveAt(Idx);
+							Idx--;
+						}
+						else
+						{
+							Subscription.UserId = User.Id;
+							Subscription.User = null;
+						}
+					}
+				}
+			}
+			return Trigger;
+		}
+
+		/// <inheritdoc/>
+		public async Task<INotificationTrigger> FindOrAddAsync(ObjectId TriggerId)
+		{
+			for (; ; )
+			{
+				// Find an existing trigger
+				INotificationTrigger? Existing = await GetAsync(TriggerId);
+				if (Existing != null)
+				{
+					return Existing;
+				}
+
+				// Try to insert a new document
+				try
+				{
+					TriggerDocument NewDocument = new TriggerDocument();
+					NewDocument.Id = TriggerId;
+					await Triggers.InsertOneAsync(NewDocument);
+					return NewDocument;
+				}
+				catch (MongoWriteException Ex)
+				{
+					if (Ex.WriteError.Category != ServerErrorCategory.DuplicateKey)
+					{
+						throw;
+					}
+				}
+			}
+		}
+
+		/// <summary>
+		/// Updates an existing document
+		/// </summary>
+		/// <param name="Trigger">The trigger to update</param>
+		/// <param name="Transaction">The update definition</param>
+		/// <returns>The updated document</returns>
+		async Task<INotificationTrigger?> TryUpdateAsync(INotificationTrigger Trigger, TransactionBuilder<TriggerDocument> Transaction)
+		{
+			TriggerDocument Document = (TriggerDocument)Trigger;
+			int NextUpdateIndex = Document.UpdateIndex + 1;
+
+			FilterDefinition<TriggerDocument> Filter = Builders<TriggerDocument>.Filter.Expr(x => x.Id == Trigger.Id && x.UpdateIndex == Document.UpdateIndex);
+			UpdateDefinition<TriggerDocument> Update = Transaction.ToUpdateDefinition().Set(x => x.UpdateIndex, NextUpdateIndex);
+
+			UpdateResult Result = await Triggers.UpdateOneAsync(Filter, Update);
+			if (Result.ModifiedCount > 0)
+			{
+				Transaction.ApplyTo(Document);
+				Document.UpdateIndex = NextUpdateIndex;
+				return Document;
+			}
+
+			return null;
+		}
+
+		/// <inheritdoc/>
+		public async Task DeleteAsync(ObjectId TriggerId)
+		{
+			await Triggers.DeleteOneAsync(x => x.Id == TriggerId);
+		}
+
+		/// <inheritdoc/>
+		public async Task DeleteAsync(List<ObjectId> TriggerIds)
+		{
+			FilterDefinition<TriggerDocument> Filter = Builders<TriggerDocument>.Filter.In(x => x.Id, TriggerIds);
+			await Triggers.DeleteManyAsync(Filter);
+		}
+
+		/// <inheritdoc/>
+		public async Task<INotificationTrigger?> FireAsync(INotificationTrigger Trigger)
+		{
+			if (Trigger.Fired)
+			{
+				return null;
+			}
+
+			for (; ; )
+			{
+				TransactionBuilder<TriggerDocument> Transaction = new TransactionBuilder<TriggerDocument>();
+				Transaction.Set(x => x.Fired, true);
+
+				INotificationTrigger? NewTrigger = await TryUpdateAsync(Trigger, Transaction);
+				if (NewTrigger != null)
+				{
+					return NewTrigger;
+				}
+
+				NewTrigger = await FindOrAddAsync(Trigger.Id); // Need to add to prevent race condition on triggering vs adding
+				if (NewTrigger == null || NewTrigger.Fired)
+				{
+					return null;
+				}
+			}
+		}
+
+		/// <inheritdoc/>
+		public async Task<INotificationTrigger?> UpdateSubscriptionsAsync(INotificationTrigger Trigger, ObjectId UserId, bool? Email, bool? Slack)
+		{
+			for (; ; )
+			{
+				// If the trigger has already fired, don't add a new subscription to it
+				if(Trigger.Fired)
+				{
+					return Trigger;
+				}
+
+				// Try to update the trigger
+				List<SubscriptionDocument> NewSubscriptions = new List<SubscriptionDocument>();
+				NewSubscriptions.AddRange(Trigger.Subscriptions.Select(x => new SubscriptionDocument(x)));
+
+				SubscriptionDocument? NewSubscription = NewSubscriptions.FirstOrDefault(x => x.UserId == UserId);
+				if (NewSubscription == null)
+				{
+					NewSubscription = new SubscriptionDocument(UserId, Email ?? false, Slack ?? false);
+					NewSubscriptions.Add(NewSubscription);
+				}
+				else
+				{
+					NewSubscription.Email = Email ?? NewSubscription.Email;
+					NewSubscription.Slack = Slack ?? NewSubscription.Slack;
+				}
+
+				TransactionBuilder<TriggerDocument> Transaction = new TransactionBuilder<TriggerDocument>();
+				Transaction.Set(x => x.Subscriptions, NewSubscriptions);
+
+				INotificationTrigger? NewTrigger = await TryUpdateAsync(Trigger, Transaction);
+				if (NewTrigger != null)
+				{
+					return NewTrigger;
+				}
+
+				NewTrigger = await GetAsync(Trigger.Id);
+				if (NewTrigger == null)
+				{
+					return null;
+				}
+			}
+		}
+	}
+}
diff --git a/Engine/Source/Programs/Horde/HordeServer/Notifications/Impl/SubscriptionCollection.cs b/Engine/Source/Programs/Horde/HordeServer/Notifications/Impl/SubscriptionCollection.cs
new file mode 100644
--- /dev/null
+++ b/Engine/Source/Programs/Horde/HordeServer/Notifications/Impl/SubscriptionCollection.cs
@@ -0,0 +1,249 @@
+// Copyright Epic Games, Inc. All Rights Reserved.
+
+using EpicGames.Core;
+using HordeCommon;
+using HordeServer.Api;
+using HordeServer.Models;
+using HordeServer.Services;
+using HordeServer.Utilities;
+using MongoDB.Bson;
+using MongoDB.Bson.Serialization.Attributes;
+using MongoDB.Driver;
+using Polly;
+using System;
+using System.Collections.Generic;
+using System.Linq;
+using System.Reflection.Metadata.Ecma335;
+using System.Text;
+using System.Threading.Tasks;
+
+namespace HordeServer.Collections.Impl
+{
+	using StreamId = StringId<IStream>;
+	using TemplateRefId = StringId<TemplateRef>;
+
+	/// <summary>
+	/// Collection of subscription documents
+	/// </summary>
+	public class SubscriptionCollection : ISubscriptionCollection
+	{
+		[BsonKnownTypes(typeof(JobCompleteEvent), typeof(LabelCompleteEvent), typeof(StepCompleteEvent))]
+		abstract class Event : IEvent
+		{
+			public static Event FromRecord(EventRecord Record)
+			{
+				JobCompleteEventRecord? JobCompleteRecord = Record as JobCompleteEventRecord;
+				if (JobCompleteRecord != null)
+				{
+					return new JobCompleteEvent(JobCompleteRecord);
+				}
+
+				LabelCompleteEventRecord? LabelCompleteRecord = Record as LabelCompleteEventRecord;
+				if (LabelCompleteRecord != null)
+				{
+					return new LabelCompleteEvent(LabelCompleteRecord);
+				}
+
+				StepCompleteEventRecord? StepCompleteRecord = Record as StepCompleteEventRecord;
+				if (StepCompleteRecord != null)
+				{
+					return new StepCompleteEvent(StepCompleteRecord);
+				}
+
+				throw new ArgumentException("Invalid record type");
+			}
+
+			public abstract EventRecord ToRecord();
+		}
+
+		class JobCompleteEvent : Event, IJobCompleteEvent
+		{
+			public StreamId StreamId { get; set; }
+			public TemplateRefId TemplateId { get; set; }
+			public LabelOutcome Outcome { get; set; }
+
+			public JobCompleteEvent()
+			{
+			}
+
+			public JobCompleteEvent(JobCompleteEventRecord Record)
+			{
+				StreamId = new StreamId(Record.StreamId);
+				TemplateId = new TemplateRefId(Record.TemplateId);
+				Outcome = Record.Outcome;
+			}
+
+			public override EventRecord ToRecord()
+			{
+				return new JobCompleteEventRecord(StreamId, TemplateId, Outcome);
+			}
+
+			public override string ToString()
+			{
+				return $"stream={StreamId}, template={TemplateId}, outcome={Outcome}";
+			}
+		}
+
+		class LabelCompleteEvent : Event, ILabelCompleteEvent
+		{
+			public StreamId StreamId { get; set; }
+			public TemplateRefId TemplateId { get; set; }
+			public string? CategoryName { get; set; }
+			public string LabelName { get; set; }
+			public LabelOutcome Outcome { get; set; }
+
+			public LabelCompleteEvent()
+			{
+				LabelName = String.Empty;
+			}
+
+			public LabelCompleteEvent(LabelCompleteEventRecord Record)
+			{
+				StreamId = new StreamId(Record.StreamId);
+				TemplateId = new TemplateRefId(Record.TemplateId);
+				CategoryName = Record.CategoryName;
+				LabelName = Record.LabelName;
+				Outcome = Record.Outcome;
+			}
+
+			public override EventRecord ToRecord()
+			{
+				return new LabelCompleteEventRecord(StreamId, TemplateId, CategoryName, LabelName, Outcome);
+			}
+
+			public override string ToString()
+			{
+				StringBuilder Result = new StringBuilder();
+				Result.Append($"stream={StreamId}, template={TemplateId}");
+				if(CategoryName != null)
+				{
+					Result.Append($", category={CategoryName}");
+				}
+				Result.Append($", label={LabelName}, outcome={Outcome}");
+				return Result.ToString();
+			}
+		}
+
+		class StepCompleteEvent : Event, IStepCompleteEvent
+		{
+			public StreamId StreamId { get; set; }
+			public TemplateRefId TemplateId { get; set; }
+			public string StepName { get; set; }
+			public JobStepOutcome Outcome { get; set; }
+
+			public StepCompleteEvent()
+			{
+				StepName = String.Empty;
+			}
+
+			public StepCompleteEvent(StepCompleteEventRecord Record)
+			{
+				StreamId = new StreamId(Record.StreamId);
+				TemplateId = new TemplateRefId(Record.TemplateId);
+				StepName = Record.StepName;
+				Outcome = Record.Outcome;
+			}
+
+			public override EventRecord ToRecord()
+			{
+				return new StepCompleteEventRecord(StreamId, TemplateId, StepName, Outcome);
+			}
+
+			public override string ToString()
+			{
+				return $"stream={StreamId}, template={TemplateId}, step={StepName}, outcome={Outcome}";
+			}
+		}
+
+		class Subscription : ISubscription
+		{
+			public string Id { get; set; }
+			public Event Event { get; set; }
+			public ObjectId UserId { get; set; }
+			public NotificationType NotificationType { get; set; }
+
+			IEvent ISubscription.Event => Event;
+
+			[BsonConstructor]
+			private Subscription()
+			{
+				this.Id = String.Empty;
+				this.Event = null!;
+			}
+
+			public Subscription(NewSubscription Subscription)
+			{
+				this.Event = Event.FromRecord(Subscription.Event);
+				this.UserId = Subscription.UserId;
+				this.NotificationType = Subscription.NotificationType;
+				this.Id = ContentHash.SHA1($"{Event}, user={UserId}, type={NotificationType.ToString()}").ToString();
+			}
+		}
+
+		IMongoCollection<Subscription> Collection;
+
+		/// <summary>
+		/// Constructor
+		/// </summary>
+		/// <param name="DatabaseService">The database service</param>
+		public SubscriptionCollection(DatabaseService DatabaseService)
+		{
+			Collection = DatabaseService.GetCollection<Subscription>("SubscriptionsV2");
+			if (!DatabaseService.ReadOnlyMode)
+			{
+				Collection.Indexes.CreateOne(new CreateIndexModel<Subscription>(Builders<Subscription>.IndexKeys.Ascending(x => x.Event)));
+				Collection.Indexes.CreateOne(new CreateIndexModel<Subscription>(Builders<Subscription>.IndexKeys.Ascending(x => x.UserId)));
+			}
+		}
+
+		/// <inheritdoc/>
+		public async Task<List<ISubscription>> AddAsync(IEnumerable<NewSubscription> NewSubscriptions)
+		{
+			List<Subscription> NewDocuments = NewSubscriptions.Select(x => new Subscription(x)).ToList();
+			try
+			{
+				await Collection.InsertManyAsync(NewDocuments, new InsertManyOptions { IsOrdered = false });
+			}
+			catch (MongoBulkWriteException Ex)
+			{
+				foreach (WriteError WriteError in Ex.WriteErrors)
+				{
+					if (WriteError.Category != ServerErrorCategory.DuplicateKey)
+					{
+						throw;
+					}
+				}
+			}
+			return NewDocuments.ConvertAll<ISubscription>(x => x);
+		}
+
+		/// <inheritdoc/>
+		public async Task RemoveAsync(IEnumerable<ISubscription> Subscriptions)
+		{
+			FilterDefinition<Subscription> Filter = Builders<Subscription>.Filter.In(x => x.Id, Subscriptions.Select(x => ((Subscription)x).Id));
+			await Collection.DeleteManyAsync(Filter);
+		}
+
+		/// <inheritdoc/>
+		public async Task<ISubscription?> GetAsync(string SubscriptionId)
+		{
+			FilterDefinition<Subscription> Filter = Builders<Subscription>.Filter.Eq(x => x.Id, SubscriptionId);
+			return await Collection.Find(Filter).FirstOrDefaultAsync();
+		}
+
+		/// <inheritdoc/>
+		public async Task<List<ISubscription>> FindSubscribersAsync(EventRecord EventRecord)
+		{
+			FilterDefinition<Subscription> Filter = Builders<Subscription>.Filter.Eq(x => x.Event, Event.FromRecord(EventRecord));
+			List<Subscription> Results = await Collection.Find(Filter).ToListAsync();
+			return Results.ConvertAll<ISubscription>(x => x);
+		}
+
+		/// <inheritdoc/>
+		public async Task<List<ISubscription>> FindSubscriptionsAsync(ObjectId UserId)
+		{
+			List<Subscription> Results = await Collection.Find(x => x.UserId == UserId).ToListAsync();
+			return Results.ConvertAll<ISubscription>(x => x);
+		}
+	}
+}
diff --git a/Engine/Source/Programs/Horde/HordeServer/Projects/IProject.cs b/Engine/Source/Programs/Horde/HordeServer/Projects/IProject.cs
new file mode 100644
--- /dev/null
+++ b/Engine/Source/Programs/Horde/HordeServer/Projects/IProject.cs
@@ -0,0 +1,248 @@
+// Copyright Epic Games, Inc. All Rights Reserved.
+
+using EpicGames.Core;
+using HordeServer.Api;
+using HordeServer.Utilities;
+using MongoDB.Bson;
+using MongoDB.Bson.Serialization.Attributes;
+using MongoDB.Driver;
+using System;
+using System.Collections.Generic;
+using System.Linq;
+using System.Text.RegularExpressions;
+using System.Threading.Tasks;
+
+namespace HordeServer.Models
+{
+	using ProjectId = StringId<IProject>;
+
+	/// <summary>
+	/// Specifies a category of streams to display on the dashboard
+	/// </summary>
+	public class StreamCategory
+	{
+		/// <summary>
+		/// Name of this group
+		/// </summary>
+		[BsonRequired]
+		public string Name { get; set; }
+
+		/// <summary>
+		/// Row to display this category on
+		/// </summary>
+		public int Row { get; set; }
+
+		/// <summary>
+		/// Whether to show this category on the nav menu
+		/// </summary>
+		public bool ShowOnNavMenu { get; set; } = true;
+
+		/// <summary>
+		/// List of stream name patterns to include for this project
+		/// </summary>
+		public List<string> IncludePatterns { get; set; } = new List<string>();
+
+		/// <summary>
+		/// Patterns for stream names to be excluded
+		/// </summary>
+		public List<string> ExcludePatterns { get; set; } = new List<string>();
+
+		/// <summary>
+		/// Private constructor
+		/// </summary>
+		private StreamCategory()
+		{
+			Name = null!;
+		}
+
+		/// <summary>
+		/// Constructor
+		/// </summary>
+		/// <param name="Name">Name of this group</param>
+		public StreamCategory(string Name)
+		{
+			this.Name = Name;
+		}
+
+		/// <summary>
+		/// Constructs a category from a request object
+		/// </summary>
+		/// <param name="Request">The request object</param>
+		public StreamCategory(CreateProjectCategoryRequest Request)
+		{
+			this.Name = Request.Name;
+			this.Row = Request.Row;
+			this.ShowOnNavMenu = Request.ShowOnNavMenu;
+			this.IncludePatterns = Request.IncludePatterns;
+			this.ExcludePatterns = Request.ExcludePatterns;
+		}
+
+		/// <summary>
+		/// Tests if a given name matches a pattern
+		/// </summary>
+		/// <param name="Name">The name to test</param>
+		/// <param name="Pattern">The pattern to match against</param>
+		/// <returns>True if the pattern matches, false otherwise</returns>
+		public static bool MatchPattern(string Name, string Pattern)
+		{
+			return Regex.IsMatch(Name, Pattern);
+		}
+	}
+
+	/// <summary>
+	/// Represents a project
+	/// </summary>
+	public interface IProject
+	{
+		/// <summary>
+		/// Identifier for the project.
+		/// </summary>
+		public ProjectId Id { get; }
+
+		/// <summary>
+		/// Name of the project
+		/// </summary>
+		public string Name { get; }
+
+		/// <summary>
+		/// Path to the config file used to configure this project
+		/// </summary>
+		public string? ConfigPath { get; }
+
+		/// <summary>
+		/// Revision of the config file used to configure this project
+		/// </summary>
+		public string? ConfigRevision { get; }
+
+		/// <summary>
+		/// Order to display on the dashboard
+		/// </summary>
+		public int Order { get; }
+
+		/// <summary>
+		/// The ACL for this object
+		/// </summary>
+		public Acl? Acl { get; }
+
+		/// <summary>
+		/// List of stream categories for this project. Controls how streams are displayed on the dashboard.
+		/// </summary>
+		public IReadOnlyList<StreamCategory> Categories { get; }
+	}
+	
+	/// <summary>
+	/// Interface for a document containing a project logo
+	/// </summary>
+	public interface IProjectLogo
+	{
+		/// <summary>
+		/// The project id
+		/// </summary>
+		public ProjectId Id { get; }
+
+		/// <summary>
+		/// Path to the logo
+		/// </summary>
+		public string Path { get; }
+
+		/// <summary>
+		/// Revision of the file
+		/// </summary>
+		public string Revision { get; }
+
+		/// <summary>
+		/// Mime type for the image
+		/// </summary>
+		public string MimeType { get; }
+
+		/// <summary>
+		/// Image data
+		/// </summary>
+		public byte[] Data { get; }
+	}
+
+	/// <summary>
+	/// Extension methods for projects
+	/// </summary>
+	public static class ProjectExtensions
+	{
+		/// <summary>
+		/// Converts this object to a public response
+		/// </summary>
+		/// <param name="Project">The project instance</param>
+		/// <param name="bIncludeStreams">Whether to include streams in the response</param>
+		/// <param name="bIncludeCategories">Whether to include categories in the response</param>
+		/// <param name="Streams">The list of streams</param>
+		/// <param name="bIncludeAcl">Whether to include the ACL in the response</param>
+		/// <returns>Response instance</returns>
+		public static GetProjectResponse ToResponse(this IProject Project, bool bIncludeStreams, bool bIncludeCategories, List<IStream>? Streams, bool bIncludeAcl)
+		{
+			List<GetProjectStreamResponse>? StreamResponses = null;
+			if(bIncludeStreams)
+			{
+				StreamResponses = Streams!.ConvertAll(x => new GetProjectStreamResponse(x.Id.ToString(), x.Name));
+			}
+
+			List<GetProjectCategoryResponse>? CategoryResponses = null;
+			if (bIncludeCategories)
+			{
+				CategoryResponses = Project.Categories.ConvertAll(x => new GetProjectCategoryResponse(x));
+				if (Streams != null)
+				{
+					foreach (IStream Stream in Streams)
+					{
+						GetProjectCategoryResponse? CategoryResponse = CategoryResponses.FirstOrDefault(x => MatchCategory(Stream.Name, x));
+						if(CategoryResponse == null)
+						{
+							int Row = (CategoryResponses.Count > 0) ? CategoryResponses.Max(x => x.Row) : 0;
+							if (CategoryResponses.Count(x => x.Row == Row) >= 3)
+							{
+								Row++;
+							}
+
+							StreamCategory OtherCategory = new StreamCategory("Other");
+							OtherCategory.Row = Row;
+							OtherCategory.IncludePatterns.Add(".*");
+
+							CategoryResponse = new GetProjectCategoryResponse(OtherCategory);
+							CategoryResponses.Add(CategoryResponse);
+						}
+						CategoryResponse.Streams!.Add(Stream.Id.ToString());
+					}
+				}
+			}
+
+			GetAclResponse? AclResponse = (bIncludeAcl && Project.Acl != null) ? new GetAclResponse(Project.Acl) : null;
+			return new GetProjectResponse(Project.Id.ToString(), Project.Name, Project.Order, StreamResponses, CategoryResponses, AclResponse);
+		}
+
+		/// <summary>
+		/// Tests if a category response matches a given stream name
+		/// </summary>
+		/// <param name="Name">The stream name</param>
+		/// <param name="Category">The category response</param>
+		/// <returns>True if the category matches</returns>
+		static bool MatchCategory(string Name, GetProjectCategoryResponse Category)
+		{
+			if (Category.IncludePatterns.Any(x => StreamCategory.MatchPattern(Name, x)))
+			{
+				if (!Category.ExcludePatterns.Any(x => StreamCategory.MatchPattern(Name, x)))
+				{
+					return true;
+				}
+			}
+			return false;
+		}
+	}
+
+	/// <summary>
+	/// Projection of a project definition to just include permissions info
+	/// </summary>
+	public interface IProjectPermissions
+	{
+		/// <summary>
+		/// ACL for the project
+		/// </summary>
+		public Acl? Acl { get; }
+	}
+}
diff --git a/Engine/Source/Programs/Horde/HordeServer/Projects/IProjectCollection.cs b/Engine/Source/Programs/Horde/HordeServer/Projects/IProjectCollection.cs
new file mode 100644
--- /dev/null
+++ b/Engine/Source/Programs/Horde/HordeServer/Projects/IProjectCollection.cs
@@ -0,0 +1,83 @@
+// Copyright Epic Games, Inc. All Rights Reserved.
+
+using HordeServer.Api;
+using HordeServer.Models;
+using HordeServer.Services;
+using HordeServer.Utilities;
+using Microsoft.Extensions.Options;
+using MongoDB.Bson;
+using MongoDB.Bson.Serialization.Attributes;
+using MongoDB.Driver;
+using System;
+using System.Collections.Generic;
+using System.Diagnostics.CodeAnalysis;
+using System.Linq;
+using System.Security.Claims;
+using System.Threading.Tasks;
+
+namespace HordeServer.Collections
+{
+	using ProjectId = StringId<IProject>;
+
+	/// <summary>
+	/// Interface for a collection of project documents
+	/// </summary>
+	public interface IProjectCollection
+	{
+		/// <summary>
+		/// Updates the project configuration
+		/// </summary>
+		/// <param name="Id">The project id</param>
+		/// <param name="ConfigPath">Path to the config file used to configure this project</param>
+		/// <param name="Revision">The config file revision</param>
+		/// <param name="Order">Order of the project</param>
+		/// <param name="Config">The configuration</param>
+		/// <returns>New project instance</returns>
+		Task<IProject?> AddOrUpdateAsync(ProjectId Id, string ConfigPath, string Revision, int Order, ProjectConfig Config);
+
+		/// <summary>
+		/// Gets all the available projects
+		/// </summary>
+		/// <returns>List of project documents</returns>
+		Task<List<IProject>> FindAllAsync();
+
+		/// <summary>
+		/// Gets a project by ID
+		/// </summary>
+		/// <param name="ProjectId">Unique id of the project</param>
+		/// <returns>The project document</returns>
+		Task<IProject?> GetAsync(ProjectId ProjectId);
+
+		/// <summary>
+		/// Gets the logo for a project
+		/// </summary>
+		/// <param name="ProjectId">The project id</param>
+		/// <returns>The project logo document</returns>
+		Task<IProjectLogo?> GetLogoAsync(ProjectId ProjectId);
+
+		/// <summary>
+		/// Sets the logo for a project
+		/// </summary>
+		/// <param name="ProjectId">The project id</param>
+		/// <param name="Path">Path to the source file</param>
+		/// <param name="Revision">Revision of the file</param>
+		/// <param name="MimeType"></param>
+		/// <param name="Data"></param>
+		/// <returns></returns>
+		Task SetLogoAsync(ProjectId ProjectId, string Path, string Revision, string MimeType, byte[] Data);
+
+		/// <summary>
+		/// Gets a project's permissions info by ID
+		/// </summary>
+		/// <param name="ProjectId">Unique id of the project</param>
+		/// <returns>The project document</returns>
+		Task<IProjectPermissions?> GetPermissionsAsync(ProjectId ProjectId);
+
+		/// <summary>
+		/// Deletes a project by id
+		/// </summary>
+		/// <param name="ProjectId">Unique id of the project</param>
+		/// <returns>True if the project was deleted</returns>
+		Task DeleteAsync(ProjectId ProjectId);
+	}
+}
diff --git a/Engine/Source/Programs/Horde/HordeServer/Projects/Impl/ProjectCollection.cs b/Engine/Source/Programs/Horde/HordeServer/Projects/Impl/ProjectCollection.cs
new file mode 100644
--- /dev/null
+++ b/Engine/Source/Programs/Horde/HordeServer/Projects/Impl/ProjectCollection.cs
@@ -0,0 +1,173 @@
+// Copyright Epic Games, Inc. All Rights Reserved.
+
+using HordeServer.Api;
+using HordeServer.Models;
+using HordeServer.Services;
+using HordeServer.Utilities;
+using Microsoft.Extensions.Options;
+using MongoDB.Bson;
+using MongoDB.Bson.Serialization.Attributes;
+using MongoDB.Driver;
+using System;
+using System.Collections.Generic;
+using System.Diagnostics.CodeAnalysis;
+using System.Linq;
+using System.Security.Claims;
+using System.Threading.Tasks;
+
+namespace HordeServer.Collections.Impl
+{
+	using ProjectId = StringId<IProject>;
+
+	/// <summary>
+	/// Wraps functionality for manipulating projects
+	/// </summary>
+	public class ProjectCollection : IProjectCollection
+	{
+		/// <summary>
+		/// Represents a project
+		/// </summary>
+		class ProjectDocument : IProject
+		{
+			public const int DefaultOrder = 128;
+
+			[BsonRequired, BsonId]
+			public ProjectId Id { get; set; }
+
+			[BsonRequired]
+			public string Name { get; set; }
+
+			public string? ConfigPath { get; set; }
+			public string? ConfigRevision { get; set; }
+
+			public int Order { get; set; } = DefaultOrder;
+			public Acl? Acl { get; set; }
+			public List<StreamCategory> Categories { get; set; } = new List<StreamCategory>();
+
+			[BsonIgnoreIfDefault, BsonDefaultValue(false)]
+			public bool Deleted { get; set; }
+
+			IReadOnlyList<StreamCategory> IProject.Categories => Categories;
+
+			[BsonConstructor]
+			private ProjectDocument()
+			{
+				Name = null!;
+			}
+
+			public ProjectDocument(ProjectId Id, string Name)
+			{
+				this.Id = Id;
+				this.Name = Name;
+			}
+		}
+
+		/// <summary>
+		/// Logo for a project
+		/// </summary>
+		class ProjectLogoDocument : IProjectLogo
+		{
+			public ProjectId Id { get; set; }
+
+			public string Path { get; set; } = String.Empty;
+			public string Revision { get; set; } = String.Empty;
+			public string MimeType { get; set; } = String.Empty;
+
+			public byte[] Data { get; set; } = Array.Empty<byte>();
+		}
+
+		/// <summary>
+		/// Projection of a project definition to just include permissions info
+		/// </summary>
+		[SuppressMessage("Design", "CA1812: Class is never instantiated")]
+		class ProjectPermissions : IProjectPermissions
+		{
+			/// <summary>
+			/// ACL for the project
+			/// </summary>
+			public Acl? Acl { get; set; }
+
+			/// <summary>
+			/// Projection to extract the permissions info from the project
+			/// </summary>
+			public static readonly ProjectionDefinition<ProjectDocument> Projection = Builders<ProjectDocument>.Projection.Include(x => x.Acl);
+		}
+
+		/// <summary>
+		/// Collection of project documents
+		/// </summary>
+		IMongoCollection<ProjectDocument> Projects;
+
+		/// <summary>
+		/// Collection of project logo documents
+		/// </summary>
+		IMongoCollection<ProjectLogoDocument> ProjectLogos;
+
+		/// <summary>
+		/// Constructor
+		/// </summary>
+		/// <param name="DatabaseService">The database service instance</param>
+		public ProjectCollection(DatabaseService DatabaseService)
+		{
+			Projects = DatabaseService.GetCollection<ProjectDocument>("Projects");
+			ProjectLogos = DatabaseService.GetCollection<ProjectLogoDocument>("ProjectLogos");
+		}
+
+		/// <inheritdoc/>
+		public async Task<IProject?> AddOrUpdateAsync(ProjectId Id, string ConfigPath, string Revision, int Order, ProjectConfig Config)
+		{
+			ProjectDocument NewProject = new ProjectDocument(Id, Config.Name);
+			NewProject.ConfigPath = ConfigPath;
+			NewProject.ConfigRevision = Revision;
+			NewProject.Order = Order;
+			NewProject.Categories = Config.Categories.ConvertAll(x => new StreamCategory(x));
+			NewProject.Acl = Acl.Merge(new Acl(), Config.Acl);
+
+			await Projects.FindOneAndReplaceAsync<ProjectDocument>(x => x.Id == Id, NewProject, new FindOneAndReplaceOptions<ProjectDocument> { IsUpsert = true });
+			return NewProject;
+		}
+
+		/// <inheritdoc/>
+		public async Task<List<IProject>> FindAllAsync()
+		{
+			List<ProjectDocument> Results = await Projects.Find(x => !x.Deleted).ToListAsync();
+			return Results.OrderBy(x => x.Order).ThenBy(x => x.Name).Select<ProjectDocument, IProject>(x => x).ToList();
+		}
+
+		/// <inheritdoc/>
+		public async Task<IProject?> GetAsync(ProjectId ProjectId)
+		{
+			return await Projects.Find<ProjectDocument>(x => x.Id == ProjectId).FirstOrDefaultAsync();
+		}
+
+		/// <inheritdoc/>
+		public async Task<IProjectLogo?> GetLogoAsync(ProjectId ProjectId)
+		{
+			return await ProjectLogos.Find(x => x.Id == ProjectId).FirstOrDefaultAsync();
+		}
+
+		/// <inheritdoc/>
+		public async Task SetLogoAsync(ProjectId ProjectId, string LogoPath, string LogoRevision, string MimeType, byte[] Data)
+		{
+			ProjectLogoDocument Logo = new ProjectLogoDocument();
+			Logo.Id = ProjectId;
+			Logo.Path = LogoPath;
+			Logo.Revision = LogoRevision;
+			Logo.MimeType = MimeType;
+			Logo.Data = Data;
+			await ProjectLogos.ReplaceOneAsync(x => x.Id == ProjectId, Logo, new ReplaceOptions { IsUpsert = true });
+		}
+
+		/// <inheritdoc/>
+		public async Task<IProjectPermissions?> GetPermissionsAsync(ProjectId ProjectId)
+		{
+			return await Projects.Find<ProjectDocument>(x => x.Id == ProjectId).Project<ProjectPermissions>(ProjectPermissions.Projection).FirstOrDefaultAsync();
+		}
+
+		/// <inheritdoc/>
+		public async Task DeleteAsync(ProjectId ProjectId)
+		{
+			await Projects.UpdateOneAsync(x => x.Id == ProjectId, Builders<ProjectDocument>.Update.Set(x => x.Deleted, true));
+		}
+	}
+}
diff --git a/Engine/Source/Programs/Horde/HordeServer/Projects/ProjectService.cs b/Engine/Source/Programs/Horde/HordeServer/Projects/ProjectService.cs
new file mode 100644
--- /dev/null
+++ b/Engine/Source/Programs/Horde/HordeServer/Projects/ProjectService.cs
@@ -0,0 +1,157 @@
+// Copyright Epic Games, Inc. All Rights Reserved.
+
+using HordeServer.Api;
+using HordeServer.Collections;
+using HordeServer.Models;
+using HordeServer.Utilities;
+using Microsoft.Extensions.Options;
+using MongoDB.Bson;
+using MongoDB.Driver;
+using System;
+using System.Collections.Generic;
+using System.Linq;
+using System.Security.Claims;
+using System.Threading.Tasks;
+
+namespace HordeServer.Services
+{
+	using ProjectId = StringId<IProject>;
+
+	/// <summary>
+	/// Cache of information about job ACLs
+	/// </summary>
+	public class ProjectPermissionsCache : GlobalPermissionsCache
+	{
+		/// <summary>
+		/// Map of project id to permissions for that project
+		/// </summary>
+		public Dictionary<ProjectId, IProjectPermissions?> Projects { get; } = new Dictionary<ProjectId, IProjectPermissions?>();
+	}
+
+	/// <summary>
+	/// Wraps functionality for manipulating projects
+	/// </summary>
+	public class ProjectService
+	{
+		/// <summary>
+		/// The ACL service
+		/// </summary>
+		AclService AclService;
+
+		/// <summary>
+		/// Collection of project documents
+		/// </summary>
+		IProjectCollection Projects;
+
+		/// <summary>
+		/// Accessor for the collection of project documents
+		/// </summary>
+		public IProjectCollection Collection => Projects;
+
+		/// <summary>
+		/// Constructor
+		/// </summary>
+		/// <param name="AclService">The ACL service</param>
+		/// <param name="Projects">Collection of project documents</param>
+		public ProjectService(AclService AclService, IProjectCollection Projects)
+		{
+			this.AclService = AclService;
+			this.Projects = Projects;
+		}
+
+		/// <summary>
+		/// Gets all the available projects
+		/// </summary>
+		/// <returns>List of project documents</returns>
+		public Task<List<IProject>> GetProjectsAsync()
+		{
+			return Projects.FindAllAsync();
+		}
+
+		/// <summary>
+		/// Gets a project by ID
+		/// </summary>
+		/// <param name="ProjectId">Unique id of the project</param>
+		/// <returns>The project document</returns>
+		public Task<IProject?> GetProjectAsync(ProjectId ProjectId)
+		{
+			return Projects.GetAsync(ProjectId);
+		}
+
+		/// <summary>
+		/// Gets a project's permissions info by ID
+		/// </summary>
+		/// <param name="ProjectId">Unique id of the project</param>
+		/// <returns>The project document</returns>
+		public Task<IProjectPermissions?> GetProjectPermissionsAsync(ProjectId ProjectId)
+		{
+			return Projects.GetPermissionsAsync(ProjectId);
+		}
+
+		/// <summary>
+		/// Deletes a project by id
+		/// </summary>
+		/// <param name="ProjectId">Unique id of the project</param>
+		public async Task DeleteProjectAsync(ProjectId ProjectId)
+		{
+			await Projects.DeleteAsync(ProjectId);
+		}
+
+		/// <summary>
+		/// Determines if the user is authorized to perform an action on a particular project
+		/// </summary>
+		/// <param name="Acl">Acl for the project to check</param>
+		/// <param name="Action">The action being performed</param>
+		/// <param name="User">The principal to authorize</param>
+		/// <param name="Cache">Cache for the scope table</param>
+		/// <returns>True if the action is authorized</returns>
+		private Task<bool> AuthorizeAsync(Acl? Acl, AclAction Action, ClaimsPrincipal User, GlobalPermissionsCache? Cache)
+		{
+			bool? Result = Acl?.Authorize(Action, User);
+			if (Result == null)
+			{
+				return AclService.AuthorizeAsync(Action, User, Cache);
+			}
+			else
+			{
+				return Task.FromResult(Result.Value);
+			}
+		}
+
+		/// <summary>
+		/// Determines if the user is authorized to perform an action on a particular project
+		/// </summary>
+		/// <param name="Project">The project to check</param>
+		/// <param name="Action">The action being performed</param>
+		/// <param name="User">The principal to authorize</param>
+		/// <param name="Cache">Cache for the scope table</param>
+		/// <returns>True if the action is authorized</returns>
+		public Task<bool> AuthorizeAsync(IProject Project, AclAction Action, ClaimsPrincipal User, GlobalPermissionsCache? Cache)
+		{
+			return AuthorizeAsync(Project.Acl, Action, User, Cache);
+		}
+
+		/// <summary>
+		/// Determines if the user is authorized to perform an action on a particular project
+		/// </summary>
+		/// <param name="ProjectId">The project id to check</param>
+		/// <param name="Action">The action being performed</param>
+		/// <param name="User">The principal to authorize</param>
+		/// <param name="Cache">Cache for project permissions</param>
+		/// <returns>True if the action is authorized</returns>
+		public async Task<bool> AuthorizeAsync(ProjectId ProjectId, AclAction Action, ClaimsPrincipal User, ProjectPermissionsCache? Cache)
+		{
+			IProjectPermissions? Permissions;
+			if (Cache == null)
+			{
+				Permissions = await GetProjectPermissionsAsync(ProjectId);
+			}
+			else if (!Cache.Projects.TryGetValue(ProjectId, out Permissions))
+			{
+				Permissions = await GetProjectPermissionsAsync(ProjectId);
+				Cache.Projects.Add(ProjectId, Permissions);
+			}
+			return await AuthorizeAsync(Permissions?.Acl, Action, User, Cache);
+		}
+	}
+}
diff --git a/Engine/Source/Programs/Horde/HordeServer/Services/IFleetManager.cs b/Engine/Source/Programs/Horde/HordeServer/Services/IFleetManager.cs
deleted file mode 100644
diff --git a/Engine/Source/Programs/Horde/HordeServer/Services/Impl/AwsFleetManager.cs b/Engine/Source/Programs/Horde/HordeServer/Services/Impl/AwsFleetManager.cs
deleted file mode 100644
diff --git a/Engine/Source/Programs/Horde/HordeServer/Services/Impl/DefaultFleetManager.cs b/Engine/Source/Programs/Horde/HordeServer/Services/Impl/DefaultFleetManager.cs
deleted file mode 100644
diff --git a/Engine/Source/Programs/Horde/HordeServer/Services/ProjectService.cs b/Engine/Source/Programs/Horde/HordeServer/Services/ProjectService.cs
deleted file mode 100644
diff --git a/Engine/Source/Programs/Horde/HordeServer/Services/StreamService.cs b/Engine/Source/Programs/Horde/HordeServer/Services/StreamService.cs
deleted file mode 100644
diff --git a/Engine/Source/Programs/Horde/HordeServer/Streams/IStream.cs b/Engine/Source/Programs/Horde/HordeServer/Streams/IStream.cs
new file mode 100644
--- /dev/null
+++ b/Engine/Source/Programs/Horde/HordeServer/Streams/IStream.cs
@@ -0,0 +1,704 @@
+// Copyright Epic Games, Inc. All Rights Reserved.
+
+using EpicGames.Core;
+using Google.Protobuf.WellKnownTypes;
+using HordeServer.Api;
+using HordeCommon;
+using HordeServer.Utilities;
+using MongoDB.Bson;
+using MongoDB.Bson.Serialization.Attributes;
+using MongoDB.Driver;
+using System;
+using System.Collections.Generic;
+using System.Diagnostics.CodeAnalysis;
+using System.IO;
+using System.Linq;
+using System.Security.Cryptography;
+using System.Text;
+using System.Text.RegularExpressions;
+using System.Threading.Tasks;
+using System.Runtime.CompilerServices;
+
+namespace HordeServer.Models
+{
+	using PoolId = StringId<IPool>;
+	using ProjectId = StringId<IProject>;
+	using StreamId = StringId<IStream>;
+	using TemplateRefId = StringId<TemplateRef>;
+
+	/// <summary>
+	/// Exception thrown when stream validation fails
+	/// </summary>
+	public class InvalidStreamException : Exception
+	{
+		/// <inheritdoc/>
+		public InvalidStreamException()
+		{
+		}
+
+		/// <inheritdoc/>
+		public InvalidStreamException(string Message) : base(Message)
+		{
+		}
+
+		/// <inheritdoc/>
+		public InvalidStreamException(string Message, Exception InnerEx) : base(Message, InnerEx)
+		{
+		}
+	}
+
+	/// <summary>
+	/// Mapping from a BuildGraph agent type to a set of machines on the farm
+	/// </summary>
+	public class AgentType
+	{
+		/// <summary>
+		/// Name of the pool of agents to use
+		/// </summary>
+		[BsonRequired]
+		public PoolId Pool { get; set; }
+
+		/// <summary>
+		/// Name of the workspace to execute on
+		/// </summary>
+		public string? Workspace { get; set; }
+
+		/// <summary>
+		/// Path to the temporary storage dir
+		/// </summary>
+		public string? TempStorageDir { get; set; }
+
+		/// <summary>
+		/// Environment variables to be set when executing the job
+		/// </summary>
+		public Dictionary<string, string>? Environment { get; set; }
+
+		/// <summary>
+		/// Default constructor
+		/// </summary>
+		[BsonConstructor]
+		private AgentType()
+		{
+		}
+
+		/// <summary>
+		/// Constructor
+		/// </summary>
+		/// <param name="Pool">The pool for this agent</param>
+		/// <param name="Workspace">Name of the workspace to use</param>
+		/// <param name="TempStorageDir">Path to the temp storage directory</param>
+		public AgentType(PoolId Pool, string Workspace, string? TempStorageDir)
+		{
+			this.Pool = Pool;
+			this.Workspace = Workspace;
+			this.TempStorageDir = TempStorageDir;
+		}
+
+		/// <summary>
+		/// Constructor
+		/// </summary>
+		/// <param name="Request">The object to construct from</param>
+		public AgentType(CreateAgentTypeRequest Request)
+		{
+			Pool = new PoolId(Request.Pool);
+			Workspace = Request.Workspace;
+			TempStorageDir = Request.TempStorageDir;
+			Environment = Request.Environment;
+		}
+
+		/// <summary>
+		/// Constructs an AgentType object from an optional request
+		/// </summary>
+		/// <param name="Request">The request object</param>
+		/// <returns>New agent type object</returns>
+		[return: NotNullIfNotNull("Request")]
+		public static AgentType? FromRequest(CreateAgentTypeRequest? Request)
+		{
+			return (Request != null) ? new AgentType(Request) : null;
+		}
+
+		/// <summary>
+		/// Creates an API response object from this stream
+		/// </summary>
+		/// <returns>The response object</returns>
+		public Api.GetAgentTypeResponse ToApiResponse()
+		{
+			return new Api.GetAgentTypeResponse(Pool.ToString(), Workspace, TempStorageDir, Environment);
+		}
+
+		/// <summary>
+		/// Creates an API response object from this stream
+		/// </summary>
+		/// <returns>The response object</returns>
+		public HordeCommon.Rpc.GetAgentTypeResponse ToRpcResponse()
+		{
+			HordeCommon.Rpc.GetAgentTypeResponse Response = new HordeCommon.Rpc.GetAgentTypeResponse();
+			if (TempStorageDir != null)
+			{
+				Response.TempStorageDir = TempStorageDir;
+			}
+			if (Environment != null)
+			{
+				Response.Environment.Add(Environment);
+			}
+			return Response;
+		}
+	}
+
+	/// <summary>
+	/// Information about a workspace type
+	/// </summary>
+	public class WorkspaceType
+	{
+		/// <summary>
+		/// Name of the Perforce cluster to use
+		/// </summary>
+		[BsonIgnoreIfNull]
+		public string? Cluster { get; set; }
+
+		/// <summary>
+		/// The Perforce server and port
+		/// </summary>
+		[BsonIgnoreIfNull]
+		public string? ServerAndPort { get; set; }
+
+		/// <summary>
+		/// The Perforce username for syncing this workspace
+		/// </summary>
+		[BsonIgnoreIfNull]
+		public string? UserName { get; set; }
+
+		/// <summary>
+		/// The Perforce password for syncing this workspace
+		/// </summary>
+		[BsonIgnoreIfNull]
+		public string? Password { get; set; }
+
+		/// <summary>
+		/// Identifier to distinguish this workspace from other workspaces. Defaults to the workspace type name.
+		/// </summary>
+		[BsonIgnoreIfNull]
+		public string? Identifier { get; set; }
+
+		/// <summary>
+		/// Override for the stream to sync
+		/// </summary>
+		[BsonIgnoreIfNull]
+		public string? Stream { get; set; }
+
+		/// <summary>
+		/// Custom view for the workspace
+		/// </summary>
+		[BsonIgnoreIfNull]
+		public List<string>? View { get; set; }
+
+		/// <summary>
+		/// Whether to use an incrementally synced workspace
+		/// </summary>
+		public bool Incremental { get; set; }
+
+		/// <summary>
+		/// Default constructor
+		/// </summary>
+		public WorkspaceType()
+		{
+		}
+
+		/// <summary>
+		/// Constructor
+		/// </summary>
+		/// <param name="Request">The object to construct from</param>
+		public WorkspaceType(CreateWorkspaceTypeRequest Request)
+		{
+			Cluster = Request.Cluster;
+			ServerAndPort = Request.ServerAndPort;
+			UserName = Request.UserName;
+			Password = Request.Password;
+			Identifier = Request.Identifier;
+			Stream = Request.Stream;
+			View = Request.View;
+			Incremental = Request.Incremental;
+		}
+
+		/// <summary>
+		/// Constructs an AgentType object from an optional request
+		/// </summary>
+		/// <param name="Request">The request object</param>
+		/// <returns>New agent type object</returns>
+		[return: NotNullIfNotNull("Request")]
+		public static WorkspaceType? FromRequest(CreateWorkspaceTypeRequest? Request)
+		{
+			return (Request != null) ? new WorkspaceType(Request) : null;
+		}
+
+		/// <summary>
+		/// Creates an API response object from this stream
+		/// </summary>
+		/// <returns>The response object</returns>
+		public Api.GetWorkspaceTypeResponse ToApiResponse()
+		{
+			return new Api.GetWorkspaceTypeResponse(Cluster, ServerAndPort, UserName, Identifier, Stream, View, Incremental);
+		}
+	}
+
+	/// <summary>
+	/// Allows triggering another downstream job on succesful completion of a step or aggregate
+	/// </summary>
+	public class ChainedJobTemplate
+	{
+		/// <summary>
+		/// Name of the target that needs to complete successfully
+		/// </summary>
+		public string Trigger { get; set; } = String.Empty;
+
+		/// <summary>
+		/// The new template to trigger
+		/// </summary>
+		public TemplateRefId TemplateRefId { get; set; }
+
+		/// <summary>
+		/// Default constructor for serialization
+		/// </summary>
+		private ChainedJobTemplate()
+		{
+		}
+
+		/// <summary>
+		/// Constructor
+		/// </summary>
+		/// <param name="Trigger">Name of the target that needs to complete</param>
+		/// <param name="TemplateRefId">The new template to trigger</param>
+		public ChainedJobTemplate(string Trigger, TemplateRefId TemplateRefId)
+		{
+			this.Trigger = Trigger;
+			this.TemplateRefId = TemplateRefId;
+		}
+
+		/// <summary>
+		/// Constructor
+		/// </summary>
+		/// <param name="Request">Request to construct from</param>
+		public ChainedJobTemplate(CreateChainedJobTemplateRequest Request)
+			: this(Request.Trigger, new TemplateRefId(Request.TemplateId))
+		{
+		}
+	}
+
+	/// <summary>
+	/// Reference to a template
+	/// </summary>
+	public class TemplateRef
+	{
+		/// <summary>
+		/// The template name (duplicated from the template object)
+		/// </summary>
+		public string Name { get; set; }
+
+		/// <summary>
+		/// Hash of the template definition
+		/// </summary>
+		public ContentHash Hash { get; set; }
+
+		/// <summary>
+		/// Whether to show badges in UGS for this schedule
+		/// </summary>
+		public bool ShowUgsBadges { get; set; }
+
+		/// <summary>
+		/// Whether to show desktop alerts for build health issues created from jobs this type
+		/// </summary>
+		public bool ShowUgsAlerts { get; set; }
+
+		/// <summary>
+		/// Notification channel for this template. Overrides the stream channel if set.
+		/// </summary>
+		public string? NotificationChannel { get; set; }
+
+		/// <summary>
+		/// Notification channel filter for this template. Errors|Warnings|Success
+		/// </summary>
+		public string? NotificationChannelFilter { get; set; }
+
+		/// <summary>
+		/// Channel for triage notification messages
+		/// </summary>
+		public string? TriageChannel { get; set; }
+
+		/// <summary>
+		/// List of schedules for this template
+		/// </summary>
+		[BsonIgnoreIfNull]
+		public Schedule? Schedule { get; set; }
+
+		/// <summary>
+		/// List of downstream templates to trigger at the same change
+		/// </summary>
+		[BsonIgnoreIfNull]
+		public List<ChainedJobTemplate>? ChainedJobs { get; set; }
+
+		/// <summary>
+		/// Custom permissions for this template
+		/// </summary>
+		public Acl? Acl { get; set; }
+
+		/// <summary>
+		/// Private constructor for serialization
+		/// </summary>
+		private TemplateRef()
+		{
+			this.Name = null!;
+			this.Hash = null!;
+		}
+
+		/// <summary>
+		/// Constructor
+		/// </summary>
+		/// <param name="Template">The template being referenced</param>
+		/// <param name="ShowUgsBadges">Whether to show badges in UGS for this job</param>
+		/// <param name="ShowUgsAlerts">Whether to show alerts in UGS for this job</param>
+		/// <param name="NotificationChannel">Notification channel for this template</param>
+		/// <param name="NotificationChannelFilter">Notification channel filter for this template</param>
+		/// <param name="TriageChannel"></param>
+		/// <param name="Schedule">Schedule for this template</param>
+		/// <param name="Triggers">List of downstream templates to trigger</param>
+		/// <param name="Acl">ACL for this template</param>
+		public TemplateRef(ITemplate Template, bool ShowUgsBadges = false, bool ShowUgsAlerts = false, string? NotificationChannel = null, string? NotificationChannelFilter = null, string? TriageChannel = null, Schedule? Schedule = null, List<ChainedJobTemplate>? Triggers = null, Acl? Acl = null)
+		{
+			this.Name = Template.Name;
+			this.Hash = Template.Id;
+			this.ShowUgsBadges = ShowUgsBadges;
+			this.ShowUgsAlerts = ShowUgsAlerts;
+			this.NotificationChannel = NotificationChannel;
+			this.NotificationChannelFilter = NotificationChannelFilter;
+			this.TriageChannel = TriageChannel;
+			this.Schedule = Schedule;
+			this.ChainedJobs = Triggers;
+			this.Acl = Acl;
+		}
+	}
+
+	/// <summary>
+	/// Query used to identify a base changelist for a preflight
+	/// </summary>
+	public class ChangeQuery
+	{
+		/// <summary>
+		/// Template to search for
+		/// </summary>
+		public TemplateRefId? TemplateRefId { get; set; }
+
+		/// <summary>
+		/// The target to look at the status for
+		/// </summary>
+		public string? Target { get; set; }
+
+		/// <summary>
+		/// Whether to match a job that contains warnings
+		/// </summary>
+		public List<JobStepOutcome>? Outcomes { get; set; }
+
+		/// <summary>
+		/// Convert to a request object
+		/// </summary>
+		/// <returns></returns>
+		public ChangeQueryRequest ToRequest()
+		{
+			return new ChangeQueryRequest { TemplateId = TemplateRefId?.ToString(), Target = Target, Outcomes = Outcomes };
+		}
+	}
+
+	/// <summary>
+	/// Definition of a query to execute to find the changelist to run a build at
+	/// </summary>
+	public class DefaultPreflight
+	{
+		/// <summary>
+		/// The template id to execute
+		/// </summary>
+		public TemplateRefId? TemplateRefId { get; set; }
+
+		/// <summary>
+		/// Query specifying a changelist to use
+		/// </summary>
+		public ChangeQuery? Change { get; set; }
+
+		/// <summary>
+		/// The job type to query for the change to use
+		/// </summary>
+		[Obsolete("Use Change.TemplateRefId instead")]
+		public TemplateRefId? ChangeTemplateRefId { get; set; }
+
+		/// <summary>
+		/// Constructor
+		/// </summary>
+		/// <param name="TemplateRefId"></param>
+		/// <param name="Change">The job type to query for the change to use</param>
+		public DefaultPreflight(TemplateRefId? TemplateRefId, ChangeQuery? Change)
+		{
+			this.TemplateRefId = TemplateRefId;
+			this.Change = Change;
+		}
+
+		/// <summary>
+		/// Convert to a request object
+		/// </summary>
+		/// <returns></returns>
+		public DefaultPreflightRequest ToRequest()
+		{
+#pragma warning disable CS0618 // Type or member is obsolete
+			ChangeQueryRequest? ChangeRequest = null;
+			if (Change != null)
+			{
+				ChangeRequest = Change.ToRequest();
+			}
+			else if (ChangeTemplateRefId != null)
+			{
+				ChangeRequest = new ChangeQueryRequest { TemplateId = ChangeTemplateRefId?.ToString() };
+			}
+
+			return new DefaultPreflightRequest { TemplateId = TemplateRefId?.ToString(), Change = ChangeRequest, ChangeTemplateId = ChangeRequest?.TemplateId };
+#pragma warning restore CS0618 // Type or member is obsolete
+		}
+	}
+
+	/// <summary>
+	/// Extension methods for template refs
+	/// </summary>
+	static class TemplateRefExtensions
+	{
+		/// <summary>
+		/// Adds a new template ref to a list
+		/// </summary>
+		/// <param name="TemplateRefs">List of template refs</param>
+		/// <param name="TemplateRef">The template ref to add</param>
+		public static void AddRef(this Dictionary<TemplateRefId, TemplateRef> TemplateRefs, TemplateRef TemplateRef)
+		{
+			TemplateRefs.Add(new TemplateRefId(TemplateRef.Name), TemplateRef);
+		}
+	}
+
+	/// <summary>
+	/// Information about a stream
+	/// </summary>
+	public interface IStream
+	{
+		/// <summary>
+		/// Name of the stream.
+		/// </summary>
+		public StreamId Id { get; }
+
+		/// <summary>
+		/// The project that this stream belongs to
+		/// </summary>
+		public ProjectId ProjectId { get; }
+
+		/// <summary>
+		/// The stream name
+		/// </summary>
+		public string Name { get; }
+
+		/// <summary>
+		/// Name of the perforce cluster
+		/// </summary>
+		public string ClusterName { get; }
+
+		/// <summary>
+		/// Path to the configuration file for this stream
+		/// </summary>
+		public string ConfigPath { get; }
+
+		/// <summary>
+		/// The revision of config file used for this stream
+		/// </summary>
+		public string ConfigRevision { get; }
+
+		/// <summary>
+		/// Order to display on the dashboard's drop-down list
+		/// </summary>
+		public int Order { get; }
+
+		/// <summary>
+		/// Notification channel for all jobs in this stream
+		/// </summary>
+		public string? NotificationChannel { get; }
+
+		/// <summary>
+		/// Notification channel filter for all jobs in this stream. Errors|Warnings|Success
+		/// </summary>
+		public string? NotificationChannelFilter { get; }
+
+		/// <summary>
+		/// Channel to post issue triage notifications
+		/// </summary>
+		public string? TriageChannel { get; }
+
+		/// <summary>
+		/// Default template to use for preflights
+		/// </summary>
+		public DefaultPreflight? DefaultPreflight { get; }
+
+		/// <summary>
+		/// List of pages to display in the dashboard
+		/// </summary>
+		public IReadOnlyList<StreamTab> Tabs { get; }
+
+		/// <summary>
+		/// Dictionary of agent types
+		/// </summary>
+		public IReadOnlyDictionary<string, AgentType> AgentTypes { get; }
+
+		/// <summary>
+		/// Dictionary of workspace types
+		/// </summary>
+		public IReadOnlyDictionary<string, WorkspaceType> WorkspaceTypes { get; }
+
+		/// <summary>
+		/// List of templates available for this stream
+		/// </summary>
+		public IReadOnlyDictionary<TemplateRefId, TemplateRef> Templates { get; }
+
+		/// <summary>
+		/// Stream is paused for builds until specified time
+		/// </summary>
+		public DateTime? PausedUntil { get; }
+		
+		/// <summary>
+		/// Comment/reason for why the stream was paused
+		/// </summary>
+		public string? PauseComment { get; }
+
+		/// <summary>
+		/// The ACL for this object
+		/// </summary>
+		public Acl? Acl { get; }
+	}
+
+	/// <summary>
+	/// Extension methods for streams
+	/// </summary>
+	static class StreamExtensions
+	{
+		/// <summary>
+		/// Tries to get an agent workspace definition from the given type name
+		/// </summary>
+		/// <param name="Stream">The stream object</param>
+		/// <param name="AgentType">The agent type</param>
+		/// <param name="Workspace">Receives the agent workspace definition</param>
+		/// <returns>True if the agent type was valid, and an agent workspace could be created</returns>
+		public static bool TryGetAgentWorkspace(this IStream Stream, AgentType AgentType, [NotNullWhen(true)] out AgentWorkspace? Workspace)
+		{
+			// Get the workspace settings
+			if (AgentType.Workspace == null)
+			{
+				// Use the default settings (fast switching workspace, clean 
+				Workspace = new AgentWorkspace(null, null, GetDefaultWorkspaceIdentifier(Stream), Stream.Name, null, false);
+				return true;
+			}
+			else
+			{
+				// Try to get the matching workspace type
+				WorkspaceType? WorkspaceType;
+				if (!Stream.WorkspaceTypes.TryGetValue(AgentType.Workspace, out WorkspaceType))
+				{
+					Workspace = null;
+					return false;
+				}
+
+				// Get the workspace identifier
+				string Identifier;
+				if (WorkspaceType.Identifier != null && !String.IsNullOrEmpty(WorkspaceType.Identifier))
+				{
+					Identifier = WorkspaceType.Identifier;
+				}
+				else if (WorkspaceType.Incremental)
+				{
+					Identifier = $"{Stream.GetEscapedName()}+{AgentType.Workspace}";
+				}
+				else
+				{
+					Identifier = GetDefaultWorkspaceIdentifier(Stream);
+				}
+
+				// Create the new workspace
+				Workspace = new AgentWorkspace(WorkspaceType.Cluster, WorkspaceType.UserName, Identifier, WorkspaceType.Stream ?? Stream.Name, WorkspaceType.View, WorkspaceType.Incremental);
+				return true;
+			}
+		}
+
+		/// <summary>
+		/// The escaped name of this stream. Removes all non-identifier characters.
+		/// </summary>
+		/// <param name="Stream">The stream object</param>
+		/// <returns>Escaped name for the stream</returns>
+		public static string GetEscapedName(this IStream Stream)
+		{
+			return Regex.Replace(Stream.Name, @"[^a-zA-Z0-9_]", "+");
+		}
+
+		/// <summary>
+		/// Gets the default identifier for workspaces created for this stream. Just includes an escaped depot name.
+		/// </summary>
+		/// <param name="Stream">The stream object</param>
+		/// <returns>The default workspace identifier</returns>
+		private static string GetDefaultWorkspaceIdentifier(IStream Stream)
+		{
+			return Regex.Replace(Stream.GetEscapedName(), @"^(\+\+[^+]*).*$", "$1");
+		}
+
+		/// <summary>
+		/// Converts to a public response object
+		/// </summary>
+		/// <param name="Stream">The stream object</param>
+		/// <param name="bIncludeAcl">Whether to include the ACL in the response object</param>
+		/// <param name="ApiTemplateRefs">The template refs for this stream. Passed separately because they have their own ACL.</param>
+		/// <returns>New response instance</returns>
+		public static Api.GetStreamResponse ToApiResponse(this IStream Stream, bool bIncludeAcl, List<GetTemplateRefResponse> ApiTemplateRefs)
+		{
+			List<GetStreamTabResponse> ApiTabs = Stream.Tabs.ConvertAll(x => x.ToResponse());
+			Dictionary<string, GetAgentTypeResponse> ApiAgentTypes = Stream.AgentTypes.ToDictionary(x => x.Key, x => x.Value.ToApiResponse());
+			Dictionary<string, GetWorkspaceTypeResponse> ApiWorkspaceTypes = Stream.WorkspaceTypes.ToDictionary(x => x.Key, x => x.Value.ToApiResponse());
+			GetAclResponse? ApiAcl = (bIncludeAcl && Stream.Acl != null)? new GetAclResponse(Stream.Acl) : null;
+			return new Api.GetStreamResponse(Stream.Id.ToString(), Stream.ProjectId.ToString(), Stream.Name, Stream.ConfigPath, Stream.ConfigRevision, Stream.Order, Stream.NotificationChannel, Stream.NotificationChannelFilter, Stream.TriageChannel, Stream.DefaultPreflight?.ToRequest(), ApiTabs, ApiAgentTypes, ApiWorkspaceTypes, ApiTemplateRefs, ApiAcl, Stream.PausedUntil, Stream.PauseComment);
+		}
+
+		/// <summary>
+		/// Converts to an RPC response object
+		/// </summary>
+		/// <param name="Stream">The stream object</param>
+		/// <returns>New response instance</returns>
+		public static HordeCommon.Rpc.GetStreamResponse ToRpcResponse(this IStream Stream)
+		{
+			HordeCommon.Rpc.GetStreamResponse Response = new HordeCommon.Rpc.GetStreamResponse();
+			Response.Name = Stream.Name;
+			Response.AgentTypes.Add(Stream.AgentTypes.ToDictionary(x => x.Key, x => x.Value.ToRpcResponse()));
+			return Response;
+		}
+
+		/// <summary>
+		/// Check if stream is paused for new builds
+		/// </summary>
+		/// <param name="Stream">The stream object</param>
+		/// <param name="CurrentTime">Current time (allow tests to pass in a fake clock)</param>
+		/// <returns>If stream is paused</returns>
+		public static bool IsPaused(this IStream Stream, DateTime CurrentTime)
+		{
+			return Stream.PausedUntil != null && Stream.PausedUntil > CurrentTime;
+		}
+	}
+
+	/// <summary>
+	/// Projection of a stream definition to just include permissions info
+	/// </summary>
+	public interface IStreamPermissions
+	{
+		/// <summary>
+		/// ACL for the stream
+		/// </summary>
+		public Acl? Acl { get; }
+
+		/// <summary>
+		/// The project containing this stream
+		/// </summary>
+		public ProjectId ProjectId { get; }
+	}
+}
diff --git a/Engine/Source/Programs/Horde/HordeServer/Streams/IStreamCollection.cs b/Engine/Source/Programs/Horde/HordeServer/Streams/IStreamCollection.cs
new file mode 100644
--- /dev/null
+++ b/Engine/Source/Programs/Horde/HordeServer/Streams/IStreamCollection.cs
@@ -0,0 +1,120 @@
+// Copyright Epic Games, Inc. All Rights Reserved.
+
+using HordeServer.Api;
+using HordeServer.Models;
+using HordeServer.Services;
+using HordeServer.Utilities;
+using MongoDB.Bson;
+using MongoDB.Bson.Serialization.Attributes;
+using MongoDB.Driver;
+using System;
+using System.Collections.Generic;
+using System.Diagnostics.CodeAnalysis;
+using System.Linq;
+using System.Threading.Tasks;
+
+namespace HordeServer.Collections
+{
+	using ProjectId = StringId<IProject>;
+	using StreamId = StringId<IStream>;
+	using TemplateRefId = StringId<TemplateRef>;
+
+	/// <summary>
+	/// Collection of stream documents
+	/// </summary>
+	public interface IStreamCollection
+	{
+		/// <summary>
+		/// Creates or replaces a stream configuration
+		/// </summary>
+		/// <param name="Id">Unique id for the new stream</param>
+		/// <param name="Stream">The current stream value. If not-null, this will attempt to replace the existing instance.</param>
+		/// <param name="ConfigPath">Path to the config file</param>
+		/// <param name="Revision">The config file revision</param>
+		/// <param name="ProjectId">The project id</param>
+		/// <param name="Config">The stream configuration</param>
+		/// <returns></returns>
+		Task<IStream?> TryCreateOrReplaceAsync(StreamId Id, IStream? Stream, string ConfigPath, string Revision, ProjectId ProjectId, StreamConfig Config);
+
+		/// <summary>
+		/// Gets a stream by ID
+		/// </summary>
+		/// <param name="StreamId">Unique id of the stream</param>
+		/// <returns>The stream document</returns>
+		Task<IStream?> GetAsync(StreamId StreamId);
+
+		/// <summary>
+		/// Gets a stream's permissions by ID
+		/// </summary>
+		/// <param name="StreamId">Unique id of the stream</param>
+		/// <returns>The stream document</returns>
+		Task<IStreamPermissions?> GetPermissionsAsync(StreamId StreamId);
+
+		/// <summary>
+		/// Enumerates all streams
+		/// </summary>
+		/// <returns></returns>
+		Task<List<IStream>> FindAllAsync();
+
+		/// <summary>
+		/// Gets all the available streams for a project
+		/// </summary>
+		/// <param name="ProjectIds">Unique id of the projects to query</param>
+		/// <returns>List of stream documents</returns>
+		Task<List<IStream>> FindForProjectsAsync(ProjectId[] ProjectIds);
+
+		/// <summary>
+		/// Updates user-facing properties for an existing stream
+		/// </summary>
+		/// <param name="Stream">The stream to update</param>
+		/// <param name="NewPausedUntil">The new datetime for pausing builds</param>
+		/// <param name="NewPauseComment">The reason for pausing</param>
+		/// <returns>The updated stream if successful, null otherwise</returns>
+		Task<IStream?> TryUpdatePauseStateAsync(IStream Stream, DateTime? NewPausedUntil, string? NewPauseComment);
+
+		/// <summary>
+		/// Attempts to update the last trigger time for a schedule
+		/// </summary>
+		/// <param name="Stream">The stream to update</param>
+		/// <param name="TemplateRefId">The template ref id</param>
+		/// <param name="LastTriggerTime">New last trigger time for the schedule</param>
+		/// <param name="LastTriggerChange">New last trigger changelist for the schedule</param>
+		/// <param name="NewActiveJobs">New list of active jobs</param>
+		/// <returns>The updated stream if successful, null otherwise</returns>
+		Task<IStream?> TryUpdateScheduleTriggerAsync(IStream Stream, TemplateRefId TemplateRefId, DateTimeOffset? LastTriggerTime, int? LastTriggerChange, List<ObjectId> NewActiveJobs);
+
+		/// <summary>
+		/// Delete a stream
+		/// </summary>
+		/// <param name="StreamId">Unique id of the stream</param>
+		/// <returns>Async task</returns>
+		Task DeleteAsync(StreamId StreamId);
+	}
+
+	static class StreamCollectionExtensions
+	{
+		/// <summary>
+		/// Creates or replaces a stream configuration
+		/// </summary>
+		/// <param name="StreamCollection">The stream collection</param>
+		/// <param name="Id">Unique id for the new stream</param>
+		/// <param name="Stream">The current stream value. If not-null, this will attempt to replace the existing instance.</param>
+		/// <param name="ConfigPath">Path to the config file</param>
+		/// <param name="Revision">The config file revision</param>
+		/// <param name="ProjectId">The project id</param>
+		/// <param name="Config">The stream configuration</param>
+		/// <returns></returns>
+		public static async Task<IStream> CreateOrReplaceAsync(this IStreamCollection StreamCollection, StreamId Id, IStream? Stream, string ConfigPath, string Revision, ProjectId ProjectId, StreamConfig Config)
+		{
+			for (; ; )
+			{
+				Stream = await StreamCollection.TryCreateOrReplaceAsync(Id, Stream, ConfigPath, Revision, ProjectId, Config);
+				if (Stream != null)
+				{
+					return Stream;
+				}
+				Stream = await StreamCollection.GetAsync(Id);
+			}
+		}
+	}
+}
diff --git a/Engine/Source/Programs/Horde/HordeServer/Streams/Impl/StreamCollection.cs b/Engine/Source/Programs/Horde/HordeServer/Streams/Impl/StreamCollection.cs
new file mode 100644
--- /dev/null
+++ b/Engine/Source/Programs/Horde/HordeServer/Streams/Impl/StreamCollection.cs
@@ -0,0 +1,422 @@
+// Copyright Epic Games, Inc. All Rights Reserved.
+
+using HordeCommon;
+using HordeServer.Api;
+using HordeServer.Models;
+using HordeServer.Services;
+using HordeServer.Utilities;
+using MongoDB.Bson;
+using MongoDB.Bson.Serialization.Attributes;
+using MongoDB.Driver;
+using System;
+using System.Collections.Generic;
+using System.Diagnostics.CodeAnalysis;
+using System.IO;
+using System.Linq;
+using System.Threading.Tasks;
+
+namespace HordeServer.Collections.Impl
+{
+	using ProjectId = StringId<IProject>;
+	using StreamId = StringId<IStream>;
+	using TemplateRefId = StringId<TemplateRef>;
+
+	/// <summary>
+	/// Collection of stream documents
+	/// </summary>
+	public class StreamCollection : IStreamCollection
+	{
+		/// <summary>
+		/// Information about a stream
+		/// </summary>
+		class StreamDocument : IStream
+		{
+			public const int DefaultOrder = 128;
+
+			[BsonRequired, BsonId]
+			public StreamId Id { get; set; }
+
+			[BsonRequired]
+			public ProjectId ProjectId { get; set; }
+
+			[BsonRequired]
+			public string Name { get; set; }
+
+			public string? ClusterName { get; set; }
+			public string ConfigPath { get; set; } = String.Empty;
+			public string ConfigRevision { get; set; } = String.Empty;
+
+			public int Order { get; set; } = DefaultOrder;
+			public string? NotificationChannel { get; set; }
+			public string? NotificationChannelFilter { get; set; }
+			public string? TriageChannel { get; set; }
+			public DefaultPreflight? DefaultPreflight { get; set; }
+			public List<StreamTab> Tabs { get; set; } = new List<StreamTab>();
+			public Dictionary<string, AgentType> AgentTypes { get; set; } = new Dictionary<string, AgentType>(StringComparer.Ordinal);
+			public Dictionary<string, WorkspaceType> WorkspaceTypes { get; set; } = new Dictionary<string, WorkspaceType>(StringComparer.Ordinal);
+			public Dictionary<TemplateRefId, TemplateRef> Templates { get; set; } = new Dictionary<TemplateRefId, TemplateRef>();
+			public DateTime? PausedUntil { get; set; }
+			public string? PauseComment { get; set; }
+			public Acl? Acl { get; set; }
+			public int UpdateIndex { get; set; }
+			public bool Deleted { get; set; }
+
+			string IStream.ClusterName => ClusterName ?? PerforceCluster.DefaultName;
+			IReadOnlyList<StreamTab> IStream.Tabs => Tabs;
+			IReadOnlyDictionary<string, AgentType> IStream.AgentTypes => AgentTypes;
+			IReadOnlyDictionary<string, WorkspaceType> IStream.WorkspaceTypes => WorkspaceTypes;
+			IReadOnlyDictionary<TemplateRefId, TemplateRef> IStream.Templates => Templates;
+
+			[BsonConstructor]
+			private StreamDocument()
+			{
+				Name = null!;
+			}
+
+			public StreamDocument(StreamId Id, string Name, ProjectId ProjectId)
+			{
+				this.Id = Id;
+				this.Name = Name;
+				this.ProjectId = ProjectId;
+			}
+		}
+
+		/// <summary>
+		/// Projection of a stream definition to just include permissions info
+		/// </summary>
+		[SuppressMessage("Design", "CA1812: Class is never instantiated")]
+		private class StreamPermissions : IStreamPermissions
+		{
+			public Acl? Acl { get; set; }
+			public ProjectId ProjectId { get; set; }
+
+			public static readonly ProjectionDefinition<StreamDocument> Projection = Builders<StreamDocument>.Projection.Include(x => x.Acl).Include(x => x.ProjectId);
+		}
+
+		/// <summary>
+		/// The stream collection
+		/// </summary>
+		IMongoCollection<StreamDocument> Streams;
+
+		/// <summary>
+		/// Clock
+		/// </summary>
+		IClock Clock;
+
+		/// <summary>
+		/// The template collection
+		/// </summary>
+		ITemplateCollection TemplateCollection;
+
+		/// <summary>
+		/// Constructor
+		/// </summary>
+		/// <param name="DatabaseService">The database service instance</param>
+		/// <param name="Clock"></param>
+		/// <param name="TemplateCollection"></param>
+		public StreamCollection(DatabaseService DatabaseService, IClock Clock, ITemplateCollection TemplateCollection)
+		{
+			this.Streams = DatabaseService.GetCollection<StreamDocument>("Streams");
+			this.Clock = Clock;
+			this.TemplateCollection = TemplateCollection;
+		}
+
+		/// <inheritdoc/>
+		public async Task<IStream?> TryCreateOrReplaceAsync(StreamId Id, IStream? Stream, string ConfigPath, string Revision, ProjectId ProjectId, StreamConfig Config)
+		{
+			List<StreamTab> Tabs = Config.Tabs.ConvertAll(x => StreamTab.FromRequest(x));
+			Dictionary<TemplateRefId, TemplateRef> TemplateRefs = await CreateTemplateRefsAsync(Config.Templates, Stream, TemplateCollection);
+
+			Dictionary<string, AgentType> AgentTypes = new Dictionary<string, AgentType>();
+			if (Config.AgentTypes != null)
+			{
+				AgentTypes = Config.AgentTypes.Where(x => x.Value != null).ToDictionary(x => x.Key, x => new AgentType(x.Value!));
+			}
+
+			Dictionary<string, WorkspaceType> WorkspaceTypes = new Dictionary<string, WorkspaceType>();
+			if (Config.WorkspaceTypes != null)
+			{
+				WorkspaceTypes = Config.WorkspaceTypes.Where(x => x.Value != null).ToDictionary(x => x.Key, x => new WorkspaceType(x.Value!));
+			}
+
+			DefaultPreflight? DefaultPreflight = Config.DefaultPreflight?.ToModel();
+			if (DefaultPreflight == null && Config.DefaultPreflightTemplate != null)
+			{
+				DefaultPreflight = new DefaultPreflight(new TemplateRefId(Config.DefaultPreflightTemplate), null);
+			}
+
+			Validate(Id, DefaultPreflight, TemplateRefs, Tabs, AgentTypes, WorkspaceTypes);
+
+			Acl? Acl = Acl.Merge(new Acl(), Config.Acl);
+			if (Stream == null)
+			{
+				return await TryCreateAsync(Id, ProjectId, ConfigPath, Revision, Config, DefaultPreflight, Tabs, AgentTypes, WorkspaceTypes, TemplateRefs, Acl);
+			}
+			else
+			{
+				return await TryReplaceAsync(Stream, ProjectId, ConfigPath, Revision, Config, DefaultPreflight, Tabs, AgentTypes, WorkspaceTypes, TemplateRefs, Acl);
+			}
+		}
+
+		/// <summary>
+		/// Creates a list of template refs from a set of request objects
+		/// </summary>
+		/// <param name="Requests">Request objects</param>
+		/// <param name="Stream">The current stream state</param>
+		/// <param name="TemplateCollection">The template service</param>
+		/// <returns>List of new template references</returns>
+		async Task<Dictionary<TemplateRefId, TemplateRef>> CreateTemplateRefsAsync(List<CreateTemplateRefRequest> Requests, IStream? Stream, ITemplateCollection TemplateCollection)
+		{
+			Dictionary<TemplateRefId, TemplateRef> NewTemplateRefs = new Dictionary<TemplateRefId, TemplateRef>();
+			foreach (CreateTemplateRefRequest Request in Requests)
+			{
+				// Create the template
+				ITemplate NewTemplate = await TemplateCollection.AddAsync(Request.Name, Request.Priority, Request.AllowPreflights, Request.InitialAgentType, Request.SubmitNewChange, Request.Counters, Request.Arguments, Request.Parameters.ConvertAll(x => x.ToModel()));
+
+				// Get an identifier for the new template ref
+				TemplateRefId NewTemplateRefId;
+				if (Request.Id != null)
+				{
+					NewTemplateRefId = new TemplateRefId(Request.Id);
+				}
+				else
+				{
+					NewTemplateRefId = TemplateRefId.Sanitize(Request.Name);
+				}
+
+				// Create the schedule object
+				Schedule? Schedule = null;
+				if (Request.Schedule != null)
+				{
+					Schedule = Request.Schedule.ToModel();
+					Schedule.LastTriggerTime = Clock.UtcNow;
+				}
+
+				// Add it to the list
+				TemplateRef NewTemplateRef = new TemplateRef(NewTemplate, Request.ShowUgsBadges, Request.ShowUgsAlerts, Request.NotificationChannel, Request.NotificationChannelFilter, Request.TriageChannel, Schedule, Request.ChainedJobs?.ConvertAll(x => new ChainedJobTemplate(x)), Acl.Merge(null, Request.Acl));
+				if (Stream != null && Stream.Templates.TryGetValue(NewTemplateRefId, out TemplateRef? OldTemplateRef))
+				{
+					if (OldTemplateRef.Schedule != null && NewTemplateRef.Schedule != null)
+					{
+						NewTemplateRef.Schedule.CopyState(OldTemplateRef.Schedule);
+					}
+				}
+				NewTemplateRefs.Add(NewTemplateRefId, NewTemplateRef);
+			}
+			foreach (TemplateRef TemplateRef in NewTemplateRefs.Values)
+			{
+				if (TemplateRef.ChainedJobs != null)
+				{
+					foreach (ChainedJobTemplate ChainedJob in TemplateRef.ChainedJobs)
+					{
+						if (!NewTemplateRefs.ContainsKey(ChainedJob.TemplateRefId))
+						{
+							throw new InvalidDataException($"Invalid template ref id '{ChainedJob.TemplateRefId}");
+						}
+					}
+				}
+			}
+			return NewTemplateRefs;
+		}
+
+		/// <inheritdoc/>
+		async Task<IStream?> TryCreateAsync(StreamId Id, ProjectId ProjectId, string ConfigPath, string ConfigRevision, StreamConfig Config, DefaultPreflight? DefaultPreflight, List<StreamTab> Tabs, Dictionary<string, AgentType> AgentTypes, Dictionary<string, WorkspaceType> WorkspaceTypes, Dictionary<TemplateRefId, TemplateRef> TemplateRefs, Acl? Acl)
+		{
+			StreamDocument NewStream = new StreamDocument(Id, Config.Name, ProjectId);
+			NewStream.ClusterName = Config.ClusterName;
+			NewStream.ConfigPath = ConfigPath;
+			NewStream.ConfigRevision = ConfigRevision;
+			NewStream.Order = Config.Order ?? StreamDocument.DefaultOrder;
+			NewStream.NotificationChannel = Config.NotificationChannel;
+			NewStream.NotificationChannelFilter = Config.NotificationChannelFilter;
+			NewStream.TriageChannel = Config.TriageChannel;
+			NewStream.DefaultPreflight = DefaultPreflight;
+			NewStream.Tabs = Tabs;
+			NewStream.AgentTypes = AgentTypes;
+			NewStream.WorkspaceTypes = WorkspaceTypes;
+			NewStream.Templates = TemplateRefs;
+			NewStream.Acl = Acl;
+
+			try
+			{
+				await Streams.InsertOneAsync(NewStream);
+				return NewStream;
+			}
+			catch (MongoWriteException Ex)
+			{
+				if (Ex.WriteError.Category == ServerErrorCategory.DuplicateKey)
+				{
+					return null;
+				}
+				else
+				{
+					throw;
+				}
+			}
+		}
+
+		/// <inheritdoc/>
+		async Task<IStream?> TryReplaceAsync(IStream StreamInterface, ProjectId ProjectId, string ConfigPath, string ConfigRevision, StreamConfig Config, DefaultPreflight? DefaultPreflight, List<StreamTab> Tabs, Dictionary<string, AgentType>? AgentTypes, Dictionary<string, WorkspaceType>? WorkspaceTypes, Dictionary<TemplateRefId, TemplateRef>? TemplateRefs, Acl? Acl)
+		{
+			int Order = Config.Order ?? StreamDocument.DefaultOrder;
+
+			StreamDocument Stream = (StreamDocument)StreamInterface;
+
+			UpdateDefinitionBuilder<StreamDocument> UpdateBuilder = Builders<StreamDocument>.Update;
+
+			List<UpdateDefinition<StreamDocument>> Updates = new List<UpdateDefinition<StreamDocument>>();
+			Updates.Add(UpdateBuilder.Set(x => x.Name, Config.Name));
+			Updates.Add(UpdateBuilder.Set(x => x.ProjectId, ProjectId));
+			Updates.Add(UpdateBuilder.Set(x => x.ClusterName, Config.ClusterName));
+			Updates.Add(UpdateBuilder.Set(x => x.ConfigPath, ConfigPath));
+			Updates.Add(UpdateBuilder.Set(x => x.ConfigRevision, ConfigRevision));
+			Updates.Add(UpdateBuilder.Set(x => x.Order, Order));
+			Updates.Add(UpdateBuilder.Set(x => x.NotificationChannel, Config.NotificationChannel));
+			Updates.Add(UpdateBuilder.Set(x => x.NotificationChannelFilter, Config.NotificationChannelFilter));
+			Updates.Add(UpdateBuilder.Set(x => x.TriageChannel, Config.TriageChannel));
+			Updates.Add(UpdateBuilder.Set(x => x.DefaultPreflight, DefaultPreflight));
+			Updates.Add(UpdateBuilder.Set(x => x.Tabs, Tabs ?? new List<StreamTab>()));
+			Updates.Add(UpdateBuilder.Set(x => x.AgentTypes, AgentTypes ?? new Dictionary<string, AgentType>()));
+			Updates.Add(UpdateBuilder.Set(x => x.WorkspaceTypes, WorkspaceTypes ?? new Dictionary<string, WorkspaceType>()));
+			Updates.Add(UpdateBuilder.Set(x => x.Templates, TemplateRefs ?? new Dictionary<TemplateRefId, TemplateRef>()));
+			Updates.Add(UpdateBuilder.SetOrUnsetNullRef(x => x.Acl, Acl));
+			Updates.Add(UpdateBuilder.Unset(x => x.Deleted));
+
+			return await TryUpdateStreamAsync(Stream, UpdateBuilder.Combine(Updates));
+		}
+
+		/// <inheritdoc/>
+		public async Task<IStream?> GetAsync(StreamId StreamId)
+		{
+			return await Streams.Find<StreamDocument>(x => x.Id == StreamId).FirstOrDefaultAsync();
+		}
+
+		/// <inheritdoc/>
+		public async Task<IStreamPermissions?> GetPermissionsAsync(StreamId StreamId)
+		{
+			return await Streams.Find<StreamDocument>(x => x.Id == StreamId).Project<StreamPermissions>(StreamPermissions.Projection).FirstOrDefaultAsync();
+		}
+
+		/// <inheritdoc/>
+		public async Task<List<IStream>> FindAllAsync()
+		{
+			List<StreamDocument> Results = await Streams.Find(Builders<StreamDocument>.Filter.Ne(x => x.Deleted, true)).ToListAsync();
+			return Results.ConvertAll<IStream>(x => x);
+		}
+
+		/// <inheritdoc/>
+		public async Task<List<IStream>> FindForProjectsAsync(ProjectId[] ProjectIds)
+		{
+			FilterDefinition<StreamDocument> Filter = Builders<StreamDocument>.Filter.In(x => x.ProjectId, ProjectIds) & Builders<StreamDocument>.Filter.Ne(x => x.Deleted, true);
+			List<StreamDocument> Results = await Streams.Find(Filter).ToListAsync();
+			return Results.ConvertAll<IStream>(x => x);
+		}
+
+		/// <inheritdoc/>
+		public async Task<IStream?> TryUpdatePauseStateAsync(IStream StreamInterface, DateTime? NewPausedUntil, string? NewPauseComment)
+		{
+			StreamDocument Stream = (StreamDocument)StreamInterface;
+
+			UpdateDefinitionBuilder<StreamDocument> UpdateBuilder = Builders<StreamDocument>.Update;
+
+			List<UpdateDefinition<StreamDocument>> Updates = new List<UpdateDefinition<StreamDocument>>();
+			Stream.PausedUntil = NewPausedUntil;
+			Stream.PauseComment = NewPauseComment;
+			Updates.Add(UpdateBuilder.Set(x => x.PausedUntil, NewPausedUntil));
+			Updates.Add(UpdateBuilder.Set(x => x.PauseComment, NewPauseComment));
+
+			return await TryUpdateStreamAsync(Stream, UpdateBuilder.Combine(Updates));
+		}
+
+		/// <inheritdoc/>
+		public async Task<IStream?> TryUpdateScheduleTriggerAsync(IStream StreamInterface, TemplateRefId TemplateRefId, DateTimeOffset? LastTriggerTime, int? LastTriggerChange, List<ObjectId> NewActiveJobs)
+		{
+			StreamDocument Stream = (StreamDocument)StreamInterface;
+			Schedule Schedule = Stream.Templates[TemplateRefId].Schedule!;
+
+			// Build the updates. MongoDB driver cannot parse TemplateRefId in expression tree; need to specify field name explicitly
+			List<UpdateDefinition<StreamDocument>> Updates = new List<UpdateDefinition<StreamDocument>>();
+			if (LastTriggerTime.HasValue && LastTriggerTime.Value != Schedule.LastTriggerTime)
+			{
+				FieldDefinition<StreamDocument, DateTimeOffset> LastTriggerTimeField = $"{nameof(Stream.Templates)}.{TemplateRefId}.{nameof(Schedule)}.{nameof(Schedule.LastTriggerTime)}";
+				Updates.Add(Builders<StreamDocument>.Update.Set(LastTriggerTimeField, LastTriggerTime.Value));
+				Schedule.LastTriggerTime = LastTriggerTime.Value;
+			}
+			if (LastTriggerChange.HasValue && LastTriggerChange.Value > Schedule.LastTriggerChange)
+			{
+				FieldDefinition<StreamDocument, int> LastTriggerChangeField = $"{nameof(Stream.Templates)}.{TemplateRefId}.{nameof(Schedule)}.{nameof(Schedule.LastTriggerChange)}";
+				Updates.Add(Builders<StreamDocument>.Update.Set(LastTriggerChangeField, LastTriggerChange.Value));
+				Schedule.LastTriggerChange = LastTriggerChange.Value;
+			}
+			if (NewActiveJobs != null)
+			{
+				FieldDefinition<StreamDocument, List<ObjectId>> Field = $"{nameof(Stream.Templates)}.{TemplateRefId}.{nameof(Schedule)}.{nameof(Schedule.ActiveJobs)}";
+				Updates.Add(Builders<StreamDocument>.Update.Set(Field, NewActiveJobs));
+				Schedule.ActiveJobs = NewActiveJobs;
+			}
+
+			return (Updates.Count == 0)? StreamInterface : await TryUpdateStreamAsync(Stream, Builders<StreamDocument>.Update.Combine(Updates));
+		}
+
+		/// <summary>
+		/// Update a stream
+		/// </summary>
+		/// <param name="Stream">The stream to update</param>
+		/// <param name="Update">The update definition</param>
+		/// <returns>The updated document, or null the update failed</returns>
+		private async Task<StreamDocument?> TryUpdateStreamAsync(StreamDocument Stream, UpdateDefinition<StreamDocument> Update)
+		{
+			FilterDefinition<StreamDocument> Filter = Builders<StreamDocument>.Filter.Expr(x => x.Id == Stream.Id && x.UpdateIndex == Stream.UpdateIndex);
+			Update = Update.Set(x => x.UpdateIndex, Stream.UpdateIndex + 1);
+
+			FindOneAndUpdateOptions<StreamDocument> Options = new FindOneAndUpdateOptions<StreamDocument> { ReturnDocument = ReturnDocument.After };
+			return await Streams.FindOneAndUpdateAsync(Filter, Update, Options);
+		}
+
+		/// <inheritdoc/>
+		public async Task DeleteAsync(StreamId StreamId)
+		{
+			await Streams.UpdateOneAsync<StreamDocument>(x => x.Id == StreamId, Builders<StreamDocument>.Update.Set(x => x.Deleted, true).Inc(x => x.UpdateIndex, 1));
+		}
+
+
+		/// <summary>
+		/// Checks the stream definition for consistency
+		/// </summary>
+		public static void Validate(StreamId StreamId, DefaultPreflight? DefaultPreflight, IReadOnlyDictionary<TemplateRefId, TemplateRef> Templates, IReadOnlyList<StreamTab> Tabs, IReadOnlyDictionary<string, AgentType> AgentTypes, IReadOnlyDictionary<string, WorkspaceType> WorkspaceTypes)
+		{
+			// Check the default preflight template is valid
+			if (DefaultPreflight != null)
+			{
+				if (DefaultPreflight.TemplateRefId != null && !Templates.ContainsKey(DefaultPreflight.TemplateRefId.Value))
+				{
+					throw new InvalidStreamException($"Default preflight template was listed as '{DefaultPreflight.TemplateRefId.Value}', but no template was found by that name");
+				}
+			}
+
+			// Check that all the templates are referenced by a tab
+			HashSet<TemplateRefId> RemainingTemplates = new HashSet<TemplateRefId>(Templates.Keys);
+			foreach (JobsTab JobsTab in Tabs.OfType<JobsTab>())
+			{
+				if (JobsTab.Templates != null)
+				{
+					RemainingTemplates.ExceptWith(JobsTab.Templates);
+				}
+			}
+			if (RemainingTemplates.Count > 0)
+			{
+				throw new InvalidStreamException(String.Join("\n", RemainingTemplates.Select(x => $"Template '{x}' is not listed on any tab for {StreamId}")));
+			}
+
+			// Check that all the agent types reference valid workspace names
+			foreach (KeyValuePair<string, AgentType> Pair in AgentTypes)
+			{
+				string? WorkspaceTypeName = Pair.Value.Workspace;
+				if (WorkspaceTypeName != null && !WorkspaceTypes.ContainsKey(WorkspaceTypeName))
+				{
+					throw new InvalidStreamException($"Agent type '{Pair.Key}' references undefined workspace type '{Pair.Value.Workspace}' in {StreamId}");
+				}
+			}
+		}
+	}
+}
diff --git a/Engine/Source/Programs/Horde/HordeServer/Streams/StreamService.cs b/Engine/Source/Programs/Horde/HordeServer/Streams/StreamService.cs
new file mode 100644
--- /dev/null
+++ b/Engine/Source/Programs/Horde/HordeServer/Streams/StreamService.cs
@@ -0,0 +1,316 @@
+// Copyright Epic Games, Inc. All Rights Reserved.
+
+using HordeServer.Api;
+using HordeCommon;
+using HordeServer.Models;
+using HordeServer.Utilities;
+using Microsoft.Extensions.Options;
+using MongoDB.Bson;
+using MongoDB.Driver;
+using System;
+using System.Collections.Generic;
+using System.Linq;
+using System.Security.Claims;
+using System.Text;
+using System.Threading.Tasks;
+using System.Linq.Expressions;
+using HordeServer.Collections;
+using Microsoft.Extensions.Caching.Memory;
+
+namespace HordeServer.Services
+{
+	using ProjectId = StringId<IProject>;
+	using StreamId = StringId<IStream>;
+	using TemplateRefId = StringId<TemplateRef>;
+
+	/// <summary>
+	/// Cache of information about stream ACLs
+	/// </summary>
+	public class StreamPermissionsCache : ProjectPermissionsCache
+	{
+		/// <summary>
+		/// Map of stream id to permissions for that stream
+		/// </summary>
+		public Dictionary<StreamId, IStreamPermissions?> Streams { get; set; } = new Dictionary<StreamId, IStreamPermissions?>();
+	}
+
+	/// <summary>
+	/// Wraps functionality for manipulating streams
+	/// </summary>
+	public sealed class StreamService : IDisposable
+	{
+		/// <summary>
+		/// The project service instance
+		/// </summary>
+		ProjectService ProjectService;
+
+		/// <summary>
+		/// Collection of stream documents
+		/// </summary>
+		IStreamCollection Streams;
+
+		/// <summary>
+		/// Cache of stream documents
+		/// </summary>
+		MemoryCache StreamCache = new MemoryCache(new MemoryCacheOptions());
+
+		/// <summary>
+		/// Accessor for the stream collection
+		/// </summary>
+		public IStreamCollection StreamCollection => Streams;
+
+		/// <summary>
+		/// Constructor
+		/// </summary>
+		/// <param name="ProjectService">The project service instance</param>
+		/// <param name="Streams">Collection of stream documents</param>
+		public StreamService(ProjectService ProjectService, IStreamCollection Streams)
+		{
+			this.ProjectService = ProjectService;
+			this.Streams = Streams;
+		}
+
+		/// <summary>
+		/// Dispose of any managed resources
+		/// </summary>
+		public void Dispose()
+		{
+			StreamCache.Dispose();
+		}
+
+		/// <summary>
+		/// Deletes an existing stream
+		/// </summary>
+		/// <param name="StreamId">Unique id of the stream</param>
+		/// <returns>Async task object</returns>
+		public async Task DeleteStreamAsync(StreamId StreamId)
+		{
+			await Streams.DeleteAsync(StreamId);
+		}
+
+		/// <summary>
+		/// Updates an existing stream
+		/// </summary>
+		/// <param name="Stream">The stream to update</param>
+		/// <param name="NewPausedUntil">The new datetime for pausing builds</param>
+		/// <param name="NewPauseComment">The reason for pausing</param>
+		/// <returns>Async task object</returns>
+		public async Task<IStream?> UpdatePauseStateAsync(IStream? Stream, DateTime? NewPausedUntil = null, string? NewPauseComment = null)
+		{
+			for (; Stream != null; Stream = await GetStreamAsync(Stream.Id))
+			{
+				IStream? NewStream = await Streams.TryUpdatePauseStateAsync(Stream, NewPausedUntil, NewPauseComment);
+				if (NewStream != null)
+				{
+					return NewStream;
+				}
+			}
+			return null;
+		}
+
+		/// <summary>
+		/// Attempts to update the last trigger time for a schedule
+		/// </summary>
+		/// <param name="Stream">The stream to update</param>
+		/// <param name="TemplateRefId">The template ref id</param>
+		/// <param name="LastTriggerTime"></param>
+		/// <param name="LastTriggerChange"></param>
+		/// <param name="AddJobs">Jobs to add</param>
+		/// <param name="RemoveJobs">Jobs to remove</param>
+		/// <returns>True if the stream was updated</returns>
+		public async Task<IStream?> UpdateScheduleTriggerAsync(IStream Stream, TemplateRefId TemplateRefId, DateTimeOffset? LastTriggerTime, int? LastTriggerChange, List<ObjectId> AddJobs, List<ObjectId> RemoveJobs)
+		{
+			IStream? NewStream = Stream;
+			while (NewStream != null)
+			{
+				TemplateRef? TemplateRef;
+				if (!NewStream.Templates.TryGetValue(TemplateRefId, out TemplateRef))
+				{
+					break;
+				}
+				if (TemplateRef.Schedule == null)
+				{
+					break;
+				}
+
+				List<ObjectId> NewActiveJobs = TemplateRef.Schedule.ActiveJobs.Except(RemoveJobs).Union(AddJobs).ToList();
+
+				NewStream = await Streams.TryUpdateScheduleTriggerAsync(NewStream, TemplateRefId, LastTriggerTime, LastTriggerChange, NewActiveJobs);
+
+				if (NewStream != null)
+				{
+					return NewStream;
+				}
+
+				NewStream = await Streams.GetAsync(Stream.Id);
+			}
+			return null;
+		}
+
+		/// <summary>
+		/// Gets all the available streams for a project
+		/// </summary>
+		/// <returns>List of stream documents</returns>
+		public Task<List<IStream>> GetStreamsAsync()
+		{
+			return Streams.FindAllAsync();
+		}
+
+		/// <summary>
+		/// Gets all the available streams for a project
+		/// </summary>
+		/// <param name="ProjectId">Unique id of the project to query</param>
+		/// <returns>List of stream documents</returns>
+		public Task<List<IStream>> GetStreamsAsync(ProjectId ProjectId)
+		{
+			return Streams.FindForProjectsAsync(new[] { ProjectId });
+		}
+
+		/// <summary>
+		/// Gets all the available streams for a project
+		/// </summary>
+		/// <param name="ProjectIds">Unique id of the project to query</param>
+		/// <returns>List of stream documents</returns>
+		public Task<List<IStream>> GetStreamsAsync(ProjectId[] ProjectIds)
+		{
+			if (ProjectIds.Length == 0)
+			{
+				return Streams.FindAllAsync();
+			}
+			else
+			{
+				return Streams.FindForProjectsAsync(ProjectIds);
+			}
+		}
+
+		/// <summary>
+		/// Gets a stream by ID
+		/// </summary>
+		/// <param name="StreamId">Unique id of the stream</param>
+		/// <returns>The stream document</returns>
+		public Task<IStream?> GetStreamAsync(StreamId StreamId)
+		{
+			return Streams.GetAsync(StreamId);
+		}
+
+		/// <summary>
+		/// Adds a cached stream interface
+		/// </summary>
+		/// <param name="Stream">The stream to add</param>
+		/// <returns>The new stream</returns>
+		private void AddCachedStream(IStream Stream)
+		{
+			MemoryCacheEntryOptions Options = new MemoryCacheEntryOptions().SetAbsoluteExpiration(TimeSpan.FromMinutes(1.0));
+			StreamCache.Set(Stream.Id, Stream, Options);
+		}
+
+		/// <summary>
+		/// Gets a cached stream interface
+		/// </summary>
+		/// <param name="StreamId">Unique id for the stream</param>
+		/// <returns>The new stream</returns>
+		public async Task<IStream?> GetCachedStream(StreamId StreamId)
+		{
+			object? Stream;
+			if (!StreamCache.TryGetValue(StreamId, out Stream))
+			{
+				Stream = await GetStreamAsync(StreamId);
+				if (Stream != null)
+				{
+					AddCachedStream((IStream)Stream);
+				}
+			}
+			return (IStream?)Stream;
+		}
+
+		/// <summary>
+		/// Gets a stream's permissions by ID
+		/// </summary>
+		/// <param name="StreamId">Unique id of the stream</param>
+		/// <returns>The stream document</returns>
+		public Task<IStreamPermissions?> GetStreamPermissionsAsync(StreamId StreamId)
+		{
+			return Streams.GetPermissionsAsync(StreamId);
+		}
+
+		/// <summary>
+		/// Determines if the user is authorized to perform an action on a particular stream
+		/// </summary>
+		/// <param name="Acl">ACL for the stream to check</param>
+		/// <param name="ProjectId">The parent project id</param>
+		/// <param name="Action">The action being performed</param>
+		/// <param name="User">The principal to authorize</param>
+		/// <param name="Cache">Cache of project permissions</param>
+		/// <returns>True if the action is authorized</returns>
+		private Task<bool> AuthorizeAsync(Acl? Acl, ProjectId ProjectId, AclAction Action, ClaimsPrincipal User, ProjectPermissionsCache? Cache)
+		{
+			bool? Result = Acl?.Authorize(Action, User);
+			if (Result == null)
+			{
+				return ProjectService.AuthorizeAsync(ProjectId, Action, User, Cache);
+			}
+			else
+			{
+				return Task.FromResult(Result.Value);
+			}
+		}
+
+		/// <summary>
+		/// Determines if the user is authorized to perform an action on a particular stream
+		/// </summary>
+		/// <param name="Stream">The stream to check</param>
+		/// <param name="Action">The action being performed</param>
+		/// <param name="User">The principal to authorize</param>
+		/// <param name="Cache">Cache of project permissions</param>
+		/// <returns>True if the action is authorized</returns>
+		public Task<bool> AuthorizeAsync(IStream Stream, AclAction Action, ClaimsPrincipal User, ProjectPermissionsCache? Cache)
+		{
+			return AuthorizeAsync(Stream.Acl, Stream.ProjectId, Action, User, Cache);
+		}
+
+		/// <summary>
+		/// Determines if the user is authorized to perform an action on a particular stream
+		/// </summary>
+		/// <param name="Stream">The stream to check</param>
+		/// <param name="Template">Template within the stream to check</param>
+		/// <param name="Action">The action being performed</param>
+		/// <param name="User">The principal to authorize</param>
+		/// <param name="Cache">Cache of project permissions</param>
+		/// <returns>True if the action is authorized</returns>
+		public Task<bool> AuthorizeAsync(IStream Stream, TemplateRef Template, AclAction Action, ClaimsPrincipal User, ProjectPermissionsCache? Cache)
+		{
+			bool? Result = Template.Acl?.Authorize(Action, User);
+			if (Result == null)
+			{
+				return AuthorizeAsync(Stream, Action, User, Cache);
+			}
+			else
+			{
+				return Task.FromResult(Result.Value);
+			}
+		}
+
+		/// <summary>
+		/// Determines if the user is authorized to perform an action on a particular project
+		/// </summary>
+		/// <param name="StreamId">The stream id to check</param>
+		/// <param name="Action">The action being performed</param>
+		/// <param name="User">The principal to authorize</param>
+		/// <param name="Cache">Cache of stream permissions</param>
+		/// <returns>True if the action is authorized</returns>
+		public async Task<bool> AuthorizeAsync(StreamId StreamId, AclAction Action, ClaimsPrincipal User, StreamPermissionsCache? Cache)
+		{
+			IStreamPermissions? Permissions;
+			if (Cache == null)
+			{
+				Permissions = await GetStreamPermissionsAsync(StreamId);
+			}
+			else if (!Cache.Streams.TryGetValue(StreamId, out Permissions))
+			{
+				Permissions = await GetStreamPermissionsAsync(StreamId);
+				Cache.Streams.Add(StreamId, Permissions);
+			}
+			return Permissions != null && await AuthorizeAsync(Permissions.Acl, Permissions.ProjectId, Action, User, Cache);
+		}
+	}
+}
diff --git a/Engine/Source/Programs/Horde/HordeServer/Users/IServiceAccount.cs b/Engine/Source/Programs/Horde/HordeServer/Users/IServiceAccount.cs
new file mode 100644
--- /dev/null
+++ b/Engine/Source/Programs/Horde/HordeServer/Users/IServiceAccount.cs
@@ -0,0 +1,48 @@
+// Copyright Epic Games, Inc. All Rights Reserved.
+
+using MongoDB.Bson;
+using System.Collections.Generic;
+
+namespace HordeServer.Models
+{
+	/// <summary>
+	/// A service account used for authenticating during server-to-server communication
+	/// For example between Robomerge and Horde.
+	/// </summary>
+	public interface IServiceAccount
+	{
+		/// <summary>
+		/// Unique id for this session
+		/// </summary>
+		public ObjectId Id { get; }
+		
+		/// <summary>
+		/// Secret token used for identifying calls made as the service account
+		/// </summary>
+		public string SecretToken { get; }
+		
+		/// <summary>
+		/// If the service account is active
+		/// </summary>
+		public bool Enabled { get; }
+		
+		/// <summary>
+		/// Description of the service account (who is it for, is there an owner etc)
+		/// </summary>
+		public string Description { get; }
+
+		/// <summary>
+		/// Add a claim
+		/// </summary>
+		/// <param name="Type">Type of claim</param>
+		/// <param name="Value">Value of claim</param>
+		/// <returns></returns>
+		public void AddClaim(string Type, string Value);
+		
+		/// <summary>
+		/// Get list of claims
+		/// </summary>
+		/// <returns>List of claims</returns>
+		public IReadOnlyList<(string Type, string Value)> GetClaims();
+	}
+}
diff --git a/Engine/Source/Programs/Horde/HordeServer/Users/Impl/IServiceAccountCollection.cs b/Engine/Source/Programs/Horde/HordeServer/Users/Impl/IServiceAccountCollection.cs
new file mode 100644
--- /dev/null
+++ b/Engine/Source/Programs/Horde/HordeServer/Users/Impl/IServiceAccountCollection.cs
@@ -0,0 +1,55 @@
+// Copyright Epic Games, Inc. All Rights Reserved.
+
+using HordeServer.Models;
+using MongoDB.Bson;
+using System.Collections.Generic;
+using System.Threading.Tasks;
+
+namespace HordeServer.Collections
+{
+	/// <summary>
+	/// Interface for a collection of service account documents
+	/// </summary>
+	public interface IServiceAccountCollection
+	{
+		/// <summary>
+		/// Adds a new service account to the collection
+		/// </summary>
+		/// <param name="SecretToken">Secret token used for identifying this service account</param>
+		/// <param name="Claims">Dictionary of claims to assign</param>
+		/// <param name="Description">Description of the account</param>
+		Task<IServiceAccount> AddAsync(string SecretToken, List<string> Claims, string Description);
+
+		/// <summary>
+		/// Get service account via ID
+		/// </summary>
+		/// <param name="Id">The unique service account id</param>
+		/// <returns>The service account</returns>
+		Task<IServiceAccount?> GetAsync(ObjectId Id);
+		
+		/// <summary>
+		/// Get service account via secret token
+		/// </summary>
+		/// <param name="SecretToken">Secret token to use for searching</param>
+		/// <returns>The service account</returns>
+		Task<IServiceAccount?> GetBySecretTokenAsync(string SecretToken);
+
+		/// <summary>
+		/// Update a service account from the collection
+		/// </summary>
+		/// <param name="Id">The service account</param>
+		/// <param name="SecretToken">If set, secret token will be set</param>
+		/// <param name="Claims">If set, claims will be set</param>
+		/// <param name="Enabled">If set, enabled flag will be set</param>
+		/// <param name="Description">If set, description will be set</param>
+		/// <returns>Async task</returns>
+		Task UpdateAsync(ObjectId Id, string? SecretToken, List<string>? Claims, bool? Enabled, string? Description);
+
+		/// <summary>
+		/// Delete a service account from the collection
+		/// </summary>
+		/// <param name="Id">The service account</param>
+		/// <returns>Async task</returns>
+		Task DeleteAsync(ObjectId Id);
+	}
+}
