Horde: More file moves.

diff --git a/Engine/Source/Programs/Horde/HordeServer/Jobs/TemplateService.cs b/Engine/Source/Programs/Horde/HordeServer/Jobs/TemplateService.cs
new file mode 100644
--- /dev/null
+++ b/Engine/Source/Programs/Horde/HordeServer/Jobs/TemplateService.cs
@@ -0,0 +1,72 @@
+// Copyright Epic Games, Inc. All Rights Reserved.
+
+using EpicGames.Core;
+using HordeServer.Api;
+using HordeServer.Collections;
+using HordeCommon;
+using HordeServer.Models;
+using MongoDB.Bson;
+using MongoDB.Driver;
+using System;
+using System.Collections.Generic;
+using System.Security.Claims;
+using System.Threading.Tasks;
+
+namespace HordeServer.Services
+{
+	/// <summary>
+	/// Wraps functionality for manipulating job templates
+	/// </summary>
+	public class TemplateService
+	{
+		/// <summary>
+		/// Collection of template documents
+		/// </summary>
+		ITemplateCollection Templates;
+
+		/// <summary>
+		/// Constructor
+		/// </summary>
+		/// <param name="Templates">Collection of template documents</param>
+		public TemplateService(ITemplateCollection Templates)
+		{
+			this.Templates = Templates;
+		}
+
+		/// <summary>
+		/// Creates a new template
+		/// </summary>
+		/// <param name="Name">Name of the new template</param>
+		/// <param name="Priority">Priority for the new template</param>
+		/// <param name="AllowPreflights">Whether to allow preflights of this template</param>
+		/// <param name="InitialAgentType">The initial agent type to parse the BuildGraph script</param>
+		/// <param name="SubmitNewChange">Path to a file within the stream to submit to generate a new changelist for jobs</param>
+		/// <param name="Counters">List of counters for the template</param>
+		/// <param name="Arguments">Common arguments for jobs started from this template</param>
+		/// <param name="Parameters">List of parameters for this template</param>
+		/// <returns>The new template document</returns>
+		public Task<ITemplate> CreateTemplateAsync(string Name, Priority? Priority = null, bool AllowPreflights = true, string? InitialAgentType = null, string? SubmitNewChange = null, List<TemplateCounter>? Counters = null, List<string>? Arguments = null, List<Parameter>? Parameters = null)
+		{
+			return Templates.AddAsync(Name, Priority, AllowPreflights, InitialAgentType, SubmitNewChange, Counters, Arguments, Parameters);
+		}
+
+		/// <summary>
+		/// Gets all the available templates
+		/// </summary>
+		/// <returns>List of template documents</returns>
+		public Task<List<ITemplate>> GetTemplatesAsync()
+		{
+			return Templates.FindAllAsync();
+		}
+
+		/// <summary>
+		/// Gets a template by ID
+		/// </summary>
+		/// <param name="TemplateId">Unique id of the template</param>
+		/// <returns>The template document</returns>
+		public Task<ITemplate?> GetTemplateAsync(ContentHash TemplateId)
+		{
+			return Templates.GetAsync(TemplateId);
+		}
+	}
+}
diff --git a/Engine/Source/Programs/Horde/HordeServer/Projects/IStream.cs b/Engine/Source/Programs/Horde/HordeServer/Projects/IStream.cs
new file mode 100644
--- /dev/null
+++ b/Engine/Source/Programs/Horde/HordeServer/Projects/IStream.cs
@@ -0,0 +1,704 @@
+// Copyright Epic Games, Inc. All Rights Reserved.
+
+using EpicGames.Core;
+using Google.Protobuf.WellKnownTypes;
+using HordeServer.Api;
+using HordeCommon;
+using HordeServer.Utilities;
+using MongoDB.Bson;
+using MongoDB.Bson.Serialization.Attributes;
+using MongoDB.Driver;
+using System;
+using System.Collections.Generic;
+using System.Diagnostics.CodeAnalysis;
+using System.IO;
+using System.Linq;
+using System.Security.Cryptography;
+using System.Text;
+using System.Text.RegularExpressions;
+using System.Threading.Tasks;
+using System.Runtime.CompilerServices;
+
+namespace HordeServer.Models
+{
+	using PoolId = StringId<IPool>;
+	using ProjectId = StringId<IProject>;
+	using StreamId = StringId<IStream>;
+	using TemplateRefId = StringId<TemplateRef>;
+
+	/// <summary>
+	/// Exception thrown when stream validation fails
+	/// </summary>
+	public class InvalidStreamException : Exception
+	{
+		/// <inheritdoc/>
+		public InvalidStreamException()
+		{
+		}
+
+		/// <inheritdoc/>
+		public InvalidStreamException(string Message) : base(Message)
+		{
+		}
+
+		/// <inheritdoc/>
+		public InvalidStreamException(string Message, Exception InnerEx) : base(Message, InnerEx)
+		{
+		}
+	}
+
+	/// <summary>
+	/// Mapping from a BuildGraph agent type to a set of machines on the farm
+	/// </summary>
+	public class AgentType
+	{
+		/// <summary>
+		/// Name of the pool of agents to use
+		/// </summary>
+		[BsonRequired]
+		public PoolId Pool { get; set; }
+
+		/// <summary>
+		/// Name of the workspace to execute on
+		/// </summary>
+		public string? Workspace { get; set; }
+
+		/// <summary>
+		/// Path to the temporary storage dir
+		/// </summary>
+		public string? TempStorageDir { get; set; }
+
+		/// <summary>
+		/// Environment variables to be set when executing the job
+		/// </summary>
+		public Dictionary<string, string>? Environment { get; set; }
+
+		/// <summary>
+		/// Default constructor
+		/// </summary>
+		[BsonConstructor]
+		private AgentType()
+		{
+		}
+
+		/// <summary>
+		/// Constructor
+		/// </summary>
+		/// <param name="Pool">The pool for this agent</param>
+		/// <param name="Workspace">Name of the workspace to use</param>
+		/// <param name="TempStorageDir">Path to the temp storage directory</param>
+		public AgentType(PoolId Pool, string Workspace, string? TempStorageDir)
+		{
+			this.Pool = Pool;
+			this.Workspace = Workspace;
+			this.TempStorageDir = TempStorageDir;
+		}
+
+		/// <summary>
+		/// Constructor
+		/// </summary>
+		/// <param name="Request">The object to construct from</param>
+		public AgentType(CreateAgentTypeRequest Request)
+		{
+			Pool = new PoolId(Request.Pool);
+			Workspace = Request.Workspace;
+			TempStorageDir = Request.TempStorageDir;
+			Environment = Request.Environment;
+		}
+
+		/// <summary>
+		/// Constructs an AgentType object from an optional request
+		/// </summary>
+		/// <param name="Request">The request object</param>
+		/// <returns>New agent type object</returns>
+		[return: NotNullIfNotNull("Request")]
+		public static AgentType? FromRequest(CreateAgentTypeRequest? Request)
+		{
+			return (Request != null) ? new AgentType(Request) : null;
+		}
+
+		/// <summary>
+		/// Creates an API response object from this stream
+		/// </summary>
+		/// <returns>The response object</returns>
+		public Api.GetAgentTypeResponse ToApiResponse()
+		{
+			return new Api.GetAgentTypeResponse(Pool.ToString(), Workspace, TempStorageDir, Environment);
+		}
+
+		/// <summary>
+		/// Creates an API response object from this stream
+		/// </summary>
+		/// <returns>The response object</returns>
+		public HordeCommon.Rpc.GetAgentTypeResponse ToRpcResponse()
+		{
+			HordeCommon.Rpc.GetAgentTypeResponse Response = new HordeCommon.Rpc.GetAgentTypeResponse();
+			if (TempStorageDir != null)
+			{
+				Response.TempStorageDir = TempStorageDir;
+			}
+			if (Environment != null)
+			{
+				Response.Environment.Add(Environment);
+			}
+			return Response;
+		}
+	}
+
+	/// <summary>
+	/// Information about a workspace type
+	/// </summary>
+	public class WorkspaceType
+	{
+		/// <summary>
+		/// Name of the Perforce cluster to use
+		/// </summary>
+		[BsonIgnoreIfNull]
+		public string? Cluster { get; set; }
+
+		/// <summary>
+		/// The Perforce server and port
+		/// </summary>
+		[BsonIgnoreIfNull]
+		public string? ServerAndPort { get; set; }
+
+		/// <summary>
+		/// The Perforce username for syncing this workspace
+		/// </summary>
+		[BsonIgnoreIfNull]
+		public string? UserName { get; set; }
+
+		/// <summary>
+		/// The Perforce password for syncing this workspace
+		/// </summary>
+		[BsonIgnoreIfNull]
+		public string? Password { get; set; }
+
+		/// <summary>
+		/// Identifier to distinguish this workspace from other workspaces. Defaults to the workspace type name.
+		/// </summary>
+		[BsonIgnoreIfNull]
+		public string? Identifier { get; set; }
+
+		/// <summary>
+		/// Override for the stream to sync
+		/// </summary>
+		[BsonIgnoreIfNull]
+		public string? Stream { get; set; }
+
+		/// <summary>
+		/// Custom view for the workspace
+		/// </summary>
+		[BsonIgnoreIfNull]
+		public List<string>? View { get; set; }
+
+		/// <summary>
+		/// Whether to use an incrementally synced workspace
+		/// </summary>
+		public bool Incremental { get; set; }
+
+		/// <summary>
+		/// Default constructor
+		/// </summary>
+		public WorkspaceType()
+		{
+		}
+
+		/// <summary>
+		/// Constructor
+		/// </summary>
+		/// <param name="Request">The object to construct from</param>
+		public WorkspaceType(CreateWorkspaceTypeRequest Request)
+		{
+			Cluster = Request.Cluster;
+			ServerAndPort = Request.ServerAndPort;
+			UserName = Request.UserName;
+			Password = Request.Password;
+			Identifier = Request.Identifier;
+			Stream = Request.Stream;
+			View = Request.View;
+			Incremental = Request.Incremental;
+		}
+
+		/// <summary>
+		/// Constructs an AgentType object from an optional request
+		/// </summary>
+		/// <param name="Request">The request object</param>
+		/// <returns>New agent type object</returns>
+		[return: NotNullIfNotNull("Request")]
+		public static WorkspaceType? FromRequest(CreateWorkspaceTypeRequest? Request)
+		{
+			return (Request != null) ? new WorkspaceType(Request) : null;
+		}
+
+		/// <summary>
+		/// Creates an API response object from this stream
+		/// </summary>
+		/// <returns>The response object</returns>
+		public Api.GetWorkspaceTypeResponse ToApiResponse()
+		{
+			return new Api.GetWorkspaceTypeResponse(Cluster, ServerAndPort, UserName, Identifier, Stream, View, Incremental);
+		}
+	}
+
+	/// <summary>
+	/// Allows triggering another downstream job on succesful completion of a step or aggregate
+	/// </summary>
+	public class ChainedJobTemplate
+	{
+		/// <summary>
+		/// Name of the target that needs to complete successfully
+		/// </summary>
+		public string Trigger { get; set; } = String.Empty;
+
+		/// <summary>
+		/// The new template to trigger
+		/// </summary>
+		public TemplateRefId TemplateRefId { get; set; }
+
+		/// <summary>
+		/// Default constructor for serialization
+		/// </summary>
+		private ChainedJobTemplate()
+		{
+		}
+
+		/// <summary>
+		/// Constructor
+		/// </summary>
+		/// <param name="Trigger">Name of the target that needs to complete</param>
+		/// <param name="TemplateRefId">The new template to trigger</param>
+		public ChainedJobTemplate(string Trigger, TemplateRefId TemplateRefId)
+		{
+			this.Trigger = Trigger;
+			this.TemplateRefId = TemplateRefId;
+		}
+
+		/// <summary>
+		/// Constructor
+		/// </summary>
+		/// <param name="Request">Request to construct from</param>
+		public ChainedJobTemplate(CreateChainedJobTemplateRequest Request)
+			: this(Request.Trigger, new TemplateRefId(Request.TemplateId))
+		{
+		}
+	}
+
+	/// <summary>
+	/// Reference to a template
+	/// </summary>
+	public class TemplateRef
+	{
+		/// <summary>
+		/// The template name (duplicated from the template object)
+		/// </summary>
+		public string Name { get; set; }
+
+		/// <summary>
+		/// Hash of the template definition
+		/// </summary>
+		public ContentHash Hash { get; set; }
+
+		/// <summary>
+		/// Whether to show badges in UGS for this schedule
+		/// </summary>
+		public bool ShowUgsBadges { get; set; }
+
+		/// <summary>
+		/// Whether to show desktop alerts for build health issues created from jobs this type
+		/// </summary>
+		public bool ShowUgsAlerts { get; set; }
+
+		/// <summary>
+		/// Notification channel for this template. Overrides the stream channel if set.
+		/// </summary>
+		public string? NotificationChannel { get; set; }
+
+		/// <summary>
+		/// Notification channel filter for this template. Errors|Warnings|Success
+		/// </summary>
+		public string? NotificationChannelFilter { get; set; }
+
+		/// <summary>
+		/// Channel for triage notification messages
+		/// </summary>
+		public string? TriageChannel { get; set; }
+
+		/// <summary>
+		/// List of schedules for this template
+		/// </summary>
+		[BsonIgnoreIfNull]
+		public Schedule? Schedule { get; set; }
+
+		/// <summary>
+		/// List of downstream templates to trigger at the same change
+		/// </summary>
+		[BsonIgnoreIfNull]
+		public List<ChainedJobTemplate>? ChainedJobs { get; set; }
+
+		/// <summary>
+		/// Custom permissions for this template
+		/// </summary>
+		public Acl? Acl { get; set; }
+
+		/// <summary>
+		/// Private constructor for serialization
+		/// </summary>
+		private TemplateRef()
+		{
+			this.Name = null!;
+			this.Hash = null!;
+		}
+
+		/// <summary>
+		/// Constructor
+		/// </summary>
+		/// <param name="Template">The template being referenced</param>
+		/// <param name="ShowUgsBadges">Whether to show badges in UGS for this job</param>
+		/// <param name="ShowUgsAlerts">Whether to show alerts in UGS for this job</param>
+		/// <param name="NotificationChannel">Notification channel for this template</param>
+		/// <param name="NotificationChannelFilter">Notification channel filter for this template</param>
+		/// <param name="TriageChannel"></param>
+		/// <param name="Schedule">Schedule for this template</param>
+		/// <param name="Triggers">List of downstream templates to trigger</param>
+		/// <param name="Acl">ACL for this template</param>
+		public TemplateRef(ITemplate Template, bool ShowUgsBadges = false, bool ShowUgsAlerts = false, string? NotificationChannel = null, string? NotificationChannelFilter = null, string? TriageChannel = null, Schedule? Schedule = null, List<ChainedJobTemplate>? Triggers = null, Acl? Acl = null)
+		{
+			this.Name = Template.Name;
+			this.Hash = Template.Id;
+			this.ShowUgsBadges = ShowUgsBadges;
+			this.ShowUgsAlerts = ShowUgsAlerts;
+			this.NotificationChannel = NotificationChannel;
+			this.NotificationChannelFilter = NotificationChannelFilter;
+			this.TriageChannel = TriageChannel;
+			this.Schedule = Schedule;
+			this.ChainedJobs = Triggers;
+			this.Acl = Acl;
+		}
+	}
+
+	/// <summary>
+	/// Query used to identify a base changelist for a preflight
+	/// </summary>
+	public class ChangeQuery
+	{
+		/// <summary>
+		/// Template to search for
+		/// </summary>
+		public TemplateRefId? TemplateRefId { get; set; }
+
+		/// <summary>
+		/// The target to look at the status for
+		/// </summary>
+		public string? Target { get; set; }
+
+		/// <summary>
+		/// Whether to match a job that contains warnings
+		/// </summary>
+		public List<JobStepOutcome>? Outcomes { get; set; }
+
+		/// <summary>
+		/// Convert to a request object
+		/// </summary>
+		/// <returns></returns>
+		public ChangeQueryRequest ToRequest()
+		{
+			return new ChangeQueryRequest { TemplateId = TemplateRefId?.ToString(), Target = Target, Outcomes = Outcomes };
+		}
+	}
+
+	/// <summary>
+	/// Definition of a query to execute to find the changelist to run a build at
+	/// </summary>
+	public class DefaultPreflight
+	{
+		/// <summary>
+		/// The template id to execute
+		/// </summary>
+		public TemplateRefId? TemplateRefId { get; set; }
+
+		/// <summary>
+		/// Query specifying a changelist to use
+		/// </summary>
+		public ChangeQuery? Change { get; set; }
+
+		/// <summary>
+		/// The job type to query for the change to use
+		/// </summary>
+		[Obsolete("Use Change.TemplateRefId instead")]
+		public TemplateRefId? ChangeTemplateRefId { get; set; }
+
+		/// <summary>
+		/// Constructor
+		/// </summary>
+		/// <param name="TemplateRefId"></param>
+		/// <param name="Change">The job type to query for the change to use</param>
+		public DefaultPreflight(TemplateRefId? TemplateRefId, ChangeQuery? Change)
+		{
+			this.TemplateRefId = TemplateRefId;
+			this.Change = Change;
+		}
+
+		/// <summary>
+		/// Convert to a request object
+		/// </summary>
+		/// <returns></returns>
+		public DefaultPreflightRequest ToRequest()
+		{
+#pragma warning disable CS0618 // Type or member is obsolete
+			ChangeQueryRequest? ChangeRequest = null;
+			if (Change != null)
+			{
+				ChangeRequest = Change.ToRequest();
+			}
+			else if (ChangeTemplateRefId != null)
+			{
+				ChangeRequest = new ChangeQueryRequest { TemplateId = ChangeTemplateRefId?.ToString() };
+			}
+
+			return new DefaultPreflightRequest { TemplateId = TemplateRefId?.ToString(), Change = ChangeRequest, ChangeTemplateId = ChangeRequest?.TemplateId };
+#pragma warning restore CS0618 // Type or member is obsolete
+		}
+	}
+
+	/// <summary>
+	/// Extension methods for template refs
+	/// </summary>
+	static class TemplateRefExtensions
+	{
+		/// <summary>
+		/// Adds a new template ref to a list
+		/// </summary>
+		/// <param name="TemplateRefs">List of template refs</param>
+		/// <param name="TemplateRef">The template ref to add</param>
+		public static void AddRef(this Dictionary<TemplateRefId, TemplateRef> TemplateRefs, TemplateRef TemplateRef)
+		{
+			TemplateRefs.Add(new TemplateRefId(TemplateRef.Name), TemplateRef);
+		}
+	}
+
+	/// <summary>
+	/// Information about a stream
+	/// </summary>
+	public interface IStream
+	{
+		/// <summary>
+		/// Name of the stream.
+		/// </summary>
+		public StreamId Id { get; }
+
+		/// <summary>
+		/// The project that this stream belongs to
+		/// </summary>
+		public ProjectId ProjectId { get; }
+
+		/// <summary>
+		/// The stream name
+		/// </summary>
+		public string Name { get; }
+
+		/// <summary>
+		/// Name of the perforce cluster
+		/// </summary>
+		public string ClusterName { get; }
+
+		/// <summary>
+		/// Path to the configuration file for this stream
+		/// </summary>
+		public string ConfigPath { get; }
+
+		/// <summary>
+		/// The revision of config file used for this stream
+		/// </summary>
+		public string ConfigRevision { get; }
+
+		/// <summary>
+		/// Order to display on the dashboard's drop-down list
+		/// </summary>
+		public int Order { get; }
+
+		/// <summary>
+		/// Notification channel for all jobs in this stream
+		/// </summary>
+		public string? NotificationChannel { get; }
+
+		/// <summary>
+		/// Notification channel filter for all jobs in this stream. Errors|Warnings|Success
+		/// </summary>
+		public string? NotificationChannelFilter { get; }
+
+		/// <summary>
+		/// Channel to post issue triage notifications
+		/// </summary>
+		public string? TriageChannel { get; }
+
+		/// <summary>
+		/// Default template to use for preflights
+		/// </summary>
+		public DefaultPreflight? DefaultPreflight { get; }
+
+		/// <summary>
+		/// List of pages to display in the dashboard
+		/// </summary>
+		public IReadOnlyList<StreamTab> Tabs { get; }
+
+		/// <summary>
+		/// Dictionary of agent types
+		/// </summary>
+		public IReadOnlyDictionary<string, AgentType> AgentTypes { get; }
+
+		/// <summary>
+		/// Dictionary of workspace types
+		/// </summary>
+		public IReadOnlyDictionary<string, WorkspaceType> WorkspaceTypes { get; }
+
+		/// <summary>
+		/// List of templates available for this stream
+		/// </summary>
+		public IReadOnlyDictionary<TemplateRefId, TemplateRef> Templates { get; }
+
+		/// <summary>
+		/// Stream is paused for builds until specified time
+		/// </summary>
+		public DateTime? PausedUntil { get; }
+		
+		/// <summary>
+		/// Comment/reason for why the stream was paused
+		/// </summary>
+		public string? PauseComment { get; }
+
+		/// <summary>
+		/// The ACL for this object
+		/// </summary>
+		public Acl? Acl { get; }
+	}
+
+	/// <summary>
+	/// Extension methods for streams
+	/// </summary>
+	static class StreamExtensions
+	{
+		/// <summary>
+		/// Tries to get an agent workspace definition from the given type name
+		/// </summary>
+		/// <param name="Stream">The stream object</param>
+		/// <param name="AgentType">The agent type</param>
+		/// <param name="Workspace">Receives the agent workspace definition</param>
+		/// <returns>True if the agent type was valid, and an agent workspace could be created</returns>
+		public static bool TryGetAgentWorkspace(this IStream Stream, AgentType AgentType, [NotNullWhen(true)] out AgentWorkspace? Workspace)
+		{
+			// Get the workspace settings
+			if (AgentType.Workspace == null)
+			{
+				// Use the default settings (fast switching workspace, clean 
+				Workspace = new AgentWorkspace(null, null, GetDefaultWorkspaceIdentifier(Stream), Stream.Name, null, false);
+				return true;
+			}
+			else
+			{
+				// Try to get the matching workspace type
+				WorkspaceType? WorkspaceType;
+				if (!Stream.WorkspaceTypes.TryGetValue(AgentType.Workspace, out WorkspaceType))
+				{
+					Workspace = null;
+					return false;
+				}
+
+				// Get the workspace identifier
+				string Identifier;
+				if (WorkspaceType.Identifier != null && !String.IsNullOrEmpty(WorkspaceType.Identifier))
+				{
+					Identifier = WorkspaceType.Identifier;
+				}
+				else if (WorkspaceType.Incremental)
+				{
+					Identifier = $"{Stream.GetEscapedName()}+{AgentType.Workspace}";
+				}
+				else
+				{
+					Identifier = GetDefaultWorkspaceIdentifier(Stream);
+				}
+
+				// Create the new workspace
+				Workspace = new AgentWorkspace(WorkspaceType.Cluster, WorkspaceType.UserName, Identifier, WorkspaceType.Stream ?? Stream.Name, WorkspaceType.View, WorkspaceType.Incremental);
+				return true;
+			}
+		}
+
+		/// <summary>
+		/// The escaped name of this stream. Removes all non-identifier characters.
+		/// </summary>
+		/// <param name="Stream">The stream object</param>
+		/// <returns>Escaped name for the stream</returns>
+		public static string GetEscapedName(this IStream Stream)
+		{
+			return Regex.Replace(Stream.Name, @"[^a-zA-Z0-9_]", "+");
+		}
+
+		/// <summary>
+		/// Gets the default identifier for workspaces created for this stream. Just includes an escaped depot name.
+		/// </summary>
+		/// <param name="Stream">The stream object</param>
+		/// <returns>The default workspace identifier</returns>
+		private static string GetDefaultWorkspaceIdentifier(IStream Stream)
+		{
+			return Regex.Replace(Stream.GetEscapedName(), @"^(\+\+[^+]*).*$", "$1");
+		}
+
+		/// <summary>
+		/// Converts to a public response object
+		/// </summary>
+		/// <param name="Stream">The stream object</param>
+		/// <param name="bIncludeAcl">Whether to include the ACL in the response object</param>
+		/// <param name="ApiTemplateRefs">The template refs for this stream. Passed separately because they have their own ACL.</param>
+		/// <returns>New response instance</returns>
+		public static Api.GetStreamResponse ToApiResponse(this IStream Stream, bool bIncludeAcl, List<GetTemplateRefResponse> ApiTemplateRefs)
+		{
+			List<GetStreamTabResponse> ApiTabs = Stream.Tabs.ConvertAll(x => x.ToResponse());
+			Dictionary<string, GetAgentTypeResponse> ApiAgentTypes = Stream.AgentTypes.ToDictionary(x => x.Key, x => x.Value.ToApiResponse());
+			Dictionary<string, GetWorkspaceTypeResponse> ApiWorkspaceTypes = Stream.WorkspaceTypes.ToDictionary(x => x.Key, x => x.Value.ToApiResponse());
+			GetAclResponse? ApiAcl = (bIncludeAcl && Stream.Acl != null)? new GetAclResponse(Stream.Acl) : null;
+			return new Api.GetStreamResponse(Stream.Id.ToString(), Stream.ProjectId.ToString(), Stream.Name, Stream.ConfigPath, Stream.ConfigRevision, Stream.Order, Stream.NotificationChannel, Stream.NotificationChannelFilter, Stream.TriageChannel, Stream.DefaultPreflight?.ToRequest(), ApiTabs, ApiAgentTypes, ApiWorkspaceTypes, ApiTemplateRefs, ApiAcl, Stream.PausedUntil, Stream.PauseComment);
+		}
+
+		/// <summary>
+		/// Converts to an RPC response object
+		/// </summary>
+		/// <param name="Stream">The stream object</param>
+		/// <returns>New response instance</returns>
+		public static HordeCommon.Rpc.GetStreamResponse ToRpcResponse(this IStream Stream)
+		{
+			HordeCommon.Rpc.GetStreamResponse Response = new HordeCommon.Rpc.GetStreamResponse();
+			Response.Name = Stream.Name;
+			Response.AgentTypes.Add(Stream.AgentTypes.ToDictionary(x => x.Key, x => x.Value.ToRpcResponse()));
+			return Response;
+		}
+
+		/// <summary>
+		/// Check if stream is paused for new builds
+		/// </summary>
+		/// <param name="Stream">The stream object</param>
+		/// <param name="CurrentTime">Current time (allow tests to pass in a fake clock)</param>
+		/// <returns>If stream is paused</returns>
+		public static bool IsPaused(this IStream Stream, DateTime CurrentTime)
+		{
+			return Stream.PausedUntil != null && Stream.PausedUntil > CurrentTime;
+		}
+	}
+
+	/// <summary>
+	/// Projection of a stream definition to just include permissions info
+	/// </summary>
+	public interface IStreamPermissions
+	{
+		/// <summary>
+		/// ACL for the stream
+		/// </summary>
+		public Acl? Acl { get; }
+
+		/// <summary>
+		/// The project containing this stream
+		/// </summary>
+		public ProjectId ProjectId { get; }
+	}
+}
diff --git a/Engine/Source/Programs/Horde/HordeServer/Projects/IStreamCollection.cs b/Engine/Source/Programs/Horde/HordeServer/Projects/IStreamCollection.cs
new file mode 100644
--- /dev/null
+++ b/Engine/Source/Programs/Horde/HordeServer/Projects/IStreamCollection.cs
@@ -0,0 +1,120 @@
+// Copyright Epic Games, Inc. All Rights Reserved.
+
+using HordeServer.Api;
+using HordeServer.Models;
+using HordeServer.Services;
+using HordeServer.Utilities;
+using MongoDB.Bson;
+using MongoDB.Bson.Serialization.Attributes;
+using MongoDB.Driver;
+using System;
+using System.Collections.Generic;
+using System.Diagnostics.CodeAnalysis;
+using System.Linq;
+using System.Threading.Tasks;
+
+namespace HordeServer.Collections
+{
+	using ProjectId = StringId<IProject>;
+	using StreamId = StringId<IStream>;
+	using TemplateRefId = StringId<TemplateRef>;
+
+	/// <summary>
+	/// Collection of stream documents
+	/// </summary>
+	public interface IStreamCollection
+	{
+		/// <summary>
+		/// Creates or replaces a stream configuration
+		/// </summary>
+		/// <param name="Id">Unique id for the new stream</param>
+		/// <param name="Stream">The current stream value. If not-null, this will attempt to replace the existing instance.</param>
+		/// <param name="ConfigPath">Path to the config file</param>
+		/// <param name="Revision">The config file revision</param>
+		/// <param name="ProjectId">The project id</param>
+		/// <param name="Config">The stream configuration</param>
+		/// <returns></returns>
+		Task<IStream?> TryCreateOrReplaceAsync(StreamId Id, IStream? Stream, string ConfigPath, string Revision, ProjectId ProjectId, StreamConfig Config);
+
+		/// <summary>
+		/// Gets a stream by ID
+		/// </summary>
+		/// <param name="StreamId">Unique id of the stream</param>
+		/// <returns>The stream document</returns>
+		Task<IStream?> GetAsync(StreamId StreamId);
+
+		/// <summary>
+		/// Gets a stream's permissions by ID
+		/// </summary>
+		/// <param name="StreamId">Unique id of the stream</param>
+		/// <returns>The stream document</returns>
+		Task<IStreamPermissions?> GetPermissionsAsync(StreamId StreamId);
+
+		/// <summary>
+		/// Enumerates all streams
+		/// </summary>
+		/// <returns></returns>
+		Task<List<IStream>> FindAllAsync();
+
+		/// <summary>
+		/// Gets all the available streams for a project
+		/// </summary>
+		/// <param name="ProjectIds">Unique id of the projects to query</param>
+		/// <returns>List of stream documents</returns>
+		Task<List<IStream>> FindForProjectsAsync(ProjectId[] ProjectIds);
+
+		/// <summary>
+		/// Updates user-facing properties for an existing stream
+		/// </summary>
+		/// <param name="Stream">The stream to update</param>
+		/// <param name="NewPausedUntil">The new datetime for pausing builds</param>
+		/// <param name="NewPauseComment">The reason for pausing</param>
+		/// <returns>The updated stream if successful, null otherwise</returns>
+		Task<IStream?> TryUpdatePauseStateAsync(IStream Stream, DateTime? NewPausedUntil, string? NewPauseComment);
+
+		/// <summary>
+		/// Attempts to update the last trigger time for a schedule
+		/// </summary>
+		/// <param name="Stream">The stream to update</param>
+		/// <param name="TemplateRefId">The template ref id</param>
+		/// <param name="LastTriggerTime">New last trigger time for the schedule</param>
+		/// <param name="LastTriggerChange">New last trigger changelist for the schedule</param>
+		/// <param name="NewActiveJobs">New list of active jobs</param>
+		/// <returns>The updated stream if successful, null otherwise</returns>
+		Task<IStream?> TryUpdateScheduleTriggerAsync(IStream Stream, TemplateRefId TemplateRefId, DateTimeOffset? LastTriggerTime, int? LastTriggerChange, List<ObjectId> NewActiveJobs);
+
+		/// <summary>
+		/// Delete a stream
+		/// </summary>
+		/// <param name="StreamId">Unique id of the stream</param>
+		/// <returns>Async task</returns>
+		Task DeleteAsync(StreamId StreamId);
+	}
+
+	static class StreamCollectionExtensions
+	{
+		/// <summary>
+		/// Creates or replaces a stream configuration
+		/// </summary>
+		/// <param name="StreamCollection">The stream collection</param>
+		/// <param name="Id">Unique id for the new stream</param>
+		/// <param name="Stream">The current stream value. If not-null, this will attempt to replace the existing instance.</param>
+		/// <param name="ConfigPath">Path to the config file</param>
+		/// <param name="Revision">The config file revision</param>
+		/// <param name="ProjectId">The project id</param>
+		/// <param name="Config">The stream configuration</param>
+		/// <returns></returns>
+		public static async Task<IStream> CreateOrReplaceAsync(this IStreamCollection StreamCollection, StreamId Id, IStream? Stream, string ConfigPath, string Revision, ProjectId ProjectId, StreamConfig Config)
+		{
+			for (; ; )
+			{
+				Stream = await StreamCollection.TryCreateOrReplaceAsync(Id, Stream, ConfigPath, Revision, ProjectId, Config);
+				if (Stream != null)
+				{
+					return Stream;
+				}
+				Stream = await StreamCollection.GetAsync(Id);
+			}
+		}
+	}
+}
diff --git a/Engine/Source/Programs/Horde/HordeServer/Projects/Impl/StreamCollection.cs b/Engine/Source/Programs/Horde/HordeServer/Projects/Impl/StreamCollection.cs
new file mode 100644
--- /dev/null
+++ b/Engine/Source/Programs/Horde/HordeServer/Projects/Impl/StreamCollection.cs
@@ -0,0 +1,422 @@
+// Copyright Epic Games, Inc. All Rights Reserved.
+
+using HordeCommon;
+using HordeServer.Api;
+using HordeServer.Models;
+using HordeServer.Services;
+using HordeServer.Utilities;
+using MongoDB.Bson;
+using MongoDB.Bson.Serialization.Attributes;
+using MongoDB.Driver;
+using System;
+using System.Collections.Generic;
+using System.Diagnostics.CodeAnalysis;
+using System.IO;
+using System.Linq;
+using System.Threading.Tasks;
+
+namespace HordeServer.Collections.Impl
+{
+	using ProjectId = StringId<IProject>;
+	using StreamId = StringId<IStream>;
+	using TemplateRefId = StringId<TemplateRef>;
+
+	/// <summary>
+	/// Collection of stream documents
+	/// </summary>
+	public class StreamCollection : IStreamCollection
+	{
+		/// <summary>
+		/// Information about a stream
+		/// </summary>
+		class StreamDocument : IStream
+		{
+			public const int DefaultOrder = 128;
+
+			[BsonRequired, BsonId]
+			public StreamId Id { get; set; }
+
+			[BsonRequired]
+			public ProjectId ProjectId { get; set; }
+
+			[BsonRequired]
+			public string Name { get; set; }
+
+			public string? ClusterName { get; set; }
+			public string ConfigPath { get; set; } = String.Empty;
+			public string ConfigRevision { get; set; } = String.Empty;
+
+			public int Order { get; set; } = DefaultOrder;
+			public string? NotificationChannel { get; set; }
+			public string? NotificationChannelFilter { get; set; }
+			public string? TriageChannel { get; set; }
+			public DefaultPreflight? DefaultPreflight { get; set; }
+			public List<StreamTab> Tabs { get; set; } = new List<StreamTab>();
+			public Dictionary<string, AgentType> AgentTypes { get; set; } = new Dictionary<string, AgentType>(StringComparer.Ordinal);
+			public Dictionary<string, WorkspaceType> WorkspaceTypes { get; set; } = new Dictionary<string, WorkspaceType>(StringComparer.Ordinal);
+			public Dictionary<TemplateRefId, TemplateRef> Templates { get; set; } = new Dictionary<TemplateRefId, TemplateRef>();
+			public DateTime? PausedUntil { get; set; }
+			public string? PauseComment { get; set; }
+			public Acl? Acl { get; set; }
+			public int UpdateIndex { get; set; }
+			public bool Deleted { get; set; }
+
+			string IStream.ClusterName => ClusterName ?? PerforceCluster.DefaultName;
+			IReadOnlyList<StreamTab> IStream.Tabs => Tabs;
+			IReadOnlyDictionary<string, AgentType> IStream.AgentTypes => AgentTypes;
+			IReadOnlyDictionary<string, WorkspaceType> IStream.WorkspaceTypes => WorkspaceTypes;
+			IReadOnlyDictionary<TemplateRefId, TemplateRef> IStream.Templates => Templates;
+
+			[BsonConstructor]
+			private StreamDocument()
+			{
+				Name = null!;
+			}
+
+			public StreamDocument(StreamId Id, string Name, ProjectId ProjectId)
+			{
+				this.Id = Id;
+				this.Name = Name;
+				this.ProjectId = ProjectId;
+			}
+		}
+
+		/// <summary>
+		/// Projection of a stream definition to just include permissions info
+		/// </summary>
+		[SuppressMessage("Design", "CA1812: Class is never instantiated")]
+		private class StreamPermissions : IStreamPermissions
+		{
+			public Acl? Acl { get; set; }
+			public ProjectId ProjectId { get; set; }
+
+			public static readonly ProjectionDefinition<StreamDocument> Projection = Builders<StreamDocument>.Projection.Include(x => x.Acl).Include(x => x.ProjectId);
+		}
+
+		/// <summary>
+		/// The stream collection
+		/// </summary>
+		IMongoCollection<StreamDocument> Streams;
+
+		/// <summary>
+		/// Clock
+		/// </summary>
+		IClock Clock;
+
+		/// <summary>
+		/// The template collection
+		/// </summary>
+		ITemplateCollection TemplateCollection;
+
+		/// <summary>
+		/// Constructor
+		/// </summary>
+		/// <param name="DatabaseService">The database service instance</param>
+		/// <param name="Clock"></param>
+		/// <param name="TemplateCollection"></param>
+		public StreamCollection(DatabaseService DatabaseService, IClock Clock, ITemplateCollection TemplateCollection)
+		{
+			this.Streams = DatabaseService.GetCollection<StreamDocument>("Streams");
+			this.Clock = Clock;
+			this.TemplateCollection = TemplateCollection;
+		}
+
+		/// <inheritdoc/>
+		public async Task<IStream?> TryCreateOrReplaceAsync(StreamId Id, IStream? Stream, string ConfigPath, string Revision, ProjectId ProjectId, StreamConfig Config)
+		{
+			List<StreamTab> Tabs = Config.Tabs.ConvertAll(x => StreamTab.FromRequest(x));
+			Dictionary<TemplateRefId, TemplateRef> TemplateRefs = await CreateTemplateRefsAsync(Config.Templates, Stream, TemplateCollection);
+
+			Dictionary<string, AgentType> AgentTypes = new Dictionary<string, AgentType>();
+			if (Config.AgentTypes != null)
+			{
+				AgentTypes = Config.AgentTypes.Where(x => x.Value != null).ToDictionary(x => x.Key, x => new AgentType(x.Value!));
+			}
+
+			Dictionary<string, WorkspaceType> WorkspaceTypes = new Dictionary<string, WorkspaceType>();
+			if (Config.WorkspaceTypes != null)
+			{
+				WorkspaceTypes = Config.WorkspaceTypes.Where(x => x.Value != null).ToDictionary(x => x.Key, x => new WorkspaceType(x.Value!));
+			}
+
+			DefaultPreflight? DefaultPreflight = Config.DefaultPreflight?.ToModel();
+			if (DefaultPreflight == null && Config.DefaultPreflightTemplate != null)
+			{
+				DefaultPreflight = new DefaultPreflight(new TemplateRefId(Config.DefaultPreflightTemplate), null);
+			}
+
+			Validate(Id, DefaultPreflight, TemplateRefs, Tabs, AgentTypes, WorkspaceTypes);
+
+			Acl? Acl = Acl.Merge(new Acl(), Config.Acl);
+			if (Stream == null)
+			{
+				return await TryCreateAsync(Id, ProjectId, ConfigPath, Revision, Config, DefaultPreflight, Tabs, AgentTypes, WorkspaceTypes, TemplateRefs, Acl);
+			}
+			else
+			{
+				return await TryReplaceAsync(Stream, ProjectId, ConfigPath, Revision, Config, DefaultPreflight, Tabs, AgentTypes, WorkspaceTypes, TemplateRefs, Acl);
+			}
+		}
+
+		/// <summary>
+		/// Creates a list of template refs from a set of request objects
+		/// </summary>
+		/// <param name="Requests">Request objects</param>
+		/// <param name="Stream">The current stream state</param>
+		/// <param name="TemplateCollection">The template service</param>
+		/// <returns>List of new template references</returns>
+		async Task<Dictionary<TemplateRefId, TemplateRef>> CreateTemplateRefsAsync(List<CreateTemplateRefRequest> Requests, IStream? Stream, ITemplateCollection TemplateCollection)
+		{
+			Dictionary<TemplateRefId, TemplateRef> NewTemplateRefs = new Dictionary<TemplateRefId, TemplateRef>();
+			foreach (CreateTemplateRefRequest Request in Requests)
+			{
+				// Create the template
+				ITemplate NewTemplate = await TemplateCollection.AddAsync(Request.Name, Request.Priority, Request.AllowPreflights, Request.InitialAgentType, Request.SubmitNewChange, Request.Counters, Request.Arguments, Request.Parameters.ConvertAll(x => x.ToModel()));
+
+				// Get an identifier for the new template ref
+				TemplateRefId NewTemplateRefId;
+				if (Request.Id != null)
+				{
+					NewTemplateRefId = new TemplateRefId(Request.Id);
+				}
+				else
+				{
+					NewTemplateRefId = TemplateRefId.Sanitize(Request.Name);
+				}
+
+				// Create the schedule object
+				Schedule? Schedule = null;
+				if (Request.Schedule != null)
+				{
+					Schedule = Request.Schedule.ToModel();
+					Schedule.LastTriggerTime = Clock.UtcNow;
+				}
+
+				// Add it to the list
+				TemplateRef NewTemplateRef = new TemplateRef(NewTemplate, Request.ShowUgsBadges, Request.ShowUgsAlerts, Request.NotificationChannel, Request.NotificationChannelFilter, Request.TriageChannel, Schedule, Request.ChainedJobs?.ConvertAll(x => new ChainedJobTemplate(x)), Acl.Merge(null, Request.Acl));
+				if (Stream != null && Stream.Templates.TryGetValue(NewTemplateRefId, out TemplateRef? OldTemplateRef))
+				{
+					if (OldTemplateRef.Schedule != null && NewTemplateRef.Schedule != null)
+					{
+						NewTemplateRef.Schedule.CopyState(OldTemplateRef.Schedule);
+					}
+				}
+				NewTemplateRefs.Add(NewTemplateRefId, NewTemplateRef);
+			}
+			foreach (TemplateRef TemplateRef in NewTemplateRefs.Values)
+			{
+				if (TemplateRef.ChainedJobs != null)
+				{
+					foreach (ChainedJobTemplate ChainedJob in TemplateRef.ChainedJobs)
+					{
+						if (!NewTemplateRefs.ContainsKey(ChainedJob.TemplateRefId))
+						{
+							throw new InvalidDataException($"Invalid template ref id '{ChainedJob.TemplateRefId}");
+						}
+					}
+				}
+			}
+			return NewTemplateRefs;
+		}
+
+		/// <inheritdoc/>
+		async Task<IStream?> TryCreateAsync(StreamId Id, ProjectId ProjectId, string ConfigPath, string ConfigRevision, StreamConfig Config, DefaultPreflight? DefaultPreflight, List<StreamTab> Tabs, Dictionary<string, AgentType> AgentTypes, Dictionary<string, WorkspaceType> WorkspaceTypes, Dictionary<TemplateRefId, TemplateRef> TemplateRefs, Acl? Acl)
+		{
+			StreamDocument NewStream = new StreamDocument(Id, Config.Name, ProjectId);
+			NewStream.ClusterName = Config.ClusterName;
+			NewStream.ConfigPath = ConfigPath;
+			NewStream.ConfigRevision = ConfigRevision;
+			NewStream.Order = Config.Order ?? StreamDocument.DefaultOrder;
+			NewStream.NotificationChannel = Config.NotificationChannel;
+			NewStream.NotificationChannelFilter = Config.NotificationChannelFilter;
+			NewStream.TriageChannel = Config.TriageChannel;
+			NewStream.DefaultPreflight = DefaultPreflight;
+			NewStream.Tabs = Tabs;
+			NewStream.AgentTypes = AgentTypes;
+			NewStream.WorkspaceTypes = WorkspaceTypes;
+			NewStream.Templates = TemplateRefs;
+			NewStream.Acl = Acl;
+
+			try
+			{
+				await Streams.InsertOneAsync(NewStream);
+				return NewStream;
+			}
+			catch (MongoWriteException Ex)
+			{
+				if (Ex.WriteError.Category == ServerErrorCategory.DuplicateKey)
+				{
+					return null;
+				}
+				else
+				{
+					throw;
+				}
+			}
+		}
+
+		/// <inheritdoc/>
+		async Task<IStream?> TryReplaceAsync(IStream StreamInterface, ProjectId ProjectId, string ConfigPath, string ConfigRevision, StreamConfig Config, DefaultPreflight? DefaultPreflight, List<StreamTab> Tabs, Dictionary<string, AgentType>? AgentTypes, Dictionary<string, WorkspaceType>? WorkspaceTypes, Dictionary<TemplateRefId, TemplateRef>? TemplateRefs, Acl? Acl)
+		{
+			int Order = Config.Order ?? StreamDocument.DefaultOrder;
+
+			StreamDocument Stream = (StreamDocument)StreamInterface;
+
+			UpdateDefinitionBuilder<StreamDocument> UpdateBuilder = Builders<StreamDocument>.Update;
+
+			List<UpdateDefinition<StreamDocument>> Updates = new List<UpdateDefinition<StreamDocument>>();
+			Updates.Add(UpdateBuilder.Set(x => x.Name, Config.Name));
+			Updates.Add(UpdateBuilder.Set(x => x.ProjectId, ProjectId));
+			Updates.Add(UpdateBuilder.Set(x => x.ClusterName, Config.ClusterName));
+			Updates.Add(UpdateBuilder.Set(x => x.ConfigPath, ConfigPath));
+			Updates.Add(UpdateBuilder.Set(x => x.ConfigRevision, ConfigRevision));
+			Updates.Add(UpdateBuilder.Set(x => x.Order, Order));
+			Updates.Add(UpdateBuilder.Set(x => x.NotificationChannel, Config.NotificationChannel));
+			Updates.Add(UpdateBuilder.Set(x => x.NotificationChannelFilter, Config.NotificationChannelFilter));
+			Updates.Add(UpdateBuilder.Set(x => x.TriageChannel, Config.TriageChannel));
+			Updates.Add(UpdateBuilder.Set(x => x.DefaultPreflight, DefaultPreflight));
+			Updates.Add(UpdateBuilder.Set(x => x.Tabs, Tabs ?? new List<StreamTab>()));
+			Updates.Add(UpdateBuilder.Set(x => x.AgentTypes, AgentTypes ?? new Dictionary<string, AgentType>()));
+			Updates.Add(UpdateBuilder.Set(x => x.WorkspaceTypes, WorkspaceTypes ?? new Dictionary<string, WorkspaceType>()));
+			Updates.Add(UpdateBuilder.Set(x => x.Templates, TemplateRefs ?? new Dictionary<TemplateRefId, TemplateRef>()));
+			Updates.Add(UpdateBuilder.SetOrUnsetNullRef(x => x.Acl, Acl));
+			Updates.Add(UpdateBuilder.Unset(x => x.Deleted));
+
+			return await TryUpdateStreamAsync(Stream, UpdateBuilder.Combine(Updates));
+		}
+
+		/// <inheritdoc/>
+		public async Task<IStream?> GetAsync(StreamId StreamId)
+		{
+			return await Streams.Find<StreamDocument>(x => x.Id == StreamId).FirstOrDefaultAsync();
+		}
+
+		/// <inheritdoc/>
+		public async Task<IStreamPermissions?> GetPermissionsAsync(StreamId StreamId)
+		{
+			return await Streams.Find<StreamDocument>(x => x.Id == StreamId).Project<StreamPermissions>(StreamPermissions.Projection).FirstOrDefaultAsync();
+		}
+
+		/// <inheritdoc/>
+		public async Task<List<IStream>> FindAllAsync()
+		{
+			List<StreamDocument> Results = await Streams.Find(Builders<StreamDocument>.Filter.Ne(x => x.Deleted, true)).ToListAsync();
+			return Results.ConvertAll<IStream>(x => x);
+		}
+
+		/// <inheritdoc/>
+		public async Task<List<IStream>> FindForProjectsAsync(ProjectId[] ProjectIds)
+		{
+			FilterDefinition<StreamDocument> Filter = Builders<StreamDocument>.Filter.In(x => x.ProjectId, ProjectIds) & Builders<StreamDocument>.Filter.Ne(x => x.Deleted, true);
+			List<StreamDocument> Results = await Streams.Find(Filter).ToListAsync();
+			return Results.ConvertAll<IStream>(x => x);
+		}
+
+		/// <inheritdoc/>
+		public async Task<IStream?> TryUpdatePauseStateAsync(IStream StreamInterface, DateTime? NewPausedUntil, string? NewPauseComment)
+		{
+			StreamDocument Stream = (StreamDocument)StreamInterface;
+
+			UpdateDefinitionBuilder<StreamDocument> UpdateBuilder = Builders<StreamDocument>.Update;
+
+			List<UpdateDefinition<StreamDocument>> Updates = new List<UpdateDefinition<StreamDocument>>();
+			Stream.PausedUntil = NewPausedUntil;
+			Stream.PauseComment = NewPauseComment;
+			Updates.Add(UpdateBuilder.Set(x => x.PausedUntil, NewPausedUntil));
+			Updates.Add(UpdateBuilder.Set(x => x.PauseComment, NewPauseComment));
+
+			return await TryUpdateStreamAsync(Stream, UpdateBuilder.Combine(Updates));
+		}
+
+		/// <inheritdoc/>
+		public async Task<IStream?> TryUpdateScheduleTriggerAsync(IStream StreamInterface, TemplateRefId TemplateRefId, DateTimeOffset? LastTriggerTime, int? LastTriggerChange, List<ObjectId> NewActiveJobs)
+		{
+			StreamDocument Stream = (StreamDocument)StreamInterface;
+			Schedule Schedule = Stream.Templates[TemplateRefId].Schedule!;
+
+			// Build the updates. MongoDB driver cannot parse TemplateRefId in expression tree; need to specify field name explicitly
+			List<UpdateDefinition<StreamDocument>> Updates = new List<UpdateDefinition<StreamDocument>>();
+			if (LastTriggerTime.HasValue && LastTriggerTime.Value != Schedule.LastTriggerTime)
+			{
+				FieldDefinition<StreamDocument, DateTimeOffset> LastTriggerTimeField = $"{nameof(Stream.Templates)}.{TemplateRefId}.{nameof(Schedule)}.{nameof(Schedule.LastTriggerTime)}";
+				Updates.Add(Builders<StreamDocument>.Update.Set(LastTriggerTimeField, LastTriggerTime.Value));
+				Schedule.LastTriggerTime = LastTriggerTime.Value;
+			}
+			if (LastTriggerChange.HasValue && LastTriggerChange.Value > Schedule.LastTriggerChange)
+			{
+				FieldDefinition<StreamDocument, int> LastTriggerChangeField = $"{nameof(Stream.Templates)}.{TemplateRefId}.{nameof(Schedule)}.{nameof(Schedule.LastTriggerChange)}";
+				Updates.Add(Builders<StreamDocument>.Update.Set(LastTriggerChangeField, LastTriggerChange.Value));
+				Schedule.LastTriggerChange = LastTriggerChange.Value;
+			}
+			if (NewActiveJobs != null)
+			{
+				FieldDefinition<StreamDocument, List<ObjectId>> Field = $"{nameof(Stream.Templates)}.{TemplateRefId}.{nameof(Schedule)}.{nameof(Schedule.ActiveJobs)}";
+				Updates.Add(Builders<StreamDocument>.Update.Set(Field, NewActiveJobs));
+				Schedule.ActiveJobs = NewActiveJobs;
+			}
+
+			return (Updates.Count == 0)? StreamInterface : await TryUpdateStreamAsync(Stream, Builders<StreamDocument>.Update.Combine(Updates));
+		}
+
+		/// <summary>
+		/// Update a stream
+		/// </summary>
+		/// <param name="Stream">The stream to update</param>
+		/// <param name="Update">The update definition</param>
+		/// <returns>The updated document, or null the update failed</returns>
+		private async Task<StreamDocument?> TryUpdateStreamAsync(StreamDocument Stream, UpdateDefinition<StreamDocument> Update)
+		{
+			FilterDefinition<StreamDocument> Filter = Builders<StreamDocument>.Filter.Expr(x => x.Id == Stream.Id && x.UpdateIndex == Stream.UpdateIndex);
+			Update = Update.Set(x => x.UpdateIndex, Stream.UpdateIndex + 1);
+
+			FindOneAndUpdateOptions<StreamDocument> Options = new FindOneAndUpdateOptions<StreamDocument> { ReturnDocument = ReturnDocument.After };
+			return await Streams.FindOneAndUpdateAsync(Filter, Update, Options);
+		}
+
+		/// <inheritdoc/>
+		public async Task DeleteAsync(StreamId StreamId)
+		{
+			await Streams.UpdateOneAsync<StreamDocument>(x => x.Id == StreamId, Builders<StreamDocument>.Update.Set(x => x.Deleted, true).Inc(x => x.UpdateIndex, 1));
+		}
+
+
+		/// <summary>
+		/// Checks the stream definition for consistency
+		/// </summary>
+		public static void Validate(StreamId StreamId, DefaultPreflight? DefaultPreflight, IReadOnlyDictionary<TemplateRefId, TemplateRef> Templates, IReadOnlyList<StreamTab> Tabs, IReadOnlyDictionary<string, AgentType> AgentTypes, IReadOnlyDictionary<string, WorkspaceType> WorkspaceTypes)
+		{
+			// Check the default preflight template is valid
+			if (DefaultPreflight != null)
+			{
+				if (DefaultPreflight.TemplateRefId != null && !Templates.ContainsKey(DefaultPreflight.TemplateRefId.Value))
+				{
+					throw new InvalidStreamException($"Default preflight template was listed as '{DefaultPreflight.TemplateRefId.Value}', but no template was found by that name");
+				}
+			}
+
+			// Check that all the templates are referenced by a tab
+			HashSet<TemplateRefId> RemainingTemplates = new HashSet<TemplateRefId>(Templates.Keys);
+			foreach (JobsTab JobsTab in Tabs.OfType<JobsTab>())
+			{
+				if (JobsTab.Templates != null)
+				{
+					RemainingTemplates.ExceptWith(JobsTab.Templates);
+				}
+			}
+			if (RemainingTemplates.Count > 0)
+			{
+				throw new InvalidStreamException(String.Join("\n", RemainingTemplates.Select(x => $"Template '{x}' is not listed on any tab for {StreamId}")));
+			}
+
+			// Check that all the agent types reference valid workspace names
+			foreach (KeyValuePair<string, AgentType> Pair in AgentTypes)
+			{
+				string? WorkspaceTypeName = Pair.Value.Workspace;
+				if (WorkspaceTypeName != null && !WorkspaceTypes.ContainsKey(WorkspaceTypeName))
+				{
+					throw new InvalidStreamException($"Agent type '{Pair.Key}' references undefined workspace type '{Pair.Value.Workspace}' in {StreamId}");
+				}
+			}
+		}
+	}
+}
diff --git a/Engine/Source/Programs/Horde/HordeServer/Projects/StreamService.cs b/Engine/Source/Programs/Horde/HordeServer/Projects/StreamService.cs
new file mode 100644
--- /dev/null
+++ b/Engine/Source/Programs/Horde/HordeServer/Projects/StreamService.cs
@@ -0,0 +1,316 @@
+// Copyright Epic Games, Inc. All Rights Reserved.
+
+using HordeServer.Api;
+using HordeCommon;
+using HordeServer.Models;
+using HordeServer.Utilities;
+using Microsoft.Extensions.Options;
+using MongoDB.Bson;
+using MongoDB.Driver;
+using System;
+using System.Collections.Generic;
+using System.Linq;
+using System.Security.Claims;
+using System.Text;
+using System.Threading.Tasks;
+using System.Linq.Expressions;
+using HordeServer.Collections;
+using Microsoft.Extensions.Caching.Memory;
+
+namespace HordeServer.Services
+{
+	using ProjectId = StringId<IProject>;
+	using StreamId = StringId<IStream>;
+	using TemplateRefId = StringId<TemplateRef>;
+
+	/// <summary>
+	/// Cache of information about stream ACLs
+	/// </summary>
+	public class StreamPermissionsCache : ProjectPermissionsCache
+	{
+		/// <summary>
+		/// Map of stream id to permissions for that stream
+		/// </summary>
+		public Dictionary<StreamId, IStreamPermissions?> Streams { get; set; } = new Dictionary<StreamId, IStreamPermissions?>();
+	}
+
+	/// <summary>
+	/// Wraps functionality for manipulating streams
+	/// </summary>
+	public sealed class StreamService : IDisposable
+	{
+		/// <summary>
+		/// The project service instance
+		/// </summary>
+		ProjectService ProjectService;
+
+		/// <summary>
+		/// Collection of stream documents
+		/// </summary>
+		IStreamCollection Streams;
+
+		/// <summary>
+		/// Cache of stream documents
+		/// </summary>
+		MemoryCache StreamCache = new MemoryCache(new MemoryCacheOptions());
+
+		/// <summary>
+		/// Accessor for the stream collection
+		/// </summary>
+		public IStreamCollection StreamCollection => Streams;
+
+		/// <summary>
+		/// Constructor
+		/// </summary>
+		/// <param name="ProjectService">The project service instance</param>
+		/// <param name="Streams">Collection of stream documents</param>
+		public StreamService(ProjectService ProjectService, IStreamCollection Streams)
+		{
+			this.ProjectService = ProjectService;
+			this.Streams = Streams;
+		}
+
+		/// <summary>
+		/// Dispose of any managed resources
+		/// </summary>
+		public void Dispose()
+		{
+			StreamCache.Dispose();
+		}
+
+		/// <summary>
+		/// Deletes an existing stream
+		/// </summary>
+		/// <param name="StreamId">Unique id of the stream</param>
+		/// <returns>Async task object</returns>
+		public async Task DeleteStreamAsync(StreamId StreamId)
+		{
+			await Streams.DeleteAsync(StreamId);
+		}
+
+		/// <summary>
+		/// Updates an existing stream
+		/// </summary>
+		/// <param name="Stream">The stream to update</param>
+		/// <param name="NewPausedUntil">The new datetime for pausing builds</param>
+		/// <param name="NewPauseComment">The reason for pausing</param>
+		/// <returns>Async task object</returns>
+		public async Task<IStream?> UpdatePauseStateAsync(IStream? Stream, DateTime? NewPausedUntil = null, string? NewPauseComment = null)
+		{
+			for (; Stream != null; Stream = await GetStreamAsync(Stream.Id))
+			{
+				IStream? NewStream = await Streams.TryUpdatePauseStateAsync(Stream, NewPausedUntil, NewPauseComment);
+				if (NewStream != null)
+				{
+					return NewStream;
+				}
+			}
+			return null;
+		}
+
+		/// <summary>
+		/// Attempts to update the last trigger time for a schedule
+		/// </summary>
+		/// <param name="Stream">The stream to update</param>
+		/// <param name="TemplateRefId">The template ref id</param>
+		/// <param name="LastTriggerTime"></param>
+		/// <param name="LastTriggerChange"></param>
+		/// <param name="AddJobs">Jobs to add</param>
+		/// <param name="RemoveJobs">Jobs to remove</param>
+		/// <returns>True if the stream was updated</returns>
+		public async Task<IStream?> UpdateScheduleTriggerAsync(IStream Stream, TemplateRefId TemplateRefId, DateTimeOffset? LastTriggerTime, int? LastTriggerChange, List<ObjectId> AddJobs, List<ObjectId> RemoveJobs)
+		{
+			IStream? NewStream = Stream;
+			while (NewStream != null)
+			{
+				TemplateRef? TemplateRef;
+				if (!NewStream.Templates.TryGetValue(TemplateRefId, out TemplateRef))
+				{
+					break;
+				}
+				if (TemplateRef.Schedule == null)
+				{
+					break;
+				}
+
+				List<ObjectId> NewActiveJobs = TemplateRef.Schedule.ActiveJobs.Except(RemoveJobs).Union(AddJobs).ToList();
+
+				NewStream = await Streams.TryUpdateScheduleTriggerAsync(NewStream, TemplateRefId, LastTriggerTime, LastTriggerChange, NewActiveJobs);
+
+				if (NewStream != null)
+				{
+					return NewStream;
+				}
+
+				NewStream = await Streams.GetAsync(Stream.Id);
+			}
+			return null;
+		}
+
+		/// <summary>
+		/// Gets all the available streams for a project
+		/// </summary>
+		/// <returns>List of stream documents</returns>
+		public Task<List<IStream>> GetStreamsAsync()
+		{
+			return Streams.FindAllAsync();
+		}
+
+		/// <summary>
+		/// Gets all the available streams for a project
+		/// </summary>
+		/// <param name="ProjectId">Unique id of the project to query</param>
+		/// <returns>List of stream documents</returns>
+		public Task<List<IStream>> GetStreamsAsync(ProjectId ProjectId)
+		{
+			return Streams.FindForProjectsAsync(new[] { ProjectId });
+		}
+
+		/// <summary>
+		/// Gets all the available streams for a project
+		/// </summary>
+		/// <param name="ProjectIds">Unique id of the project to query</param>
+		/// <returns>List of stream documents</returns>
+		public Task<List<IStream>> GetStreamsAsync(ProjectId[] ProjectIds)
+		{
+			if (ProjectIds.Length == 0)
+			{
+				return Streams.FindAllAsync();
+			}
+			else
+			{
+				return Streams.FindForProjectsAsync(ProjectIds);
+			}
+		}
+
+		/// <summary>
+		/// Gets a stream by ID
+		/// </summary>
+		/// <param name="StreamId">Unique id of the stream</param>
+		/// <returns>The stream document</returns>
+		public Task<IStream?> GetStreamAsync(StreamId StreamId)
+		{
+			return Streams.GetAsync(StreamId);
+		}
+
+		/// <summary>
+		/// Adds a cached stream interface
+		/// </summary>
+		/// <param name="Stream">The stream to add</param>
+		/// <returns>The new stream</returns>
+		private void AddCachedStream(IStream Stream)
+		{
+			MemoryCacheEntryOptions Options = new MemoryCacheEntryOptions().SetAbsoluteExpiration(TimeSpan.FromMinutes(1.0));
+			StreamCache.Set(Stream.Id, Stream, Options);
+		}
+
+		/// <summary>
+		/// Gets a cached stream interface
+		/// </summary>
+		/// <param name="StreamId">Unique id for the stream</param>
+		/// <returns>The new stream</returns>
+		public async Task<IStream?> GetCachedStream(StreamId StreamId)
+		{
+			object? Stream;
+			if (!StreamCache.TryGetValue(StreamId, out Stream))
+			{
+				Stream = await GetStreamAsync(StreamId);
+				if (Stream != null)
+				{
+					AddCachedStream((IStream)Stream);
+				}
+			}
+			return (IStream?)Stream;
+		}
+
+		/// <summary>
+		/// Gets a stream's permissions by ID
+		/// </summary>
+		/// <param name="StreamId">Unique id of the stream</param>
+		/// <returns>The stream document</returns>
+		public Task<IStreamPermissions?> GetStreamPermissionsAsync(StreamId StreamId)
+		{
+			return Streams.GetPermissionsAsync(StreamId);
+		}
+
+		/// <summary>
+		/// Determines if the user is authorized to perform an action on a particular stream
+		/// </summary>
+		/// <param name="Acl">ACL for the stream to check</param>
+		/// <param name="ProjectId">The parent project id</param>
+		/// <param name="Action">The action being performed</param>
+		/// <param name="User">The principal to authorize</param>
+		/// <param name="Cache">Cache of project permissions</param>
+		/// <returns>True if the action is authorized</returns>
+		private Task<bool> AuthorizeAsync(Acl? Acl, ProjectId ProjectId, AclAction Action, ClaimsPrincipal User, ProjectPermissionsCache? Cache)
+		{
+			bool? Result = Acl?.Authorize(Action, User);
+			if (Result == null)
+			{
+				return ProjectService.AuthorizeAsync(ProjectId, Action, User, Cache);
+			}
+			else
+			{
+				return Task.FromResult(Result.Value);
+			}
+		}
+
+		/// <summary>
+		/// Determines if the user is authorized to perform an action on a particular stream
+		/// </summary>
+		/// <param name="Stream">The stream to check</param>
+		/// <param name="Action">The action being performed</param>
+		/// <param name="User">The principal to authorize</param>
+		/// <param name="Cache">Cache of project permissions</param>
+		/// <returns>True if the action is authorized</returns>
+		public Task<bool> AuthorizeAsync(IStream Stream, AclAction Action, ClaimsPrincipal User, ProjectPermissionsCache? Cache)
+		{
+			return AuthorizeAsync(Stream.Acl, Stream.ProjectId, Action, User, Cache);
+		}
+
+		/// <summary>
+		/// Determines if the user is authorized to perform an action on a particular stream
+		/// </summary>
+		/// <param name="Stream">The stream to check</param>
+		/// <param name="Template">Template within the stream to check</param>
+		/// <param name="Action">The action being performed</param>
+		/// <param name="User">The principal to authorize</param>
+		/// <param name="Cache">Cache of project permissions</param>
+		/// <returns>True if the action is authorized</returns>
+		public Task<bool> AuthorizeAsync(IStream Stream, TemplateRef Template, AclAction Action, ClaimsPrincipal User, ProjectPermissionsCache? Cache)
+		{
+			bool? Result = Template.Acl?.Authorize(Action, User);
+			if (Result == null)
+			{
+				return AuthorizeAsync(Stream, Action, User, Cache);
+			}
+			else
+			{
+				return Task.FromResult(Result.Value);
+			}
+		}
+
+		/// <summary>
+		/// Determines if the user is authorized to perform an action on a particular project
+		/// </summary>
+		/// <param name="StreamId">The stream id to check</param>
+		/// <param name="Action">The action being performed</param>
+		/// <param name="User">The principal to authorize</param>
+		/// <param name="Cache">Cache of stream permissions</param>
+		/// <returns>True if the action is authorized</returns>
+		public async Task<bool> AuthorizeAsync(StreamId StreamId, AclAction Action, ClaimsPrincipal User, StreamPermissionsCache? Cache)
+		{
+			IStreamPermissions? Permissions;
+			if (Cache == null)
+			{
+				Permissions = await GetStreamPermissionsAsync(StreamId);
+			}
+			else if (!Cache.Streams.TryGetValue(StreamId, out Permissions))
+			{
+				Permissions = await GetStreamPermissionsAsync(StreamId);
+				Cache.Streams.Add(StreamId, Permissions);
+			}
+			return Permissions != null && await AuthorizeAsync(Permissions.Acl, Permissions.ProjectId, Action, User, Cache);
+		}
+	}
+}
diff --git a/Engine/Source/Programs/Horde/HordeServer/Services/IAvatarService.cs b/Engine/Source/Programs/Horde/HordeServer/Services/IAvatarService.cs
deleted file mode 100644
diff --git a/Engine/Source/Programs/Horde/HordeServer/Services/TemplateService.cs b/Engine/Source/Programs/Horde/HordeServer/Services/TemplateService.cs
deleted file mode 100644
diff --git a/Engine/Source/Programs/Horde/HordeServer/Streams/IStream.cs b/Engine/Source/Programs/Horde/HordeServer/Streams/IStream.cs
deleted file mode 100644
diff --git a/Engine/Source/Programs/Horde/HordeServer/Streams/IStreamCollection.cs b/Engine/Source/Programs/Horde/HordeServer/Streams/IStreamCollection.cs
deleted file mode 100644
diff --git a/Engine/Source/Programs/Horde/HordeServer/Streams/Impl/StreamCollection.cs b/Engine/Source/Programs/Horde/HordeServer/Streams/Impl/StreamCollection.cs
deleted file mode 100644
diff --git a/Engine/Source/Programs/Horde/HordeServer/Streams/StreamService.cs b/Engine/Source/Programs/Horde/HordeServer/Streams/StreamService.cs
deleted file mode 100644
diff --git a/Engine/Source/Programs/Horde/HordeServer/Users/IAvatarService.cs b/Engine/Source/Programs/Horde/HordeServer/Users/IAvatarService.cs
new file mode 100644
--- /dev/null
+++ b/Engine/Source/Programs/Horde/HordeServer/Users/IAvatarService.cs
@@ -0,0 +1,28 @@
+// Copyright Epic Games, Inc. All Rights Reserved.
+
+using HordeServer.Models;
+using Microsoft.Extensions.Logging;
+using Microsoft.Extensions.Options;
+using Newtonsoft.Json.Linq;
+using System;
+using System.Collections.Generic;
+using System.Linq;
+using System.Net.Http;
+using System.Threading.Tasks;
+
+namespace HordeServer.Services
+{
+	/// <summary>
+	/// Provides avatar images
+	/// </summary>
+	public interface IAvatarService
+	{
+		/// <summary>
+		/// Gets a users's avatar
+		/// </summary>
+		/// <param name="User"></param>
+		/// <returns></returns>
+		Task<IAvatar?> GetAvatarAsync(IUser User);
+	}
+}
+
