Horde: More file moves.

diff --git a/Engine/Source/Programs/Horde/HordeServer/Agents/AgentSoftware.cs b/Engine/Source/Programs/Horde/HordeServer/Agents/AgentSoftware.cs
new file mode 100644
--- /dev/null
+++ b/Engine/Source/Programs/Horde/HordeServer/Agents/AgentSoftware.cs
@@ -0,0 +1,47 @@
+// Copyright Epic Games, Inc. All Rights Reserved.
+
+using HordeServer.Api;
+using HordeServer.Collections;
+using HordeServer.Services;
+using HordeServer.Utilities;
+using MongoDB.Bson;
+using MongoDB.Bson.Serialization.Attributes;
+using System;
+using System.Collections.Generic;
+using System.IO;
+using System.IO.Compression;
+using System.Linq;
+using System.Text.Json;
+using System.Threading.Tasks;
+
+namespace HordeServer.Models
+{
+	using AgentSoftwareVersion = StringId<IAgentSoftwareCollection>;
+	using AgentSoftwareChannelName = StringId<AgentSoftwareChannels>;
+
+	/// <summary>
+	/// A software channel
+	/// </summary>
+	public interface IAgentSoftwareChannel
+	{
+		/// <summary>
+		/// The channel id
+		/// </summary>
+		public AgentSoftwareChannelName Name { get; set; }
+
+		/// <summary>
+		/// Name of the user that made the last modification
+		/// </summary>
+		public string? ModifiedBy { get; set; }
+
+		/// <summary>
+		/// Last modification time
+		/// </summary>
+		public DateTime ModifiedTime { get; set; }
+
+		/// <summary>
+		/// The software version number
+		/// </summary>
+		public string Version { get; set; }
+	}
+}
diff --git a/Engine/Source/Programs/Horde/HordeServer/Agents/IAgentSoftwareCollection.cs b/Engine/Source/Programs/Horde/HordeServer/Agents/IAgentSoftwareCollection.cs
new file mode 100644
--- /dev/null
+++ b/Engine/Source/Programs/Horde/HordeServer/Agents/IAgentSoftwareCollection.cs
@@ -0,0 +1,54 @@
+// Copyright Epic Games, Inc. All Rights Reserved.
+
+using EpicGames.Core;
+using HordeCommon;
+using HordeServer.Models;
+using HordeServer.Services;
+using MongoDB.Bson;
+using MongoDB.Bson.Serialization.Attributes;
+using MongoDB.Driver;
+using System;
+using System.Collections.Generic;
+using System.Globalization;
+using System.IO;
+using System.IO.Compression;
+using System.Linq;
+using System.Text.Json;
+using System.Threading.Tasks;
+
+namespace HordeServer.Collections
+{
+	/// <summary>
+	/// Interface for a collection of template documents
+	/// </summary>
+	public interface IAgentSoftwareCollection
+	{
+		/// <summary>
+		/// Adds a new software revision
+		/// </summary>
+		/// <param name="Version">The version number</param>
+		/// <param name="Data">Zip file containing the new software</param>
+		/// <returns>New software instance</returns>
+		Task<bool> AddAsync(string Version, byte[] Data);
+
+		/// <summary>
+		/// Tests whether a given version exists
+		/// </summary>
+		/// <param name="Version">Version of the software</param>
+		/// <returns>True if it exists, false otherwise</returns>
+		Task<bool> ExistsAsync(string Version);
+
+		/// <summary>
+		/// Removes a software archive 
+		/// </summary>
+		/// <param name="Version">Version of the software to delete</param>
+		Task<bool> RemoveAsync(string Version);
+
+		/// <summary>
+		/// Downloads software of a given revision
+		/// </summary>
+		/// <param name="Version">Version of the software</param>
+		/// <returns>Data for the given software</returns>
+		Task<byte[]?> GetAsync(string Version);
+	}
+}
diff --git a/Engine/Source/Programs/Horde/HordeServer/Agents/Impl/AgentSoftwareCollection.cs b/Engine/Source/Programs/Horde/HordeServer/Agents/Impl/AgentSoftwareCollection.cs
new file mode 100644
--- /dev/null
+++ b/Engine/Source/Programs/Horde/HordeServer/Agents/Impl/AgentSoftwareCollection.cs
@@ -0,0 +1,210 @@
+// Copyright Epic Games, Inc. All Rights Reserved.
+
+using EpicGames.Core;
+using HordeCommon;
+using HordeServer.Models;
+using HordeServer.Services;
+using HordeServer.Utilities;
+using MongoDB.Bson;
+using MongoDB.Bson.Serialization.Attributes;
+using MongoDB.Driver;
+using System;
+using System.Collections.Generic;
+using System.Globalization;
+using System.IO;
+using System.IO.Compression;
+using System.Linq;
+using System.Text.Json;
+using System.Threading.Tasks;
+
+namespace HordeServer.Collections.Impl
+{
+	using AgentSoftwareVersion = StringId<IAgentSoftwareCollection>;
+
+	/// <summary>
+	/// Collection of agent software documents
+	/// </summary>
+	public class AgentSoftwareCollection : IAgentSoftwareCollection
+	{
+		/// <summary>
+		/// Information about agent software
+		/// </summary>
+		class AgentSoftwareDocument
+		{
+			[BsonId]
+			public string Id { get; set; }
+
+			[BsonIgnoreIfNull]
+			public byte[]? Data { get; set; }
+
+			[BsonIgnoreIfNull]
+            public List<string>? ChunkIds { get; set; }
+
+            [BsonConstructor]
+			private AgentSoftwareDocument()
+			{
+				this.Id = null!;
+				this.Data = null;
+                this.ChunkIds = null;			
+            }
+
+			public AgentSoftwareDocument(string Id, byte[]? Data, List<string>? ChunkIds = null)
+			{
+				this.Id = Id;
+				this.Data = Data;
+                this.ChunkIds = ChunkIds;
+            }
+		}
+
+		class ChunkDocument
+		{
+			[BsonId]
+			public string Id { get; set; }
+
+			public byte[] Data { get; set; } = null!;
+
+			public string Version { get; set; }
+
+
+            [BsonConstructor]
+			private ChunkDocument()
+			{
+				this.Id = null!;
+				this.Data = null!;
+                this.Version = null!;
+            }
+
+			public ChunkDocument(string Id, string Version, byte[] Data)
+			{
+				this.Id = Id;
+				this.Data = Data;
+                this.Version = Version;
+            }
+		}
+
+		// MongoDB document size limit is 16 megs, we chunk at this with a 1k buffer for anything else in document
+		readonly int ChunkSize = 16 * 1024 * 1024 - 1024;
+
+
+
+		/// <summary>
+		/// Template documents
+		/// </summary>
+		IMongoCollection<AgentSoftwareDocument> Collection;
+
+		/// <summary>
+		/// Agent template chunks
+		/// </summary>
+		IMongoCollection<ChunkDocument> AgentChunks;
+
+
+		/// <summary>
+		/// Constructor
+		/// </summary>
+		/// <param name="DatabaseService">The database service singleton</param>
+		public AgentSoftwareCollection(DatabaseService DatabaseService)
+		{
+			Collection = DatabaseService.GetCollection<AgentSoftwareDocument>("AgentSoftware");
+			AgentChunks = DatabaseService.GetCollection<ChunkDocument>("AgentSoftwareChunks");
+		}
+
+		/// <inheritdoc/>
+		public async Task<bool> AddAsync(string Version, byte[] Data)
+		{
+
+			if (await ExistsAsync(Version))
+			{
+                return false;
+            }
+
+            AgentSoftwareDocument Document;
+
+            // Check if agent will fit in document, otherwise we need to chunk
+            if (Data.Length <= ChunkSize)
+			{
+				Document = new AgentSoftwareDocument(Version, Data);
+				await Collection.InsertOneAsync(Document);
+                return true;
+            }
+
+			int TotalRead = 0;            
+            byte[] Buffer = new byte[ChunkSize];            
+			List<ChunkDocument> Chunks = new List<ChunkDocument>();
+            using (MemoryStream ChunkStream = new MemoryStream(Data))
+            {
+                for (; ; )
+                {
+					int BytesRead = ChunkStream.Read(Buffer, 0, ChunkSize);
+					Chunks.Add(new ChunkDocument(ObjectId.GenerateNewId().ToString(), Version, Buffer.Take(BytesRead).ToArray()));
+					TotalRead += BytesRead;
+					if (TotalRead == Data.Length)
+					{
+						break;
+					}
+                }
+            }
+
+            // insert the chunks
+            await AgentChunks.InsertManyAsync(Chunks);
+
+            Document = new AgentSoftwareDocument(Version, null, Chunks.Select( Chunk => Chunk.Id).ToList());
+			await Collection.InsertOneAsync(Document);
+            return true;
+        }
+
+		/// <inheritdoc/>
+		public async Task<bool> ExistsAsync(string Version)
+		{
+			return await Collection.Find(x => x.Id == Version).AnyAsync();
+		}
+
+		/// <inheritdoc/>
+		public async Task<bool> RemoveAsync(string Version)
+		{
+			// Delete any chunks
+			await AgentChunks.DeleteManyAsync(x => x.Version == Version);
+
+			DeleteResult Result = await Collection.DeleteOneAsync(x => x.Id == Version);			
+			return Result.DeletedCount > 0;
+		}
+
+		/// <inheritdoc/>
+		public async Task<byte[]?> GetAsync(string Version)
+		{			
+			IFindFluent<AgentSoftwareDocument, AgentSoftwareDocument> Query = Collection.Find(x => x.Id == Version);
+			AgentSoftwareDocument? Software = await Query.FirstOrDefaultAsync();
+
+			if (Software == null)
+			{
+                return null;
+            }
+
+			if (Software.Data != null && Software.Data.Length > 0)
+			{
+                return Software.Data;
+            }
+
+			if (Software.ChunkIds != null)
+			{
+				FilterDefinitionBuilder<ChunkDocument> FilterBuilder = Builders<ChunkDocument>.Filter;
+				FilterDefinition<ChunkDocument> Filter = FilterBuilder.Empty;
+				Filter &= FilterBuilder.In(x => x.Id, Software.ChunkIds);
+
+				IFindFluent<ChunkDocument, ChunkDocument> ChunkQuery = AgentChunks.Find(x => Software.ChunkIds.Contains(x.Id));				
+				List<ChunkDocument> Chunks = await ChunkQuery.ToListAsync();
+
+				using (MemoryStream DataStream = new MemoryStream())
+				{
+					foreach (string ChunkId in Software.ChunkIds)
+					{
+                        DataStream.Write(Chunks.First(x => x.Id == ChunkId).Data);
+                    }
+
+                    return DataStream.ToArray();
+                }
+			}
+
+            return null;
+        }
+	}
+}
diff --git a/Engine/Source/Programs/Horde/HordeServer/Agents/Impl/AgentSoftwareService.cs b/Engine/Source/Programs/Horde/HordeServer/Agents/Impl/AgentSoftwareService.cs
new file mode 100644
--- /dev/null
+++ b/Engine/Source/Programs/Horde/HordeServer/Agents/Impl/AgentSoftwareService.cs
@@ -0,0 +1,348 @@
+// Copyright Epic Games, Inc. All Rights Reserved.
+
+using HordeCommon;
+using HordeServer.Collections;
+using HordeServer.Models;
+using HordeServer.Utilities;
+using Microsoft.IdentityModel.Tokens;
+using MongoDB.Bson;
+using MongoDB.Bson.Serialization.Attributes;
+using MongoDB.Driver;
+using System;
+using System.Collections.Generic;
+using System.Diagnostics;
+using System.IO;
+using System.Linq;
+using System.Threading.Channels;
+using System.Threading.Tasks;
+using EpicGames.Core;
+using Microsoft.Extensions.Logging;
+using System.Runtime.InteropServices;
+using Microsoft.Extensions.Options;
+using System.IO.Compression;
+using System.Security.Cryptography.X509Certificates;
+
+namespace HordeServer.Services
+{
+	using AgentSoftwareChannelName = StringId<AgentSoftwareChannels>;
+
+	/// <summary>
+	/// Information about a channel
+	/// </summary>
+	public class AgentSoftwareChannel : IAgentSoftwareChannel
+	{
+		/// <summary>
+		/// The channel id
+		/// </summary>
+		public AgentSoftwareChannelName Name { get; set; }
+
+		/// <summary>
+		/// Name of the user that made the last modification
+		/// </summary>
+		public string? ModifiedBy { get; set; }
+
+		/// <summary>
+		/// Last modification time
+		/// </summary>
+		public DateTime ModifiedTime { get; set; }
+
+		/// <summary>
+		/// The software revision number
+		/// </summary>
+		public string Version { get; set; } = String.Empty;
+
+		/// <summary>
+		/// Private constructor for serialization
+		/// </summary>
+		[BsonConstructor]
+		private AgentSoftwareChannel()
+		{
+		}
+
+		/// <summary>
+		/// Constructor
+		/// </summary>
+		/// <param name="Name">Name of the channel</param>
+		public AgentSoftwareChannel(AgentSoftwareChannelName Name)
+		{
+			this.Name = Name;
+		}
+	}
+
+	/// <summary>
+	/// Singleton document used to track different versions of the agent software
+	/// </summary>
+	[SingletonDocument("5f455039d97900f2b6c735a9")]
+	public class AgentSoftwareChannels : SingletonBase
+	{
+		/// <summary>
+		/// The next version number
+		/// </summary>
+		public int NextVersion { get; set; } = 1;
+
+		/// <summary>
+		/// List of channels
+		/// </summary>
+		public List<AgentSoftwareChannel> Channels { get; set; } = new List<AgentSoftwareChannel>();
+
+		/// <summary>
+		/// Finds an existing channel by the given name, or adds a new one
+		/// </summary>
+		/// <param name="Name"></param>
+		/// <returns></returns>
+		public AgentSoftwareChannel FindOrAddChannel(AgentSoftwareChannelName Name)
+		{
+			AgentSoftwareChannel? Channel = Channels.FirstOrDefault(x => x.Name == Name);
+			if (Channel == null)
+			{
+				Channel = new AgentSoftwareChannel(Name);
+				Channels.Add(Channel);
+			}
+			return Channel;
+		}
+	}
+
+	/// <summary>
+	/// Wrapper for a collection of software revisions. 
+	/// </summary>
+	public class AgentSoftwareService
+	{
+		/// <summary>
+		/// Name of the default channel
+		/// </summary>
+		public static AgentSoftwareChannelName DefaultChannelName { get; } = new AgentSoftwareChannelName("default");
+
+		/// <summary>
+		/// Collection of software documents
+		/// </summary>
+		IAgentSoftwareCollection Collection;
+
+		/// <summary>
+		/// Channels singleton
+		/// </summary>
+		ISingletonDocument<AgentSoftwareChannels> Singleton;
+
+		/// <summary>
+		/// Cached copy of the channels singleton
+		/// </summary>
+		LazyCachedValue<Task<AgentSoftwareChannels>> ChannelsDocument;
+
+		/// <summary>
+		/// The server settings
+		/// </summary>
+		IOptionsMonitor<ServerSettings> Settings;
+
+		/// <summary>
+		///  Logger for controller
+		/// </summary>
+		private readonly ILogger<AgentSoftwareService> Logger;
+
+
+		/// <summary>
+		/// Constructor
+		/// </summary>
+		/// <param name="Collection">The software collection</param>
+		/// <param name="Singleton">The channels singleton</param>		
+		/// <param name="Settings">The settings monitor</param>		
+		/// <param name="Logger">The logger instance</param>
+		public AgentSoftwareService(IAgentSoftwareCollection Collection, ISingletonDocument<AgentSoftwareChannels> Singleton, IOptionsMonitor<ServerSettings> Settings, ILogger<AgentSoftwareService> Logger)
+		{
+			this.Collection = Collection;
+			this.Singleton = Singleton;
+			this.Settings = Settings;
+			this.Logger = Logger;
+
+			ChannelsDocument = new LazyCachedValue<Task<AgentSoftwareChannels>>(() => Singleton.GetAsync(), TimeSpan.FromSeconds(20.0));
+
+			Task.Run(() => RegisterDefaultAgent(5000));
+		}
+
+		/// <summary>
+		/// Finds all software archives matching a set of criteria
+		/// </summary>
+		public async Task<List<IAgentSoftwareChannel>> FindChannelsAsync()
+		{
+			AgentSoftwareChannels Current = await ChannelsDocument.GetLatest();
+			return Current.Channels.ConvertAll<IAgentSoftwareChannel>(x => x);
+		}
+
+		/// <summary>
+		/// Gets a single software archive
+		/// </summary>
+		/// <param name="Name">Unique id of the channel to find software for</param>
+		public async Task<IAgentSoftwareChannel?> GetChannelAsync(AgentSoftwareChannelName Name)
+		{
+			AgentSoftwareChannels Current = await ChannelsDocument.GetLatest();
+			return Current.Channels.FirstOrDefault(x => x.Name == Name);
+		}
+
+		/// <summary>
+		/// Gets a cached channel name
+		/// </summary>
+		/// <param name="Name">The cached channel name</param>
+		/// <returns></returns>
+		public async Task<IAgentSoftwareChannel?> GetCachedChannelAsync(AgentSoftwareChannelName Name)
+		{
+			AgentSoftwareChannels Document = await ChannelsDocument.GetCached();
+			return Document.Channels.FirstOrDefault(x => x.Name == Name);
+		}
+
+		/// <summary>
+		/// Removes a channel
+		/// </summary>
+		/// <param name="Name">The channel id</param>
+		/// <returns>Async task</returns>
+		public async Task DeleteChannelAsync(AgentSoftwareChannelName Name)
+		{
+			for (; ; )
+			{
+				AgentSoftwareChannels Current = await Singleton.GetAsync();
+
+				int ChannelIdx = Current.Channels.FindIndex(x => x.Name == Name);
+				if (ChannelIdx == -1)
+				{
+					break;
+				}
+
+				string Version = Current.Channels[ChannelIdx].Version;
+				Current.Channels.RemoveAt(ChannelIdx);
+
+				if (await Singleton.TryUpdateAsync(Current))
+				{
+					if (!Current.Channels.Any(x => x.Version == Version))
+					{
+						await Collection.RemoveAsync(Version);
+					}
+					break;
+				}
+			}
+		}
+
+		/// <summary>
+		/// Updates a new software revision
+		/// </summary>
+		/// <param name="Name">Name of the channel</param>
+		/// <param name="Author">Name of the user uploading this file</param>
+		/// <param name="Data">The input data stream. This should be a zip archive containing the HordeAgent executable.</param>
+		/// <returns>Unique id for the file</returns>
+		public async Task<string> SetArchiveAsync(AgentSoftwareChannelName Name, string? Author, byte[] Data)
+		{
+			// Upload the software
+			string Version = AgentUtilities.ReadVersion(Data);
+			await Collection.AddAsync(Version, Data);
+
+			// Update the channel
+			for(; ;)
+			{
+				AgentSoftwareChannels Instance = await Singleton.GetAsync();
+
+				AgentSoftwareChannel Channel = Instance.FindOrAddChannel(Name);
+				Channel.ModifiedBy = Author;
+				Channel.ModifiedTime = DateTime.UtcNow;
+				Channel.Version = Version;
+
+				if (await Singleton.TryUpdateAsync(Instance))
+				{
+					break;
+				}
+			}
+			return Version;
+		}
+
+		/// <summary>
+		/// Gets the zip file for a given channel
+		/// </summary>
+		/// <param name="Name">The channel name</param>
+		/// <returns>Data for the given archive</returns>
+		public async Task<byte[]?> GetArchiveAsync(AgentSoftwareChannelName Name)
+		{
+			AgentSoftwareChannels Instance = await Singleton.GetAsync();
+
+			AgentSoftwareChannel? Channel = Instance.Channels.FirstOrDefault(x => x.Name == Name);
+			if (Channel == null)
+			{
+				return null;
+			}
+			else
+			{
+				return await Collection.GetAsync(Channel.Version);
+			}
+		}
+
+		/// <summary>
+		/// Gets the zip file with a given version number
+		/// </summary>
+		/// <param name="Version">The version</param>
+		/// <returns>Data for the given archive</returns>
+		public Task<byte[]?> GetArchiveAsync(string Version)
+		{
+			return Collection.GetAsync(Version);
+		}
+		async Task RegisterDefaultAgent(int DelayMs)
+		{
+			await Task.Delay(DelayMs);
+
+			// Check whether we have an installed agent zip
+			FileReference AgentZip = FileReference.Combine(Program.AppDir, "DefaultAgent/Agent.zip");
+			if (!Settings.CurrentValue.SingleInstance || !FileReference.Exists(AgentZip))
+			{
+				return;
+			}
+
+			Logger.LogInformation("Checking for default agent software update");
+
+			string AgentHash = ContentHash.MD5(AgentZip).ToString();
+
+			FileReference AgentHashFile = FileReference.Combine(Program.DataDir, "Agent/DefaultAgentHash");
+
+			try
+			{
+
+				if (FileReference.Exists(AgentHashFile) && FileReference.ReadAllText(AgentHashFile).Trim() == AgentHash)
+				{
+					Logger.LogInformation("Default agent software is up to date");
+					return;
+				}
+			}
+			catch (Exception Ex)
+			{
+				Logger.LogError(Ex, "Error checking default agent software, {Message}", Ex.Message);
+			}
+
+            byte[] Bytes = await File.ReadAllBytesAsync(AgentZip.ToString());
+
+			using X509Certificate2? GrpcCertificate = Program.ReadGrpcCertificate(Settings.CurrentValue);
+
+			if (GrpcCertificate == null)
+			{
+				throw new Exception("Unable to register default agent without valid grpc certicate");
+			}
+
+			// construct agent app settings
+			
+			string ProfileName = "Default";
+
+			string ServerUrl = $"https://{System.Net.Dns.GetHostName()}:{Settings.CurrentValue.HttpsPort}";			
+
+			Dictionary<string, object> DefaultProfile = new Dictionary<string, object>() { { "Name", ProfileName }, { "Environment", "prod" }, { "Thumbprint", GrpcCertificate.Thumbprint }, { "Url", ServerUrl } };
+
+			List<object> ServerProfiles = new List<object>() { DefaultProfile };
+
+			Dictionary<string, object> AgentSettings = new Dictionary<string, object>() { { "ServerProfiles", ServerProfiles }, { "Server", ProfileName }};
+
+			Bytes = AgentUtilities.UpdateAppSettings(Bytes, AgentSettings);
+
+			string Version = await SetArchiveAsync(new AgentSoftwareChannelName("default"), null, Bytes);
+
+			if (!DirectoryReference.Exists(AgentHashFile.Directory))
+			{
+				DirectoryReference.CreateDirectory(AgentHashFile.Directory);
+			}
+
+			FileReference.WriteAllText(AgentHashFile, AgentHash);
+
+			Logger.LogInformation($"Updated default agent software to {AgentHash}");
+
+		}
+	}
+}
diff --git a/Engine/Source/Programs/Horde/HordeServer/Agents/Impl/Tasks/ConformTaskSource.cs b/Engine/Source/Programs/Horde/HordeServer/Agents/Impl/Tasks/ConformTaskSource.cs
new file mode 100644
--- /dev/null
+++ b/Engine/Source/Programs/Horde/HordeServer/Agents/Impl/Tasks/ConformTaskSource.cs
@@ -0,0 +1,384 @@
+// Copyright Epic Games, Inc. All Rights Reserved.
+
+using Google.Protobuf;
+using Google.Protobuf.Reflection;
+using Google.Protobuf.WellKnownTypes;
+using HordeCommon;
+using HordeCommon.Rpc.Tasks;
+using HordeServer.Api;
+using HordeServer.Collections;
+using HordeServer.Models;
+using HordeServer.Services;
+using HordeServer.Utilities;
+using Microsoft.Extensions.Hosting;
+using Microsoft.Extensions.Logging;
+using MongoDB.Bson;
+using System;
+using System.Collections.Generic;
+using System.Linq;
+using System.Threading;
+using System.Threading.Tasks;
+
+namespace HordeServer.Tasks.Impl
+{
+	/// <summary>
+	/// Generates tasks telling agents to sync their workspaces
+	/// </summary>
+	public sealed class ConformTaskSource : TaskSourceBase<ConformTask>, IHostedService, IDisposable
+	{
+		/// <inheritdoc/>
+		public override string Type => "Conform";
+
+		DatabaseService DatabaseService;
+		IAgentCollection AgentCollection;
+		PoolService PoolService;
+		SingletonDocument<ConformList> ConformList;
+		PerforceLoadBalancer PerforceLoadBalancer;
+		ILogFileService LogService;
+		ILogger Logger;
+		ElectedTick TickConformList;
+
+		/// <summary>
+		/// Constructor
+		/// </summary>
+		/// <param name="DatabaseService"></param>
+		/// <param name="AgentCollection"></param>
+		/// <param name="PoolService"></param>
+		/// <param name="LogService"></param>
+		/// <param name="PerforceLoadBalancer"></param>
+		/// <param name="Logger"></param>
+		public ConformTaskSource(DatabaseService DatabaseService, IAgentCollection AgentCollection, PoolService PoolService, ILogFileService LogService, PerforceLoadBalancer PerforceLoadBalancer, ILogger<ConformTaskSource> Logger)
+		{
+			this.DatabaseService = DatabaseService;
+			this.AgentCollection = AgentCollection;
+			this.PoolService = PoolService;
+			this.ConformList = new SingletonDocument<ConformList>(DatabaseService);
+			this.PerforceLoadBalancer = PerforceLoadBalancer;
+			this.LogService = LogService;
+			this.Logger = Logger;
+			this.TickConformList = new ElectedTick(DatabaseService, new ObjectId("60afc5cf555a9a76aff0a50c"), CleanConformListAsync, TimeSpan.FromMinutes(1.0), Logger);
+		}
+
+		/// <inheritdoc/>
+		public void Dispose()
+		{
+			TickConformList.Dispose();
+		}
+
+		/// <inheritdoc/>
+		public Task StartAsync(CancellationToken CancellationToken)
+		{
+			TickConformList.Start();
+			return Task.CompletedTask;
+		}
+
+		/// <inheritdoc/>
+		public async Task StopAsync(CancellationToken CancellationToken)
+		{
+			await TickConformList.StopAsync();
+		}
+
+		/// <summary>
+		/// Clean up the conform list of any outdated entries
+		/// </summary>
+		/// <param name="CancellationToken">Cancellation token</param>
+		/// <returns>Async task</returns>
+		async Task CleanConformListAsync(CancellationToken CancellationToken)
+		{
+			DateTime UtcNow = DateTime.UtcNow;
+			DateTime LastCheckTimeUtc = UtcNow - TimeSpan.FromMinutes(30.0);
+
+			// Get the current state of the conform list
+			ConformList List = await ConformList.GetAsync();
+
+			// Update any leases that are older than LastCheckTimeUtc
+			Dictionary<ObjectId, bool> RemoveLeases = new Dictionary<ObjectId, bool>();
+			foreach (ConformListEntry Entry in Enumerable.Concat(List.Entries, List.Servers.SelectMany(x => x.Entries)))
+			{
+				if (Entry.LastCheckTimeUtc < LastCheckTimeUtc)
+				{
+					IAgent? Agent = await AgentCollection.GetAsync(Entry.AgentId);
+
+					bool Remove = false;
+					if (Agent == null || !Agent.Leases.Any(x => x.Id == Entry.LeaseId))
+					{
+						Logger.LogWarning("Removing invalid lease from conform list: {AgentId} lease {LeaseId}", Entry.AgentId, Entry.LeaseId);
+						Remove = true;
+					}
+
+					RemoveLeases[Entry.LeaseId] = Remove;
+				}
+			}
+
+			// If there's anything to change, update the list
+			if (RemoveLeases.Count > 0)
+			{
+				await ConformList.UpdateAsync(List => UpdateConformList(List, UtcNow, RemoveLeases));
+			}
+		}
+
+		/// <summary>
+		/// Remove items from the conform list, and update timestamps for them
+		/// </summary>
+		/// <param name="List">The list to update</param>
+		/// <param name="UtcNow">Current time</param>
+		/// <param name="RemoveLeases">List of leases to update. Entries with values set to true will be removed, entries with values set to false will have their timestamp updated.</param>
+		static void UpdateConformList(ConformList List, DateTime UtcNow, Dictionary<ObjectId, bool> RemoveLeases)
+		{
+			UpdateConformList(List.Entries, UtcNow, RemoveLeases);
+			foreach (ConformListServer Server in List.Servers)
+			{
+				UpdateConformList(Server.Entries, UtcNow, RemoveLeases);
+			}
+		}
+
+		/// <summary>
+		/// Remove items from the conform list, and update timestamps for them
+		/// </summary>
+		/// <param name="Entries">The list to update</param>
+		/// <param name="UtcNow">Current time</param>
+		/// <param name="RemoveLeases">List of leases to update. Entries with values set to true will be removed, entries with values set to false will have their timestamp updated.</param>
+		static void UpdateConformList(List<ConformListEntry> Entries, DateTime UtcNow, Dictionary<ObjectId, bool> RemoveLeases)
+		{
+			for (int Idx = 0; Idx < Entries.Count; Idx++)
+			{
+				bool Remove;
+				if (RemoveLeases.TryGetValue(Entries[Idx].LeaseId, out Remove))
+				{
+					if (Remove)
+					{
+						Entries.RemoveAt(Idx--);
+					}
+					else if (Entries[Idx].LastCheckTimeUtc < UtcNow)
+					{
+						Entries[Idx].LastCheckTimeUtc = UtcNow;
+					}
+				}
+			}
+		}
+
+		/// <inheritdoc/>
+		public override async Task<AgentLease?> AssignLeaseAsync(IAgent Agent, CancellationToken CancellationToken)
+		{
+			DateTime UtcNow = DateTime.UtcNow;
+			if (!await IsConformPendingAsync(Agent, UtcNow))
+			{
+				return null;
+			}
+
+			if (Agent.Leases.Count == 0)
+			{
+				ConformTask Task = new ConformTask();
+				if (await GetWorkspacesAsync(Agent, Task.Workspaces))
+				{
+					ObjectId LeaseId = ObjectId.GenerateNewId();
+					if (!await AllocateConformLeaseAsync(Agent.Id, Task.Workspaces, LeaseId))
+					{
+						return null;
+					}
+
+					ILogFile Log = await LogService.CreateLogFileAsync(ObjectId.Empty, Agent.SessionId, LogType.Json);
+					Task.LogId = Log.Id.ToString();
+
+					byte[] Payload = Any.Pack(Task).ToByteArray();
+
+					return new AgentLease(LeaseId, "Updating workspaces", null, null, Log.Id, LeaseState.Pending, null, true, Payload);
+				}
+			}
+
+			if (Agent.RequestConform)
+			{
+				return AgentLease.Drain;
+			}
+
+			return null;
+		}
+
+		/// <inheritdoc/>
+		public async Task<bool> GetWorkspacesAsync(IAgent Agent, IList<HordeCommon.Rpc.Messages.AgentWorkspace> Workspaces)
+		{
+			Globals Globals = await DatabaseService.GetGlobalsAsync();
+
+			HashSet<AgentWorkspace> ConformWorkspaces = await PoolService.GetWorkspacesAsync(Agent, DateTime.UtcNow);
+			foreach (AgentWorkspace ConformWorkspace in ConformWorkspaces)
+			{
+				PerforceCluster? Cluster = Globals.FindPerforceCluster(ConformWorkspace.Cluster);
+				if (Cluster == null || !await Agent.TryAddWorkspaceMessage(ConformWorkspace, Cluster, PerforceLoadBalancer, Workspaces))
+				{
+					return false;
+				}
+			}
+
+			return true;
+		}
+
+		/// <inheritdoc/>
+		public override Task OnLeaseFinishedAsync(IAgent Agent, ObjectId LeaseId, ConformTask Payload, LeaseOutcome Outcome, ReadOnlyMemory<byte> Output, ILogger Logger)
+		{
+			return ReleaseConformLeaseAsync(LeaseId);
+		}
+
+		/// <summary>
+		/// Atempt to allocate a conform resource for the given lease
+		/// </summary>
+		/// <param name="AgentId">The agent id</param>
+		/// <param name="Workspaces">List of workspaces that are required</param>
+		/// <param name="LeaseId">The lease id</param>
+		/// <returns>True if the resource was allocated, false otherwise</returns>
+		private async Task<bool> AllocateConformLeaseAsync(AgentId AgentId, IEnumerable<HordeCommon.Rpc.Messages.AgentWorkspace> Workspaces, ObjectId LeaseId)
+		{
+			Globals Globals = await DatabaseService.GetGlobalsAsync();
+			for (; ; )
+			{
+				ConformList CurrentValue = await ConformList.GetAsync();
+				if (Globals.MaxConformCount != 0 && CurrentValue.Entries.Count + CurrentValue.Servers.Sum(x => x.Entries.Count) >= Globals.MaxConformCount)
+				{
+					return false;
+				}
+
+				HashSet<string> Servers = new HashSet<string>(StringComparer.OrdinalIgnoreCase);
+				foreach (HordeCommon.Rpc.Messages.AgentWorkspace Workspace in Workspaces)
+				{
+					if (Servers.Add(Workspace.ServerAndPort))
+					{
+						PerforceCluster? Cluster = Globals.FindPerforceCluster(Workspace.Cluster);
+						if (Cluster == null)
+						{
+							Logger.LogWarning("Unable to find perforce cluster '{Cluster}' for conform", Workspace.Cluster);
+							return false;
+						}
+
+						PerforceServer? Server = Cluster.Servers.FirstOrDefault(x => x.ServerAndPort.Equals(Workspace.BaseServerAndPort, StringComparison.Ordinal));
+						if (Server == null)
+						{
+							Logger.LogWarning("Unable to find perforce server '{Server}' in cluster '{Cluster}' for conform", Workspace.BaseServerAndPort, Workspace.Cluster);
+							return false;
+						}
+
+						ConformListServer? ConformServer = CurrentValue.Servers.FirstOrDefault(x => x.Cluster.Equals(Workspace.Cluster, StringComparison.Ordinal) && x.ServerAndPort.Equals(Workspace.ServerAndPort, StringComparison.Ordinal));
+						if (ConformServer == null)
+						{
+							ConformServer = new ConformListServer { Cluster = Workspace.Cluster, ServerAndPort = Workspace.ServerAndPort };
+							CurrentValue.Servers.Add(ConformServer);
+						}
+
+						if (Server.MaxConformCount != 0 && ConformServer.Entries.Count > Server.MaxConformCount)
+						{
+							return false;
+						}
+
+						ConformListEntry Entry = new ConformListEntry();
+						Entry.AgentId = AgentId;
+						Entry.LeaseId = LeaseId;
+						Entry.LastCheckTimeUtc = DateTime.UtcNow;
+						ConformServer.Entries.Add(Entry);
+					}
+				}
+
+				if (await ConformList.TryUpdateAsync(CurrentValue))
+				{
+					Logger.LogInformation("Added conform lease {LeaseId}", LeaseId);
+					return true;
+				}
+			}
+		}
+
+		/// <summary>
+		/// Terminate a conform lease
+		/// </summary>
+		/// <param name="LeaseId">The lease id</param>
+		/// <returns>Async task</returns>
+		public async Task ReleaseConformLeaseAsync(ObjectId LeaseId)
+		{
+			for (; ; )
+			{
+				ConformList CurrentValue = await ConformList.GetAsync();
+				if (CurrentValue.Entries.RemoveAll(x => x.LeaseId == LeaseId) + CurrentValue.Servers.Sum(x => x.Entries.RemoveAll(x => x.LeaseId == LeaseId)) == 0)
+				{
+					Logger.LogInformation("Conform lease {LeaseId} is not in singelton", LeaseId);
+					break;
+				}
+				if (await ConformList.TryUpdateAsync(CurrentValue))
+				{
+					Logger.LogInformation("Removed conform lease {LeaseId}", LeaseId);
+					break;
+				}
+			}
+		}
+
+		/// <summary>
+		/// Determine if an agent should be conformed
+		/// </summary>
+		/// <param name="Agent">The agent to test</param>
+		/// <param name="UtcNow">Current time</param>
+		/// <returns>True if the agent should be conformed</returns>
+		private async Task<bool> IsConformPendingAsync(IAgent Agent, DateTime UtcNow)
+		{
+			// If a conform was manually requested, allow it to run even if the agent is disabled
+			if (Agent.RequestConform)
+			{
+				return !IsConformCoolDownPeriod(Agent, UtcNow);
+			}
+
+			// Otherwise only run if the agent is enabled
+			if (Agent.Enabled)
+			{
+				// If we've attempted (and failed) a conform, run again after a certain time
+				if (Agent.ConformAttemptCount.HasValue)
+				{
+					return !IsConformCoolDownPeriod(Agent, UtcNow);
+				}
+
+				// Always run a conform every 24h
+				if (UtcNow > Agent.LastConformTime + TimeSpan.FromDays(1.0))
+				{
+					return true;
+				}
+
+				// Check if the workspaces have changed (first check against a cached list of workspaces, then an accurate one)
+				HashSet<AgentWorkspace> Workspaces = await PoolService.GetWorkspacesAsync(Agent, UtcNow - TimeSpan.FromSeconds(30.0));
+				if (!Workspaces.SetEquals(Agent.Workspaces))
+				{
+					Workspaces = await PoolService.GetWorkspacesAsync(Agent, UtcNow);
+					if (!Workspaces.SetEquals(Agent.Workspaces))
+					{
+						return true;
+					}
+				}
+			}
+			return false;
+		}
+
+		/// <summary>
+		/// Determines if the given agent is in a cooldown period, waiting to retry after a failed conform
+		/// </summary>
+		/// <param name="Agent">The agent to test</param>
+		/// <param name="UtcNow">Current time</param>
+		/// <returns></returns>
+		private static bool IsConformCoolDownPeriod(IAgent Agent, DateTime UtcNow)
+		{
+			if (!Agent.ConformAttemptCount.HasValue)
+			{
+				return false;
+			}
+
+			TimeSpan RetryTime;
+			switch (Agent.ConformAttemptCount.Value)
+			{
+				case 1:
+					RetryTime = TimeSpan.FromMinutes(5.0);
+					break;
+				case 2:
+					RetryTime = TimeSpan.FromMinutes(20.0);
+					break;
+				case 3:
+					RetryTime = TimeSpan.FromHours(1.0);
+					break;
+				default:
+					RetryTime = TimeSpan.FromHours(6.0);
+					break;
+			}
+			return UtcNow < Agent.LastConformTime + RetryTime;
+		}
+	}
+}
diff --git a/Engine/Source/Programs/Horde/HordeServer/Agents/Impl/Tasks/RestartTaskSource.cs b/Engine/Source/Programs/Horde/HordeServer/Agents/Impl/Tasks/RestartTaskSource.cs
new file mode 100644
--- /dev/null
+++ b/Engine/Source/Programs/Horde/HordeServer/Agents/Impl/Tasks/RestartTaskSource.cs
@@ -0,0 +1,52 @@
+// Copyright Epic Games, Inc. All Rights Reserved.
+
+using Google.Protobuf;
+using Google.Protobuf.Reflection;
+using Google.Protobuf.WellKnownTypes;
+using HordeCommon;
+using HordeCommon.Rpc.Tasks;
+using HordeServer.Api;
+using HordeServer.Models;
+using HordeServer.Services;
+using MongoDB.Bson;
+using System;
+using System.Collections.Generic;
+using System.Linq;
+using System.Threading;
+using System.Threading.Tasks;
+
+namespace HordeServer.Tasks.Impl
+{
+	class RestartTaskSource : TaskSourceBase<RestartTask>
+	{
+		public override string Type => "Restart";
+
+		ILogFileService LogService;
+
+		public RestartTaskSource(ILogFileService LogService)
+		{
+			this.LogService = LogService;
+		}
+
+		public override async Task<AgentLease?> AssignLeaseAsync(IAgent Agent, CancellationToken CancellationToken)
+		{
+			if (!Agent.RequestRestart)
+			{
+				return null;
+			}
+			if (Agent.Leases.Count > 0)
+			{
+				return AgentLease.Drain;
+			}
+
+			ILogFile Log = await LogService.CreateLogFileAsync(ObjectId.Empty, Agent.SessionId, LogType.Json);
+
+			RestartTask Task = new RestartTask();
+			Task.LogId = Log.Id.ToString();
+
+			byte[] Payload = Any.Pack(Task).ToByteArray();
+
+			return new AgentLease(ObjectId.GenerateNewId(), "Restart", null, null, Log.Id, LeaseState.Pending, null, true, Payload);
+		}
+	}
+}
diff --git a/Engine/Source/Programs/Horde/HordeServer/Agents/Impl/Tasks/ShutdownTaskSource.cs b/Engine/Source/Programs/Horde/HordeServer/Agents/Impl/Tasks/ShutdownTaskSource.cs
new file mode 100644
--- /dev/null
+++ b/Engine/Source/Programs/Horde/HordeServer/Agents/Impl/Tasks/ShutdownTaskSource.cs
@@ -0,0 +1,53 @@
+// Copyright Epic Games, Inc. All Rights Reserved.
+
+using EpicGames.Core;
+using Google.Protobuf;
+using Google.Protobuf.Reflection;
+using Google.Protobuf.WellKnownTypes;
+using HordeCommon;
+using HordeCommon.Rpc.Tasks;
+using HordeServer.Api;
+using HordeServer.Models;
+using HordeServer.Services;
+using MongoDB.Bson;
+using System;
+using System.Collections.Generic;
+using System.Linq;
+using System.Threading;
+using System.Threading.Tasks;
+
+namespace HordeServer.Tasks.Impl
+{
+	class ShutdownTaskSource : TaskSourceBase<ShutdownTask>
+	{
+		public override string Type => "Shutdown";
+
+		ILogFileService LogService;
+
+		public ShutdownTaskSource(ILogFileService LogService)
+		{
+			this.LogService = LogService;
+		}
+
+		public override async Task<AgentLease?> AssignLeaseAsync(IAgent Agent, CancellationToken CancellationToken)
+		{
+			if (!Agent.RequestShutdown)
+			{
+				return null;
+			}
+			if (Agent.Leases.Count > 0)
+			{
+				return AgentLease.Drain;
+			}
+
+			ILogFile Log = await LogService.CreateLogFileAsync(ObjectId.Empty, Agent.SessionId, LogType.Json);
+
+			ShutdownTask Task = new ShutdownTask();
+			Task.LogId = Log.Id.ToString();
+
+			byte[] Payload = Any.Pack(Task).ToByteArray();
+
+			return new AgentLease(ObjectId.GenerateNewId(), "Shutdown", null, null, Log.Id, LeaseState.Pending, null, true, Payload);
+		}
+	}
+}
diff --git a/Engine/Source/Programs/Horde/HordeServer/Agents/Impl/Tasks/UpgradeTaskSource.cs b/Engine/Source/Programs/Horde/HordeServer/Agents/Impl/Tasks/UpgradeTaskSource.cs
new file mode 100644
--- /dev/null
+++ b/Engine/Source/Programs/Horde/HordeServer/Agents/Impl/Tasks/UpgradeTaskSource.cs
@@ -0,0 +1,84 @@
+// Copyright Epic Games, Inc. All Rights Reserved.
+
+using Google.Protobuf;
+using Google.Protobuf.Reflection;
+using Google.Protobuf.WellKnownTypes;
+using HordeCommon;
+using HordeCommon.Rpc.Tasks;
+using HordeServer.Api;
+using HordeServer.Models;
+using HordeServer.Services;
+using HordeServer.Utilities;
+using MongoDB.Bson;
+using System;
+using System.Collections.Generic;
+using System.Linq;
+using System.Threading;
+using System.Threading.Tasks;
+
+namespace HordeServer.Tasks.Impl
+{
+	using AgentSoftwareChannelName = StringId<AgentSoftwareChannels>;
+
+	class UpgradeTaskSource : TaskSourceBase<UpgradeTask>
+	{
+		public override string Type => "Upgrade";
+
+		AgentSoftwareService AgentSoftwareService;
+		ILogFileService LogService;
+		IClock Clock;
+
+		public UpgradeTaskSource(AgentSoftwareService AgentSoftwareService, ILogFileService LogService, IClock Clock)
+		{
+			this.AgentSoftwareService = AgentSoftwareService;
+			this.LogService = LogService;
+			this.Clock = Clock;
+		}
+
+		public async override Task<AgentLease?> AssignLeaseAsync(IAgent Agent, CancellationToken CancellationToken)
+		{
+			string? RequiredVersion = await GetRequiredSoftwareVersion(Agent);
+			if (RequiredVersion != null && Agent.Version != RequiredVersion)
+			{
+				AgentLease? Lease = null;
+				if (Agent.Leases.Count == 0 && (Agent.LastUpgradeTime == null || Agent.LastUpgradeTime.Value + TimeSpan.FromMinutes(5.0) < Clock.UtcNow || Agent.LastUpgradeVersion != RequiredVersion.ToString()))
+				{
+					ILogFile LogFile = await LogService.CreateLogFileAsync(ObjectId.Empty, Agent.SessionId, LogType.Json);
+
+					UpgradeTask Task = new UpgradeTask();
+					Task.SoftwareId = RequiredVersion.ToString();
+					Task.LogId = LogFile.Id.ToString();
+
+					byte[] Payload;
+					if (Agent.Version == "5.0.0-17425336" || Agent.Version == "5.0.0-17448746")
+					{
+						Any Any = new Any();
+						Any.TypeUrl = "type.googleapis.com/Horde.UpgradeTask";
+						Any.Value = Task.ToByteString();
+						Payload = Any.ToByteArray();
+					}
+					else
+					{
+						Payload = Any.Pack(Task).ToByteArray();
+					}
+
+					Lease = new AgentLease(ObjectId.GenerateNewId(), $"Upgrade to {RequiredVersion}", null, null, LogFile.Id, LeaseState.Pending, null, true, Payload);
+				}
+				return Lease;
+			}
+			return null;
+		}
+
+		/// <summary>
+		/// Determines the client software version that should be installed on an agent
+		/// </summary>
+		/// <param name="Agent">The agent instance</param>
+		/// <returns>Unique id of the client version this agent should be running</returns>
+		public async Task<string?> GetRequiredSoftwareVersion(IAgent Agent)
+		{
+			AgentSoftwareChannelName ChannelName = Agent.Channel ?? AgentSoftwareService.DefaultChannelName;
+			IAgentSoftwareChannel? Channel = await AgentSoftwareService.GetCachedChannelAsync(ChannelName);
+			return Channel?.Version;
+		}
+	}
+}
diff --git a/Engine/Source/Programs/Horde/HordeServer/Agents/Impl/UpgradeService.cs b/Engine/Source/Programs/Horde/HordeServer/Agents/Impl/UpgradeService.cs
new file mode 100644
--- /dev/null
+++ b/Engine/Source/Programs/Horde/HordeServer/Agents/Impl/UpgradeService.cs
@@ -0,0 +1,110 @@
+// Copyright Epic Games, Inc. All Rights Reserved.
+
+using HordeServer.Collections;
+using HordeServer.Collections.Impl;
+using HordeServer.Models;
+using Microsoft.Extensions.DependencyInjection;
+using Microsoft.Extensions.Logging;
+using MongoDB.Driver;
+using Serilog.Events;
+using System;
+using System.Collections.Generic;
+using System.Linq;
+using System.Threading.Tasks;
+
+namespace HordeServer.Services
+{
+	/// <summary>
+	/// Service to perform schema upgrades on the database
+	/// </summary>
+	public class UpgradeService
+	{
+		/// <summary>
+		/// The current schema version
+		/// </summary>
+		public const int LatestSchemaVersion = 5;
+
+		/// <summary>
+		/// The database service instance
+		/// </summary>
+		DatabaseService DatabaseService { get; set; }
+
+		/// <summary>
+		/// The DI service provider
+		/// </summary>
+		IServiceProvider ServiceProvider;
+
+		/// <summary>
+		/// Logger instance
+		/// </summary>
+		ILogger<UpgradeService> Logger;
+
+		/// <summary>
+		/// Constructor
+		/// </summary>
+		/// <param name="DatabaseService">The database service singleton</param>
+		/// <param name="ServiceProvider">The DI service provider</param>
+		/// <param name="Logger">Logger instance</param>
+		public UpgradeService(DatabaseService DatabaseService, IServiceProvider ServiceProvider, ILogger<UpgradeService> Logger)
+		{
+			this.DatabaseService = DatabaseService;
+			this.ServiceProvider = ServiceProvider;
+			this.Logger = Logger;
+		}
+
+		/// <summary>
+		/// Update the database schema
+		/// </summary>
+		/// <param name="FromVersion">The version to upgrade from. Automatically detected if this is null.</param>
+		/// <returns>Async task</returns>
+		public async Task UpgradeSchemaAsync(int? FromVersion)
+		{
+			Globals Globals = await DatabaseService.GetGlobalsAsync();
+			int SchemaVersion = FromVersion ?? Globals.SchemaVersion ?? 0;
+
+			while (SchemaVersion < LatestSchemaVersion)
+			{
+				Logger.LogInformation("Upgrading from schema version {Version}", SchemaVersion);
+
+				// Handle the current version
+				//				if (SchemaVersion == 0)
+				//				{
+				//					IssueCollection IssueCollection = new IssueCollection(DatabaseService);
+				//					IJobCollection JobCollection = ServiceProvider.GetService<IJobCollection>();
+				//					await IssueCollection.UpdateSchemaLogIdFields(JobCollection, Logger);
+				//				}
+				//				if(SchemaVersion == 2)
+				//				{
+				//					IssueCollection IssueCollection = new IssueCollection(DatabaseService);
+				//					await IssueCollection.FixResolvedIssuesAsync(Logger);
+				//				}
+
+				if (SchemaVersion == 3)
+				{
+					IJobCollection JobCollection = ServiceProvider.GetRequiredService<IJobCollection>();
+					await JobCollection.UpgradeDocumentsAsync();
+				}
+				if (SchemaVersion == 4)
+				{
+					UserCollectionV1 UserCollectionV1 = new UserCollectionV1(ServiceProvider.GetRequiredService<DatabaseService>());
+					using UserCollectionV2 UserCollectionV2 = new UserCollectionV2(ServiceProvider.GetRequiredService<DatabaseService>(), ServiceProvider.GetRequiredService<ILogger<UserCollectionV2>>());
+					await UserCollectionV2.ResaveDocumentsAsync(UserCollectionV1);
+				}
+
+				// Increment the version number
+				SchemaVersion++;
+
+				// Try to update the current schema version number
+				while (Globals.SchemaVersion == null || Globals.SchemaVersion < SchemaVersion)
+				{
+					Globals.SchemaVersion = SchemaVersion;
+					if (await DatabaseService.TryUpdateSingletonAsync(Globals))
+					{
+						break;
+					}
+					Globals = await DatabaseService.GetGlobalsAsync();
+				}
+			}
+		}
+	}
+}
diff --git a/Engine/Source/Programs/Horde/HordeServer/Fleet/IPoolCollection.cs b/Engine/Source/Programs/Horde/HordeServer/Fleet/IPoolCollection.cs
new file mode 100644
--- /dev/null
+++ b/Engine/Source/Programs/Horde/HordeServer/Fleet/IPoolCollection.cs
@@ -0,0 +1,81 @@
+// Copyright Epic Games, Inc. All Rights Reserved.
+
+using EpicGames.Horde.Common;
+using HordeServer.Models;
+using HordeServer.Services;
+using HordeServer.Utilities;
+using MongoDB.Bson;
+using MongoDB.Bson.Serialization;
+using MongoDB.Bson.Serialization.Attributes;
+using MongoDB.Driver;
+using System;
+using System.Collections.Generic;
+using System.Linq;
+using System.Linq.Expressions;
+using System.Threading.Tasks;
+
+namespace HordeServer.Collections
+{
+	using PoolId = StringId<IPool>;
+
+	/// <summary>
+	/// Collection of pool documents
+	/// </summary>
+	public interface IPoolCollection
+	{
+		/// <summary>
+		/// Creates a new pool
+		/// </summary>
+		/// <param name="Id">Unique id for the new pool</param>
+		/// <param name="Name">Name of the new pool</param>
+		/// <param name="Condition">Condition for agents to be included in this pool</param>
+		/// <param name="EnableAutoscaling">Whether to enable autoscaling for this pool</param>
+		/// <param name="MinAgents">Minimum number of agents in the pool</param>
+		/// <param name="NumReserveAgents">Minimum number of idle agents to maintain</param>
+		/// <param name="Properties">Properties for the pool</param>
+		/// <returns>The new pool document</returns>
+		Task<IPool> AddAsync(PoolId Id, string Name, Condition? Condition = null, bool? EnableAutoscaling = null, int? MinAgents = null, int? NumReserveAgents = null, IEnumerable<KeyValuePair<string, string>>? Properties = null);
+
+		/// <summary>
+		/// Enumerates all the pools
+		/// </summary>
+		/// <returns>The pool documents</returns>
+		Task<List<IPool>> GetAsync();
+
+		/// <summary>
+		/// Gets a pool by ID
+		/// </summary>
+		/// <param name="Id">Unique id of the pool</param>
+		/// <returns>The pool document</returns>
+		Task<IPool?> GetAsync(PoolId Id);
+
+		/// <summary>
+		/// Gets a list of all valid pool ids
+		/// </summary>
+		/// <returns>List of pool ids</returns>
+		Task<List<PoolId>> GetPoolIdsAsync();
+
+		/// <summary>
+		/// Gets a pool by ID
+		/// </summary>
+		/// <param name="Id">Unique id of the pool</param>
+		/// <returns>The pool document</returns>
+		Task<bool> DeleteAsync(PoolId Id);
+
+		/// <summary>
+		/// Updates a pool
+		/// </summary>
+		/// <param name="Pool">The pool to update</param>
+		/// <param name="NewName">New name for the pool</param>
+		/// <param name="NewCondition">New condition for the pool</param>
+		/// <param name="NewEnableAutoscaling">New setting for whether to enable autoscaling</param>
+		/// <param name="NewMinAgents">Minimum number of agents in the pool</param>
+		/// <param name="NewNumReserveAgents">Minimum number of idle agents to maintain</param>
+		/// <param name="NewWorkspaces">New workspaces for the pool</param>
+		/// <param name="NewProperties">New properties for the pool</param>
+		/// <param name="LastScaleUpTime">New time for last (auto) scale up</param>
+		/// <param name="LastScaleDownTime">New time for last (auto) scale down</param>
+		/// <returns>Async task</returns>
+		Task<IPool?> TryUpdateAsync(IPool Pool, string? NewName = null, Condition? NewCondition = null, bool? NewEnableAutoscaling = null, int? NewMinAgents = null, int? NewNumReserveAgents = null, List<AgentWorkspace>? NewWorkspaces = null, Dictionary<string, string?>? NewProperties = null, DateTime? LastScaleUpTime = null, DateTime? LastScaleDownTime = null);
+	}
+}
diff --git a/Engine/Source/Programs/Horde/HordeServer/Fleet/ITelemetryCollection.cs b/Engine/Source/Programs/Horde/HordeServer/Fleet/ITelemetryCollection.cs
new file mode 100644
--- /dev/null
+++ b/Engine/Source/Programs/Horde/HordeServer/Fleet/ITelemetryCollection.cs
@@ -0,0 +1,41 @@
+// Copyright Epic Games, Inc. All Rights Reserved.
+
+using HordeServer.Models;
+using HordeServer.Utilities;
+using System;
+using System.Collections.Generic;
+using System.Linq;
+using System.Threading.Tasks;
+
+namespace HordeServer.Collections
+{
+	using PoolId = StringId<IPool>;
+	using StreamId = StringId<IStream>;
+
+	/// <summary>
+	/// Collection of utilization collection
+	/// </summary>
+	public interface ITelemetryCollection
+	{
+		/// <summary>
+		/// Adds entries for the given utilization
+		/// </summary>
+		/// <param name="Telemetry">Telemetry data to add</param>
+		/// <returns>Async task</returns>
+		Task AddUtilizationTelemetryAsync(IUtilizationTelemetry Telemetry);
+
+		/// <summary>
+		/// Finds utilization data matching the given criteria
+		/// </summary>
+		/// <param name="StartTimeUtc">Start time to query utilization for</param>
+		/// <param name="FinishTimeUtc">Finish time to query utilization for</param>
+		/// <returns>The utilization data</returns>
+		Task<List<IUtilizationTelemetry>> GetUtilizationTelemetryAsync(DateTime StartTimeUtc, DateTime FinishTimeUtc);
+
+		/// <summary>
+		/// Finds the latest utilization data
+		/// </summary>
+		/// <returns>The utilization data</returns>
+		Task<IUtilizationTelemetry?> GetLatestUtilizationTelemetryAsync();
+	}
+}
diff --git a/Engine/Source/Programs/Horde/HordeServer/Fleet/Impl/PoolCollection.cs b/Engine/Source/Programs/Horde/HordeServer/Fleet/Impl/PoolCollection.cs
new file mode 100644
--- /dev/null
+++ b/Engine/Source/Programs/Horde/HordeServer/Fleet/Impl/PoolCollection.cs
@@ -0,0 +1,244 @@
+// Copyright Epic Games, Inc. All Rights Reserved.
+
+using EpicGames.Horde.Common;
+using HordeServer.Models;
+using HordeServer.Services;
+using HordeServer.Utilities;
+using MongoDB.Bson;
+using MongoDB.Bson.Serialization;
+using MongoDB.Bson.Serialization.Attributes;
+using MongoDB.Driver;
+using System;
+using System.Collections.Generic;
+using System.Linq;
+using System.Linq.Expressions;
+using System.Threading.Tasks;
+
+namespace HordeServer.Collections.Impl
+{
+	using PoolId = StringId<IPool>;
+
+	/// <summary>
+	/// Collection of pool documents
+	/// </summary>
+	public class PoolCollection : IPoolCollection
+	{
+		class PoolDocument : IPool
+		{
+			// Properties
+			[BsonId]
+			public PoolId Id { get; set; }
+			[BsonRequired]
+			public string Name { get; set; } = null!;
+			[BsonIgnoreIfNull]
+			public Condition? Condition { get; set; }
+			public List<AgentWorkspace> Workspaces { get; set; } = new List<AgentWorkspace>();
+			public Dictionary<string, string> Properties { get; set; } = new Dictionary<string, string>();
+			[BsonIgnoreIfDefault(true)]
+			public bool EnableAutoscaling { get; set; } = true;
+			public int? MinAgents { get; set; }
+			public int? NumReserveAgents { get; set; }
+			public DateTime? LastScaleUpTime { get; set; }
+			public DateTime? LastScaleDownTime { get; set; }
+			public int UpdateIndex { get; set; }
+
+			// Read-only wrappers
+			IReadOnlyList<AgentWorkspace> IPool.Workspaces => Workspaces;
+			IReadOnlyDictionary<string, string> IPool.Properties => Properties;
+
+			public PoolDocument()
+			{
+			}
+
+			public PoolDocument(IPool Other)
+			{
+				Id = Other.Id;
+				Name = Other.Name;
+				Condition = Other.Condition;
+				Workspaces.AddRange(Other.Workspaces);
+				Properties = new Dictionary<string, string>(Other.Properties);
+				EnableAutoscaling = Other.EnableAutoscaling;
+				MinAgents = Other.MinAgents;
+				NumReserveAgents = Other.NumReserveAgents;
+				LastScaleUpTime = Other.LastScaleUpTime;
+				LastScaleDownTime = Other.LastScaleDownTime;
+				UpdateIndex = Other.UpdateIndex;
+			}
+		}
+
+		/// <summary>
+		/// Collection of pool documents
+		/// </summary>
+		IMongoCollection<PoolDocument> Pools;
+
+		/// <summary>
+		/// Constructor
+		/// </summary>
+		/// <param name="DatabaseService">The database service instance</param>
+		public PoolCollection(DatabaseService DatabaseService)
+		{
+			Pools = DatabaseService.GetCollection<PoolDocument>("Pools");
+		}
+
+		/// <inheritdoc/>
+		public async Task<IPool> AddAsync(PoolId Id, string Name, Condition? Condition, bool? EnableAutoscaling, int? MinAgents, int? NumReserveAgents, IEnumerable<KeyValuePair<string, string>>? Properties)
+		{
+			PoolDocument Pool = new PoolDocument();
+			Pool.Id = Id;
+			Pool.Name = Name;
+			Pool.Condition = Condition;
+			if (EnableAutoscaling != null)
+			{
+				Pool.EnableAutoscaling = EnableAutoscaling.Value;
+			}
+			Pool.MinAgents = MinAgents;
+			Pool.NumReserveAgents = NumReserveAgents;
+			if (Properties != null)
+			{
+				Pool.Properties = new Dictionary<string, string>(Properties);
+			}
+			await Pools.InsertOneAsync(Pool);
+			return Pool;
+		}
+
+		/// <inheritdoc/>
+		public async Task<List<IPool>> GetAsync()
+		{
+			List<PoolDocument> Results = await Pools.Find(FilterDefinition<PoolDocument>.Empty).ToListAsync();
+			return Results.ConvertAll<IPool>(x => x);
+		}
+
+		/// <inheritdoc/>
+		public async Task<IPool?> GetAsync(PoolId Id)
+		{
+			return await Pools.Find(x => x.Id == Id).FirstOrDefaultAsync();
+		}
+
+		/// <inheritdoc/>
+		public async Task<List<PoolId>> GetPoolIdsAsync()
+		{
+			ProjectionDefinition<PoolDocument, BsonDocument> Projection = Builders<PoolDocument>.Projection.Include(x => x.Id);
+			List<BsonDocument> Results = await Pools.Find(FilterDefinition<PoolDocument>.Empty).Project(Projection).ToListAsync();
+			return Results.ConvertAll(x => new PoolId(x.GetElement("_id").Value.AsString));
+		}
+
+		/// <inheritdoc/>
+		public async Task<bool> DeleteAsync(PoolId Id)
+		{
+			FilterDefinition<PoolDocument> Filter = Builders<PoolDocument>.Filter.Eq(x => x.Id, Id);
+			DeleteResult Result = await Pools.DeleteOneAsync(Filter);
+			return Result.DeletedCount > 0;
+		}
+
+		/// <summary>
+		/// Attempts to update a pool, upgrading it to the latest document schema if necessary
+		/// </summary>
+		/// <param name="PoolToUpdate">Interface for the document to update</param>
+		/// <param name="Transaction">The transaction to execute</param>
+		/// <returns>New pool definition, or null on failure.</returns>
+		async Task<IPool?> TryUpdateAsync(IPool PoolToUpdate, TransactionBuilder<PoolDocument> Transaction)
+		{
+			if (Transaction.IsEmpty)
+			{
+				return PoolToUpdate;
+			}
+
+			Transaction.Set(x => x.UpdateIndex, PoolToUpdate.UpdateIndex + 1);
+
+			PoolDocument? MutablePool = PoolToUpdate as PoolDocument;
+			if (MutablePool != null)
+			{
+				UpdateResult Result = await Pools.UpdateOneAsync(x => x.Id == PoolToUpdate.Id && x.UpdateIndex == PoolToUpdate.UpdateIndex, Transaction.ToUpdateDefinition());
+				if (Result.ModifiedCount == 0)
+				{
+					return null;
+				}
+
+				Transaction.ApplyTo(MutablePool);
+			}
+			else
+			{
+				MutablePool = new PoolDocument(PoolToUpdate);
+				Transaction.ApplyTo(MutablePool);
+
+				ReplaceOneResult Result = await Pools.ReplaceOneAsync<PoolDocument>(x => x.Id == PoolToUpdate.Id && x.UpdateIndex == PoolToUpdate.UpdateIndex, MutablePool);
+				if (Result.ModifiedCount == 0)
+				{
+					return null;
+				}
+			}
+			return MutablePool;
+		}
+
+		/// <inheritdoc/>
+		public Task<IPool?> TryUpdateAsync(IPool Pool, string? NewName, Condition? NewCondition, bool? NewEnableAutoscaling, int? NewMinAgents, int? NewNumReserveAgents, List<AgentWorkspace>? NewWorkspaces, Dictionary<string, string?>? NewProperties, DateTime? LastScaleUpTime, DateTime? LastScaleDownTime)
+		{
+			TransactionBuilder<PoolDocument> Transaction = new TransactionBuilder<PoolDocument>();
+			if (NewName != null)
+			{
+				Transaction.Set(x => x.Name, NewName);
+			}
+			if (NewCondition != null)
+			{
+				if (NewCondition.IsEmpty())
+				{
+					Transaction.Unset(x => x.Condition!);
+				}
+				else
+				{
+					Transaction.Set(x => x.Condition, NewCondition);
+				}
+			}
+			if (NewEnableAutoscaling != null)
+			{
+				if (NewEnableAutoscaling.Value)
+				{
+					Transaction.Unset(x => x.EnableAutoscaling);
+				}
+				else
+				{
+					Transaction.Set(x => x.EnableAutoscaling, NewEnableAutoscaling.Value);
+				}
+			}
+			if (NewMinAgents != null)
+			{
+				if (NewMinAgents.Value < 0)
+				{
+					Transaction.Unset(x => x.MinAgents!);
+				}
+				else
+				{
+					Transaction.Set(x => x.MinAgents, NewMinAgents.Value);
+				}
+			}
+			if (NewNumReserveAgents != null)
+			{
+				if (NewNumReserveAgents.Value < 0)
+				{
+					Transaction.Unset(x => x.NumReserveAgents!);
+				}
+				else
+				{
+					Transaction.Set(x => x.NumReserveAgents, NewNumReserveAgents.Value);
+				}
+			}
+			if (NewWorkspaces != null)
+			{
+				Transaction.Set(x => x.Workspaces, NewWorkspaces);
+			}
+			if (NewProperties != null)
+			{
+				Transaction.UpdateDictionary(x => x.Properties, NewProperties);
+			}
+			if (LastScaleUpTime != null)
+			{
+				Transaction.Set(x => x.LastScaleUpTime, LastScaleUpTime);
+			}
+			if (LastScaleDownTime != null)
+			{
+				Transaction.Set(x => x.LastScaleDownTime, LastScaleDownTime);
+			}
+			return TryUpdateAsync(Pool, Transaction);
+		}
+	}
+}
diff --git a/Engine/Source/Programs/Horde/HordeServer/Fleet/Impl/TelemetryCollection.cs b/Engine/Source/Programs/Horde/HordeServer/Fleet/Impl/TelemetryCollection.cs
new file mode 100644
--- /dev/null
+++ b/Engine/Source/Programs/Horde/HordeServer/Fleet/Impl/TelemetryCollection.cs
@@ -0,0 +1,136 @@
+// Copyright Epic Games, Inc. All Rights Reserved.
+
+using EpicGames.Core;
+using HordeServer.Models;
+using HordeServer.Services;
+using HordeServer.Utilities;
+using Microsoft.Extensions.Hosting;
+using MongoDB.Bson;
+using MongoDB.Bson.Serialization.Attributes;
+using MongoDB.Driver;
+using System;
+using System.Collections.Generic;
+using System.Linq;
+using System.Threading;
+using System.Threading.Tasks;
+
+namespace HordeServer.Collections.Impl
+{
+	using PoolId = StringId<IPool>;
+	using StreamId = StringId<IStream>;
+
+	/// <summary>
+	/// Collection of utilization data
+	/// </summary>
+	public class TelemetryCollection : ITelemetryCollection
+	{
+		class UtilizationDocument : IUtilizationTelemetry
+		{
+			public DateTime StartTime { get; set; }
+			public DateTime FinishTime { get; set; }
+
+			public int NumAgents { get; set; }
+			public List<PoolUtilizationDocument> Pools { get; set; }
+			IReadOnlyList<IPoolUtilizationTelemetry> IUtilizationTelemetry.Pools => Pools;
+
+			public double AdminTime { get; set; }
+			public int UpdateIndex { get; set; }
+
+			[BsonConstructor]
+			private UtilizationDocument()
+			{
+				Pools = new List<PoolUtilizationDocument>();
+			}
+
+			public UtilizationDocument(IUtilizationTelemetry Other)
+			{
+				this.StartTime = Other.StartTime;
+				this.FinishTime = Other.FinishTime;
+				this.NumAgents = Other.NumAgents;
+				this.Pools = Other.Pools.ConvertAll(x => new PoolUtilizationDocument(x));
+				this.AdminTime = Other.AdminTime;
+			}
+		}
+
+		class PoolUtilizationDocument : IPoolUtilizationTelemetry
+		{
+			public PoolId PoolId { get; set; }
+			public int NumAgents { get; set; }
+
+			public List<StreamUtilizationDocument> Streams { get; set; } 
+			IReadOnlyList<IStreamUtilizationTelemetry> IPoolUtilizationTelemetry.Streams => Streams;
+
+			public double AdminTime { get; set; }
+			public double OtherTime { get; set; }
+
+			[BsonConstructor]
+			private PoolUtilizationDocument()
+			{
+				Streams = new List<StreamUtilizationDocument>();
+			}
+
+			public PoolUtilizationDocument(IPoolUtilizationTelemetry Other)
+			{
+				this.PoolId = Other.PoolId;
+				this.NumAgents = Other.NumAgents;
+				this.Streams = Other.Streams.ConvertAll(x => new StreamUtilizationDocument(x));
+				this.AdminTime = Other.AdminTime;
+				this.OtherTime = Other.OtherTime;
+			}
+		}
+
+		class StreamUtilizationDocument : IStreamUtilizationTelemetry
+		{
+			public StreamId StreamId { get; set; }
+			public double Time { get; set; }
+
+			[BsonConstructor]
+			private StreamUtilizationDocument()
+			{
+			}
+
+			public StreamUtilizationDocument(IStreamUtilizationTelemetry Other)
+			{
+				this.StreamId = Other.StreamId;
+				this.Time = Other.Time;
+			}
+		}
+
+		IMongoCollection<UtilizationDocument> Utilization;
+
+		/// <summary>
+		/// Constructor
+		/// </summary>
+		/// <param name="Database"></param>
+		public TelemetryCollection(DatabaseService Database)
+		{
+			Utilization = Database.GetCollection<UtilizationDocument>("Utilization");
+
+			if (!Database.ReadOnlyMode)
+			{
+				Utilization.Indexes.CreateOne(new CreateIndexModel<UtilizationDocument>(Builders<UtilizationDocument>.IndexKeys.Ascending(x => x.FinishTime).Ascending(x => x.StartTime)));
+			}
+		}
+
+		/// <inheritdoc/>
+		public async Task AddUtilizationTelemetryAsync(IUtilizationTelemetry NewTelemetry)
+		{
+			await Utilization.InsertOneAsync(new UtilizationDocument(NewTelemetry));
+		}
+
+		/// <inheritdoc/>
+		public async Task<List<IUtilizationTelemetry>> GetUtilizationTelemetryAsync(DateTime StartTimeUtc, DateTime FinishTimeUtc)
+		{
+			FilterDefinition<UtilizationDocument> Filter = Builders<UtilizationDocument>.Filter.Gte(x => x.FinishTime, StartTimeUtc) & Builders<UtilizationDocument>.Filter.Lte(x => x.StartTime, FinishTimeUtc);
+
+			List<UtilizationDocument> Documents = await Utilization.Find(Filter).ToListAsync();
+			return Documents.ConvertAll<IUtilizationTelemetry>(x => x);
+		}
+
+		/// <inheritdoc/>
+		public async Task<IUtilizationTelemetry?> GetLatestUtilizationTelemetryAsync()
+		{
+			return await Utilization.Find(FilterDefinition<UtilizationDocument>.Empty).SortByDescending(x => x.FinishTime).FirstOrDefaultAsync();
+		}
+	}
+}
diff --git a/Engine/Source/Programs/Horde/HordeServer/Fleet/PoolService.cs b/Engine/Source/Programs/Horde/HordeServer/Fleet/PoolService.cs
new file mode 100644
--- /dev/null
+++ b/Engine/Source/Programs/Horde/HordeServer/Fleet/PoolService.cs
@@ -0,0 +1,185 @@
+// Copyright Epic Games, Inc. All Rights Reserved.
+
+using EpicGames.Horde.Common;
+using HordeCommon;
+using HordeServer.Api;
+using HordeServer.Collections;
+using HordeServer.Models;
+using HordeServer.Utilities;
+using Microsoft.Extensions.Logging;
+using Microsoft.Extensions.Options;
+using MongoDB.Bson;
+using MongoDB.Driver;
+using System;
+using System.Collections.Generic;
+using System.Linq;
+using System.Threading;
+using System.Threading.Tasks;
+
+namespace HordeServer.Services
+{
+	using PoolId = StringId<IPool>;
+
+	/// <summary>
+	/// Wraps functionality for manipulating pools
+	/// </summary>
+	public class PoolService
+	{
+		/// <summary>
+		/// The database service instance
+		/// </summary>
+		DatabaseService DatabaseService;
+
+		/// <summary>
+		/// Collection of pool documents
+		/// </summary>
+		IPoolCollection Pools;
+
+		/// <summary>
+		/// Returns the current time
+		/// </summary>
+		IClock Clock;
+
+		/// <summary>
+		/// Cached set of pools, along with the timestamp that it was obtained
+		/// </summary>
+		Tuple<DateTime, Dictionary<PoolId, IPool>>? CachedPoolLookup;
+
+		/// <summary>
+		/// Constructor
+		/// </summary>
+		/// <param name="DatabaseService"></param>
+		/// <param name="Pools">Collection of pool documents</param>
+		/// <param name="Clock"></param>
+		public PoolService(DatabaseService DatabaseService, IPoolCollection Pools, IClock Clock)
+		{
+			this.DatabaseService = DatabaseService;
+			this.Pools = Pools;
+			this.Clock = Clock;
+		}
+
+		/// <summary>
+		/// Creates a new pool
+		/// </summary>
+		/// <param name="Name">Name of the new pool</param>
+		/// <param name="Condition">Condition for agents to be automatically included in this pool</param>
+		/// <param name="EnableAutoscaling">Whether to enable autoscaling for this pool</param>
+		/// <param name="MinAgents">Minimum number of agents in the pool</param>
+		/// <param name="NumReserveAgents">Minimum number of idle agents to maintain</param>
+		/// <param name="Properties">Properties for the new pool</param>
+		/// <returns>The new pool document</returns>
+		public Task<IPool> CreatePoolAsync(string Name, Condition? Condition = null, bool? EnableAutoscaling = null, int? MinAgents = null, int? NumReserveAgents = null, Dictionary<string, string>? Properties = null)
+		{
+			return Pools.AddAsync(PoolId.Sanitize(Name), Name, Condition, EnableAutoscaling, MinAgents, NumReserveAgents, Properties);
+		}
+
+		/// <summary>
+		/// Deletes a pool
+		/// </summary>
+		/// <param name="PoolId">Unique id of the pool</param>
+		/// <returns>Async task object</returns>
+		public Task<bool> DeletePoolAsync(PoolId PoolId)
+		{
+			return Pools.DeleteAsync(PoolId);
+		}
+
+		/// <summary>
+		/// Updates an existing pool
+		/// </summary>
+		/// <param name="Pool">The pool to update</param>
+		/// <param name="NewName">The new name for the pool</param>
+		/// <param name="NewCondition">New requirements for the pool</param>
+		/// <param name="NewEnableAutoscaling">Whether to enable autoscaling</param>
+		/// <param name="NewMinAgents">Minimum number of agents in the pool</param>
+		/// <param name="NewNumReserveAgents">Minimum number of idle agents to maintain</param>
+		/// <param name="NewProperties">Properties on the pool to update. Any properties with a value of null will be removed.</param>
+		/// <returns>Async task object</returns>
+		public async Task<IPool?> UpdatePoolAsync(IPool? Pool, string? NewName = null, Condition? NewCondition = null, bool? NewEnableAutoscaling = null, int? NewMinAgents = null, int? NewNumReserveAgents = null, Dictionary<string, string?>? NewProperties = null)
+		{
+			for (; Pool != null; Pool = await Pools.GetAsync(Pool.Id))
+			{
+				IPool? NewPool = await Pools.TryUpdateAsync(Pool, NewName, NewCondition, NewEnableAutoscaling, NewMinAgents, NewNumReserveAgents, null, NewProperties, null, null);
+				if (NewPool != null)
+				{
+					return NewPool;
+				}
+			}
+			return Pool;
+		}
+
+		/// <summary>
+		/// Gets all the available pools
+		/// </summary>
+		/// <returns>List of pool documents</returns>
+		public Task<List<IPool>> GetPoolsAsync()
+		{
+			return Pools.GetAsync();
+		}
+
+		/// <summary>
+		/// Gets a pool by ID
+		/// </summary>
+		/// <param name="PoolId">Unique id of the pool</param>
+		/// <returns>The pool document</returns>
+		public Task<IPool?> GetPoolAsync(PoolId PoolId)
+		{
+			return Pools.GetAsync(PoolId);
+		}
+
+		/// <summary>
+		/// Get a list of workspaces for the given agent
+		/// </summary>
+		/// <param name="Agent">The agent to return workspaces for</param>
+		/// <param name="ValidAtTime">Absolute time at which we expect the results to be valid. Values may be cached as long as they are after this time.</param>
+		/// <returns>List of workspaces</returns>
+		public async Task<HashSet<AgentWorkspace>> GetWorkspacesAsync(IAgent Agent, DateTime ValidAtTime)
+		{
+			HashSet<AgentWorkspace> Workspaces = new HashSet<AgentWorkspace>();
+
+			Dictionary<PoolId, IPool> PoolMapping = await GetPoolLookupAsync(ValidAtTime);
+			foreach (PoolId PoolId in Agent.GetPools())
+			{
+				IPool? Pool;
+				if (PoolMapping.TryGetValue(PoolId, out Pool))
+				{
+					Workspaces.UnionWith(Pool.Workspaces);
+				}
+			}
+
+			Globals Globals = await DatabaseService.GetGlobalsAsync();
+			Workspaces.UnionWith(Agent.GetAutoSdkWorkspaces(Globals, Workspaces.ToList()));
+
+			return Workspaces;
+		}
+
+		/// <summary>
+		/// Gets a mapping from pool identifiers to definitions
+		/// </summary>
+		/// <param name="ValidAtTime">Absolute time at which we expect the results to be valid. Values may be cached as long as they are after this time.</param>
+		/// <returns>Map of pool ids to pool documents</returns>
+		private async Task<Dictionary<PoolId, IPool>> GetPoolLookupAsync(DateTime ValidAtTime)
+		{
+			Tuple<DateTime, Dictionary<PoolId, IPool>>? CachedPoolLookupCopy = CachedPoolLookup;
+			if (CachedPoolLookupCopy == null || CachedPoolLookupCopy.Item1 < ValidAtTime)
+			{
+				// Get a new list of cached pools
+				DateTime NewCacheTime = Clock.UtcNow;
+				List<IPool> NewPools = await Pools.GetAsync();
+				Tuple<DateTime, Dictionary<PoolId, IPool>> NewCachedPoolLookup = Tuple.Create(NewCacheTime, NewPools.ToDictionary(x => x.Id, x => x));
+
+				// Try to swap it with the current version
+				while (CachedPoolLookupCopy == null || CachedPoolLookupCopy.Item1 < NewCacheTime)
+				{
+					Tuple<DateTime, Dictionary<PoolId, IPool>>? OriginalValue = Interlocked.CompareExchange(ref CachedPoolLookup, NewCachedPoolLookup, CachedPoolLookupCopy);
+					if (OriginalValue == CachedPoolLookupCopy)
+					{
+						CachedPoolLookupCopy = NewCachedPoolLookup;
+						break;
+					}
+					CachedPoolLookupCopy = OriginalValue;
+				}
+			}
+			return CachedPoolLookupCopy.Item2;
+		}
+	}
+}
diff --git a/Engine/Source/Programs/Horde/HordeServer/Fleet/PoolUpdateService.cs b/Engine/Source/Programs/Horde/HordeServer/Fleet/PoolUpdateService.cs
new file mode 100644
--- /dev/null
+++ b/Engine/Source/Programs/Horde/HordeServer/Fleet/PoolUpdateService.cs
@@ -0,0 +1,131 @@
+// Copyright Epic Games, Inc. All Rights Reserved.
+
+using HordeServer.Collections;
+using HordeServer.Models;
+using HordeServer.Utilities;
+using Microsoft.Extensions.Logging;
+using System;
+using System.Collections.Generic;
+using System.Linq;
+using System.Threading;
+using System.Threading.Tasks;
+
+namespace HordeServer.Services
+{
+	using PoolId = StringId<IPool>;
+
+	/// <summary>
+	/// Periodically updates pool documents to contain the correct workspaces
+	/// </summary>
+	public class PoolUpdateService : TickedBackgroundService
+	{
+		/// <summary>
+		/// Collection of pool documents
+		/// </summary>
+		IPoolCollection Pools;
+
+		/// <summary>
+		/// Collection of stream documents
+		/// </summary>
+		IStreamCollection Streams;
+
+		/// <summary>
+		/// The logger instance for this service
+		/// </summary>
+		ILogger<PoolService> Logger;
+
+		/// <summary>
+		/// Constructor
+		/// </summary>
+		/// <param name="Pools">Collection of pool documents</param>
+		/// <param name="Streams">Collection of stream documents</param>
+		/// <param name="Logger">Logger instance</param>
+		public PoolUpdateService(IPoolCollection Pools, IStreamCollection Streams, ILogger<PoolService> Logger)
+			: base(TimeSpan.FromSeconds(30.0), Logger)
+		{
+			this.Pools = Pools;
+			this.Streams = Streams;
+			this.Logger = Logger;
+		}
+
+		/// <summary>
+		/// Execute the background task
+		/// </summary>
+		/// <param name="StoppingToken">Cancellation token for the async task</param>
+		/// <returns>Async task</returns>
+		protected override async Task TickAsync(CancellationToken StoppingToken)
+		{
+			// Capture the start time for this operation. We use this to attempt to sequence updates to agents, and prevent overriding another server's updates.
+			DateTime StartTime = DateTime.UtcNow;
+
+			// Update the list
+			bool bRetryUpdate = true;
+			while (bRetryUpdate && !StoppingToken.IsCancellationRequested)
+			{
+				Logger.LogDebug("Updating pool->workspace map");
+
+				// Assume this will be the last iteration
+				bRetryUpdate = false;
+
+				// Capture the list of pools at the start of this update
+				List<IPool> CurrentPools = await Pools.GetAsync();
+
+				// Lookup table of pool id to workspaces
+				Dictionary<PoolId, List<AgentWorkspace>> PoolToAgentWorkspaces = new Dictionary<PoolId, List<AgentWorkspace>>();
+
+				// Populate the workspace list from the current stream
+				List<IStream> ActiveStreams = await Streams.FindAllAsync();
+				foreach (IStream ActiveStream in ActiveStreams)
+				{
+					foreach (KeyValuePair<string, AgentType> AgentTypePair in ActiveStream.AgentTypes)
+					{
+						// Create the new agent workspace
+						AgentWorkspace? AgentWorkspace;
+						if (ActiveStream.TryGetAgentWorkspace(AgentTypePair.Value, out AgentWorkspace))
+						{
+							AgentType AgentType = AgentTypePair.Value;
+
+							// Find or add a list of workspaces for this pool
+							List<AgentWorkspace>? AgentWorkspaces;
+							if (!PoolToAgentWorkspaces.TryGetValue(AgentType.Pool, out AgentWorkspaces))
+							{
+								AgentWorkspaces = new List<AgentWorkspace>();
+								PoolToAgentWorkspaces.Add(AgentType.Pool, AgentWorkspaces);
+							}
+
+							// Add it to the list
+							if (!AgentWorkspaces.Contains(AgentWorkspace))
+							{
+								AgentWorkspaces.Add(AgentWorkspace);
+							}
+						}
+					}
+				}
+
+				// Update the list of workspaces for each pool
+				foreach (IPool CurrentPool in CurrentPools)
+				{
+					// Get the new list of workspaces for this pool
+					List<AgentWorkspace>? NewWorkspaces;
+					if (!PoolToAgentWorkspaces.TryGetValue(CurrentPool.Id, out NewWorkspaces))
+					{
+						NewWorkspaces = new List<AgentWorkspace>();
+					}
+
+					// Update the pools document
+					if (!AgentWorkspace.SetEquals(CurrentPool.Workspaces, NewWorkspaces) || CurrentPool.Workspaces.Count != NewWorkspaces.Count)
+					{
+						Logger.LogInformation("New workspaces for pool {Pool}:{Workspaces}", CurrentPool.Id, String.Join("", NewWorkspaces.Select(x => $"\n  Identifier=\"{x.Identifier}\", Stream={x.Stream}")));
+
+						IPool? Result = await Pools.TryUpdateAsync(CurrentPool, NewWorkspaces: NewWorkspaces);
+						if (Result == null)
+						{
+							Logger.LogInformation("Pool modified; will retry");
+							bRetryUpdate = true;
+						}
+					}
+				}
+			}
+		}
+	}
+}
diff --git a/Engine/Source/Programs/Horde/HordeServer/Fleet/TelemetryService.cs b/Engine/Source/Programs/Horde/HordeServer/Fleet/TelemetryService.cs
new file mode 100644
--- /dev/null
+++ b/Engine/Source/Programs/Horde/HordeServer/Fleet/TelemetryService.cs
@@ -0,0 +1,142 @@
+// Copyright Epic Games, Inc. All Rights Reserved.
+
+using Google.Protobuf.WellKnownTypes;
+using HordeCommon.Rpc.Tasks;
+using HordeServer.Collections;
+using HordeServer.Models;
+using HordeServer.Utilities;
+using Microsoft.Extensions.Logging;
+using MongoDB.Bson;
+using System;
+using System.Collections.Generic;
+using System.Linq;
+using System.Threading;
+using System.Threading.Tasks;
+
+namespace HordeServer.Services
+{
+	using PoolId = StringId<IPool>;
+
+	/// <summary>
+	/// Service which updates telemetry periodically
+	/// </summary>
+	public class TelemetryService : ElectedBackgroundService
+	{
+		/// <summary>
+		/// Collection of telemetry documents
+		/// </summary>
+		ITelemetryCollection TelemetryCollection;
+
+		/// <summary>
+		/// Collection of agent documents
+		/// </summary>
+		IAgentCollection AgentCollection;
+
+		/// <summary>
+		/// Collection of lease documents
+		/// </summary>
+		ILeaseCollection LeaseCollection;
+
+		/// <summary>
+		/// Collection of pool documents
+		/// </summary>
+		IPoolCollection PoolCollection;
+
+		/// <summary>
+		/// Logger for diagnostic messages
+		/// </summary>
+		ILogger Logger;
+
+		/// <summary>
+		/// Constructor
+		/// </summary>
+		/// <param name="DatabaseService"></param>
+		/// <param name="TelemetryCollection"></param>
+		/// <param name="AgentCollection"></param>
+		/// <param name="LeaseCollection"></param>
+		/// <param name="PoolCollection"></param>
+		/// <param name="Logger"></param>
+		public TelemetryService(DatabaseService DatabaseService, ITelemetryCollection TelemetryCollection, IAgentCollection AgentCollection, ILeaseCollection LeaseCollection, IPoolCollection PoolCollection, ILogger<TelemetryService> Logger)
+			: base(DatabaseService, new ObjectId("600082cb9b822d7afee6f0d1"), Logger)
+		{
+			this.TelemetryCollection = TelemetryCollection;
+			this.AgentCollection = AgentCollection;
+			this.LeaseCollection = LeaseCollection;
+			this.PoolCollection = PoolCollection;
+			this.Logger = Logger;
+		}
+
+		/// <inheritdoc/>
+		protected override async Task<DateTime> TickLeaderAsync(CancellationToken StoppingToken)
+		{
+			DateTime CurrentTime = DateTime.UtcNow;
+			DateTime NextTickTime = CurrentTime + TimeSpan.FromMinutes(10.0);
+
+			// Find the last time that we need telemetry for
+			DateTime MaxTime = CurrentTime.Date + TimeSpan.FromHours(CurrentTime.Hour);
+
+			// Get the latest telemetry data
+			IUtilizationTelemetry? Latest = await TelemetryCollection.GetLatestUtilizationTelemetryAsync();
+			TimeSpan Interval = TimeSpan.FromHours(1.0);
+			int Count = (Latest == null) ? (7 * 24) : (int)Math.Round((MaxTime - Latest.FinishTime) / Interval);
+			DateTime MinTime = MaxTime - Count * Interval;
+
+			// Query all the current data
+			List<IAgent> Agents = await AgentCollection.FindAsync();
+			List<IPool> Pools = await PoolCollection.GetAsync();
+			List<ILease> Leases = await LeaseCollection.FindLeasesAsync(MinTime: MinTime);
+
+			// Find all the agents
+			Dictionary<AgentId, List<PoolId>> AgentToPoolIds = Agents.ToDictionary(x => x.Id, x => x.GetPools().ToList());
+
+			// Generate all the telemetry data
+			DateTime BucketMinTime = MinTime;
+			for (int Idx = 0; Idx < Count; Idx++)
+			{
+				DateTime BucketMaxTime = BucketMinTime + Interval;
+				Logger.LogInformation("Creating telemetry for {MinTime} to {MaxTime}", BucketMinTime, BucketMaxTime);
+
+				NewUtilizationTelemetry Telemetry = new NewUtilizationTelemetry(BucketMinTime, BucketMaxTime);
+				Telemetry.NumAgents = Agents.Count;
+				foreach (PoolId PoolId in AgentToPoolIds.Values.SelectMany(x => x))
+				{
+					Telemetry.FindOrAddPool(PoolId).NumAgents++;
+				}
+				foreach (ILease Lease in Leases)
+				{
+					if (Lease.StartTime < BucketMaxTime && Lease.FinishTime >= BucketMinTime)
+					{
+						List<PoolId>? LeasePools;
+						if (AgentToPoolIds.TryGetValue(Lease.AgentId, out LeasePools))
+						{
+							DateTime FinishTime = Lease.FinishTime ?? BucketMaxTime;
+							double Time = new TimeSpan(Math.Min(FinishTime.Ticks, BucketMaxTime.Ticks) - Math.Max(Lease.StartTime.Ticks, BucketMinTime.Ticks)).TotalHours;
+
+							foreach (PoolId PoolId in LeasePools)
+							{
+								NewPoolUtilizationTelemetry PoolTelemetry = Telemetry.FindOrAddPool(PoolId);
+								if (Lease.PoolId == null || Lease.StreamId == null)
+								{
+									PoolTelemetry.AdminTime += Time;
+								}
+								else if (PoolId == Lease.PoolId)
+								{
+									PoolTelemetry.FindOrAddStream(Lease.StreamId.Value).Time += Time;
+								}
+								else
+								{
+									PoolTelemetry.OtherTime += Time;
+								}
+							}
+						}
+					}
+				}
+				await TelemetryCollection.AddUtilizationTelemetryAsync(Telemetry);
+
+				BucketMinTime = BucketMaxTime;
+			}
+
+			return NextTickTime;
+		}
+	}
+}
diff --git a/Engine/Source/Programs/Horde/HordeServer/Issues/IIssue.cs b/Engine/Source/Programs/Horde/HordeServer/Issues/IIssue.cs
new file mode 100644
--- /dev/null
+++ b/Engine/Source/Programs/Horde/HordeServer/Issues/IIssue.cs
@@ -0,0 +1,244 @@
+// Copyright Epic Games, Inc. All Rights Reserved.
+
+using EpicGames.Core;
+using HordeServer.Api;
+using HordeServer.Collections;
+using HordeCommon;
+using HordeServer.Services;
+using HordeServer.Utilities;
+using MongoDB.Bson;
+using MongoDB.Bson.Serialization.Attributes;
+using MongoDB.Bson.Serialization.Serializers;
+using System;
+using System.Collections.Generic;
+using System.Diagnostics.CodeAnalysis;
+using System.Linq;
+using System.Runtime.CompilerServices;
+
+namespace HordeServer.Models
+{
+	using StreamId = StringId<IStream>;
+	using TemplateRefId = StringId<TemplateRef>;
+
+	/// <summary>
+	/// The severity of an issue
+	/// </summary>
+	public enum IssueSeverity
+	{
+		/// <summary>
+		/// Unspecified severity
+		/// </summary>
+		Unspecified,
+
+		/// <summary>
+		/// This error represents a warning
+		/// </summary>
+		Warning,
+
+		/// <summary>
+		/// This issue represents an error
+		/// </summary>
+		Error,
+	}
+
+	/// <summary>
+	/// Stores information about a build health issue
+	/// </summary>
+	public interface IIssue
+	{
+		/// <summary>
+		/// Id for resolved by system
+		/// </summary>
+		static ObjectId ResolvedByTimeoutId { get; } = ObjectId.Parse("609592712b5c90b5bcf88c48");
+
+		/// <summary>
+		/// Id for resolved by unknown user
+		/// </summary>
+		static ObjectId ResolvedByUnknownId { get; } = ObjectId.Parse("609593b83e9b0b6dde620cf3");
+
+		/// <summary>
+		/// The unique object id
+		/// </summary>
+		public int Id { get; }
+
+		/// <summary>
+		/// Summary for this issue
+		/// </summary>
+		public string Summary { get; }
+
+		/// <summary>
+		/// Summary set by users
+		/// </summary>
+		public string? UserSummary { get; }
+
+		/// <summary>
+		/// Fingerprint for this issue
+		/// </summary>
+		public IIssueFingerprint Fingerprint { get; }
+
+		/// <summary>
+		/// Severity of this issue
+		/// </summary>
+		public IssueSeverity Severity { get; }
+
+		/// <summary>
+		/// User id of the owner
+		/// </summary>
+		public ObjectId? OwnerId { get; }
+
+		/// <summary>
+		/// User id of the person that nominated the owner
+		/// </summary>
+		public ObjectId? NominatedById { get; }
+
+		/// <summary>
+		/// Time at which the issue was created
+		/// </summary>
+		public DateTime CreatedAt { get; }
+
+		/// <summary>
+		/// Time that the current owner was nominated
+		/// </summary>
+		public DateTime? NominatedAt { get; }
+
+		/// <summary>
+		/// Time at which the issue was acknowledged
+		/// </summary>
+		public DateTime? AcknowledgedAt { get; }
+
+		/// <summary>
+		/// Time at which the issue was resolved
+		/// </summary>
+		public DateTime? ResolvedAt { get; }
+
+		/// <summary>
+		/// User that resolved the issue
+		/// </summary>
+		public ObjectId? ResolvedById { get; }
+
+		/// <summary>
+		/// Time at which the issue was verified fixed
+		/// </summary>
+		public DateTime? VerifiedAt { get; }
+
+		/// <summary>
+		/// Time at which the issue was last seen.
+		/// </summary>
+		public DateTime LastSeenAt { get; }
+
+		/// <summary>
+		/// Fix changelist for this issue
+		/// </summary>
+		public int? FixChange { get; }
+
+		/// <summary>
+		/// List of streams affected by this issue
+		/// </summary>
+		public IReadOnlyList<IIssueStream> Streams { get; }
+
+		/// <summary>
+		/// Whether all suspects should be notified about this issue
+		/// </summary>
+		public bool NotifySuspects { get; }
+
+		/// <summary>
+		/// Update index for this instance
+		/// </summary>
+		public int UpdateIndex { get; }
+	}
+
+	/// <summary>
+	/// Information about a stream affected by an issue
+	/// </summary>
+	public interface IIssueStream
+	{
+		/// <summary>
+		/// The stream id
+		/// </summary>
+		StreamId StreamId { get; }
+
+		/// <summary>
+		/// Whether this stream contains the fix change
+		/// </summary>
+		bool? ContainsFix { get; }
+	}
+
+	/// <summary>
+	/// Suspect for an issue
+	/// </summary>
+	public interface IIssueSuspect
+	{
+		/// <summary>
+		/// Unique id for this suspect
+		/// </summary>
+		ObjectId Id { get; }
+
+		/// <summary>
+		/// Issue that this suspect belongs to
+		/// </summary>
+		int IssueId { get; }
+
+		/// <summary>
+		/// The user id of the change's author
+		/// </summary>
+		ObjectId AuthorId { get; }
+
+		/// <summary>
+		/// The change suspected of causing this issue (in the origin stream)
+		/// </summary>
+		int Change { get; }
+
+		/// <summary>
+		/// Time at which the author declined the issue
+		/// </summary>
+		DateTime? DeclinedAt { get; }
+	}
+
+	/// <summary>
+	/// Extension methods for issues
+	/// </summary>
+	static class IssueExtensions
+	{
+		/// <summary>
+		/// Creates a lookup from stream id to whether it's fixed
+		/// </summary>
+		/// <param name="Issue"></param>
+		/// <returns></returns>
+		public static Dictionary<StreamId, bool> GetFixStreamIds(this IIssue Issue)
+		{
+			Dictionary<StreamId, bool> FixStreamIds = new Dictionary<StreamId, bool>();
+			foreach (IIssueStream Stream in Issue.Streams)
+			{
+				if (Stream.ContainsFix.HasValue)
+				{
+					FixStreamIds[Stream.StreamId] = Stream.ContainsFix.Value;
+				}
+			}
+			return FixStreamIds;
+		}
+
+		/// <summary>
+		/// Find the first fix-failed step for this issue
+		/// </summary>
+		/// <param name="Issue"></param>
+		/// <param name="Spans"></param>
+		/// <returns></returns>
+		public static IIssueStep? FindFixFailedStep(this IIssue Issue, IEnumerable<IIssueSpan> Spans)
+		{
+			IIssueStep? FixFailedStep = null;
+			if (Issue.FixChange != null && Issue.FixChange.Value >= 0)
+			{
+				foreach (IIssueSpan Span in Spans)
+				{
+					IIssueStream? Stream = Issue.Streams.FirstOrDefault(x => x.StreamId == Span.StreamId);
+					if(Stream != null && (Stream.ContainsFix ?? false) && Span.LastFailure.Change >= Issue.FixChange.Value)
+					{
+						FixFailedStep = Span.LastFailure;
+						break;
+					}
+				}
+			}
+			return FixFailedStep;
+		}
+	}
+}
\ No newline at end of file
