Horde: More file moves.

diff --git a/Engine/Source/Programs/Horde/HordeServer/Issues/IIssueCollection.cs b/Engine/Source/Programs/Horde/HordeServer/Issues/IIssueCollection.cs
new file mode 100644
--- /dev/null
+++ b/Engine/Source/Programs/Horde/HordeServer/Issues/IIssueCollection.cs
@@ -0,0 +1,544 @@
+// Copyright Epic Games, Inc. All Rights Reserved.
+
+using HordeCommon;
+using HordeServer.Models;
+using HordeServer.Services;
+using HordeServer.Utilities;
+using MessagePack.Formatters;
+using Microsoft.AspNetCore.Connections.Features;
+using Microsoft.Extensions.Logging;
+using MongoDB.Bson;
+using MongoDB.Bson.Serialization;
+using MongoDB.Bson.Serialization.Attributes;
+using MongoDB.Driver;
+using System;
+using System.Collections.Generic;
+using System.Diagnostics.CodeAnalysis;
+using System.IO;
+using System.Linq;
+using System.Linq.Expressions;
+using System.Text.RegularExpressions;
+using System.Threading.Tasks;
+using System.Transactions;
+
+namespace HordeServer.Collections
+{
+	using StreamId = StringId<IStream>;
+	using TemplateRefId = StringId<TemplateRef>;
+
+	/// <summary>
+	/// Fingerprint for an issue
+	/// </summary>
+	public class NewIssueFingerprint : IIssueFingerprint
+	{
+		/// <summary>
+		/// The type of issue
+		/// </summary>
+		public string Type { get; }
+
+		/// <summary>
+		/// List of keys which identify this issue.
+		/// </summary>
+		public CaseInsensitiveStringSet Keys { get; set; }
+
+		/// <summary>
+		/// Set of keys which should trigger a negative match
+		/// </summary>
+		public CaseInsensitiveStringSet? RejectKeys { get; set; }
+
+		/// <summary>
+		/// Constructor
+		/// </summary>
+		/// <param name="Type">The type of issue</param>
+		/// <param name="Keys">Keys which uniquely identify this issue</param>
+		/// <param name="RejectKeys">Keys which should not match with this issue</param>
+		public NewIssueFingerprint(string Type, IEnumerable<string> Keys, IEnumerable<string>? RejectKeys)
+		{
+			this.Type = Type;
+			this.Keys = new CaseInsensitiveStringSet(Keys);
+
+			if (RejectKeys != null && RejectKeys.Any())
+			{
+				this.RejectKeys = new CaseInsensitiveStringSet(RejectKeys);
+			}
+		}
+
+		/// <summary>
+		/// Copy constructor
+		/// </summary>
+		/// <param name="Other">The fingerprint to copy from</param>
+		public NewIssueFingerprint(IIssueFingerprint Other)
+			: this(Other.Type, Other.Keys, Other.RejectKeys)
+		{
+		}
+
+		/// <summary>
+		/// Merges another fingerprint into this one
+		/// </summary>
+		/// <param name="Other">The other fingerprint to merge with</param>
+		public NewIssueFingerprint MergeWith(IIssueFingerprint Other)
+		{
+			NewIssueFingerprint Result = new NewIssueFingerprint(this);
+
+			Result.Keys.UnionWith(Other.Keys);
+			if (Other.RejectKeys != null)
+			{
+				if (Result.RejectKeys == null)
+				{
+					Result.RejectKeys = new CaseInsensitiveStringSet();
+				}
+				Result.RejectKeys.UnionWith(Other.RejectKeys);
+			}
+
+			return Result;
+		}
+
+		/// <inheritdoc/>
+		public override bool Equals(object? Other)
+		{
+			NewIssueFingerprint? OtherFingerprint = Other as NewIssueFingerprint;
+			if(OtherFingerprint == null || !Type.Equals(OtherFingerprint.Type, StringComparison.Ordinal) || !Keys.SetEquals(OtherFingerprint.Keys))
+			{
+				return false;
+			}
+
+			if (RejectKeys == null)
+			{
+				if (OtherFingerprint.RejectKeys != null && OtherFingerprint.RejectKeys.Count > 0)
+				{
+					return false;
+				}
+			}
+			else
+			{
+				if (OtherFingerprint.RejectKeys == null || !RejectKeys.SetEquals(OtherFingerprint.RejectKeys))
+				{
+					return false;
+				}
+			}
+
+			return true;
+		}
+
+		/// <summary>
+		/// Gets a hashcode for this fingerprint
+		/// </summary>
+		/// <returns>The hashcode value</returns>
+		public override int GetHashCode()
+		{
+			int Result = StringComparer.Ordinal.GetHashCode(Type);
+			foreach (string Key in Keys)
+			{
+				Result = HashCode.Combine(Result, Key);
+			}
+			if (RejectKeys != null)
+			{
+				foreach (string RejectKey in RejectKeys)
+				{
+					Result = HashCode.Combine(Result, RejectKey);
+				}
+			}
+			return Result;
+		}
+
+		/// <summary>
+		/// Merges two fingerprints togetherthis issue fingerprint with another
+		/// </summary>
+		/// <param name="A">The first fingerprint</param>
+		/// <param name="B">The second fingerprint to merge with</param>
+		/// <returns>Merged fingerprint</returns>
+		public static NewIssueFingerprint Merge(IIssueFingerprint A, IIssueFingerprint B)
+		{
+			NewIssueFingerprint NewFingerprint = new NewIssueFingerprint(A);
+			NewFingerprint.Keys.UnionWith(B.Keys);
+
+			if (B.RejectKeys != null)
+			{
+				if (NewFingerprint.RejectKeys == null)
+				{
+					NewFingerprint.RejectKeys = new CaseInsensitiveStringSet();
+				}
+				NewFingerprint.RejectKeys.UnionWith(B.RejectKeys);
+			}
+
+			return NewFingerprint;
+		}
+
+		/// <inheritdoc/>
+		public override string ToString()
+		{
+			return $"{Type}: \"{String.Join("\", \"", Keys)}\"";
+		}
+	}
+
+	/// <summary>
+	/// Identifies a particular changelist and job
+	/// </summary>
+	public class NewIssueStepData
+	{
+		/// <summary>
+		/// The changelist number
+		/// </summary>
+		public int Change { get; set; }
+
+		/// <summary>
+		/// Severity of the issue
+		/// </summary>
+		public IssueSeverity Severity { get; set; }
+
+		/// <summary>
+		/// Name of the job
+		/// </summary>
+		public string JobName { get; set; }
+
+		/// <summary>
+		/// The unique job id
+		/// </summary>
+		public ObjectId JobId { get; set; }
+
+		/// <summary>
+		/// Batch id for the step
+		/// </summary>
+		public SubResourceId BatchId { get; set; }
+
+		/// <summary>
+		/// Id of the step
+		/// </summary>
+		public SubResourceId StepId { get; set; }
+
+		/// <summary>
+		/// Time that the step started
+		/// </summary>
+		public DateTime StepTime { get; set; }
+
+		/// <summary>
+		/// The log id for this step
+		/// </summary>
+		public ObjectId? LogId { get; set; }
+
+		/// <summary>
+		/// Whether to notify suspects for this step
+		/// </summary>
+		public bool NotifySuspects { get; set; }
+
+		/// <summary>
+		/// Constructor
+		/// </summary>
+		/// <param name="Change">The changelist number for this job</param>
+		/// <param name="Severity">Severity of the issue in this step</param>
+		/// <param name="JobName">The job name</param>
+		/// <param name="JobId">The unique job id</param>
+		/// <param name="BatchId">The batch id</param>
+		/// <param name="StepId">The step id</param>
+		/// <param name="StepTime">Time that the step started</param>
+		/// <param name="LogId">Unique id of the log file for this step</param>
+		/// <param name="NotifySuspects">Whether to notify suspects for this step</param>
+		public NewIssueStepData(int Change, IssueSeverity Severity, string JobName, ObjectId JobId, SubResourceId BatchId, SubResourceId StepId, DateTime StepTime, ObjectId? LogId, bool NotifySuspects)
+		{
+			this.Change = Change;
+			this.Severity = Severity;
+			this.JobName = JobName;
+			this.JobId = JobId;
+			this.BatchId = BatchId;
+			this.StepId = StepId;
+			this.StepTime = StepTime;
+			this.LogId = LogId;
+			this.NotifySuspects = NotifySuspects;
+		}
+
+		/// <summary>
+		/// Constructor
+		/// </summary>
+		/// <param name="Job">The job being built</param>
+		/// <param name="Batch">Batch of the job for the step</param>
+		/// <param name="Step">The step being built</param>
+		/// <param name="Severity">Severity of the issue in this step</param>
+		/// <param name="NotifySuspects">Whether to notify suspects for this step failing</param>
+		public NewIssueStepData(IJob Job, IJobStepBatch Batch, IJobStep Step, IssueSeverity Severity, bool NotifySuspects)
+			: this(Job.Change, Severity, Job.Name, Job.Id, Batch.Id, Step.Id, Step.StartTimeUtc ?? default, Step.LogId, NotifySuspects)
+		{
+		}
+
+		/// <summary>
+		/// Constructor
+		/// </summary>
+		/// <param name="JobStepRef">The jobstep reference</param>
+		public NewIssueStepData(IJobStepRef JobStepRef)
+			: this(JobStepRef.Change, IssueSeverity.Unspecified, JobStepRef.JobName, JobStepRef.Id.JobId, JobStepRef.Id.BatchId, JobStepRef.Id.StepId, JobStepRef.StartTimeUtc, JobStepRef.LogId, false)
+		{
+		}
+	}
+
+	/// <summary>
+	/// Information about a suspect changelist that may have caused an issue
+	/// </summary>
+	public class NewIssueSuspectData
+	{
+		/// <summary>
+		/// Author of the changelist
+		/// </summary>
+		public ObjectId AuthorId { get; set; }
+
+		/// <summary>
+		/// The submitted changelist
+		/// </summary>
+		public int Change { get; set; }
+
+		/// <summary>
+		/// The time that the user declined causing this issue
+		/// </summary>
+		public DateTime? DeclinedAt { get; set; }
+
+		/// <summary>
+		/// Constructor
+		/// </summary>
+		/// <param name="AuthorId">Author of the change</param>
+		/// <param name="Change">The changelist number</param>
+		/// <param name="DeclinedAt">The time that the user declined this issue</param>
+		public NewIssueSuspectData(ObjectId AuthorId, int Change, DateTime? DeclinedAt)
+		{
+			this.AuthorId = AuthorId;
+			this.Change = Change;
+			this.DeclinedAt = DeclinedAt;
+		}
+	}
+
+	/// <summary>
+	/// Information about a suspect changelist that may have caused an issue
+	/// </summary>
+	public class NewIssueSpanSuspectData
+	{
+		/// <summary>
+		/// The submitted changelist
+		/// </summary>
+		public int Change { get; set; }
+
+		/// <summary>
+		/// Author of the changelist
+		/// </summary>
+		public ObjectId AuthorId { get; set; }
+
+		/// <summary>
+		/// The original changelist number, if merged from another branch. For changes merged between several branches, this is the originally submitted change.
+		/// </summary>
+		public int? OriginatingChange { get; set; }
+
+		/// <summary>
+		/// Constructor
+		/// </summary>
+		/// <param name="Change">The changelist number</param>
+		/// <param name="AuthorId">Author of the change</param>
+		public NewIssueSpanSuspectData(int Change, ObjectId AuthorId)
+		{
+			this.Change = Change;
+			this.AuthorId = AuthorId;
+		}
+	}
+
+	/// <summary>
+	/// Interface for a token used to indicate a position in the transaction queue. Transactions are opportunistic, and will fail if the token is longer valid at the time that it is used.
+	/// </summary>
+	[SuppressMessage("Design", "CA1040:Avoid empty interfaces")]
+	public interface IIssueSequenceToken
+	{
+		/// <summary>
+		/// Reset the token to the last operation
+		/// </summary>
+		/// <returns>Async task</returns>
+		public Task ResetAsync();
+	}
+
+	/// <summary>
+	/// Interface for an issue collection
+	/// </summary>
+	public interface IIssueCollection
+	{
+		/// <summary>
+		/// Gets a token used to sequence operations to the issue collection
+		/// </summary>
+		/// <returns>Sequence token instance</returns>
+		Task<IIssueSequenceToken> GetSequenceTokenAsync();
+
+		/// <summary>
+		/// Assign a unique issue id. For compatibility with the legacy issue system.
+		/// </summary>
+		/// <param name="Token">The current sequence token</param>
+		/// <returns>Sequence token</returns>
+		Task<int?> ReserveUniqueIdAsync(IIssueSequenceToken Token);
+
+		#region Issues
+
+		/// <summary>
+		/// Creates a new issue
+		/// </summary>
+		/// <param name="Token">Token for the transaction</param>
+		/// <param name="Summary">Summary text for the issue</param>
+		/// <param name="Span">The initial span for the issue</param>
+		/// <returns>The new issue instance</returns>
+		Task<IIssue?> AddIssueAsync(IIssueSequenceToken Token, string Summary, IIssueSpan Span);
+
+		/// <summary>
+		/// Retrieves and issue by id
+		/// </summary>
+		/// <param name="IssueId">Unique id of the issue</param>
+		/// <returns>The issue matching the given id, or null</returns>
+		Task<IIssue?> GetIssueAsync(int IssueId);
+
+		/// <summary>
+		/// Gets the suspects for an issue
+		/// </summary>
+		/// <param name="Issue">The issue to retrieve suspects for</param>
+		/// <returns>List of suspects</returns>
+		Task<List<IIssueSuspect>> GetSuspectsAsync(IIssue Issue);
+
+		/// <summary>
+		/// Searches for open issues
+		/// </summary>
+		/// <param name="Ids">Set of issue ids to find</param>
+		/// <param name="UserId">The user to find issues for</param>
+		/// <param name="StreamId">The stream affected by the issue</param>
+		/// <param name="MinChange">Minimum changelist affected by the issue</param>
+		/// <param name="MaxChange">Maximum changelist affected by the issue</param>
+		/// <param name="Resolved">Include issues that are now resolved</param>
+		/// <param name="Index">Index within the results to return</param>
+		/// <param name="Count">Number of results</param>
+		/// <returns>List of streams open in the given stream at the given changelist</returns>
+		Task<List<IIssue>> FindIssuesAsync(IEnumerable<int>? Ids = null, ObjectId? UserId = null, StreamId? StreamId = null, int? MinChange = null, int? MaxChange = null, bool? Resolved = null, int? Index = null, int? Count = null);
+
+		/// <summary>
+		/// Searches for open issues
+		/// </summary>
+		/// <param name="Changes">List of suspect changes</param>
+		/// <returns>List of issues that are affected by the given changes</returns>
+		Task<List<IIssue>> FindIssuesForSuspectsAsync(List<int> Changes);
+
+		/// <summary>
+		/// Updates the state of an issue
+		/// </summary>
+		/// <param name="Issue">The issue to update</param>
+		/// <param name="NewSeverity">New severity for the issue</param>
+		/// <param name="NewSummary">New summary for the issue</param>
+		/// <param name="NewUserSummary">New user summary for the issue</param>
+		/// <param name="NewOwnerId">New owner of the issue</param>
+		/// <param name="NewNominatedById">Person that nominated the new owner</param>
+		/// <param name="NewAcknowledged">Whether the issue has been acknowledged</param>
+		/// <param name="NewDeclinedById">Name of a user that has declined the issue</param>
+		/// <param name="NewFixChange">Fix changelist for the issue. Pass 0 to clear the fix changelist, -1 for systemic issue.</param>
+		/// <param name="NewFixStreamIds">List of stream ids with a bool indicating whether the fix changelist is in this stream</param>
+		/// <param name="NewResolvedById">User that resolved the issue (may be ObjectId.Empty to clear)</param>
+		/// <param name="NewLastSeenAt"></param>
+		/// <param name="NewNotifySuspects">Whether all suspects should be notified about this issue</param>
+		/// <returns>True if the issue was updated</returns>
+		Task<IIssue?> UpdateIssueAsync(IIssue Issue, IssueSeverity? NewSeverity = null, string? NewSummary = null, string? NewUserSummary = null, ObjectId? NewOwnerId = null, ObjectId? NewNominatedById = null, bool? NewAcknowledged = null, ObjectId? NewDeclinedById = null, int? NewFixChange = null, Dictionary<StreamId, bool>? NewFixStreamIds = null, ObjectId? NewResolvedById = null, DateTime? NewLastSeenAt = null, bool? NewNotifySuspects = null);
+
+		/// <summary>
+		/// Updates the list of suspects for an issue, and optionally sets the resolved state
+		/// </summary>
+		/// <param name="Issue">The issue to update</param>
+		/// <returns>True if the issue was updated</returns>
+		Task<IIssue?> UpdateIssueDerivedDataAsync(IIssue Issue);
+
+		#endregion
+
+		#region Spans
+
+		/// <summary>
+		/// Creates a new span from the given failure
+		/// </summary>
+		/// <param name="Token">Token to sequence the addition of new spans</param>
+		/// <param name="Stream">The stream containing the span</param>
+		/// <param name="TemplateRefId">The template being executed</param>
+		/// <param name="NodeName">Name of the step</param>
+		/// <param name="Fingerprint">Fingerprint for this issue</param>
+		/// <param name="LastSuccess">Last time the build succeeded before the span</param>
+		/// <param name="Failure">List of failing steps</param>
+		/// <param name="NextSuccess">First time the build succeeded after the span</param>
+		/// <param name="Suspects">Suspects for this span</param>
+		/// <returns>New span, or null if the sequence token is not valid</returns>
+		Task<IIssueSpan?> AddSpanAsync(IIssueSequenceToken Token, IStream Stream, TemplateRefId TemplateRefId, string NodeName, NewIssueFingerprint Fingerprint, NewIssueStepData? LastSuccess, NewIssueStepData Failure, NewIssueStepData? NextSuccess, List<NewIssueSpanSuspectData> Suspects);
+
+		/// <summary>
+		/// Gets a particular span
+		/// </summary>
+		/// <param name="SpanId">Unique id of the span</param>
+		/// <returns>New span, or null if the sequence token is not valid</returns>
+		Task<IIssueSpan?> GetSpanAsync(ObjectId SpanId);
+
+		/// <summary>
+		/// Updates the given span. Note that data in the span's issue may be derived from this, and the issue should be updated afterwards.
+		/// </summary>
+		/// <param name="Span">Span to update</param>
+		/// <param name="NewSeverity">New severity for the span</param>
+		/// <param name="NewLastSuccess">New last successful step</param>
+		/// <param name="NewFailure">New failed step</param>
+		/// <param name="NewNextSuccess">New next successful step</param>
+		/// <param name="NewSuspects">New suspects for the span</param>
+		/// <param name="NewModified">New modified flag for the span</param>
+		/// <returns>The updated span, or null on failure</returns>
+		Task<IIssueSpan?> UpdateSpanAsync(IIssueSpan Span, IssueSeverity? NewSeverity = null, NewIssueStepData? NewLastSuccess = null, NewIssueStepData? NewFailure = null, NewIssueStepData? NewNextSuccess = null, List<NewIssueSpanSuspectData>? NewSuspects = null, bool? NewModified = null);
+
+		/// <summary>
+		/// Updates the given span
+		/// </summary>
+		/// <param name="Token">The token</param>
+		/// <param name="Span">The span to update</param>
+		/// <param name="Issue">The issue to attach to</param>
+		/// <param name="NewSeverity">New severity for the span</param>
+		/// <param name="NewFingerprint">New fingerprint for the span</param>
+		/// <param name="NewSummary">New summary for the issue</param>
+		/// <returns>The updated span, or null on failure</returns>
+		Task<bool> AttachSpanToIssueAsync(IIssueSequenceToken Token, IIssueSpan Span, IIssue Issue, IssueSeverity? NewSeverity = null, NewIssueFingerprint? NewFingerprint = null, string? NewSummary = null);
+
+		/// <summary>
+		/// Gets all the spans for a particular issue
+		/// </summary>
+		/// <param name="IssueId">Issue id</param>
+		/// <returns>List of spans</returns>
+		Task<List<IIssueSpan>> FindSpansAsync(int IssueId);
+
+		/// <summary>
+		/// Retrieves multiple spans
+		/// </summary>
+		/// <param name="SpanIds">The span ids</param>
+		/// <returns>List of spans</returns>
+		Task<List<IIssueSpan>> FindSpansAsync(IEnumerable<ObjectId> SpanIds);
+
+		/// <summary>
+		/// Finds the open issues for a given stream
+		/// </summary>
+		/// <param name="StreamId">The stream id</param>
+		/// <param name="TemplateId">The template id</param>
+		/// <param name="Name">Name of the node</param>
+		/// <param name="Change">Changelist number to query</param>
+		/// <returns>List of open issues</returns>
+		Task<List<IIssueSpan>> FindOpenSpansAsync(StreamId StreamId, TemplateRefId TemplateId, string Name, int Change);
+
+		/// <summary>
+		/// Finds a span marked as modified
+		/// </summary>
+		/// <returns></returns>
+		Task<IIssueSpan?> FindModifiedSpanAsync();
+
+		#endregion
+
+		#region Steps
+			
+		/// <summary>
+		/// Find steps for the given spans
+		/// </summary>
+		/// <param name="SpanIds">Span ids</param>
+		/// <returns>List of steps</returns>
+		Task<List<IIssueStep>> FindStepsAsync(IEnumerable<ObjectId> SpanIds);
+
+		/// <summary>
+		/// Find steps for the given spans
+		/// </summary>
+		/// <param name="JobId">The job id</param>
+		/// <param name="BatchId">The batch id</param>
+		/// <param name="StepId">The step id</param>
+		/// <returns>List of steps</returns>
+		Task<List<IIssueStep>> FindStepsAsync(ObjectId JobId, SubResourceId? BatchId, SubResourceId? StepId);
+
+		#endregion
+	}
+}
diff --git a/Engine/Source/Programs/Horde/HordeServer/Issues/IIssueFingerprint.cs b/Engine/Source/Programs/Horde/HordeServer/Issues/IIssueFingerprint.cs
new file mode 100644
--- /dev/null
+++ b/Engine/Source/Programs/Horde/HordeServer/Issues/IIssueFingerprint.cs
@@ -0,0 +1,79 @@
+// Copyright Epic Games, Inc. All Rights Reserved.
+
+using HordeServer.Utilities;
+using System;
+using System.Collections.Generic;
+using System.Linq;
+using System.Threading.Tasks;
+
+namespace HordeServer.Models
+{
+	/// <summary>
+	/// Fingerprint for an issue
+	/// </summary>
+	public interface IIssueFingerprint
+	{
+		/// <summary>
+		/// The type of issue
+		/// </summary>
+		public string Type { get; }
+
+		/// <summary>
+		/// List of keys which identify this issue.
+		/// </summary>
+		public CaseInsensitiveStringSet Keys { get; }
+
+		/// <summary>
+		/// Set of keys which should trigger a negative match
+		/// </summary>
+		public CaseInsensitiveStringSet? RejectKeys { get; }
+	}
+
+	/// <summary>
+	/// Extension methods for <see cref="IIssueFingerprint"/>
+	/// </summary>
+	public static class IssueFingerprintExtensions
+	{
+		/// <summary>
+		/// Checks if a fingerprint matches another fingerprint
+		/// </summary>
+		/// <param name="Fingerprint">The first fingerprint to compare</param>
+		/// <param name="Other">The other fingerprint to compare to</param>
+		/// <returns>True is the fingerprints match</returns>
+		public static bool IsMatch(this IIssueFingerprint Fingerprint, IIssueFingerprint Other)
+		{
+			if (!Fingerprint.Type.Equals(Other.Type, StringComparison.Ordinal))
+			{
+				return false;
+			}
+			if (Fingerprint.Keys.Count > 0 && Other.Keys.Count > 0 && !Fingerprint.Keys.Any(x => Other.Keys.Contains(x)))
+			{
+				return false;
+			}
+			if (Fingerprint.RejectKeys != null && Fingerprint.RejectKeys.Any(x => Other.Keys.Contains(x)))
+			{
+				return false;
+			}
+			return true;
+		}
+
+		/// <summary>
+		/// Checks if a fingerprint matches another fingerprint for creating a new span
+		/// </summary>
+		/// <param name="Fingerprint">The first fingerprint to compare</param>
+		/// <param name="Other">The other fingerprint to compare to</param>
+		/// <returns>True is the fingerprints match</returns>
+		public static bool IsMatchForNewSpan(this IIssueFingerprint Fingerprint, IIssueFingerprint Other)
+		{
+			if (!Fingerprint.Type.Equals(Other.Type, StringComparison.Ordinal))
+			{
+				return false;
+			}
+			if (Fingerprint.RejectKeys != null && Fingerprint.RejectKeys.Any(x => Other.Keys.Contains(x)))
+			{
+				return false;
+			}
+			return true;
+		}
+	}
+}
diff --git a/Engine/Source/Programs/Horde/HordeServer/Issues/IIssueHandler.cs b/Engine/Source/Programs/Horde/HordeServer/Issues/IIssueHandler.cs
new file mode 100644
--- /dev/null
+++ b/Engine/Source/Programs/Horde/HordeServer/Issues/IIssueHandler.cs
@@ -0,0 +1,161 @@
+// Copyright Epic Games, Inc. All Rights Reserved.
+
+using EpicGames.Core;
+using HordeServer.Models;
+using HordeServer.Collections;
+using System;
+using System.Collections.Generic;
+using System.Linq;
+using System.Threading.Tasks;
+using System.Diagnostics.CodeAnalysis;
+
+namespace HordeServer.IssueHandlers
+{
+	/// <summary>
+	/// Information about a changelist and a value ranking the likelihood that it caused an issue
+	/// </summary>
+	class SuspectChange
+	{
+		/// <summary>
+		/// Set of extensions to treat as code
+		/// </summary>
+		static readonly HashSet<StringView> CodeExtensions = new HashSet<StringView>(StringViewComparer.OrdinalIgnoreCase)
+		{
+			".c",
+			".cc",
+			".cpp",
+			".m",
+			".mm",
+			".rc",
+			".cs",
+			".csproj",
+			".h",
+			".hpp",
+			".inl",
+			".usf",
+			".ush",
+			".uproject",
+			".uplugin",
+			".sln"
+		};
+
+		/// <summary>
+		/// Set of file extensions to treat as content
+		/// </summary>
+		static readonly HashSet<StringView> ContentExtensions = new HashSet<StringView>(StringViewComparer.OrdinalIgnoreCase)
+		{
+			".uasset",
+			".umap",
+			".ini"
+		};
+
+		/// <summary>
+		/// The change detials
+		/// </summary>
+		public ChangeDetails Details { get; set; }
+
+		/// <summary>
+		/// Rank for the change. A value less than or equal to zero indicates a lack of culpability, positive values indicate
+		/// the possibility of a change being the culprit.
+		/// </summary>
+		public int Rank { get; set; }
+
+		/// <summary>
+		/// Whether the change contains code
+		/// </summary>
+		public bool ContainsCode { get; }
+
+		/// <summary>
+		/// Whether the change modifies content
+		/// </summary>
+		public bool ContainsContent { get; }
+
+		/// <summary>
+		/// Constructor
+		/// </summary>
+		/// <param name="Details">The changelist details</param>
+		public SuspectChange(ChangeDetails Details)
+		{
+			this.Details = Details;
+
+			foreach (ChangeFile File in Details.Files)
+			{
+				int Idx = File.Path.LastIndexOf('.');
+				if (Idx != -1)
+				{
+					StringView Extension = new StringView(File.Path, Idx);
+					if (CodeExtensions.Contains(Extension))
+					{
+						ContainsCode = true;
+					}
+					if (ContentExtensions.Contains(Extension))
+					{
+						ContainsContent = true;
+					}
+					if (ContainsCode && ContainsContent)
+					{
+						break;
+					}
+				}
+			}
+		}
+
+		/// <summary>
+		/// Determines whether this change modifies the given file
+		/// </summary>
+		/// <param name="FileToCheck">The file to look for</param>
+		/// <returns>True if the change modifies the given file</returns>
+		public bool ModifiesFile(string FileToCheck)
+		{
+			foreach (ChangeFile File in Details.Files)
+			{
+				if (File.Path.EndsWith(FileToCheck, StringComparison.OrdinalIgnoreCase) && (File.Length == FileToCheck.Length || File.Path[File.Path.Length - FileToCheck.Length - 1] == '/'))
+				{
+					return true;
+				}
+			}
+			return false;
+		}
+	}
+
+	/// <summary>
+	/// Interface for issue matchers
+	/// </summary>
+	interface IIssueHandler
+	{
+		/// <summary>
+		/// Identifier for the type of issue
+		/// </summary>
+		string Type { get; }
+
+		/// <summary>
+		/// Priority of this matcher
+		/// </summary>
+		int Priority { get; }
+
+		/// <summary>
+		/// Match the given event and produce a fingerprint
+		/// </summary>
+		/// <param name="Job">The job that spawned the event</param>
+		/// <param name="Node">Node that was executed</param>
+		/// <param name="EventData">The event data</param>
+		/// <param name="Fingerprint">Receives the fingerprint on success</param>
+		/// <returns>True if the match is successful</returns>
+		bool TryGetFingerprint(IJob Job, INode Node, ILogEventData EventData, [NotNullWhen(true)] out NewIssueFingerprint? Fingerprint);
+
+		/// <summary>
+		/// Rank all the suspect changes for a given fingerprint
+		/// </summary>
+		/// <param name="Fingerprint">The issue fingerprint</param>
+		/// <param name="Suspects">Potential suspects</param>
+		void RankSuspects(IIssueFingerprint Fingerprint, List<SuspectChange> Suspects);
+
+		/// <summary>
+		/// Gets the summary text for an issue
+		/// </summary>
+		/// <param name="Fingerprint">The fingerprint</param>
+		/// <param name="Severity">Severity of the issue</param>
+		/// <returns>The summary text</returns>
+		string GetSummary(IIssueFingerprint Fingerprint, IssueSeverity Severity);
+	}
+}
diff --git a/Engine/Source/Programs/Horde/HordeServer/Issues/IIssueService.cs b/Engine/Source/Programs/Horde/HordeServer/Issues/IIssueService.cs
new file mode 100644
--- /dev/null
+++ b/Engine/Source/Programs/Horde/HordeServer/Issues/IIssueService.cs
@@ -0,0 +1,151 @@
+// Copyright Epic Games, Inc. All Rights Reserved.
+
+using HordeServer.Models;
+using HordeServer.Utilities;
+using Microsoft.Extensions.Hosting;
+using MongoDB.Bson;
+using System;
+using System.Collections.Generic;
+using System.Linq;
+using System.Threading.Tasks;
+
+namespace HordeServer.Services
+{
+	using StreamId = StringId<IStream>;
+
+	/// <summary>
+	/// Service which manages build health issues
+	/// </summary>
+	public interface IIssueService : IHostedService
+	{
+		/// <summary>
+		/// Delegate for issue creation events
+		/// </summary>
+		public delegate void IssueUpdatedEvent(IIssue Issue);
+
+		/// <summary>
+		/// Event triggered when a label state changes
+		/// </summary>
+		event IssueUpdatedEvent? OnIssueUpdated;
+
+		/// <summary>
+		/// Gets a list of cached open issues
+		/// </summary>
+		/// <returns>List of cached issues</returns>
+		IReadOnlyList<IIssueDetails> CachedOpenIssues { get; }
+
+		/// <summary>
+		/// Find a cached issue, or update it
+		/// </summary>
+		/// <param name="IssueId">The issue id</param>
+		/// <returns>The cached issue data</returns>
+		Task<IIssueDetails?> GetCachedIssueDetailsAsync(int IssueId);
+
+		/// <summary>
+		/// Gets the details for a particular issue
+		/// </summary>
+		/// <param name="Issue">The issue id</param>
+		/// <returns>The issue details</returns>
+		Task<IIssueDetails> GetIssueDetailsAsync(IIssue Issue);
+
+		/// <summary>
+		/// Determines whether to show alerts for a given issue
+		/// </summary>
+		/// <param name="Issue">The issue to query</param>
+		/// <param name="Spans">Spans for the issue</param>
+		/// <returns>True if we should show alerts for the given issue</returns>
+		bool ShowDesktopAlertsForIssue(IIssue Issue, IReadOnlyList<IIssueSpan> Spans);
+
+		/// <summary>
+		/// Gets an issue with the given id
+		/// </summary>
+		/// <param name="IssueId"></param>
+		/// <returns></returns>
+		Task<IIssue?> GetIssueAsync(int IssueId);
+
+		/// <summary>
+		/// Gets the spans for the given issue
+		/// </summary>
+		/// <param name="Issue">The issue id</param>
+		/// <returns>List of spans for the issue</returns>
+		Task<List<IIssueSpan>> GetIssueSpansAsync(IIssue Issue);
+
+		/// <summary>
+		/// Gets the spans for the given issue
+		/// </summary>
+		/// <param name="Span">The issue span</param>
+		/// <returns>List of steps for the issue</returns>
+		Task<List<IIssueStep>> GetIssueStepsAsync(IIssueSpan Span);
+
+		/// <summary>
+		/// Gets the suspects for an issue
+		/// </summary>
+		/// <param name="Issue"></param>
+		/// <returns></returns>
+		Task<List<IIssueSuspect>> GetIssueSuspectsAsync(IIssue Issue);
+
+		/// <summary>
+		/// Searches for open issues
+		/// </summary>
+		/// <param name="Ids">List of issue ids to include</param>
+		/// <param name="UserId">User to filter issues returned by</param>
+		/// <param name="StreamId">The stream to return issues for</param>
+		/// <param name="MinChange">The minimum changelist number to include</param>
+		/// <param name="MaxChange">The maximum changelist number to include</param>
+		/// <param name="Resolved">Whether to include results that are resolved</param>
+		/// <param name="Index">Index within the results to return</param>
+		/// <param name="Count">Number of results</param>
+		/// <returns>List of streams open in the given stream at the given changelist</returns>
+		Task<List<IIssue>> FindIssuesAsync(IEnumerable<int>? Ids = null, ObjectId? UserId = null, StreamId? StreamId = null, int? MinChange = null, int? MaxChange = null, bool? Resolved = null, int? Index = null, int? Count = null);
+
+		/// <summary>
+		/// Searches for open issues affecting a job
+		/// </summary>
+		/// <param name="Ids"></param>
+		/// <param name="Job"></param>
+		/// <param name="Graph"></param>
+		/// <param name="StepId"></param>
+		/// <param name="BatchId"></param>
+		/// <param name="LabelIdx"></param>
+		/// <param name="UserId"></param>
+		/// <param name="Resolved">Whether to include results that are resolved</param>
+		/// <param name="Index">Index within the results to return</param>
+		/// <param name="Count">Number of results</param>
+		/// <returns></returns>
+		Task<List<IIssue>> FindIssuesForJobAsync(int[]? Ids, IJob Job, IGraph Graph, SubResourceId? StepId = null, SubResourceId? BatchId = null, int? LabelIdx = null, ObjectId? UserId = null, bool? Resolved = null, int? Index = null, int? Count = null);
+
+		/// <summary>
+		/// Find the events for a particular issue
+		/// </summary>
+		/// <param name="IssueId">The issue id</param>
+		/// <param name="LogIds">Log ids to include</param>
+		/// <param name="Index">Index of the first event to return</param>
+		/// <param name="Count">Number of issues to return</param>
+		/// <returns>List of log events</returns>
+		Task<List<ILogEvent>> FindEventsForIssueAsync(int IssueId, ObjectId[]? LogIds = null, int Index = 0, int Count = 10);
+
+		/// <summary>
+		/// Updates the state of an issue
+		/// </summary>
+		/// <param name="Id">The current issue id</param>
+		/// <param name="Summary">New summary for the issue</param>
+		/// <param name="OwnerId">New owner of the issue</param>
+		/// <param name="NominatedById">Person that nominated the new owner</param>
+		/// <param name="Acknowledged">Whether the issue has been acknowledged</param>
+		/// <param name="DeclinedById">Name of the user that has declined this issue</param>
+		/// <param name="FixChange">Fix changelist for the issue</param>
+		/// <param name="ResolvedById">Whether the issue has been resolved</param>
+		/// <returns>True if the issue was updated</returns>
+		Task<bool> UpdateIssueAsync(int Id, string? Summary = null, ObjectId? OwnerId = null, ObjectId? NominatedById = null, bool? Acknowledged = null, ObjectId? DeclinedById = null, int? FixChange = null, ObjectId? ResolvedById = null);
+
+		/// <summary>
+		/// Marks a step as complete
+		/// </summary>
+		/// <param name="Job">The job to update</param>
+		/// <param name="Graph">Graph for the job</param>
+		/// <param name="BatchId">Unique id of the batch</param>
+		/// <param name="StepId">Unique id of the step</param>
+		/// <returns>Async task</returns>
+		Task UpdateCompleteStep(IJob Job, IGraph Graph, SubResourceId BatchId, SubResourceId StepId);
+	}
+}
diff --git a/Engine/Source/Programs/Horde/HordeServer/Issues/IIssueStep.cs b/Engine/Source/Programs/Horde/HordeServer/Issues/IIssueStep.cs
new file mode 100644
--- /dev/null
+++ b/Engine/Source/Programs/Horde/HordeServer/Issues/IIssueStep.cs
@@ -0,0 +1,68 @@
+// Copyright Epic Games, Inc. All Rights Reserved.
+
+using MongoDB.Bson;
+using System;
+using System.Collections.Generic;
+using System.Linq;
+using System.Threading.Tasks;
+
+using SubResourceId = HordeServer.Models.SubResourceId;
+
+namespace HordeServer.Models
+{
+	/// <summary>
+	/// Identifies a particular changelist and job that contributes to a span
+	/// </summary>
+	public interface IIssueStep
+	{
+		/// <summary>
+		/// The span that this step belongs to
+		/// </summary>
+		public ObjectId SpanId { get; }
+
+		/// <summary>
+		/// The changelist number
+		/// </summary>
+		public int Change { get; }
+
+		/// <summary>
+		/// Severity of the issue in this step
+		/// </summary>
+		public IssueSeverity Severity { get; }
+
+		/// <summary>
+		/// Name of the job
+		/// </summary>
+		public string JobName { get; }
+
+		/// <summary>
+		/// The unique job id
+		/// </summary>
+		public ObjectId JobId { get; }
+
+		/// <summary>
+		/// Unique id of the batch within the job
+		/// </summary>
+		public SubResourceId BatchId { get; }
+
+		/// <summary>
+		/// Unique id of the step within the job
+		/// </summary>
+		public SubResourceId StepId { get; }
+
+		/// <summary>
+		/// Time that the step started
+		/// </summary>
+		public DateTime StepTime { get; }
+
+		/// <summary>
+		/// Whether to notify suspects for this issue
+		/// </summary>
+		public bool NotifySuspects { get; }
+
+		/// <summary>
+		/// The log id for this step
+		/// </summary>
+		public ObjectId? LogId { get; }
+	}
+}
diff --git a/Engine/Source/Programs/Horde/HordeServer/Issues/Impl/Handlers/CompileIssueHandler.cs b/Engine/Source/Programs/Horde/HordeServer/Issues/Impl/Handlers/CompileIssueHandler.cs
new file mode 100644
--- /dev/null
+++ b/Engine/Source/Programs/Horde/HordeServer/Issues/Impl/Handlers/CompileIssueHandler.cs
@@ -0,0 +1,132 @@
+// Copyright Epic Games, Inc. All Rights Reserved.
+
+using HordeCommon;
+using HordeServer.Collections;
+using HordeCommon.Rpc;
+using HordeServer.Services;
+using HordeServer.Utilities;
+using Microsoft.Extensions.Logging;
+using MongoDB.Bson;
+using MongoDB.Bson.Serialization;
+using MongoDB.Driver;
+using System;
+using System.Collections.Generic;
+using System.Diagnostics.CodeAnalysis;
+using System.IO;
+using System.Linq;
+using System.Security.Policy;
+using System.Text;
+using System.Threading.Tasks;
+using System.Text.Json;
+using HordeServer.Models;
+using EpicGames.Core;
+
+namespace HordeServer.IssueHandlers.Impl
+{
+	/// <summary>
+	/// Instance of a particular compile error
+	/// </summary>
+	class CompileIssueHandler : IIssueHandler
+	{
+		/// <inheritdoc/>
+		public string Type => "Compile";
+
+		/// <inheritdoc/>
+		public int Priority => 10;
+
+		/// <summary>
+		/// Determines if the given event id matches
+		/// </summary>
+		/// <param name="EventId">The event id to compare</param>
+		/// <returns>True if the given event id matches</returns>
+		public static bool IsMatchingEventId(EventId? EventId)
+		{
+			return EventId == KnownLogEvents.Compiler || EventId == KnownLogEvents.AutomationTool_SourceFileLine || EventId == KnownLogEvents.MSBuild;
+		}
+
+		/// <summary>
+		/// Extracts a list of source files from an event
+		/// </summary>
+		/// <param name="Event">The event data</param>
+		/// <param name="SourceFiles">List of source files</param>
+		static void GetSourceFiles(ILogEventData Event, List<string> SourceFiles)
+		{
+			foreach (ILogEventLine Line in Event.Lines)
+			{
+				JsonElement Properties;
+				if (Line.Data.TryGetProperty("properties", out Properties) && Properties.ValueKind == JsonValueKind.Object)
+				{
+					foreach (JsonProperty Property in Properties.EnumerateObject())
+					{
+						if (Property.NameEquals("file") && Property.Value.ValueKind == JsonValueKind.String)
+						{
+							AddSourceFile(SourceFiles, Property.Value.GetString()!);
+						}
+						if (Property.Value.HasStringProperty("type", "SourceFile") && Property.Value.TryGetStringProperty("relativePath", out string? Value))
+						{
+							AddSourceFile(SourceFiles, Value);
+						}
+					}
+				}
+			}
+		}
+
+		/// <summary>
+		/// Add a new source file to a list of unique source files
+		/// </summary>
+		/// <param name="SourceFiles">List of source files</param>
+		/// <param name="RelativePath">File to add</param>
+		static void AddSourceFile(List<string> SourceFiles, string RelativePath)
+		{
+			int EndIdx = RelativePath.LastIndexOfAny(new char[] { '/', '\\' }) + 1;
+
+			string FileName = RelativePath.Substring(EndIdx);
+			if (!SourceFiles.Any(x => x.Equals(FileName, StringComparison.OrdinalIgnoreCase)))
+			{
+				SourceFiles.Add(FileName);
+			}
+		}
+
+		/// <inheritdoc/>
+		public bool TryGetFingerprint(IJob Job, INode Node, ILogEventData EventData, [NotNullWhen(true)] out NewIssueFingerprint? Fingerprint)
+		{
+			if (!IsMatchingEventId(EventData.EventId))
+			{
+				Fingerprint = null;
+				return false;
+			}
+
+			List<string> NewFileNames = new List<string>();
+			GetSourceFiles(EventData, NewFileNames);
+			Fingerprint = new NewIssueFingerprint(Type, NewFileNames, null);
+			return true;
+		}
+
+		/// <inheritdoc/>
+		public string GetSummary(IIssueFingerprint Fingerprint, IssueSeverity Severity)
+		{
+			string Type = (Severity == IssueSeverity.Warning) ? "Compile warnings" : "Compile errors";
+			string List = StringUtils.FormatList(Fingerprint.Keys.ToArray(), 2);
+			return $"{Type} in {List}";
+		}
+
+		/// <inheritdoc/>
+		public void RankSuspects(IIssueFingerprint Fingerprint, List<SuspectChange> Suspects)
+		{
+			foreach (SuspectChange Change in Suspects)
+			{
+				if (Change.ContainsCode)
+				{
+					if (Fingerprint.Keys.Any(x => Change.ModifiesFile(x)))
+					{
+						Change.Rank += 20;
+					}
+					else
+					{
+						Change.Rank += 10;
+					}
+				}
+			}
+		}
+	}
+}
diff --git a/Engine/Source/Programs/Horde/HordeServer/Issues/Impl/Handlers/ContentIssueHandler.cs b/Engine/Source/Programs/Horde/HordeServer/Issues/Impl/Handlers/ContentIssueHandler.cs
new file mode 100644
--- /dev/null
+++ b/Engine/Source/Programs/Horde/HordeServer/Issues/Impl/Handlers/ContentIssueHandler.cs
@@ -0,0 +1,105 @@
+// Copyright Epic Games, Inc. All Rights Reserved.
+
+using HordeCommon;
+using HordeServer.Collections;
+using HordeServer.Models;
+using HordeServer.Utilities;
+using Microsoft.Extensions.Logging;
+using MongoDB.Bson;
+using MongoDB.Bson.Serialization;
+using MongoDB.Driver;
+using System;
+using System.Collections.Generic;
+using System.Diagnostics.CodeAnalysis;
+using System.IO;
+using System.Linq;
+using System.Security.Policy;
+using System.Text;
+using System.Threading.Tasks;
+using EpicGames.Core;
+
+namespace HordeServer.IssueHandlers.Impl
+{
+	/// <summary>
+	/// Instance of a particular compile error
+	/// </summary>
+	class ContentIssueHandler : IIssueHandler
+	{
+		/// <inheritdoc/>
+		public string Type => "Content";
+
+		/// <inheritdoc/>
+		public int Priority => 10;
+
+		/// <summary>
+		/// Determines if the given event id matches
+		/// </summary>
+		/// <param name="EventId">The event id to compare</param>
+		/// <returns>True if the given event id matches</returns>
+		public static bool IsMatchingEventId(EventId? EventId)
+		{
+			return EventId == KnownLogEvents.Engine_AssetLog;
+		}
+
+		/// <summary>
+		/// Adds all the assets from the given log event
+		/// </summary>
+		/// <param name="Event">The log event to parse</param>
+		/// <param name="AssetNames">Receives the referenced asset names</param>
+		public static void GetAssetNames(ILogEventData Event, HashSet<string> AssetNames)
+		{
+			foreach (ILogEventLine Line in Event.Lines)
+			{
+				string? RelativePath;
+				if (Line.Data.TryGetNestedProperty("properties.asset.relativePath", out RelativePath))
+				{
+					int EndIdx = RelativePath.LastIndexOfAny(new char[] { '/', '\\' }) + 1;
+					string FileName = RelativePath.Substring(EndIdx);
+					AssetNames.Add(FileName);
+				}
+			}
+		}
+
+		/// <inheritdoc/>
+		public bool TryGetFingerprint(IJob Job, INode Node, ILogEventData EventData, [NotNullWhen(true)] out NewIssueFingerprint? Fingerprint)
+		{
+			if (!IsMatchingEventId(EventData.EventId))
+			{
+				Fingerprint = null;
+				return false;
+			}
+
+			HashSet<string> NewAssetNames = new HashSet<string>();
+			GetAssetNames(EventData, NewAssetNames);
+			Fingerprint = new NewIssueFingerprint(Type, NewAssetNames, null);
+			return true;
+		}
+
+		/// <inheritdoc/>
+		public string GetSummary(IIssueFingerprint Fingerprint, IssueSeverity Severity)
+		{
+			string Type = (Severity == IssueSeverity.Warning) ? "Warnings" : "Errors";
+			string List = StringUtils.FormatList(Fingerprint.Keys.ToArray(), 2);
+			return $"{Type} in {List}";
+		}
+
+		/// <inheritdoc/>
+		public void RankSuspects(IIssueFingerprint Fingerprint, List<SuspectChange> Suspects)
+		{
+			foreach (SuspectChange Suspect in Suspects)
+			{
+				if (Suspect.ContainsContent)
+				{
+					if (Suspect.Details.Files.Any(x => Fingerprint.Keys.Any(y => x.Path.Contains(y, StringComparison.OrdinalIgnoreCase))))
+					{
+						Suspect.Rank += 20;
+					}
+					else
+					{
+						Suspect.Rank += 10;
+					}
+				}
+			}
+		}
+	}
+}
diff --git a/Engine/Source/Programs/Horde/HordeServer/Issues/Impl/Handlers/CopyrightIssueHandler.cs b/Engine/Source/Programs/Horde/HordeServer/Issues/Impl/Handlers/CopyrightIssueHandler.cs
new file mode 100644
--- /dev/null
+++ b/Engine/Source/Programs/Horde/HordeServer/Issues/Impl/Handlers/CopyrightIssueHandler.cs
@@ -0,0 +1,109 @@
+// Copyright Epic Games, Inc. All Rights Reserved.
+
+using HordeCommon;
+using HordeServer.Collections;
+using HordeCommon.Rpc;
+using HordeServer.Models;
+using HordeServer.Utilities;
+using Microsoft.Extensions.Logging;
+using MongoDB.Bson;
+using MongoDB.Bson.Serialization;
+using MongoDB.Driver;
+using System;
+using System.Collections.Generic;
+using System.Diagnostics.CodeAnalysis;
+using System.IO;
+using System.Linq;
+using System.Security.Policy;
+using System.Text;
+using System.Threading.Tasks;
+using EpicGames.Core;
+
+namespace HordeServer.IssueHandlers.Impl
+{
+	/// <summary>
+	/// Instance of a particular compile error
+	/// </summary>
+	class CopyrightIssueHandler : IIssueHandler
+	{
+		/// <inheritdoc/>
+		public string Type => "Copyright";
+
+		/// <inheritdoc/>
+		public int Priority => 10;
+
+		/// <summary>
+		/// Filenames containing errors or warnings
+		/// </summary>
+		public HashSet<string> FileNames { get; set; } = new HashSet<string>();
+
+		/// <summary>
+		/// Determines if the given event id matches
+		/// </summary>
+		/// <param name="EventId">The event id to compare</param>
+		/// <returns>True if the given event id matches</returns>
+		public static bool IsMatchingEventId(EventId? EventId)
+		{
+			return EventId == KnownLogEvents.AutomationTool_MissingCopyright;
+		}
+
+		/// <summary>
+		/// Extracts a list of source files from an event
+		/// </summary>
+		/// <param name="Event">The event data</param>
+		/// <param name="SourceFiles">List of source files</param>
+		public static void GetSourceFiles(ILogEventData Event, HashSet<string> SourceFiles)
+		{
+			foreach (ILogEventLine Line in Event.Lines)
+			{
+				string? RelativePath;
+				if (Line.Data.TryGetNestedProperty("properties.file.relativePath", out RelativePath) || Line.Data.TryGetNestedProperty("properties.file", out RelativePath))
+				{
+					int EndIdx = RelativePath.LastIndexOfAny(new char[] { '/', '\\' }) + 1;
+					string FileName = RelativePath.Substring(EndIdx);
+					SourceFiles.Add(FileName);
+				}
+			}
+		}
+
+		/// <inheritdoc/>
+		public bool TryGetFingerprint(IJob Job, INode Node, ILogEventData EventData, [NotNullWhen(true)] out NewIssueFingerprint? Fingerprint)
+		{
+			if (!IsMatchingEventId(EventData.EventId))
+			{
+				Fingerprint = null;
+				return false;
+			}
+
+			HashSet<string> NewFileNames = new HashSet<string>(StringComparer.OrdinalIgnoreCase);
+			GetSourceFiles(EventData, NewFileNames);
+			Fingerprint = new NewIssueFingerprint(Type, NewFileNames, null);
+			return true;
+		}
+
+		/// <inheritdoc/>
+		public void RankSuspects(IIssueFingerprint Fingerprint, List<SuspectChange> Suspects)
+		{
+			foreach (SuspectChange Suspect in Suspects)
+			{
+				if (Suspect.ContainsCode)
+				{
+					if (FileNames.Any(x => Suspect.ModifiesFile(x)))
+					{
+						Suspect.Rank += 20;
+					}
+					else
+					{
+						Suspect.Rank += 10;
+					}
+				}
+			}
+		}
+
+		/// <inheritdoc/>
+		public string GetSummary(IIssueFingerprint Fingerprint, IssueSeverity Severity)
+		{
+			return $"Missing copyright notice in {StringUtils.FormatList(Fingerprint.Keys.ToArray(), 2)}";
+		}
+	}
+}
diff --git a/Engine/Source/Programs/Horde/HordeServer/Issues/Impl/Handlers/DefaultIssueHandler.cs b/Engine/Source/Programs/Horde/HordeServer/Issues/Impl/Handlers/DefaultIssueHandler.cs
new file mode 100644
--- /dev/null
+++ b/Engine/Source/Programs/Horde/HordeServer/Issues/Impl/Handlers/DefaultIssueHandler.cs
@@ -0,0 +1,55 @@
+// Copyright Epic Games, Inc. All Rights Reserved.
+
+using HordeServer.Collections;
+using HordeServer.Models;
+using System;
+using System.Collections.Generic;
+using System.Diagnostics.CodeAnalysis;
+using System.Linq;
+using System.Threading.Tasks;
+
+namespace HordeServer.IssueHandlers.Impl
+{
+	/// <summary>
+	/// Instance of a particular compile error
+	/// </summary>
+	class DefaultIssueHandler : IIssueHandler
+	{
+		/// <summary>
+		/// Name of the handler
+		/// </summary>
+		public const string Type = "Default";
+
+		/// <inheritdoc/>
+		string IIssueHandler.Type => Type;
+
+		/// <inheritdoc/>
+		public int Priority => 0;
+
+		/// <inheritdoc/>
+		public bool TryGetFingerprint(IJob Job, INode Node, ILogEventData EventData, [NotNullWhen(true)] out NewIssueFingerprint? Fingerprint)
+		{
+			Fingerprint = new NewIssueFingerprint(Type, new[] { Node.Name }, null);
+			return true;
+		}
+
+		/// <inheritdoc/>
+		public void RankSuspects(IIssueFingerprint Fingerprint, List<SuspectChange> Suspects)
+		{
+		}
+
+		/// <inheritdoc/>
+		public string GetSummary(IIssueFingerprint Fingerprint, IssueSeverity Severity)
+		{
+			string NodeName = Fingerprint.Keys.FirstOrDefault() ?? "(unknown)";
+			if(Severity == IssueSeverity.Warning)
+			{
+				return $"Warnings in {NodeName}";
+			}
+			else
+			{
+				return $"Errors in {NodeName}";
+			}
+		}
+	}
+}
diff --git a/Engine/Source/Programs/Horde/HordeServer/Issues/Impl/Handlers/GauntletIssueHandler.cs b/Engine/Source/Programs/Horde/HordeServer/Issues/Impl/Handlers/GauntletIssueHandler.cs
new file mode 100644
--- /dev/null
+++ b/Engine/Source/Programs/Horde/HordeServer/Issues/Impl/Handlers/GauntletIssueHandler.cs
@@ -0,0 +1,177 @@
+// Copyright Epic Games, Inc. All Rights Reserved.
+
+using HordeCommon;
+using HordeServer.Models;
+using HordeServer.Utilities;
+using Microsoft.AspNetCore.Authentication;
+using Microsoft.AspNetCore.Connections.Features;
+using Microsoft.Extensions.Logging;
+using Serilog.Formatting.Json;
+using System;
+using System.Collections.Generic;
+using System.Linq;
+using System.Security.Cryptography.X509Certificates;
+using System.Text.Json;
+using System.Threading.Tasks;
+using System.Diagnostics.CodeAnalysis;
+using HordeServer.Collections;
+using EpicGames.Core;
+
+namespace HordeServer.IssueHandlers.Impl
+{
+	/// <summary>
+	/// Instance of a particular Gauntlet error
+	/// </summary>
+	class GauntletIssueHandler : IIssueHandler
+	{
+		/// <summary>
+		/// Prefix for unit test keys
+		/// </summary>
+		const string UnitTestPrefix = "UnitTest:";
+
+		/// <summary>
+		/// Prefix for screenshot test keys
+		/// </summary>
+		const string ScreenshotTestPrefix = "Screenshot:";
+
+		/// <inheritdoc/>
+		public string Type => "Gauntlet";
+
+		/// <inheritdoc/>
+		public int Priority => 10;
+
+		/// <summary>
+		/// Determines if the given event id matches
+		/// </summary>
+		/// <param name="EventId">The event id to compare</param>
+		/// <returns>True if the given event id matches</returns>
+		public static bool IsMatchingEventId(EventId? EventId)
+		{
+			return EventId == KnownLogEvents.Gauntlet_UnitTest || EventId == KnownLogEvents.Gauntlet_ScreenshotTest;
+		}
+
+		/// <summary>
+		/// Parses symbol names from a log event
+		/// </summary>
+		/// <param name="EventData">The log event data</param>
+		/// <param name="UnitTestNames">Receives a set of the unit test names</param>
+		private static void GetUnitTestNames(ILogEventData EventData, HashSet<string> UnitTestNames)
+		{
+			string? Group = null;
+			string? Name = null;
+
+			foreach (ILogEventLine Line in EventData.Lines)
+			{
+				string? Value;
+				if (Line.Data.TryGetNestedProperty("properties.group", out Value))
+				{
+					Group = Value;
+				}
+				if (Line.Data.TryGetNestedProperty("properties.name", out Value))
+				{
+					Name = Value;
+				}
+			}
+
+			if(Group != null && Name != null)
+			{
+				UnitTestNames.Add($"{UnitTestPrefix}:{Group}/{Name}");
+			}
+		}
+
+		/// <summary>
+		/// Parses screenshot test names from a log event
+		/// </summary>
+		/// <param name="EventData">The event data</param>
+		/// <param name="ScreenshotTestNames">Receives the parsed screenshot test names</param>
+		private static void GetScreenshotTestNames(ILogEventData EventData, HashSet<string> ScreenshotTestNames)
+		{
+			foreach (ILogEventLine Line in EventData.Lines)
+			{
+				if (Line.Data.TryGetProperty("properties", JsonValueKind.Object, out JsonElement Properties))
+				{
+					foreach (JsonProperty Property in Properties.EnumerateObject())
+					{
+						JsonElement Value = Property.Value;
+						if (Value.ValueKind == JsonValueKind.Object && Value.HasStringProperty("type", "Screenshot") && Value.TryGetStringProperty("name", out string? Name))
+						{
+							ScreenshotTestNames.Add($"{ScreenshotTestPrefix}:{Name}");
+						}
+					}
+				}
+			}
+		}
+
+		/// <inheritdoc/>
+		public bool TryGetFingerprint(IJob Job, INode Node, ILogEventData EventData, [NotNullWhen(true)] out NewIssueFingerprint? Fingerprint)
+		{
+			if(!IsMatchingEventId(EventData.EventId))
+			{
+				Fingerprint = null;
+				return false;
+			}
+
+			HashSet<string> Keys = new HashSet<string>();
+			GetUnitTestNames(EventData, Keys);
+			GetScreenshotTestNames(EventData, Keys);
+
+			if (Keys.Count == 0)
+			{
+				Fingerprint = null;
+				return false;
+			}
+
+			Fingerprint = new NewIssueFingerprint(Type, Keys, null);
+			return true;
+		}
+
+		/// <inheritdoc/>
+		public void RankSuspects(IIssueFingerprint Fingerprint, List<SuspectChange> Changes)
+		{
+			foreach (SuspectChange Change in Changes)
+			{
+				if (Change.ContainsCode)
+				{
+					Change.Rank += 10;
+				}
+			}
+		}
+
+		/// <inheritdoc/>
+		public string GetSummary(IIssueFingerprint Fingerprint, IssueSeverity Severity)
+		{
+			List<string> UnitTestNames = Fingerprint.Keys.Where(x => x.StartsWith(UnitTestPrefix, StringComparison.Ordinal)).Select(x => x.Substring(UnitTestPrefix.Length + 1)).ToList();
+			if (UnitTestNames.Count > 0)
+			{
+				HashSet<string> GroupNames = new HashSet<string>(StringComparer.OrdinalIgnoreCase);
+				HashSet<string> TestNames = new HashSet<string>(StringComparer.OrdinalIgnoreCase);
+				foreach (string UnitTestName in UnitTestNames)
+				{
+					int Idx = UnitTestName.IndexOf('/', StringComparison.OrdinalIgnoreCase);
+					if (Idx != -1)
+					{
+						GroupNames.Add(UnitTestName.Substring(0, Idx));
+					}
+					TestNames.Add(UnitTestName.Substring(Idx + 1));
+				}
+
+				if (GroupNames.Count == 1)
+				{
+					return $"{GroupNames.First()} test failures: {StringUtils.FormatList(TestNames.ToArray(), 3)}";
+				}
+				else
+				{
+					return $"{StringUtils.FormatList(GroupNames.OrderBy(x => x).ToArray(), 100)} test failures";
+				}
+			}
+
+			List<string> ScreenshotTestNames = Fingerprint.Keys.Where(x => x.StartsWith(ScreenshotTestPrefix, StringComparison.Ordinal)).Select(x => x.Substring(ScreenshotTestPrefix.Length + 1)).ToList();
+			if (ScreenshotTestNames.Count > 0)
+			{
+				return $"Screenshot test failures: {StringUtils.FormatList(ScreenshotTestNames.ToArray(), 3)}";
+			}
+
+			return "Test failures";
+		}
+	}
+}
diff --git a/Engine/Source/Programs/Horde/HordeServer/Issues/Impl/Handlers/SymbolIssueHandler.cs b/Engine/Source/Programs/Horde/HordeServer/Issues/Impl/Handlers/SymbolIssueHandler.cs
new file mode 100644
--- /dev/null
+++ b/Engine/Source/Programs/Horde/HordeServer/Issues/Impl/Handlers/SymbolIssueHandler.cs
@@ -0,0 +1,105 @@
+// Copyright Epic Games, Inc. All Rights Reserved.
+
+using HordeCommon;
+using HordeServer.Collections;
+using HordeServer.Models;
+using HordeServer.Utilities;
+using Microsoft.Extensions.Logging;
+using MongoDB.Bson;
+using MongoDB.Bson.Serialization;
+using System;
+using System.Collections.Generic;
+using System.Diagnostics.CodeAnalysis;
+using System.Linq;
+using System.Text.RegularExpressions;
+using System.Threading.Tasks;
+using EpicGames.Core;
+
+namespace HordeServer.IssueHandlers.Impl
+{
+	/// <summary>
+	/// Instance of a particular compile error
+	/// </summary>
+	class SymbolIssueHandler : IIssueHandler
+	{
+		/// <inheritdoc/>
+		public string Type => "Symbol";
+
+		/// <inheritdoc/>
+		public int Priority => 10;
+
+		/// <summary>
+		/// Determines if the given event id matches
+		/// </summary>
+		/// <param name="EventId">The event id to compare</param>
+		/// <returns>True if the given event id matches</returns>
+		public static bool IsMatchingEventId(EventId? EventId)
+		{
+			return EventId == KnownLogEvents.Linker_UndefinedSymbol || EventId == KnownLogEvents.Linker_DuplicateSymbol || EventId == KnownLogEvents.Linker;
+		}
+
+		/// <summary>
+		/// Parses symbol names from a log event
+		/// </summary>
+		/// <param name="EventData">The log event data</param>
+		/// <param name="SymbolNames">Receives the list of symbol names</param>
+		public static void GetSymbolNames(ILogEventData EventData, SortedSet<string> SymbolNames)
+		{
+			foreach (ILogEventLine Line in EventData.Lines)
+			{
+				string? Identifier;
+				if (Line.Data.TryGetNestedProperty("properties.symbol.identifier", out Identifier))
+				{
+					SymbolNames.Add(Identifier);
+				}
+			}
+		}
+
+		/// <inheritdoc/>
+		public void RankSuspects(IIssueFingerprint Fingerprint, List<SuspectChange> Changes)
+		{
+			HashSet<string> Names = new HashSet<string>();
+			foreach (string Name in Fingerprint.Keys)
+			{
+				Names.UnionWith(Name.Split("::", StringSplitOptions.RemoveEmptyEntries));
+			}
+
+			foreach (SuspectChange Change in Changes)
+			{
+				if (Change.ContainsCode)
+				{
+					int Matches = Names.Count(x => Change.Details.Files.Any(y => y.Path.Contains(x, StringComparison.OrdinalIgnoreCase)));
+					Change.Rank += 10 + (10 * Matches);
+				}
+			}
+		}
+
+		/// <inheritdoc/>
+		public string GetSummary(IIssueFingerprint Fingerprint, IssueSeverity Severity)
+		{
+			HashSet<string> Symbols = Fingerprint.Keys;
+			if (Symbols.Count == 1)
+			{
+				return $"Undefined symbol '{Symbols.First()}'";
+			}
+			else
+			{
+				return $"Undefined symbols: {StringUtils.FormatList(Symbols.ToArray(), 3)}";
+			}
+		}
+
+		public bool TryGetFingerprint(IJob Job, INode Node, ILogEventData EventData, [NotNullWhen(true)] out NewIssueFingerprint? Fingerprint)
+		{
+			if (!IsMatchingEventId(EventData.EventId))
+			{
+				Fingerprint = null;
+				return false;
+			}
+
+			SortedSet<string> SymbolNames = new SortedSet<string>();
+			GetSymbolNames(EventData, SymbolNames);
+			Fingerprint = new NewIssueFingerprint(Type, SymbolNames, null);
+			return true;
+		}
+	}
+}
