Horde: More file moves.

diff --git a/Engine/Source/Programs/Horde/HordeServer/Jobs/Impl/ArtifactService.cs b/Engine/Source/Programs/Horde/HordeServer/Jobs/Impl/ArtifactService.cs
new file mode 100644
--- /dev/null
+++ b/Engine/Source/Programs/Horde/HordeServer/Jobs/Impl/ArtifactService.cs
@@ -0,0 +1,261 @@
+// Copyright Epic Games, Inc. All Rights Reserved.
+
+using Amazon.Extensions.NETCore.Setup;
+using Amazon.Runtime;
+using Amazon.S3;
+using Amazon.S3.Model;
+using HordeServer.Api;
+using HordeServer.Models;
+using HordeServer.Storage;
+using HordeServer.Utilities;
+using Microsoft.Extensions.Configuration;
+using Microsoft.Extensions.Hosting;
+using Microsoft.Extensions.Logging;
+using Microsoft.Extensions.Options;
+using MongoDB.Bson;
+using MongoDB.Driver;
+using System;
+using System.Collections.Generic;
+using System.IO;
+using System.Linq;
+using System.Security.Claims;
+using System.Text;
+using System.Threading;
+using System.Threading.Tasks;
+
+namespace HordeServer.Services
+{
+	/// <summary>
+	/// Wraps functionality for manipulating artifacts
+	/// </summary>
+	public class ArtifactService
+	{
+		/// <summary>
+		/// Backend storage container
+		/// </summary>
+		private readonly IStorageBackend StorageBackend;
+
+		/// <summary>
+		/// Information Logger
+		/// </summary>
+		private readonly ILogger<ArtifactService> Logger;
+
+		/// <summary>
+		/// Collection of log documents
+		/// </summary>
+		private readonly IMongoCollection<Artifact> Artifacts;
+
+		/// <summary>
+		/// Constructor
+		/// </summary>
+		/// <param name="DatabaseService">The database service</param>
+		/// <param name="StorageBackend">The storage backend</param>
+		/// <param name="Logger">Log interface</param>
+		public ArtifactService(DatabaseService DatabaseService, IStorageBackend<ArtifactService> StorageBackend, ILogger<ArtifactService> Logger)
+		{
+			this.StorageBackend = StorageBackend;
+			this.Logger = Logger;
+
+			// Initialize Artifacts table
+			Artifacts = DatabaseService.Artifacts;
+			if (!DatabaseService.ReadOnlyMode)
+			{
+				Artifacts.Indexes.CreateOne(new CreateIndexModel<Artifact>(Builders<Artifact>.IndexKeys.Ascending(x => x.JobId)));
+			}
+		}
+
+		/// <summary>
+		/// Creates a new artifact
+		/// </summary>
+		/// <param name="JobId">Unique id of the job that owns this artifact</param>
+		/// <param name="StepId">Optional Step id</param>
+		/// <param name="Name">Name of artifact</param>
+		/// <param name="MimeType">Type of artifact</param>
+		/// <param name="Data">The data to write</param>
+		/// <returns>The new log file document</returns>
+		public async Task<Artifact> CreateArtifactAsync(ObjectId JobId, SubResourceId? StepId, string Name, string MimeType, System.IO.Stream Data)
+		{
+			// upload first
+			string ArtifactName = ValidateName(Name);
+			await StorageBackend.WriteAsync(GetPath(JobId, StepId, ArtifactName), Data);
+
+			// then create entry
+			Artifact NewArtifact = new Artifact(JobId, StepId, ArtifactName, Data.Length, MimeType);
+			await Artifacts.InsertOneAsync(NewArtifact);
+			return NewArtifact;
+		}
+
+		/// <summary>
+		/// Gets all the available artifacts for a job
+		/// </summary>
+		/// <param name="JobId">Unique id of the job to query</param>
+		/// <param name="StepId">Unique id of the Step to query</param>
+		/// <param name="Name">Name of the artifact</param>
+		/// <returns>List of artifact documents</returns>
+		public async Task<List<Artifact>> GetArtifactsAsync(ObjectId? JobId, SubResourceId? StepId, string? Name)
+		{
+			FilterDefinitionBuilder<Artifact> Builder = Builders<Artifact>.Filter;
+
+			FilterDefinition<Artifact> Filter = FilterDefinition<Artifact>.Empty;
+			if (JobId != null)
+			{
+				Filter &= Builder.Eq(x => x.JobId, JobId.Value);
+			}
+			if (StepId != null)
+			{
+				Filter &= Builder.Eq(x => x.StepId, StepId.Value);
+			}
+			if (Name != null)
+			{
+				Filter &= Builder.Eq(x => x.Name, Name);
+			}
+
+			return await Artifacts.Find(Filter).ToListAsync();
+		}
+
+		/// <summary>
+		/// Gets a specific list of artifacts based on id
+		/// </summary>
+		/// <param name="ArtifactIds">The list of artifact Ids</param>
+		/// <returns>List of artifact documents</returns>
+		public async Task<List<Artifact>> GetArtifactsAsync(IEnumerable<ObjectId> ArtifactIds)
+		{
+			FilterDefinitionBuilder<Artifact> Builder = Builders<Artifact>.Filter;
+
+			FilterDefinition<Artifact> Filter = FilterDefinition<Artifact>.Empty;
+			Filter &= Builder.In(x => x.Id, ArtifactIds);
+			
+			return await Artifacts.Find(Filter).ToListAsync();
+		}
+
+		/// <summary>
+		/// Gets an artifact by ID
+		/// </summary>
+		/// <param name="ArtifactId">Unique id of the artifact</param>
+		/// <returns>The artifact document</returns>
+		public async Task<Artifact?> GetArtifactAsync(ObjectId ArtifactId)
+		{
+			return await Artifacts.Find<Artifact>(x => x.Id == ArtifactId).FirstOrDefaultAsync();
+		}
+
+		/// <summary>
+		/// Attempts to update a single artifact document, if the update counter is valid
+		/// </summary>
+		/// <param name="Current">The current artifact document</param>
+		/// <param name="Update">The update definition</param>
+		/// <returns>True if the document was updated, false if another writer updated the document first</returns>
+		private async Task<bool> TryUpdateArtifactAsync(Artifact Current, UpdateDefinition<Artifact> Update)
+		{
+			UpdateResult Result = await Artifacts.UpdateOneAsync<Artifact>(x => x.Id == Current.Id && x.UpdateIndex == Current.UpdateIndex, Update.Set(x => x.UpdateIndex, Current.UpdateIndex + 1));
+			return Result.ModifiedCount == 1;
+		}
+
+		/// <summary>
+		/// Updates an artifact
+		/// </summary>
+		/// <param name="Artifact">The artifact</param>
+		/// <param name="NewMimeType">New mime type</param>
+		/// <param name="NewData">New data</param>
+		/// <returns>Async task</returns>
+		public async Task<bool> UpdateArtifactAsync(Artifact? Artifact, string NewMimeType, System.IO.Stream NewData)
+		{
+			while (Artifact != null)
+			{
+				UpdateDefinitionBuilder<Artifact> UpdateBuilder = Builders<Artifact>.Update;
+
+				List<UpdateDefinition<Artifact>> Updates = new List<UpdateDefinition<Artifact>>();
+				Updates.Add(UpdateBuilder.Set(x => x.MimeType, NewMimeType));
+				Updates.Add(UpdateBuilder.Set(x => x.Length, NewData.Length));
+
+				// re-upload the data to external
+				string ArtifactName = ValidateName(Artifact.Name);
+				await StorageBackend.WriteAsync(GetPath(Artifact.JobId, Artifact.StepId, ArtifactName), NewData);
+
+				if (await TryUpdateArtifactAsync(Artifact, UpdateBuilder.Combine(Updates)))
+				{
+					return true;
+				}
+
+				Artifact = await GetArtifactAsync(Artifact.Id);
+			}
+			return false;
+		}
+
+		/// <summary>
+		/// gets artifact data
+		/// </summary>
+		/// <param name="Artifact">The artifact</param>
+		/// <returns>The chunk data</returns>
+		public async Task<System.IO.Stream> OpenArtifactReadStreamAsync(Artifact Artifact)
+		{
+			System.IO.Stream? Stream = await StorageBackend.ReadAsync(GetPath(Artifact.JobId, Artifact.StepId, Artifact.Name));
+			if (Stream == null)
+			{
+				throw new Exception($"Unable to get artifact {Artifact.Id}");
+			}
+			return Stream;
+		}
+
+		/// <summary>
+		/// Get the path for an artifact
+		/// </summary>
+		/// <param name="JobId"></param>
+		/// <param name="StepId"></param>
+		/// <param name="Name"></param>
+		/// <returns></returns>
+		private static string GetPath(ObjectId JobId, SubResourceId? StepId, string Name)
+		{
+			if (StepId == null)
+			{
+				return $"{JobId}/{Name}";
+			}
+			else
+			{
+				return $"{JobId}/{StepId.Value}/{Name}";
+			}
+		}
+
+		/// <summary>
+		/// Checks that the name given is valid for an artifact
+		/// </summary>
+		/// <param name="Name"></param>
+		private static string ValidateName(string Name)
+		{
+			string NewName = Name.Replace('\\', '/');
+			if (NewName.Length == 0)
+			{
+				throw new ArgumentException("Artifact name is empty");
+			}
+			else if (NewName.StartsWith('/'))
+			{
+				throw new ArgumentException($"Artifact has an absolute path ({NewName})");
+			}
+			else if (NewName.EndsWith('/'))
+			{
+				throw new ArgumentException($"Artifact does not have a file name ({NewName})");
+			}
+			else if (NewName.Contains("//", StringComparison.Ordinal))
+			{
+				throw new ArgumentException($"Artifact name contains an invalid directory ({NewName})");
+			}
+
+			string InvalidChars = ":<>|\"";
+			for (int Idx = 0; Idx < NewName.Length; Idx++)
+			{
+				int MinIdx = Idx;
+				for (; Idx < NewName.Length && NewName[Idx] != '/'; Idx++)
+				{
+					if (InvalidChars.Contains(NewName[Idx], StringComparison.Ordinal))
+					{
+						throw new ArgumentException($"Invalid character in artifact name ({NewName})");
+					}
+				}
+				if ((Idx == MinIdx + 1 && NewName[MinIdx] == '.') || (Idx == MinIdx + 2 && NewName[MinIdx] == '.' && NewName[MinIdx + 1] == '.'))
+				{
+					throw new ArgumentException($"Artifact may not contain symbolic directory names ({NewName})");
+				}
+			}
+			return NewName;
+		}
+	}
+}
\ No newline at end of file
diff --git a/Engine/Source/Programs/Horde/HordeServer/Jobs/Impl/JobTaskSource.cs b/Engine/Source/Programs/Horde/HordeServer/Jobs/Impl/JobTaskSource.cs
new file mode 100644
--- /dev/null
+++ b/Engine/Source/Programs/Horde/HordeServer/Jobs/Impl/JobTaskSource.cs
@@ -0,0 +1,1096 @@
+// Copyright Epic Games, Inc. All Rights Reserved.
+
+using EpicGames.Core;
+using EpicGames.Horde.Common;
+using Google.Protobuf;
+using Google.Protobuf.Reflection;
+using Google.Protobuf.WellKnownTypes;
+using HordeServer.Api;
+using HordeServer.Collections;
+using HordeCommon;
+using HordeServer.Models;
+using HordeCommon.Rpc;
+using HordeCommon.Rpc.Tasks;
+using HordeServer.Tasks;
+using HordeServer.Utilities;
+using Microsoft.Extensions.Hosting;
+using Microsoft.Extensions.Logging;
+using Microsoft.Extensions.Options;
+using MongoDB.Bson;
+using MongoDB.Driver;
+using System;
+using System.Collections.Generic;
+using System.Diagnostics;
+using System.Diagnostics.CodeAnalysis;
+using System.Linq;
+using System.Runtime.CompilerServices;
+using System.Text;
+using System.Threading;
+using System.Threading.Tasks;
+using HordeServer.Services;
+using EpicGames.Horde.Compute;
+
+namespace HordeServer.Tasks.Impl
+{
+	using PoolId = StringId<IPool>;
+	using StreamId = StringId<IStream>;
+
+	/// <summary>
+	/// Background service to dispatch pending work to agents in priority order.
+	/// </summary>
+	public sealed class JobTaskSource : TaskSourceBase<ExecuteJobTask>, IDisposable
+	{
+		/// <inheritdoc/>
+		public override string Type => "Job";
+
+		/// <summary>
+		/// An item in the queue to be executed
+		/// </summary>
+		class QueueItem
+		{
+			/// <summary>
+			/// The stream for this job
+			/// </summary>
+			public IStream Stream;
+
+			/// <summary>
+			/// The job instance
+			/// </summary>
+			public IJob Job;
+
+			/// <summary>
+			/// Index of the batch within this job to be executed
+			/// </summary>
+			public int BatchIdx;
+
+			/// <summary>
+			/// The pool of machines to allocate from
+			/// </summary>
+			public PoolId PoolId;
+
+			/// <summary>
+			/// The type of workspace that this item should run in
+			/// </summary>
+			public AgentWorkspace Workspace;
+
+			/// <summary>
+			/// Task for creating a lease and assigning to a waiter
+			/// </summary>
+			public Task? AssignTask;
+
+			/// <summary>
+			/// Accessor for the batch referenced by this item
+			/// </summary>
+			public IJobStepBatch Batch
+			{
+				get { return Job.Batches[BatchIdx]; }
+			}
+
+			/// <summary>
+			/// Returns an identifier describing this unique batch
+			/// </summary>
+			public (ObjectId, SubResourceId) Id
+			{
+				get { return (Job.Id, Batch.Id); }
+			}
+
+			/// <summary>
+			/// Constructor
+			/// </summary>
+			/// <param name="Stream">The stream containing this job</param>
+			/// <param name="Job">The job instance</param>
+			/// <param name="BatchIdx">The batch index to execute</param>
+			/// <param name="PoolId">Unique id of the pool of machines to allocate from</param>
+			/// <param name="Workspace">The workspace that this job should run in</param>
+			public QueueItem(IStream Stream, IJob Job, int BatchIdx, PoolId PoolId, AgentWorkspace Workspace)
+			{
+				this.Stream = Stream;
+				this.Job = Job;
+				this.BatchIdx = BatchIdx;
+				this.PoolId = PoolId;
+				this.Workspace = Workspace;
+			}
+		}
+
+		/// <summary>
+		/// Comparer for items in the queue
+		/// </summary>
+		class QueueItemComparer : IComparer<QueueItem>
+		{
+			/// <summary>
+			/// Compare two items
+			/// </summary>
+			/// <param name="X">First item to compare</param>
+			/// <param name="Y">Second item to compare</param>
+			/// <returns>Negative value if X is a higher priority than Y</returns>
+			public int Compare([AllowNull] QueueItem X, [AllowNull] QueueItem Y)
+			{
+				if (X == null)
+				{
+					return (Y == null) ? 0 : -1;
+				}
+				else if (Y == null)
+				{
+					return 1;
+				}
+
+				int Delta = Y.Batch.SchedulePriority - X.Batch.SchedulePriority;
+				if (Delta == 0)
+				{
+					Delta = X.Job.Id.CompareTo(Y.Job.Id);
+					if (Delta == 0)
+					{
+						Delta = (int)X.Batch.Id.Value - (int)Y.Batch.Id.Value;
+					}
+				}
+				return Delta;
+			}
+		}
+
+		/// <summary>
+		/// Information about an agent waiting for work
+		/// </summary>
+		class QueueWaiter
+		{
+			/// <summary>
+			/// The agent performing the wait
+			/// </summary>
+			public IAgent Agent { get; }
+
+			/// <summary>
+			/// Task to wait for a lease to be assigned
+			/// </summary>
+			public Task<AgentLease?> Task => LeaseSource.Task;
+
+			/// <summary>
+			/// Completion source for the waiting agent. If a new queue item becomes available, the result will be passed through this.
+			/// </summary>
+			public TaskCompletionSource<AgentLease?> LeaseSource { get; } = new TaskCompletionSource<AgentLease?>();
+
+			/// <summary>
+			/// Constructor
+			/// </summary>
+			/// <param name="Agent">The agent waiting for a task</param>
+			public QueueWaiter(IAgent Agent)
+			{
+				this.Agent = Agent;
+			}
+		}
+
+		/// <summary>
+		/// The database service instance
+		/// </summary>
+		DatabaseService DatabaseService;
+
+		/// <summary>
+		/// The stream service instance
+		/// </summary>
+		StreamService StreamService;
+
+		/// <summary>
+		/// The log file service instance
+		/// </summary>
+		ILogFileService LogFileService;
+
+		/// <summary>
+		/// Collection of agent documents
+		/// </summary>
+		IAgentCollection AgentsCollection;
+
+		/// <summary>
+		/// Collection of job documents
+		/// </summary>
+		IJobCollection Jobs;
+
+		/// <summary>
+		/// Collection of jobstepref documents
+		/// </summary>
+		IJobStepRefCollection JobStepRefs;
+
+		/// <summary>
+		/// Collection of graph documents
+		/// </summary>
+		IGraphCollection Graphs;
+
+		/// <summary>
+		/// Collection of pool documents
+		/// </summary>
+		IPoolCollection PoolCollection;
+
+		/// <summary>
+		/// Collection of UGS metadata documents
+		/// </summary>
+		IUgsMetadataCollection UgsMetadataCollection;
+
+		/// <summary>
+		/// The Perforce load balancer
+		/// </summary>
+		PerforceLoadBalancer PerforceLoadBalancer;
+
+		/// <summary>
+		/// The application lifetime
+		/// </summary>
+		IHostApplicationLifetime ApplicationLifetime;
+
+		/// <summary>
+		/// Settings instance
+		/// </summary>
+		IOptionsMonitor<ServerSettings> Settings;
+
+		/// <summary>
+		/// Log writer
+		/// </summary>
+		ILogger<JobTaskSource> Logger;
+
+		/// <summary>
+		/// Object used for ensuring mutual exclusion to the queues
+		/// </summary>
+		object LockObject = new object();
+
+		/// <summary>
+		/// List of items waiting to be executed
+		/// </summary>
+		SortedSet<QueueItem> Queue = new SortedSet<QueueItem>(new QueueItemComparer());
+
+		/// <summary>
+		/// Map from batch id to the corresponding queue item
+		/// </summary>
+		Dictionary<(ObjectId, SubResourceId), QueueItem> BatchIdToQueueItem = new Dictionary<(ObjectId, SubResourceId), QueueItem>();
+
+		/// <summary>
+		/// Set of long-poll tasks waiting to be satisfied 
+		/// </summary>
+		HashSet<QueueWaiter> Waiters = new HashSet<QueueWaiter>();
+
+		/// <summary>
+		/// During a background queue refresh operation, any updated batches are added to this dictionary for merging into the updated queue.
+		/// </summary>
+		List<QueueItem>? NewQueueItemsDuringUpdate;
+
+		/// <summary>
+		/// Cache of pools
+		/// </summary>
+		Dictionary<PoolId, IPool> CachedPoolIdToInstance = new Dictionary<PoolId, IPool>();
+
+		/// <summary>
+		/// Cache of stream objects. Used to resolve agent types.
+		/// </summary>
+		private Dictionary<StreamId, IStream> Streams = new Dictionary<StreamId, IStream>();
+
+		/// <summary>
+		/// Background tick registration
+		/// </summary>
+		BackgroundTick Ticker;
+
+		/// <summary>
+		/// Interval between querying the database for jobs to execute
+		/// </summary>
+		static readonly TimeSpan RefreshInterval = TimeSpan.FromSeconds(5.0);
+
+		/// <summary>
+		/// Constructor
+		/// </summary>
+		/// <param name="DatabaseService">The database service instance</param>
+		/// <param name="Agents">The agents collection</param>
+		/// <param name="Jobs">The jobs collection</param>
+		/// <param name="JobStepRefs">Collection of jobstepref documents</param>
+		/// <param name="Graphs">The graphs collection</param>
+		/// <param name="Pools">The pools collection</param>
+		/// <param name="UgsMetadataCollection">Ugs metadata collection</param>
+		/// <param name="StreamService">The stream service instance</param>
+		/// <param name="LogFileService">The log file service instance</param>
+		/// <param name="PerforceLoadBalancer">Perforce load balancer</param>
+		/// <param name="ApplicationLifetime">Application lifetime interface</param>
+		/// <param name="Settings">Settings for the server</param>
+		/// <param name="Logger">Log writer</param>
+		public JobTaskSource(DatabaseService DatabaseService, IAgentCollection Agents, IJobCollection Jobs, IJobStepRefCollection JobStepRefs, IGraphCollection Graphs, IPoolCollection Pools, IUgsMetadataCollection UgsMetadataCollection, StreamService StreamService, ILogFileService LogFileService, PerforceLoadBalancer PerforceLoadBalancer, IHostApplicationLifetime ApplicationLifetime, IOptionsMonitor<ServerSettings> Settings, ILogger<JobTaskSource> Logger)
+		{
+			this.DatabaseService = DatabaseService;
+			this.AgentsCollection = Agents;
+			this.Jobs = Jobs;
+			this.JobStepRefs = JobStepRefs;
+			this.Graphs = Graphs;
+			this.PoolCollection = Pools;
+			this.UgsMetadataCollection = UgsMetadataCollection;
+			this.StreamService = StreamService;
+			this.LogFileService = LogFileService;
+			this.PerforceLoadBalancer = PerforceLoadBalancer;
+			this.ApplicationLifetime = ApplicationLifetime;
+			this.Ticker = new BackgroundTick(TickAsync, RefreshInterval, Logger);
+			this.Settings = Settings;
+			this.Logger = Logger;
+		}
+
+		/// <inheritdoc/>
+		public async ValueTask DisposeAsync()
+		{
+			await Ticker.DisposeAsync();
+		}
+
+		/// <inheritdoc/>
+		public void Dispose()
+		{
+			Ticker.Dispose();
+		}
+
+		/// <summary>
+		/// Gets an object containing the stats of the queue for diagnostic purposes.
+		/// </summary>
+		/// <returns>Status object</returns>
+		public object GetStatus()
+		{
+			lock (LockObject)
+			{
+				List<object> OutputItems = new List<object>();
+				foreach (QueueItem QueueItem in Queue)
+				{
+					OutputItems.Add(new { JobId = QueueItem.Job.Id.ToString(), BatchId = QueueItem.Batch.Id.ToString(), PoolId = QueueItem.PoolId.ToString(), Workspace = QueueItem.Workspace });
+				}
+
+				List<object> OutputWaiters = new List<object>();
+				foreach (QueueWaiter Waiter in Waiters)
+				{
+					OutputWaiters.Add(new { Id = Waiter.Agent.Id.ToString(), Pools = Waiter.Agent.GetPools().Select(x => x.ToString()).ToList(), Workspaces = Waiter.Agent.Workspaces });
+				}
+
+				return new { Items = OutputItems, Waiters = OutputWaiters };
+			}
+		}
+
+		/// <summary>
+		/// Cancel any pending wait for an agent, allowing it to cycle its session state immediately
+		/// </summary>
+		/// <param name="AgentId">The agent id</param>
+		public void CancelLongPollForAgent(AgentId AgentId)
+		{
+			QueueWaiter? Waiter;
+			lock (LockObject)
+			{
+				Waiter = Waiters.FirstOrDefault(x => x.Agent.Id == AgentId);
+			}
+			if(Waiter != null)
+			{
+				Waiter.LeaseSource.TrySetResult(null);
+			}
+		}
+
+		/// <summary>
+		/// Background task
+		/// </summary>
+		/// <param name="StoppingToken">Token that indicates that the service should shut down</param>
+		/// <returns>Async task</returns>
+		async Task TickAsync(CancellationToken StoppingToken)
+		{
+			// Set the NewBatchIdToQueueItem member, so we capture any updated jobs during the DB query.
+			lock (LockObject)
+			{
+				NewQueueItemsDuringUpdate = new List<QueueItem>();
+			}
+
+			// Query all the current streams
+			List<IStream> StreamsList = await StreamService.GetStreamsAsync();
+			Streams = StreamsList.ToDictionary(x => x.Id, x => x);
+
+			// Find all the pools which are valid (ie. have at least one online agent)
+			DateTime UtcNow = DateTime.UtcNow;
+			List<IAgent> Agents = await AgentsCollection.FindAsync();
+			List<IPool> Pools = await PoolCollection.GetAsync();
+
+			// Find all the pools which are currently online
+			HashSet<PoolId> OnlinePools = new HashSet<PoolId>(Agents.Where(x => x.IsSessionValid(UtcNow) && x.Enabled).SelectMany(x => x.ExplicitPools));
+			foreach (IPool Pool in Pools)
+			{
+				if (Pool.Condition != null && !OnlinePools.Contains(Pool.Id) && Agents.Any(x => x.IsSessionValid(UtcNow) && x.SatisfiesCondition(Pool.Condition) && x.Enabled))
+				{
+					OnlinePools.Add(Pool.Id);
+				}
+			}
+
+			// Find lists of valid pools and online pools
+			HashSet<PoolId> ValidPools = new HashSet<PoolId>(OnlinePools.Union(Agents.Where(x => !x.IsSessionValid(UtcNow) || !x.Enabled).SelectMany(x => x.ExplicitPools)));
+			foreach (IPool Pool in Pools)
+			{
+				if (Pool.Condition != null && !ValidPools.Contains(Pool.Id) && Agents.Any(x => !x.IsSessionValid(UtcNow) && x.SatisfiesCondition(Pool.Condition) && x.Enabled))
+				{
+					ValidPools.Add(Pool.Id);
+				}
+			}
+
+			// Query all the current pools
+			CachedPoolIdToInstance = Pools.ToDictionary(x => x.Id, x => x);
+
+			// New list of queue items
+			SortedSet<QueueItem> NewQueue = new SortedSet<QueueItem>(Queue.Comparer);
+			Dictionary<(ObjectId, SubResourceId), QueueItem> NewBatchIdToQueueItem = new Dictionary<(ObjectId, SubResourceId), QueueItem>();
+
+			// Query for a new list of jobs for the queue
+			List<IJob> NewJobs = await Jobs.GetDispatchQueueAsync();
+			for(int Idx = 0; Idx < NewJobs.Count; Idx++)
+			{
+				IJob? NewJob = NewJobs[Idx];
+
+				if (NewJob.GraphHash == null)
+				{
+					Logger.LogError("Job {JobId} has a null graph hash and can't be started.", NewJob.Id);
+					await Jobs.TryRemoveFromDispatchQueueAsync(NewJob);
+					continue;
+				}
+				if (NewJob.AbortedByUser != null)
+				{
+					Logger.LogError("Job {JobId} was aborted but not removed from dispatch queue", NewJob.Id);
+					await Jobs.TryRemoveFromDispatchQueueAsync(NewJob);
+					continue;
+				}
+
+				// Get the graph for this job
+				IGraph Graph = await Graphs.GetAsync(NewJob.GraphHash);
+
+				// Get the stream. If it fails, skip the whole job.
+				IStream? Stream;
+				if (!Streams.TryGetValue(NewJob.StreamId, out Stream))
+				{
+					NewJob = await Jobs.SkipAllBatchesAsync(NewJob, Graph, JobStepBatchError.UnknownStream);
+					continue;
+				}
+
+				// Update all the batches
+				bool IsRunning = false;
+				for (int BatchIdx = 0; NewJob != null && BatchIdx < NewJob.Batches.Count; BatchIdx++)
+				{
+					IJobStepBatch Batch = NewJob.Batches[BatchIdx];
+					IsRunning |= (Batch.State == JobStepBatchState.Ready || Batch.State == JobStepBatchState.Starting || Batch.State == JobStepBatchState.Running || Batch.State == JobStepBatchState.Stopping);
+
+					// Skip any batches which aren't ready to execute
+					if (Batch.State == JobStepBatchState.Ready)
+					{
+						// Validate the agent type and workspace settings
+						IPool? Pool;
+						if (!Stream.AgentTypes.TryGetValue(Graph.Groups[Batch.GroupIdx].AgentType, out AgentType? AgentType))
+						{
+							NewJob = await SkipBatchAsync(NewJob, BatchIdx, Graph, JobStepBatchError.UnknownAgentType);
+						}
+						else if (!CachedPoolIdToInstance.TryGetValue(AgentType.Pool, out Pool))
+						{
+							NewJob = await SkipBatchAsync(NewJob, BatchIdx, Graph, JobStepBatchError.UnknownPool);
+						}
+						else if (!ValidPools.Contains(AgentType.Pool))
+						{
+							NewJob = await SkipBatchAsync(NewJob, BatchIdx, Graph, JobStepBatchError.NoAgentsInPool);
+						}
+						else if (!OnlinePools.Contains(AgentType.Pool))
+						{
+							NewJob = await SkipBatchAsync(NewJob, BatchIdx, Graph, JobStepBatchError.NoAgentsOnline);
+						}
+						else if (!Stream.TryGetAgentWorkspace(AgentType, out AgentWorkspace? Workspace))
+						{
+							NewJob = await SkipBatchAsync(NewJob, BatchIdx, Graph, JobStepBatchError.UnknownWorkspace);
+						}
+						else
+						{
+							QueueItem NewQueueItem = new QueueItem(Stream, NewJob, BatchIdx, AgentType.Pool, Workspace);
+							NewQueue.Add(NewQueueItem);
+							NewBatchIdToQueueItem[(NewJob.Id, Batch.Id)] = NewQueueItem;
+						}
+					}
+				}
+
+				// Add a warning if a job looks to be idle
+				if (!IsRunning && NewJob != null)
+				{
+					Logger.LogError("Job {JobId} is in dispatch queue but not currently executing", NewJob.Id);
+					await Jobs.TryRemoveFromDispatchQueueAsync(NewJob);
+					continue;
+				}
+			}
+
+			// Update the queue
+			lock (LockObject)
+			{
+				Queue = NewQueue;
+				BatchIdToQueueItem = NewBatchIdToQueueItem;
+
+				// Merge the new queue items with the queue
+				foreach (QueueItem NewQueueItem in NewQueueItemsDuringUpdate)
+				{
+					QueueItem? ExistingQueueItem;
+					if (!NewBatchIdToQueueItem.TryGetValue((NewQueueItem.Job.Id, NewQueueItem.Batch.Id), out ExistingQueueItem))
+					{
+						// Always just add this item
+						Queue.Add(NewQueueItem);
+						BatchIdToQueueItem[NewQueueItem.Id] = NewQueueItem;
+					}
+					else if (NewQueueItem.Job.UpdateIndex > ExistingQueueItem.Job.UpdateIndex)
+					{
+						// Replace the existing item
+						Queue.Remove(ExistingQueueItem);
+						Queue.Add(NewQueueItem);
+						BatchIdToQueueItem[NewQueueItem.Id] = NewQueueItem;
+					}
+				}
+
+				// Clear out the list to capture queue items during an update
+				NewQueueItemsDuringUpdate = null;
+			}
+		}
+
+		private async Task<IJob?> SkipBatchAsync(IJob Job, int BatchIdx, IGraph Graph, JobStepBatchError Reason)
+		{
+			Logger.LogInformation("Skipping batch {BatchIdx} for job {JobId} (reason: {Reason})", BatchIdx, Job.Id, Reason);
+
+			IReadOnlyList<(LabelState, LabelOutcome)> OldLabelStates = Job.GetLabelStates(Graph);
+			IJob? NewJob = await Jobs.SkipBatchAsync(Job, BatchIdx, Graph, Reason);
+			if(NewJob != null)
+			{
+				IReadOnlyList<(LabelState, LabelOutcome)> NewLabelStates = NewJob.GetLabelStates(Graph);
+				await UpdateUgsBadges(NewJob, Graph, OldLabelStates, NewLabelStates);
+			}
+			return NewJob;
+		}
+
+		/// <summary>
+		/// Updates the current state of a job
+		/// </summary>
+		/// <param name="Job">The job that has been updated</param>
+		/// <param name="Graph">Graph for the job</param>
+		/// <returns>Async task</returns>
+		public void UpdateQueuedJob(IJob Job, IGraph Graph)
+		{
+			IStream? Stream;
+			Streams.TryGetValue(Job.StreamId, out Stream);
+			UpdateQueuedJob(Job, Graph, Stream);
+		}
+
+		void AssignAnyQueueItemToWaiter(QueueWaiter Waiter)
+		{
+			lock (Waiters)
+			{
+				foreach(QueueItem Item in BatchIdToQueueItem.Values)
+				{
+					if (TryAssignItemToWaiter(Item, Waiter))
+					{
+						break;
+					}
+				}
+			}
+		}
+
+		/// <summary>
+		/// Attempt to find a waiter that can handle the given queue item
+		/// </summary>
+		/// <param name="Item">The queue item</param>
+		/// <returns></returns>
+		void AssignQueueItemToAnyWaiter(QueueItem Item)
+		{
+			if (Item.AssignTask == null && Item.Batch.SessionId == null)
+			{
+				lock (Waiters)
+				{
+					if (Item.AssignTask == null && Item.Batch.SessionId == null)
+					{
+						foreach (QueueWaiter Waiter in Waiters)
+						{
+							if (TryAssignItemToWaiter(Item, Waiter))
+							{
+								break;
+							}
+						}
+					}
+				}
+			}
+		}
+
+		/// <summary>
+		/// Attempts to assign a queue item to an agent waiting for work
+		/// </summary>
+		/// <param name="Item"></param>
+		/// <param name="Waiter"></param>
+		/// <returns></returns>
+		bool TryAssignItemToWaiter(QueueItem Item, QueueWaiter Waiter)
+		{
+			if (Item.AssignTask == null && Item.Batch.SessionId == null && Waiter.Agent.Enabled && Waiter.Agent.Leases.Count == 0 && Waiter.Agent.IsInPool(Item.PoolId))
+			{
+				Task StartTask = new Task<Task>(() => TryCreateLeaseAsync(Item, Waiter));
+				Task ExecuteTask = StartTask.ContinueWith(Task => Task, TaskScheduler.Default);
+
+				if (Interlocked.CompareExchange(ref Item.AssignTask, ExecuteTask, null) == null)
+				{
+					StartTask.Start(TaskScheduler.Default);
+					return true;
+				}
+			}
+			return false;
+		}
+
+		/// <summary>
+		/// Updates the current state of a job
+		/// </summary>
+		/// <param name="Job">The job that has been updated</param>
+		/// <param name="Graph">Graph for the job</param>
+		/// <param name="Stream">The stream containing the job</param>
+		public void UpdateQueuedJob(IJob Job, IGraph Graph, IStream? Stream)
+		{
+			List<TaskCompletionSource<bool>> CompleteWaiters = new List<TaskCompletionSource<bool>>();
+			lock (LockObject)
+			{
+				for (int BatchIdx = 0; BatchIdx < Job.Batches.Count; BatchIdx++)
+				{
+					IJobStepBatch Batch = Job.Batches[BatchIdx];
+					if (Batch.State == JobStepBatchState.Ready && Stream != null && Batch.AgentId == null)
+					{
+						// Check if this item is already in the list.
+						QueueItem? ExistingItem;
+						if (BatchIdToQueueItem.TryGetValue((Job.Id, Batch.Id), out ExistingItem))
+						{
+							// Make sure this is a newer version of the job. There's no guarantee that this is the latest revision.
+							if (Job.UpdateIndex > ExistingItem.Job.UpdateIndex)
+							{
+								if (Batch.SchedulePriority == ExistingItem.Batch.SchedulePriority)
+								{
+									ExistingItem.Job = Job;
+									ExistingItem.BatchIdx = BatchIdx;
+								}
+								else
+								{
+									RemoveQueueItem(ExistingItem);
+									InsertQueueItem(Stream, Job, BatchIdx, ExistingItem.PoolId, ExistingItem.Workspace);
+								}
+							}
+							continue;
+						}
+
+						// Get the group being executed by this batch
+						INodeGroup Group = Graph.Groups[Batch.GroupIdx];
+
+						// Get the requirements for the new queue item
+						AgentType? AgentType;
+						if (Stream.AgentTypes.TryGetValue(Group.AgentType, out AgentType))
+						{
+							AgentWorkspace? Workspace;
+							if (Stream.TryGetAgentWorkspace(AgentType, out Workspace))
+							{
+								InsertQueueItem(Stream, Job, BatchIdx, AgentType.Pool, Workspace);
+							}
+						}
+					}
+					else
+					{
+						// Check if this item is already in the list. Remove it if it is.
+						QueueItem? ExistingItem;
+						if (BatchIdToQueueItem.TryGetValue((Job.Id, Batch.Id), out ExistingItem))
+						{
+							if (Job.UpdateIndex > ExistingItem.Job.UpdateIndex)
+							{
+								RemoveQueueItem(ExistingItem);
+							}
+						}
+					}
+				}
+			}
+
+			// Awake all the threads that have been assigned new work items. Has do be done outside the lock to prevent continuations running within it (see Waiter.CompletionSource for more info).
+			foreach (TaskCompletionSource<bool> CompleteWaiter in CompleteWaiters)
+			{
+				CompleteWaiter.TrySetResult(true);
+			}
+		}
+
+		/// <inheritdoc/>
+		public override async Task<AgentLease?> AssignLeaseAsync(IAgent Agent, CancellationToken CancellationToken)
+		{
+			QueueWaiter Waiter = new QueueWaiter(Agent);
+			lock (LockObject)
+			{
+				AssignAnyQueueItemToWaiter(Waiter);
+				if (Waiter.LeaseSource.Task.IsCompleted)
+				{
+					return Waiter.Task.Result;
+				}
+				Waiters.Add(Waiter);
+			}
+
+			AgentLease? Lease;
+			using (CancellationToken.Register(() => Waiter.LeaseSource.TrySetResult(null)))
+			{
+				Lease = await Waiter.Task;
+			}
+
+			lock (LockObject)
+			{
+				Waiters.Remove(Waiter);
+			}
+
+			return Lease;
+		}
+
+		/// <inheritdoc/>
+		public override Task CancelLeaseAsync(IAgent Agent, ObjectId LeaseId, ExecuteJobTask Task)
+		{
+			return CancelLeaseAsync(Agent, Task.JobId.ToObjectId(), Task.BatchId.ToSubResourceId());
+		}
+
+		/// <summary>
+		/// Assign a new batch to be executed by the given agent
+		/// </summary>
+		/// <param name="Item">The item to create a lease for</param>
+		/// <param name="Waiter">The agent waiting for work</param>
+		/// <returns>New work to execute</returns>
+		private async Task<AgentLease?> TryCreateLeaseAsync(QueueItem Item, QueueWaiter Waiter)
+		{
+			IJob Job = Item.Job;
+			IJobStepBatch Batch = Job.Batches[Item.BatchIdx];
+			IAgent Agent = Waiter.Agent;
+			Logger.LogDebug("Assigning job {JobId}, batch {BatchId} to waiter (agent {AgentID})", Job.Id, Batch.Id, Agent.Id);
+
+			// Generate a new unique id for the lease
+			ObjectId LeaseId = ObjectId.GenerateNewId();
+
+			// The next time to try assigning to another agent
+			DateTime BackOffTime = DateTime.UtcNow + TimeSpan.FromMinutes(1.0);
+
+			// Try to update the job with this agent id
+			ObjectId LogId = (await LogFileService.CreateLogFileAsync(Job.Id, Agent.SessionId, LogType.Json)).Id;
+			if (await Jobs.TryAssignLeaseAsync(Item.Job, Item.BatchIdx, Item.PoolId, Agent.Id, Agent.SessionId!.Value, LeaseId, LogId))
+			{
+				// Get the lease name
+				StringBuilder LeaseName = new StringBuilder($"{Item.Stream.Name} - ");
+				if (Job.PreflightChange > 0)
+				{
+					LeaseName.Append((Job.Change > 0) ? $"Preflight CL {Job.PreflightChange} against CL {Job.Change}" : $"Preflight CL {Job.PreflightChange} against latest");
+				}
+				else
+				{
+					LeaseName.Append((Job.Change > 0) ? $"CL {Job.Change}" : "Latest CL");
+				}
+				LeaseName.Append($" - {Job.Name}");
+
+				// Get the global settings
+				Globals Globals = await DatabaseService.GetGlobalsAsync();
+
+				// Encode the payload
+				ExecuteJobTask? Task = await CreateExecuteJobTaskAsync(Item.Stream, Job, Batch, Agent, Item.Workspace, LogId);
+				if (Task != null)
+				{
+					byte[] Payload = Any.Pack(Task).ToByteArray();
+
+					// Create the lease and try to set it on the waiter. If this fails, the waiter has already moved on, and the lease can be cancelled.
+					AgentLease Lease = new AgentLease(LeaseId, LeaseName.ToString(), Job.StreamId, Item.PoolId, LogId, LeaseState.Pending, null, true, Payload);
+					if (Waiter.LeaseSource.TrySetResult(Lease))
+					{
+						Logger.LogDebug("Assigned lease {LeaseId} to agent {AgentId}", LeaseId, Agent.Id);
+						return Lease;
+					}
+				}
+
+				// Cancel the lease
+				Logger.LogDebug("Unable to assign lease {LeaseId} to agent {AgentId}, cancelling", LeaseId, Agent.Id);
+				await CancelLeaseAsync(Waiter.Agent, Job.Id, Batch.Id);
+			}
+			else
+			{
+				// Unable to assign job
+				Logger.LogTrace("Refreshing queue entries for job {JobId}", Job.Id);
+
+				// Get the new copy of the job
+				IJob? NewJob = await Jobs.GetAsync(Job.Id);
+				if (NewJob == null)
+				{
+					lock (LockObject)
+					{
+						List<QueueItem> RemoveItems = Queue.Where(x => x.Job == Job).ToList();
+						foreach (QueueItem RemoveItem in RemoveItems)
+						{
+							RemoveQueueItem(RemoveItem);
+						}
+					}
+				}
+				else
+				{
+					IGraph Graph = await Graphs.GetAsync(NewJob.GraphHash);
+					UpdateQueuedJob(NewJob, Graph);
+				}
+			}
+
+			// Clear out the assignment for this item, and try to reassign it
+			Item.AssignTask = null;
+			AssignQueueItemToAnyWaiter(Item);
+			return null;
+		}
+
+		async Task<ExecuteJobTask?> CreateExecuteJobTaskAsync(IStream Stream, IJob Job, IJobStepBatch Batch, IAgent Agent, AgentWorkspace Workspace, ObjectId LogId)
+		{
+			// Get the lease name
+			StringBuilder LeaseName = new StringBuilder($"{Stream.Name} - ");
+			if (Job.PreflightChange > 0)
+			{
+				LeaseName.Append((Job.Change > 0) ? $"Preflight CL {Job.PreflightChange} against CL {Job.Change}" : $"Preflight CL {Job.PreflightChange} against latest");
+			}
+			else
+			{
+				LeaseName.Append((Job.Change > 0) ? $"CL {Job.Change}" : "Latest CL");
+			}
+			LeaseName.Append($" - {Job.Name}");
+
+			// Get the global settings
+			Globals Globals = await DatabaseService.GetGlobalsAsync();
+
+			// Encode the payload
+			ExecuteJobTask Task = new ExecuteJobTask();
+			Task.JobId = Job.Id.ToString();
+			Task.BatchId = Batch.Id.ToString();
+			Task.LogId = LogId.ToString();
+			Task.JobName = LeaseName.ToString();
+
+			List<HordeCommon.Rpc.Messages.AgentWorkspace> Workspaces = new List<HordeCommon.Rpc.Messages.AgentWorkspace>();
+
+			PerforceCluster? Cluster = Globals.FindPerforceCluster(Workspace.Cluster);
+			if (Cluster == null)
+			{
+				return null;
+			}
+
+			AgentWorkspace? AutoSdkWorkspace = Agent.GetAutoSdkWorkspace(Cluster);
+			if (AutoSdkWorkspace != null)
+			{
+				if (!await Agent.TryAddWorkspaceMessage(AutoSdkWorkspace, Cluster, PerforceLoadBalancer, Workspaces))
+				{
+					return null;
+				}
+				Task.AutoSdkWorkspace = Workspaces.Last();
+			}
+
+			if (!await Agent.TryAddWorkspaceMessage(Workspace, Cluster, PerforceLoadBalancer, Workspaces))
+			{
+				return null;
+			}
+
+			Task.Workspace = Workspaces.Last();
+			return Task;
+		}
+
+		/// <summary>
+		/// Send any badge updates for this job
+		/// </summary>
+		/// <param name="Job">The job being updated</param>
+		/// <param name="Graph">Graph for the job</param>
+		/// <param name="OldLabelStates">Previous badge states for the job</param>
+		/// <returns>Async task</returns>
+		public async Task UpdateUgsBadges(IJob Job, IGraph Graph, IReadOnlyList<(LabelState, LabelOutcome)> OldLabelStates)
+		{
+			await UpdateUgsBadges(Job, Graph, OldLabelStates, Job.GetLabelStates(Graph));
+		}
+
+		/// <summary>
+		/// Send any badge updates for this job
+		/// </summary>
+		/// <param name="Job">The job being updated</param>
+		/// <param name="Graph">Graph for the job</param>
+		/// <param name="OldLabelStates">Previous badge states for the job</param>
+		/// <param name="NewLabelStates">The new badge states for the job</param>
+		/// <returns>Async task</returns>
+		public async Task UpdateUgsBadges(IJob Job, IGraph Graph, IReadOnlyList<(LabelState, LabelOutcome)> OldLabelStates, IReadOnlyList<(LabelState, LabelOutcome)> NewLabelStates)
+		{
+			if (!Job.ShowUgsBadges || Job.PreflightChange != 0)
+			{
+				return;
+			}
+
+			IReadOnlyDictionary<int, UgsBadgeState> OldStates = Job.GetUgsBadgeStates(Graph, OldLabelStates);
+			IReadOnlyDictionary<int, UgsBadgeState> NewStates = Job.GetUgsBadgeStates(Graph, NewLabelStates);
+
+			// Figure out a list of all the badges that have been modified
+			List<int> UpdateLabels = new List<int>();
+			foreach (KeyValuePair<int, UgsBadgeState> Pair in OldStates)
+			{
+				if (!NewStates.ContainsKey(Pair.Key))
+				{
+					UpdateLabels.Add(Pair.Key);
+				}
+			}
+			foreach (KeyValuePair<int, UgsBadgeState> Pair in NewStates)
+			{
+				if (!OldStates.TryGetValue(Pair.Key, out UgsBadgeState PrevState) || PrevState != Pair.Value)
+				{
+					UpdateLabels.Add(Pair.Key);
+				}
+			}
+
+			// Cached stream for this job
+			IStream? Stream = null;
+
+			// Send all the updates
+			Dictionary<int, IUgsMetadata> MetadataCache = new Dictionary<int, IUgsMetadata>();
+			foreach (int LabelIdx in UpdateLabels)
+			{
+				ILabel Label = Graph.Labels[LabelIdx];
+
+				// Skip if this label has no UGS name.
+				if (Label.UgsName == null)
+				{
+					continue;
+				}
+
+				// Get the new state
+				if (!NewStates.TryGetValue(LabelIdx, out UgsBadgeState NewState))
+				{
+					NewState = UgsBadgeState.Skipped;
+				}
+
+				// Get the stream
+				if (Stream == null)
+				{
+					Stream = await StreamService.GetStreamAsync(Job.StreamId);
+					if (Stream == null)
+					{
+						Logger.LogError("Unable to fetch definition for stream {StreamId}", Job.StreamId);
+						break;
+					}
+				}
+
+				// The changelist number to display the badge for
+				int Change;
+				if (Label.Change == LabelChange.Code)
+				{
+					Change = Job.CodeChange;
+				}
+				else
+				{
+					Change = Job.Change;
+				}
+
+				// Get the current metadata state
+				IUgsMetadata? Metadata;
+				if (!MetadataCache.TryGetValue(Change, out Metadata))
+				{
+					Metadata = await UgsMetadataCollection.FindOrAddAsync(Stream.Name, Change, Label.UgsProject);
+					MetadataCache[Change] = Metadata;
+				}
+
+				// Apply the update
+				Uri LabelUrl = new Uri(Settings.CurrentValue.DashboardUrl, $"job/{Job.Id}?label={LabelIdx}");
+				Logger.LogInformation("Updating state of badge {BadgeName} at {Change} to {NewState} ({LabelUrl})", Label.UgsName, Change, NewState, LabelUrl);
+				Metadata = await UgsMetadataCollection.UpdateBadgeAsync(Metadata, Label.UgsName!, LabelUrl, NewState);
+				MetadataCache[Change] = Metadata;
+			}
+		}
+
+		/// <inheritdoc/>
+		public override async Task OnLeaseFinishedAsync(IAgent Agent, ObjectId LeaseId, ExecuteJobTask Task, LeaseOutcome Outcome, ReadOnlyMemory<byte> Output, ILogger Logger)
+		{
+			if (Outcome != LeaseOutcome.Success)
+			{
+				AgentId AgentId = Agent.Id;
+				ObjectId JobId = Task.JobId.ToObjectId();
+				SubResourceId BatchId = Task.BatchId.ToSubResourceId();
+
+				// Update the batch
+				for (; ; )
+				{
+					IJob? Job = await Jobs.GetAsync(JobId);
+					if (Job == null)
+					{
+						break;
+					}
+
+					int BatchIdx = Job.Batches.FindIndex(x => x.Id == BatchId);
+					if (BatchIdx == -1)
+					{
+						break;
+					}
+
+					IJobStepBatch Batch = Job.Batches[BatchIdx];
+					if (Batch.AgentId != AgentId)
+					{
+						break;
+					}
+
+					int RunningStepIdx = Batch.Steps.FindIndex(x => x.State == JobStepState.Running);
+
+					IGraph Graph = await Graphs.GetAsync(Job.GraphHash);
+					if (await Jobs.TryFailBatchAsync(Job, BatchIdx, Graph))
+					{
+						if (Batch.Error != JobStepBatchError.None)
+						{
+							Logger.LogInformation("Failed job {JobId}, batch {BatchId} with error {Error}", Job.Id, Batch.Id, Batch.Error);
+						}
+						if (RunningStepIdx != -1)
+						{
+							await JobStepRefs.UpdateAsync(Job, Batch, Batch.Steps[RunningStepIdx], Graph);
+						}
+						break;
+					}
+				}
+			}
+		}
+
+		/// <summary>
+		/// 
+		/// </summary>
+		/// <param name="Agent"></param>
+		/// <param name="JobId"></param>
+		/// <param name="BatchId"></param>
+		/// <returns></returns>
+		async Task CancelLeaseAsync(IAgent Agent, ObjectId JobId, SubResourceId BatchId)
+		{
+			Logger.LogDebug("Cancelling lease for job {JobId}, batch {BatchId}", JobId, BatchId);
+
+			// Update the batch
+			for (; ; )
+			{
+				IJob? Job = await Jobs.GetAsync(JobId);
+				if (Job == null)
+				{
+					break;
+				}
+
+				int BatchIdx = Job.Batches.FindIndex(x => x.Id == BatchId);
+				if (BatchIdx == -1)
+				{
+					break;
+				}
+
+				IJobStepBatch Batch = Job.Batches[BatchIdx];
+				if (Batch.AgentId != Agent.Id)
+				{
+					break;
+				}
+
+				if (await Jobs.TryCancelLeaseAsync(Job, BatchIdx))
+				{
+					break;
+				}
+			}
+		}
+
+		/// <summary>
+		/// Inserts an item into the queue
+		/// </summary>
+		/// <param name="Stream">The stream containing the job</param>
+		/// <param name="Job"></param>
+		/// <param name="BatchIdx"></param>
+		/// <param name="PoolId">The pool to use</param>
+		/// <param name="Workspace">The workspace for this item to run in</param>
+		/// <returns></returns>
+		void InsertQueueItem(IStream Stream, IJob Job, int BatchIdx, PoolId PoolId, AgentWorkspace Workspace)
+		{
+			Logger.LogDebug("Adding queued job {JobId}, batch {BatchId} [Pool: {Pool}, Workspace: {Workspace}]", Job.Id, Job.Batches[BatchIdx].Id, PoolId, Workspace.Identifier);
+
+			QueueItem NewItem = new QueueItem(Stream, Job, BatchIdx, PoolId, Workspace);
+			BatchIdToQueueItem[NewItem.Id] = NewItem;
+			Queue.Add(NewItem);
+
+			AssignQueueItemToAnyWaiter(NewItem);
+		}
+
+		/// <summary>
+		/// Removes an item from the queue
+		/// </summary>
+		/// <param name="Item">Item to remove</param>
+		void RemoveQueueItem(QueueItem Item)
+		{
+			Logger.LogDebug("Removing queued job {JobId}, batch {BatchId}", Item.Job.Id, Item.Batch.Id);
+
+			Queue.Remove(Item);
+			BatchIdToQueueItem.Remove(Item.Id);
+		}
+	}
+}
diff --git a/Engine/Source/Programs/Horde/HordeServer/Jobs/SubResourceId.cs b/Engine/Source/Programs/Horde/HordeServer/Jobs/SubResourceId.cs
new file mode 100644
--- /dev/null
+++ b/Engine/Source/Programs/Horde/HordeServer/Jobs/SubResourceId.cs
@@ -0,0 +1,166 @@
+// Copyright Epic Games, Inc. All Rights Reserved.
+
+using MongoDB.Bson.Serialization;
+using MongoDB.Bson.Serialization.Attributes;
+using System;
+using System.Collections.Generic;
+using System.Diagnostics;
+using System.Globalization;
+using System.Linq;
+using System.Threading.Tasks;
+
+namespace HordeServer.Models
+{
+	/// <summary>
+	/// Identifier for subresources. Assigning unique ids to subresources prevents against race conditions using indices when subresources are added and removed.
+	/// 
+	/// Subresource identifiers are stored as 16-bit integers formatted as a 4-digit hex code, in order to keep URLs short. Calling Next() will generate a new
+	/// identifier with more entropy than just incrementing the value but an identical period before repeating, in order to make URL fragments more distinctive.
+	/// </summary>
+	[BsonSerializer(typeof(SubResourceIdSerializer))]
+	public struct SubResourceId : IEquatable<SubResourceId>
+	{
+		/// <summary>
+		/// The unique identifier value
+		/// </summary>
+		public ushort Value { get; }
+
+		/// <summary>
+		/// Constructor
+		/// </summary>
+		/// <param name="Value">New identifier for this subresource</param>
+		public SubResourceId(ushort Value)
+		{
+			this.Value = Value;
+		}
+
+		/// <summary>
+		/// Creates a new random subresource id. We use random numbers for this to increase distinctiveness.
+		/// </summary>
+		/// <returns>New subresource id</returns>
+		public static SubResourceId Random()
+		{
+			return new SubResourceId((ushort)Stopwatch.GetTimestamp());
+		}
+
+		/// <summary>
+		/// Updates the current value, and returns a copy of the previous value.
+		/// </summary>
+		/// <returns>New subresource identifier</returns>
+		public SubResourceId Next()
+		{
+			// 771 is a factor of 65535, so we won't repeat when wrapping round 64k.
+			return new SubResourceId((ushort)((int)Value + 771));
+		}
+
+		/// <summary>
+		/// Parse a subresource id from a string
+		/// </summary>
+		/// <param name="Text">Text to parse</param>
+		/// <returns>New subresource id</returns>
+		public static SubResourceId Parse(string Text)
+		{
+			return new SubResourceId(ushort.Parse(Text, NumberStyles.HexNumber, CultureInfo.InvariantCulture));
+		}
+
+		/// <summary>
+		/// Converts this identifier to a string
+		/// </summary>
+		/// <returns>String representation of this id</returns>
+		public override string ToString()
+		{
+			return Value.ToString("x4", CultureInfo.InvariantCulture);
+		}
+
+		/// <inheritdoc/>
+		public override bool Equals(object? Other)
+		{
+			SubResourceId? OtherId = Other as SubResourceId?;
+			return OtherId != null && Equals(OtherId.Value);
+		}
+
+		/// <inheritdoc/>
+		public bool Equals(SubResourceId Other)
+		{
+			return Value == Other.Value;
+		}
+
+		/// <inheritdoc/>
+		public override int GetHashCode()
+		{
+			return (int)Value;
+		}
+
+		/// <summary>
+		/// Equality operator for identifiers
+		/// </summary>
+		/// <param name="Left">First identifier to compare</param>
+		/// <param name="Right">Second identifier to compare</param>
+		/// <returns>True if the identifiers are equal</returns>
+		public static bool operator ==(SubResourceId Left, SubResourceId Right)
+		{
+			return Left.Value == Right.Value;
+		}
+
+		/// <summary>
+		/// Inequality operator for identifiers
+		/// </summary>
+		/// <param name="Left">First identifier to compare</param>
+		/// <param name="Right">Second identifier to compare</param>
+		/// <returns>True if the identifiers are equal</returns>
+		public static bool operator !=(SubResourceId Left, SubResourceId Right)
+		{
+			return Left.Value != Right.Value;
+		}
+	}
+
+	/// <summary>
+	/// Extension methods for manipulating subresource ids
+	/// </summary>
+	public static class SubResourceIdExtensions
+	{
+		/// <summary>
+		/// Parse a string as a subresource identifier
+		/// </summary>
+		/// <param name="Text">Text to parse</param>
+		/// <returns>The new subresource identifier</returns>
+		public static SubResourceId ToSubResourceId(this string Text)
+		{
+			return new SubResourceId(ushort.Parse(Text, NumberStyles.HexNumber, CultureInfo.InvariantCulture));
+		}
+	}
+
+	/// <summary>
+	/// Serializer for subresource ids
+	/// </summary>
+	public class SubResourceIdSerializer : IBsonSerializer<SubResourceId>
+	{
+		/// <inheritdoc/>
+		public Type ValueType => typeof(SubResourceId);
+
+		/// <inheritdoc/>
+		public object Deserialize(BsonDeserializationContext Context, BsonDeserializationArgs Args)
+		{
+			return new SubResourceId((ushort)Context.Reader.ReadInt32());
+		}
+
+		/// <inheritdoc/>
+		public void Serialize(BsonSerializationContext Context, BsonSerializationArgs Args, object Value)
+		{
+			SubResourceId Id = (SubResourceId)Value;
+			Context.Writer.WriteInt32((int)Id.Value);
+		}
+
+		/// <inheritdoc/>
+		public void Serialize(BsonSerializationContext Context, BsonSerializationArgs Args, SubResourceId Id)
+		{
+			Context.Writer.WriteInt32((int)Id.Value);
+		}
+
+		/// <inheritdoc/>
+		SubResourceId IBsonSerializer<SubResourceId>.Deserialize(BsonDeserializationContext Context, BsonDeserializationArgs Vrgs)
+		{
+			return new SubResourceId((ushort)Context.Reader.ReadInt32());
+		}
+	}
+}
diff --git a/Engine/Source/Programs/Horde/HordeServer/Logs/Data/LogChunkData.cs b/Engine/Source/Programs/Horde/HordeServer/Logs/Data/LogChunkData.cs
new file mode 100644
--- /dev/null
+++ b/Engine/Source/Programs/Horde/HordeServer/Logs/Data/LogChunkData.cs
@@ -0,0 +1,280 @@
+// Copyright Epic Games, Inc. All Rights Reserved.
+
+using EpicGames.Core;
+using HordeServer.Api;
+using HordeServer.Utilities;
+using System;
+using System.Collections.Generic;
+using System.IO;
+using System.Linq;
+using System.Text;
+using System.Threading.Tasks;
+
+namespace HordeServer.Logs
+{
+	/// <summary>
+	/// Data for a log chunk
+	/// </summary>
+	public class LogChunkData
+	{
+		/// <summary>
+		/// Offset of this chunk
+		/// </summary>
+		public long Offset { get; }
+
+		/// <summary>
+		/// Length of this chunk
+		/// </summary>
+		public int Length { get; }
+
+		/// <summary>
+		/// Line index of this chunk
+		/// </summary>
+		public int LineIndex { get; }
+
+		/// <summary>
+		/// Number of lines in this chunk
+		/// </summary>
+		public int LineCount { get; }
+			
+		/// <summary>
+		/// List of sub-chunks
+		/// </summary>
+		public IReadOnlyList<LogSubChunkData> SubChunks { get; }
+
+		/// <summary>
+		/// Offset of each sub-chunk within this chunk
+		/// </summary>
+		public int[] SubChunkOffset { get; }
+
+		/// <summary>
+		/// Total line count of the chunk after each sub-chunk
+		/// </summary>
+		public int[] SubChunkLineIndex { get; }
+
+		/// <summary>
+		/// Constructor
+		/// </summary>
+		/// <param name="Offset">Offset of this chunk</param>
+		/// <param name="LineIndex">First line index of this chunk</param>
+		/// <param name="SubChunks">Sub-chunks for this chunk</param>
+		public LogChunkData(long Offset, int LineIndex, IReadOnlyList<LogSubChunkData> SubChunks)
+		{
+			this.Offset = Offset;
+			this.LineIndex = LineIndex;
+			this.SubChunks = SubChunks;
+			this.SubChunkOffset = CreateSumLookup(SubChunks, x => x.Length);
+			this.SubChunkLineIndex = CreateSumLookup(SubChunks, x => x.LineCount);
+
+			if(SubChunks.Count > 0)
+			{
+				int LastSubChunkIdx = SubChunks.Count - 1;
+				Length = SubChunkOffset[LastSubChunkIdx] + SubChunks[LastSubChunkIdx].Length;
+				LineCount = SubChunkLineIndex[LastSubChunkIdx] + SubChunks[LastSubChunkIdx].LineCount;
+			}
+		}
+
+		/// <summary>
+		/// Creates a lookup table by summing a field across a list of subchunks
+		/// </summary>
+		/// <param name="SubChunks"></param>
+		/// <param name="Field"></param>
+		/// <returns></returns>
+		static int[] CreateSumLookup(IReadOnlyList<LogSubChunkData> SubChunks, Func<LogSubChunkData, int> Field)
+		{
+			int[] Total = new int[SubChunks.Count];
+
+			int Value = 0;
+			for (int Idx = 0; Idx + 1 < SubChunks.Count; Idx++)
+			{
+				Value += Field(SubChunks[Idx]);
+				Total[Idx + 1] = Value;
+			}
+
+			return Total;
+		}
+
+		/// <summary>
+		/// Gets the offset of a line within the file
+		/// </summary>
+		/// <param name="LineIdx">The line index</param>
+		/// <returns>Offset of the line within the file</returns>
+		public long GetLineOffsetWithinChunk(int LineIdx)
+		{
+			int SubChunkIdx = SubChunkLineIndex.BinarySearch(LineIdx);
+			if (SubChunkIdx < 0)
+			{
+				SubChunkIdx = ~SubChunkIdx - 1;
+			}
+			return SubChunkOffset[SubChunkIdx] + SubChunks[SubChunkIdx].InflateText().LineOffsets[LineIdx - SubChunkLineIndex[SubChunkIdx]];
+		}
+
+		/// <summary>
+		/// Gets the sub chunk index for the given line
+		/// </summary>
+		/// <param name="ChunkLineIdx">Line index within the chunk</param>
+		/// <returns>Subchunk line index</returns>
+		public int GetSubChunkForLine(int ChunkLineIdx)
+		{
+			int SubChunkIdx = SubChunkLineIndex.BinarySearch(ChunkLineIdx);
+			if (SubChunkIdx < 0)
+			{
+				SubChunkIdx = ~SubChunkIdx - 1;
+			}
+			return SubChunkIdx;
+		}
+
+		/// <summary>
+		/// Gets the index of the sub-chunk containing the given offset
+		/// </summary>
+		/// <param name="Offset">Offset to search for</param>
+		/// <returns>Index of the sub-chunk</returns>
+		public int GetSubChunkForOffsetWithinChunk(int Offset)
+		{
+			int SubChunkIdx = SubChunkOffset.BinarySearch(Offset);
+			if (SubChunkIdx < 0)
+			{
+				SubChunkIdx = ~SubChunkIdx - 1;
+			}
+			return SubChunkIdx;
+		}
+
+		/// <summary>
+		/// Signature for output data
+		/// </summary>
+		const int CurrentSignature = 'L' | ('C' << 8) | ('D' << 16);
+
+		/// <summary>
+		/// Read a log chunk from the given stream
+		/// </summary>
+		/// <param name="Reader">The reader to read from</param>
+		/// <param name="Offset">Offset of this chunk within the file</param>
+		/// <param name="LineIndex">Line index of this chunk</param>
+		/// <returns>New log chunk data</returns>
+		public static LogChunkData Read(MemoryReader Reader, long Offset, int LineIndex)
+		{
+			int Signature = Reader.ReadInt32();
+			if ((Signature & 0xffffff) != CurrentSignature)
+			{
+				List<LogSubChunkData> SubChunks = new List<LogSubChunkData>();
+				SubChunks.Add(new LogSubChunkData(LogType.Json, Offset, LineIndex, new ReadOnlyLogText(Reader.Memory)));
+				Reader.Offset = Reader.Memory.Length;
+				return new LogChunkData(Offset, LineIndex, SubChunks);
+			}
+
+			int Version = Signature >> 24;
+			if (Version == 0)
+			{
+				List<LogSubChunkData> SubChunks = ReadSubChunkList(Reader, Offset, LineIndex);
+				Reader.ReadVariableLengthBytes();
+				return new LogChunkData(Offset, LineIndex, SubChunks);
+			}
+			else
+			{
+				List<LogSubChunkData> SubChunks = ReadSubChunkList(Reader, Offset, LineIndex);
+				return new LogChunkData(Offset, LineIndex, SubChunks);
+			}
+		}
+
+		/// <summary>
+		/// Read a list of sub-chunks from the given stream
+		/// </summary>
+		/// <param name="Reader">The reader to read from</param>
+		/// <param name="SubChunkOffset">Offset of this chunk within the file</param>
+		/// <param name="SubChunkLineIndex">Line index of this chunk</param>
+		/// <returns>List of sub-chunks</returns>
+		static List<LogSubChunkData> ReadSubChunkList(MemoryReader Reader, long SubChunkOffset, int SubChunkLineIndex)
+		{
+			int NumSubChunks = Reader.ReadInt32();
+
+			List<LogSubChunkData> SubChunks = new List<LogSubChunkData>();
+			for (int Idx = 0; Idx < NumSubChunks; Idx++)
+			{
+				LogSubChunkData SubChunkData = Reader.ReadLogSubChunkData(SubChunkOffset, SubChunkLineIndex);
+				SubChunkOffset += SubChunkData.Length;
+				SubChunkLineIndex += SubChunkData.LineCount;
+				SubChunks.Add(SubChunkData);
+			}
+
+			return SubChunks;
+		}
+
+		/// <summary>
+		/// Construct an object from flat memory buffer
+		/// </summary>
+		/// <param name="Memory">Memory buffer</param>
+		/// <param name="Offset">Offset of this chunk within the file</param>
+		/// <param name="LineIndex">Line index of this chunk</param>
+		/// <returns>Log chunk data</returns>
+		public static LogChunkData FromMemory(ReadOnlyMemory<byte> Memory, long Offset, int LineIndex)
+		{
+			MemoryReader Reader = new MemoryReader(Memory);
+			LogChunkData ChunkData = Read(Reader, Offset, LineIndex);
+			Reader.CheckOffset(Memory.Length);
+			return ChunkData;
+		}
+
+		/// <summary>
+		/// Write the chunk data to a stream
+		/// </summary>
+		/// <returns>Serialized data</returns>
+		public void Write(MemoryWriter Writer)
+		{
+			Writer.WriteInt32(CurrentSignature | (1 << 24));
+
+			Writer.WriteInt32(SubChunks.Count);
+			foreach (LogSubChunkData SubChunk in SubChunks)
+			{
+				Writer.WriteLogSubChunkData(SubChunk);
+			}
+		}
+
+		/// <summary>
+		/// Construct an object from flat memory buffer
+		/// </summary>
+		/// <returns>Log chunk data</returns>
+		public byte[] ToByteArray()
+		{
+			byte[] Data = new byte[GetSerializedSize()];
+			MemoryWriter Writer = new MemoryWriter(Data);
+			Write(Writer);
+			Writer.CheckOffset(Data.Length);
+			return Data;
+		}
+
+		/// <summary>
+		/// Determines the size of the serialized buffer
+		/// </summary>
+		public int GetSerializedSize()
+		{
+			return sizeof(int) + (sizeof(int) + SubChunks.Sum(x => x.GetSerializedSize()));
+		}
+	}
+
+	/// <summary>
+	/// Extensions for the log chunk data
+	/// </summary>
+	static class LogChunkDataExtensions
+	{
+		/// <summary>
+		/// Read a log chunk from the given stream
+		/// </summary>
+		/// <param name="Reader">The reader to read from</param>
+		/// <param name="Offset">Offset of this chunk within the file</param>
+		/// <param name="LineIndex">Line index of this chunk</param>
+		/// <returns>New log chunk data</returns>
+		public static LogChunkData ReadLogChunkData(this MemoryReader Reader, long Offset, int LineIndex)
+		{
+			return LogChunkData.Read(Reader, Offset, LineIndex);
+		}
+
+		/// <summary>
+		/// Write the chunk data to a stream
+		/// </summary>
+		/// <returns>Serialized data</returns>
+		public static void WriteLogChunkData(this MemoryWriter Writer, LogChunkData ChunkData)
+		{
+			ChunkData.Write(Writer);
+		}
+	}
+}
diff --git a/Engine/Source/Programs/Horde/HordeServer/Logs/Data/LogIndexData.cs b/Engine/Source/Programs/Horde/HordeServer/Logs/Data/LogIndexData.cs
new file mode 100644
--- /dev/null
+++ b/Engine/Source/Programs/Horde/HordeServer/Logs/Data/LogIndexData.cs
@@ -0,0 +1,640 @@
+// Copyright Epic Games, Inc. All Rights Reserved.
+
+using EpicGames.Core;
+using HordeServer.Utilities;
+using OpenTracing;
+using OpenTracing.Util;
+using System;
+using System.Collections.Generic;
+using System.IO;
+using System.Linq;
+using System.Reflection;
+using System.Text;
+using System.Threading.Tasks;
+
+namespace HordeServer.Logs
+{
+	/// <summary>
+	/// Contains a source block of text to be indexed
+	/// </summary>
+	public class LogIndexBlock
+	{
+		/// <summary>
+		/// First line within the file
+		/// </summary>
+		public int LineIndex { get; }
+
+		/// <summary>
+		/// Number of lines in this block
+		/// </summary>
+		public int LineCount { get; }
+
+		/// <summary>
+		/// The plain text for this block
+		/// </summary>
+		public ILogText? CachedPlainText { get; private set; }
+
+		/// <summary>
+		/// The compressed plain text data
+		/// </summary>
+		public ReadOnlyMemory<byte> CompressedPlainText { get; }
+
+		/// <summary>
+		/// Constructor
+		/// </summary>
+		/// <param name="LineIndex">Index of the first line within this block</param>
+		/// <param name="LineCount">Number of lines in the block</param>
+		/// <param name="CachedPlainText">The decompressed plain text</param>
+		/// <param name="CompressedPlainText">The compressed text data</param>
+		public LogIndexBlock(int LineIndex, int LineCount, ILogText? CachedPlainText, ReadOnlyMemory<byte> CompressedPlainText)
+		{
+			this.LineIndex = LineIndex;
+			this.LineCount = LineCount;
+			this.CachedPlainText = CachedPlainText;
+			this.CompressedPlainText = CompressedPlainText;
+		}
+
+		/// <summary>
+		/// Accessor for the decompressed plaintext
+		/// </summary>
+		public ILogText InflatePlainText()
+		{
+			if (CachedPlainText == null)
+			{
+				CachedPlainText = new ReadOnlyLogText(CompressedPlainText.DecompressBzip2());
+			}
+			return CachedPlainText;
+		}
+	}
+
+	/// <summary>
+	/// Extension methods for <see cref="LogIndexBlock"/>
+	/// </summary>
+	public static class LogIndexBlockExtensions
+	{
+		/// <summary>
+		/// Deserialize an index block
+		/// </summary>
+		/// <param name="Reader">Reader to deserialize from</param>
+		/// <returns>The new index block</returns>
+		public static LogIndexBlock ReadLogIndexBlock(this MemoryReader Reader)
+		{
+			int LineIndex = Reader.ReadInt32();
+			int LineCount = Reader.ReadInt32();
+			ReadOnlyMemory<byte> CompressedPlainText = Reader.ReadVariableLengthBytes();
+			return new LogIndexBlock(LineIndex, LineCount, null, CompressedPlainText);
+		}
+
+		/// <summary>
+		/// Serialize an index block
+		/// </summary>
+		/// <param name="Writer">Writer to serialize to</param>
+		/// <param name="Block">The block to serialize</param>
+		public static void WriteLogIndexBlock(this MemoryWriter Writer, LogIndexBlock Block)
+		{
+			Writer.WriteInt32(Block.LineIndex);
+			Writer.WriteInt32(Block.LineCount);
+			Writer.WriteVariableLengthBytes(Block.CompressedPlainText.Span);
+		}
+
+		/// <summary>
+		/// Gets the serialized size of a block
+		/// </summary>
+		/// <param name="Block">The block to serialize</param>
+		public static int GetSerializedSize(this LogIndexBlock Block)
+		{
+			return sizeof(int) + sizeof(int) + (sizeof(int) + Block.CompressedPlainText.Length);
+		}
+	}
+
+	/// <summary>
+	/// Stats for a search
+	/// </summary>
+	public class LogSearchStats
+	{
+		/// <summary>
+		/// Number of blocks that were scanned
+		/// </summary>
+		public int NumScannedBlocks { get; set; }
+
+		/// <summary>
+		/// Number of bytes that had to be scanned for results
+		/// </summary>
+		public int NumScannedBytes { get; set; }
+
+		/// <summary>
+		/// Number of blocks that were skipped
+		/// </summary>
+		public int NumSkippedBlocks { get; set; }
+
+		/// <summary>
+		/// Number of blocks that had to be decompressed
+		/// </summary>
+		public int NumDecompressedBlocks { get; set; }
+
+		/// <summary>
+		/// Number of blocks that were searched but did not contain the search term
+		/// </summary>
+		public int NumFalsePositiveBlocks { get; set; }
+	}
+
+	/// <summary>
+	/// Data for a log chunk
+	/// </summary>
+	public class LogIndexData
+	{
+		/// <summary>
+		/// Index for tokens into the block list
+		/// </summary>
+		ReadOnlyTrie? CachedTrie;
+
+		/// <summary>
+		/// Number of bits in the index devoted to the block index
+		/// </summary>
+		int NumBlockBits;
+
+		/// <summary>
+		/// List of text blocks
+		/// </summary>
+		LogIndexBlock[] Blocks;
+
+		/// <summary>
+		/// Empty index data
+		/// </summary>
+		public static LogIndexData Empty { get; } = new LogIndexData(ReadOnlyTrie.Empty, 0, Array.Empty<LogIndexBlock>());
+
+		/// <summary>
+		/// Number of lines covered by the index
+		/// </summary>
+		public int LineCount => (Blocks.Length > 0) ? (Blocks[Blocks.Length - 1].LineIndex + Blocks[Blocks.Length - 1].LineCount) : 0;
+
+		/// <summary>
+		/// Constructor
+		/// </summary>
+		/// <param name="CachedTrie">Index into the text blocks</param>
+		/// <param name="NumBlockBits">Number of bits devoted to the block index</param>
+		/// <param name="Blocks">Bloom filters for this log file</param>
+		public LogIndexData(ReadOnlyTrie? CachedTrie, int NumBlockBits, LogIndexBlock[] Blocks)
+		{
+			this.CachedTrie = CachedTrie;
+			this.NumBlockBits = NumBlockBits;
+			this.Blocks = Blocks;
+		}
+
+		/// <summary>
+		/// Updates the blocks in this index relative to a give line index
+		/// </summary>
+		/// <param name="LineIndex"></param>
+		public void SetBaseLineIndex(int LineIndex)
+		{
+			for(int BlockIdx = 0; BlockIdx < Blocks.Length; BlockIdx++)
+			{
+				LogIndexBlock Block = Blocks[BlockIdx];
+				Blocks[BlockIdx] = new LogIndexBlock(LineIndex, Block.LineCount, Block.CachedPlainText, Block.CompressedPlainText);
+				LineIndex += Block.LineCount;
+			}
+		}
+
+		/// <summary>
+		/// Builds a trie for this index
+		/// </summary>
+		/// <returns>The trie for this data</returns>
+		ReadOnlyTrie BuildTrie()
+		{
+			if(CachedTrie == null)
+			{
+				ReadOnlyTrieBuilder Builder = new ReadOnlyTrieBuilder();
+				for (int BlockIdx = 0; BlockIdx < Blocks.Length; BlockIdx++)
+				{
+					LogIndexBlock Block = Blocks[BlockIdx];
+					LogToken.GetTokens(Block.InflatePlainText().Data.Span, Token => Builder.Add((Token << NumBlockBits) | (uint)BlockIdx));
+				}
+				CachedTrie = Builder.Build();
+			}
+			return CachedTrie;
+		}
+
+		/// <summary>
+		/// Create an index from an array of blocks
+		/// </summary>
+		/// <param name="Indexes">List of indexes to merge</param>
+		/// <returns>Index data</returns>
+		public static LogIndexData Merge(IEnumerable<LogIndexData> Indexes)
+		{
+			using IScope Scope = GlobalTracer.Instance.BuildSpan("LogIndexData.Merge").StartActive();
+
+			// Create the combined block list
+			LogIndexBlock[] NewBlocks = Indexes.SelectMany(x => x.Blocks).ToArray();
+
+			// Figure out how many bits to devote to the block size
+			int NewNumBlockBits = 0;
+			while(NewBlocks.Length > (1 << NewNumBlockBits))
+			{
+				NewNumBlockBits += 4;
+			}
+
+			// Add all the blocks into a combined index
+			ReadOnlyTrieBuilder NewTrieBuilder = new ReadOnlyTrieBuilder();
+
+			int BlockCount = 0;
+			foreach(LogIndexData Index in Indexes)
+			{
+				ulong BlockMask = ((1UL << Index.NumBlockBits) - 1);
+				foreach (ulong Value in Index.BuildTrie())
+				{
+					ulong Token = Value >> Index.NumBlockBits;
+					int BlockIdx = (int)(Value & BlockMask);
+
+					ulong NewToken = (Token << NewNumBlockBits) | (ulong)(long)(BlockCount + BlockIdx);
+					NewTrieBuilder.Add(NewToken);
+				}
+				BlockCount += Index.Blocks.Length;
+			}
+
+			// Construct the new index data
+			ReadOnlyTrie NewTrie = NewTrieBuilder.Build();
+			return new LogIndexData(NewTrie, NewNumBlockBits, NewBlocks);
+		}
+
+		/// <summary>
+		/// Search for the given text in the index
+		/// </summary>
+		/// <param name="FirstLineIndex">First line index to search from</param>
+		/// <param name="Text">Text to search for</param>
+		/// <param name="Stats">Receives stats for the search</param>
+		/// <returns>List of line numbers for the text</returns>
+		public IEnumerable<int> Search(int FirstLineIndex, SearchText Text, LogSearchStats Stats)
+		{
+			int LastBlockCount = 0;
+			foreach (int BlockIdx in EnumeratePossibleBlocks(Text.Bytes, FirstLineIndex))
+			{
+				LogIndexBlock Block = Blocks[BlockIdx];
+
+				Stats.NumScannedBlocks++;
+				Stats.NumDecompressedBlocks += (Block.CachedPlainText == null)? 1 : 0;
+
+				Stats.NumSkippedBlocks += BlockIdx - LastBlockCount;
+				LastBlockCount = BlockIdx + 1;
+
+				// Decompress the text
+				ILogText BlockText = Block.InflatePlainText();
+
+				// Find the initial offset within this block
+				int Offset = 0;
+				if(FirstLineIndex > Block.LineIndex)
+				{
+					int LineIndexWithinBlock = FirstLineIndex - Block.LineIndex;
+					Offset = BlockText.LineOffsets[LineIndexWithinBlock];
+				}
+
+				// Search within this block
+				for(; ;)
+				{
+					// Find the next offset
+					int NextOffset = BlockText.Data.Span.FindNextOcurrence(Offset, Text);
+					if(NextOffset == -1)
+					{
+						Stats.NumScannedBytes += BlockText.Length - Offset;
+						break;
+					}
+
+					// Update the stats
+					Stats.NumScannedBytes += NextOffset - Offset;
+					Offset = NextOffset;
+
+					// Check it's not another match within the same line
+					int LineIndexWithinBlock = BlockText.GetLineIndexForOffset(Offset);
+					yield return Block.LineIndex + LineIndexWithinBlock;
+
+					// Move to the next line
+					Offset = BlockText.LineOffsets[LineIndexWithinBlock + 1];
+				}
+
+				// If the last scanned bytes is zero, we didn't have any matches from this chunk
+				if(Offset == 0 && CachedTrie != null)
+				{
+					Stats.NumFalsePositiveBlocks++;
+				}
+			}
+			Stats.NumSkippedBlocks += Blocks.Length - LastBlockCount;
+		}
+
+		/// <summary>
+		/// Search for the given text in the index
+		/// </summary>
+		/// <param name="Text">Text to search for</param>
+		/// <param name="LineIndex">The first </param>
+		/// <returns>List of line numbers for the text</returns>
+		IEnumerable<int> EnumeratePossibleBlocks(ReadOnlyMemory<byte> Text, int LineIndex)
+		{
+			// Find the starting block index
+			int BlockIdx = Blocks.BinarySearch(x => x.LineIndex, LineIndex);
+			if(BlockIdx < 0)
+			{
+				BlockIdx = Math.Max(~BlockIdx - 1, 0);
+			}
+
+			// Make sure we're not starting out of range
+			if (BlockIdx < Blocks.Length)
+			{
+				// In order to be considered a possible block for a positive match, the block must contain an exact list of tokens parsed from the
+				// search text, along with a set of partial matches (for substrings that are not guaranteed to start on token boundaries). The former
+				// is relatively cheap to compute, and done by stepping through all enumerators for each token in a wave, only returning blocks
+				// which all contain the token. At the very least, this must contain blocks in the index from BlockIdx onwards.
+				List<IEnumerator<int>> Enumerators = new List<IEnumerator<int>>();
+				Enumerators.Add(Enumerable.Range(BlockIdx, Blocks.Length - BlockIdx).GetEnumerator());
+
+				// The second aspect of the search requires a more expensive search through the trie. This is done by running a set of arbitrary 
+				// delegates to filter the matches returned from the enumerators.
+				List<Predicate<int>> Predicates = new List<Predicate<int>>();
+
+				// If the index has a trie, tokenize the input and generate a list of enumerators and predicates.
+				if (CachedTrie != null)
+				{
+					// Find a list of filters for matching blocks
+					HashSet<ulong> Tokens = new HashSet<ulong>();
+					for (int TokenPos = 0; TokenPos < Text.Length;)
+					{
+						ReadOnlySpan<byte> Token = LogToken.GetTokenText(Text.Span, TokenPos);
+						if (TokenPos == 0)
+						{
+							GetUnalignedTokenPredicate(Token, Token.Length == Text.Length, Predicates);
+						}
+						else
+						{
+							GetAlignedTokenPredicate(Token, TokenPos + Token.Length == Text.Length, Tokens, Predicates);
+						}
+						TokenPos += Token.Length;
+					}
+
+					// Create an enumerator for each token
+					foreach (ulong Token in Tokens)
+					{
+						ulong MinValue = (Token << NumBlockBits) | (ulong)(long)BlockIdx;
+						ulong MaxValue = MinValue + (1UL << NumBlockBits) - 1;
+
+						IEnumerator<int> Enumerator = CachedTrie.EnumerateRange(MinValue, MaxValue).Select(x => (int)(x - MinValue)).GetEnumerator();
+						if (!Enumerator.MoveNext())
+						{
+							yield break;
+						}
+
+						Enumerators.Add(Enumerator);
+					}
+				}
+
+				// Enumerate the matches
+				for (; ; )
+				{
+					// Advance all the enumerators that are behind the current block index. If they are all equal, we have a match.
+					bool bMatch = true;
+					foreach (IEnumerator<int> Enumerator in Enumerators)
+					{
+						while (Enumerator.Current < BlockIdx)
+						{
+							if (!Enumerator.MoveNext())
+							{
+								yield break;
+							}
+						}
+
+						if (Enumerator.Current > BlockIdx)
+						{
+							BlockIdx = Enumerator.Current;
+							bMatch = false;
+						}
+					}
+
+					// Return the match and move to the next block
+					if (bMatch)
+					{
+						if (Predicates.All(Predicate => Predicate(BlockIdx)))
+						{
+							yield return BlockIdx;
+						}
+						BlockIdx++;
+					}
+				}
+			}
+		}
+
+		/// <summary>
+		/// Gets predicates for matching a token that starts 
+		/// </summary>
+		/// <param name="Text">The token text</param>
+		/// <param name="bAllowPartialMatch">Whether to allow a partial match of the token</param>
+		/// <param name="Tokens">Set of aligned tokens that are required</param>
+		/// <param name="Predicates">List of predicates for the search</param>
+		void GetAlignedTokenPredicate(ReadOnlySpan<byte> Text, bool bAllowPartialMatch, HashSet<ulong> Tokens, List<Predicate<int>> Predicates)
+		{
+			for (int Offset = 0; Offset < Text.Length; Offset += LogToken.MaxTokenBytes)
+			{
+				ulong Token = LogToken.GetWindowedTokenValue(Text, Offset);
+				if (Offset + LogToken.MaxTokenBytes > Text.Length && bAllowPartialMatch)
+				{
+					ulong TokenMask = LogToken.GetWindowedTokenMask(Text, Offset, true);
+					Predicates.Add(BlockIdx => BlockContainsToken(BlockIdx, Token, TokenMask));
+					break;
+				}
+				Tokens.Add(Token);
+			}
+		}
+
+		/// <summary>
+		/// Generates a predicate for matching a token which may or may not start on a regular token boundary
+		/// </summary>
+		/// <param name="Text">The token text</param>
+		/// <param name="bAllowPartialMatch">Whether to allow a partial match of the token</param>
+		/// <param name="Predicates">List of predicates for the search</param>
+		void GetUnalignedTokenPredicate(ReadOnlySpan<byte> Text, bool bAllowPartialMatch, List<Predicate<int>> Predicates)
+		{
+			byte[] TextCopy = Text.ToArray();
+
+			Lazy<HashSet<int>> Blocks = new Lazy<HashSet<int>>(() =>
+			{
+				HashSet<int> Union = new HashSet<int>();
+				for (int Shift = 0; Shift < LogToken.MaxTokenBytes; Shift++)
+				{
+					HashSet<int> Blocks = new HashSet<int>(BlocksContainingToken(TextCopy.AsSpan(), -Shift, bAllowPartialMatch));
+					for (int Offset = -Shift + LogToken.MaxTokenBytes; Offset < TextCopy.Length && Blocks.Count > 0; Offset += LogToken.MaxTokenBytes)
+					{
+						Blocks.IntersectWith(BlocksContainingToken(TextCopy.AsSpan(), Offset, bAllowPartialMatch));
+					}
+					if (Blocks.Count > 0)
+					{
+						Union.UnionWith(Blocks);
+					}
+				}
+				return Union;
+			});
+
+			Predicates.Add(BlockIdx => Blocks.Value.Contains(BlockIdx));
+		}
+
+		/// <summary>
+		/// Tests whether a block contains a particular token
+		/// </summary>
+		/// <param name="BlockIdx">Index of the block to search</param>
+		/// <param name="Token">The token to test</param>
+		/// <param name="TokenMask">Mask of which bits in the token are valid</param>
+		/// <returns>True if the given block contains a token</returns>
+		bool BlockContainsToken(int BlockIdx, ulong Token, ulong TokenMask)
+		{
+			Token = (Token << NumBlockBits) | (uint)BlockIdx;
+			TokenMask = (TokenMask << NumBlockBits) | ((1UL << NumBlockBits) - 1);
+			return CachedTrie!.EnumerateValues((Value, ValueMask) => (Value & TokenMask) == (Token & ValueMask)).Any();
+		}
+
+		/// <summary>
+		/// Tests whether a block contains a particular token
+		/// </summary>
+		/// <param name="Text">The token to test</param>
+		/// <param name="Offset">Offset of the window into the token to test</param>
+		/// <param name="bAllowPartialMatch">Whether to allow a partial match of the token</param>
+		/// <returns>True if the given block contains a token</returns>
+		IEnumerable<int> BlocksContainingToken(ReadOnlySpan<byte> Text, int Offset, bool bAllowPartialMatch)
+		{
+			ulong Token = LogToken.GetWindowedTokenValue(Text, Offset) << NumBlockBits;
+			ulong TokenMask = LogToken.GetWindowedTokenMask(Text, Offset, bAllowPartialMatch) << NumBlockBits;
+			ulong BlockMask = (1UL << NumBlockBits) - 1;
+			return CachedTrie!.EnumerateValues((Value, ValueMask) => (Value & TokenMask) == (Token & ValueMask)).Select(x => (int)(x & BlockMask)).Distinct();
+		}
+
+		/// <summary>
+		/// Deserialize the index from memory
+		/// </summary>
+		/// <param name="Reader">Reader to deserialize from</param>
+		/// <returns>Index data</returns>
+		public static LogIndexData Read(MemoryReader Reader)
+		{
+			int Version = Reader.ReadInt32();
+			if (Version == 0)
+			{
+				return LogIndexData.Empty;
+			}
+			else if (Version == 1)
+			{
+				ReadOnlyTrie Index = Reader.ReadTrie();
+				int NumBlockBits = Reader.ReadInt32();
+
+				LogIndexBlock[] Blocks = new LogIndexBlock[Reader.ReadInt32()];
+				for (int Idx = 0; Idx < Blocks.Length; Idx++)
+				{
+					int LineIndex = (Idx > 0) ? (Blocks[Idx - 1].LineIndex + Blocks[Idx - 1].LineCount) : 0;
+					int LineCount = Reader.ReadInt32();
+					ReadOnlyMemory<byte> CompressedPlainText = Reader.ReadVariableLengthBytes();
+					Reader.ReadTrie();
+					Blocks[Idx] = new LogIndexBlock(LineIndex, LineCount, null, CompressedPlainText);
+				}
+
+				return new LogIndexData(Index, NumBlockBits, Blocks);
+			}
+			else if (Version == 2)
+			{
+				ReadOnlyTrie Index = Reader.ReadTrie();
+				int NumBlockBits = Reader.ReadInt32();
+
+				LogIndexBlock[] Blocks = new LogIndexBlock[Reader.ReadInt32()];
+				for (int Idx = 0; Idx < Blocks.Length; Idx++)
+				{
+					int LineIndex = Reader.ReadInt32();
+					int LineCount = Reader.ReadInt32();
+					ReadOnlyMemory<byte> CompressedPlainText = Reader.ReadVariableLengthBytes();
+					Reader.ReadTrie();
+					Blocks[Idx] = new LogIndexBlock(LineIndex, LineCount, null, CompressedPlainText);
+				}
+
+				return new LogIndexData(Index, NumBlockBits, Blocks);
+			}
+			else if (Version == 3)
+			{
+				ReadOnlyTrie Index = Reader.ReadTrie();
+				int NumBlockBits = Reader.ReadInt32();
+				LogIndexBlock[] Blocks = Reader.ReadVariableLengthArray(() => Reader.ReadLogIndexBlock());
+				for(int Idx = 1; Idx < Blocks.Length; Idx++)
+				{
+					if(Blocks[Idx].LineIndex == 0)
+					{
+						int LineIndex = Blocks[Idx - 1].LineIndex + Blocks[Idx - 1].LineCount;
+						Blocks[Idx] = new LogIndexBlock(LineIndex, Blocks[Idx].LineCount, Blocks[Idx].CachedPlainText, Blocks[Idx].CompressedPlainText);
+					}
+				}
+				return new LogIndexData(Index, NumBlockBits, Blocks);
+			}
+			else
+			{
+				throw new InvalidDataException($"Invalid index version number {Version}");
+			}
+		}
+
+		/// <summary>
+		/// Serialize an index into memory
+		/// </summary>
+		/// <param name="Writer">Writer to serialize to</param>
+		public void Write(MemoryWriter Writer)
+		{
+			Writer.WriteInt32(3);
+			Writer.WriteTrie(BuildTrie());
+			Writer.WriteInt32(NumBlockBits);
+			Writer.WriteVariableLengthArray(Blocks, x => Writer.WriteLogIndexBlock(x));
+		}
+
+		/// <summary>
+		/// Deserialize the index from memory
+		/// </summary>
+		/// <param name="Memory">Memory to deserialize from</param>
+		/// <returns>Index data</returns>
+		public static LogIndexData FromMemory(ReadOnlyMemory<byte> Memory)
+		{
+			MemoryReader Reader = new MemoryReader(Memory);
+			return Read(Reader);
+		}
+
+		/// <summary>
+		/// Serializes the index
+		/// </summary>
+		/// <returns>Index data</returns>
+		public byte[] ToByteArray()
+		{
+			byte[] Buffer = new byte[GetSerializedSize()];
+
+			MemoryWriter Writer = new MemoryWriter(Buffer);
+			Write(Writer);
+			Writer.CheckOffset(Buffer.Length);
+
+			return Buffer;
+		}
+
+		/// <summary>
+		/// Gets the serialized size of this index
+		/// </summary>
+		/// <returns>The serialized size</returns>
+		public int GetSerializedSize()
+		{
+			return sizeof(int) + BuildTrie().GetSerializedSize() + sizeof(int) + (sizeof(int) + Blocks.Sum(x => x.GetSerializedSize()));
+		}
+	}
+
+	static class LogIndexDataExtensions
+	{
+		/// <summary>
+		/// Deserialize the index from memory
+		/// </summary>
+		/// <param name="Reader">Reader to deserialize from</param>
+		/// <returns>Index data</returns>
+		public static LogIndexData ReadLogIndexData(this MemoryReader Reader)
+		{
+			return LogIndexData.Read(Reader);
+		}
+
+		/// <summary>
+		/// Serialize an index into memory
+		/// </summary>
+		/// <param name="Writer">Writer to serialize to</param>
+		/// <param name="Index">The index to write</param>
+		public static void WriteLogIndexData(this MemoryWriter Writer, LogIndexData Index)
+		{
+			Index.Write(Writer);
+		}
+	}
+}
diff --git a/Engine/Source/Programs/Horde/HordeServer/Logs/Data/LogSubChunkData.cs b/Engine/Source/Programs/Horde/HordeServer/Logs/Data/LogSubChunkData.cs
new file mode 100644
--- /dev/null
+++ b/Engine/Source/Programs/Horde/HordeServer/Logs/Data/LogSubChunkData.cs
@@ -0,0 +1,255 @@
+// Copyright Epic Games, Inc. All Rights Reserved.
+
+using EpicGames.Core;
+using HordeServer.Api;
+using HordeServer.Utilities;
+using ICSharpCode.SharpZipLib.BZip2;
+using Microsoft.AspNetCore.Routing.Constraints;
+using System;
+using System.Collections.Generic;
+using System.IO;
+using System.Linq;
+using System.Runtime.CompilerServices;
+using System.Text;
+using System.Threading.Tasks;
+
+namespace HordeServer.Logs
+{
+	/// <summary>
+	/// Pending data for a sub-chunk
+	/// </summary>
+	public class LogSubChunkData
+	{
+		/// <summary>
+		/// Type of data stored in this subchunk
+		/// </summary>
+		public LogType Type { get; }
+
+		/// <summary>
+		/// Offset within the file of this sub-chunk
+		/// </summary>
+		public long Offset { get; }
+
+		/// <summary>
+		/// Length of this sub-chunk
+		/// </summary>
+		public int Length { get; }
+
+		/// <summary>
+		/// Index of the first line in this sub-chunk
+		/// </summary>
+		public int LineIndex { get; }
+
+		/// <summary>
+		/// Gets the number of lines in this sub-chunk
+		/// </summary>
+		public int LineCount { get; }
+
+		/// <summary>
+		/// Text data
+		/// </summary>
+		ILogText? TextInternal;
+
+		/// <summary>
+		/// Compressed text data
+		/// </summary>
+		ReadOnlyMemory<byte> CompressedTextInternal;
+
+		/// <summary>
+		/// The index for this sub-chunk
+		/// </summary>
+		LogIndexData? IndexInternal;
+
+		/// <summary>
+		/// The log text
+		/// </summary>
+		public ILogText InflateText()
+		{
+			if (TextInternal == null)
+			{
+				TextInternal = new ReadOnlyLogText(CompressedTextInternal.DecompressBzip2());
+			}
+			return TextInternal;
+		}
+
+		/// <summary>
+		/// The compressed log text
+		/// </summary>
+		public ReadOnlyMemory<byte> DeflateText()
+		{
+			if(CompressedTextInternal.IsEmpty)
+			{
+				CompressedTextInternal = TextInternal!.Data.CompressBzip2();
+			}
+			return CompressedTextInternal;
+		}
+
+		/// <summary>
+		/// Index for tokens in this chunk
+		/// </summary>
+		public LogIndexData BuildIndex()
+		{
+			if(IndexInternal == null)
+			{
+				ILogText PlainText = InflateText();
+				if (Type != LogType.Text)
+				{
+					PlainText = PlainText.ToPlainText();
+				}
+
+				LogIndexBlock[] Blocks = new LogIndexBlock[1];
+				Blocks[0] = new LogIndexBlock(LineIndex, LineCount, PlainText, PlainText.Data.CompressBzip2());
+				IndexInternal = new LogIndexData(null, 0, Blocks);
+			}
+			return IndexInternal;
+		}
+
+		/// <summary>
+		/// Constructor
+		/// </summary>
+		/// <param name="Type">Type of data stored in this subchunk</param>
+		/// <param name="Offset">Offset within the file of this sub-chunk</param>
+		/// <param name="LineIndex">Index of the first line in this sub-chunk</param>
+		/// <param name="Text">Text to add</param>
+		public LogSubChunkData(LogType Type, long Offset, int LineIndex, ILogText Text)
+		{
+			this.Type = Type;
+			this.Offset = Offset;
+			this.Length = Text.Data.Length;
+			this.LineIndex = LineIndex;
+			this.LineCount = Text.LineCount;
+			this.TextInternal = Text;
+		}
+
+		/// <summary>
+		/// Constructor for raw data
+		/// </summary>
+		/// <param name="Type">Type of data stored in this subchunk</param>
+		/// <param name="Offset">Offset within the file of this sub-chunk</param>
+		/// <param name="Length">Length of the uncompressed data</param>
+		/// <param name="LineIndex">Index of the first line</param>
+		/// <param name="LineCount">Number of lines in the uncompressed text</param>
+		/// <param name="CompressedText">Compressed text data</param>
+		/// <param name="Index">Index data</param>
+		public LogSubChunkData(LogType Type, long Offset, int Length, int LineIndex, int LineCount, ReadOnlyMemory<byte> CompressedText, LogIndexData? Index)
+		{
+			this.Type = Type;
+			this.Offset = Offset;
+			this.Length = Length;
+			this.LineIndex = LineIndex;
+			this.LineCount = LineCount;
+			this.CompressedTextInternal = CompressedText;
+			this.IndexInternal = Index;
+		}
+
+		/// <summary>
+		/// Constructs a sub-chunk from a block of memory. Uses slices of the given memory buffer rather than copying the data.
+		/// </summary>
+		/// <param name="Reader">The reader to read from</param>
+		/// <param name="Offset">Offset of the sub-chunk within this file</param>
+		/// <param name="LineIndex">Index of the first line in this subchunk</param>
+		/// <returns>New subchunk data</returns>
+		public static LogSubChunkData Read(MemoryReader Reader, long Offset, int LineIndex)
+		{
+			int Version = Reader.ReadInt32(); // Version placeholder
+			if (Version == 0)
+			{
+				LogType Type = (LogType)Reader.ReadInt32();
+				int Length = Reader.ReadInt32();
+				int LineCount = Reader.ReadInt32();
+				ReadOnlyMemory<byte> CompressedText = Reader.ReadVariableLengthBytes();
+				Reader.ReadVariableLengthBytes();
+				return new LogSubChunkData(Type, Offset, Length, LineIndex, LineCount, CompressedText, null);
+			}
+			else if (Version == 1)
+			{
+				LogType Type = (LogType)Reader.ReadInt32();
+				int Length = Reader.ReadInt32();
+				int LineCount = Reader.ReadInt32();
+				ReadOnlyMemory<byte> CompressedText = Reader.ReadVariableLengthBytes();
+				ReadOnlyMemory<byte> CompressedPlainText = Reader.ReadVariableLengthBytes();
+				ReadOnlyTrie Trie = Reader.ReadTrie();
+				LogIndexBlock IndexBlock = new LogIndexBlock(LineIndex, LineCount, null, CompressedPlainText);
+				LogIndexData Index = new LogIndexData(Trie, 0, new[] { IndexBlock });
+				return new LogSubChunkData(Type, Offset, Length, LineIndex, LineCount, CompressedText, Index);
+			}
+			else
+			{
+				LogType Type = (LogType)Reader.ReadInt32();
+				int Length = Reader.ReadInt32();
+				int LineCount = Reader.ReadInt32();
+				ReadOnlyMemory<byte> CompressedText = Reader.ReadVariableLengthBytes();
+				LogIndexData Index = Reader.ReadLogIndexData();
+				Index.SetBaseLineIndex(LineIndex); // Fix for incorrectly saved data
+				return new LogSubChunkData(Type, Offset, Length, LineIndex, LineCount, CompressedText, Index);
+			}
+		}
+
+		/// <summary>
+		/// Serializes the sub-chunk to a stream
+		/// </summary>
+		/// <param name="Writer">Writer to output to</param>
+		public void Write(MemoryWriter Writer)
+		{
+			Writer.WriteInt32(2); // Version placeholder
+
+			Writer.WriteInt32((int)Type);
+			Writer.WriteInt32(Length);
+			Writer.WriteInt32(LineCount);
+
+			Writer.WriteVariableLengthBytes(DeflateText().Span);
+			Writer.WriteLogIndexData(BuildIndex());
+		}
+
+		/// <summary>
+		/// Serializes this object to a byte array
+		/// </summary>
+		/// <returns>Byte array</returns>
+		public byte[] ToByteArray()
+		{
+			byte[] Data = new byte[GetSerializedSize()];
+
+			MemoryWriter Writer = new MemoryWriter(Data);
+			Write(Writer);
+			Writer.CheckOffset(Data.Length);
+
+			return Data;
+		}
+
+		/// <summary>
+		/// Determines the size of serialized data
+		/// </summary>
+		public int GetSerializedSize()
+		{
+			return sizeof(int) + sizeof(int) + sizeof(int) + sizeof(int) + (sizeof(int) + DeflateText().Length) + BuildIndex().GetSerializedSize();
+		}
+	}
+
+	/// <summary>
+	/// Pending data for a sub-chunk
+	/// </summary>
+	public static class LogSubChunkDataExtensions
+	{
+		/// <summary>
+		/// Constructs a sub-chunk from a block of memory. Uses slices of the given memory buffer rather than copying the data.
+		/// </summary>
+		/// <param name="Reader">The reader to read from</param>
+		/// <param name="Offset">Offset of this sub-chunk within the file</param>
+		/// <param name="LineIndex">Index of the first line within this chunk</param>
+		/// <returns>New subchunk data</returns>
+		public static LogSubChunkData ReadLogSubChunkData(this MemoryReader Reader, long Offset, int LineIndex)
+		{
+			return LogSubChunkData.Read(Reader, Offset, LineIndex);
+		}
+
+		/// <summary>
+		/// Serializes the sub-chunk to a stream
+		/// </summary>
+		/// <param name="Writer">Writer to output to</param>
+		/// <param name="SubChunkData">The sub-chunk data to write</param>
+		public static void WriteLogSubChunkData(this MemoryWriter Writer, LogSubChunkData SubChunkData)
+		{
+			SubChunkData.Write(Writer);
+		}
+	}
+}
