Horde: More file moves.

diff --git a/Engine/Source/Programs/Horde/HordeServer/Logs/Data/LogText.cs b/Engine/Source/Programs/Horde/HordeServer/Logs/Data/LogText.cs
new file mode 100644
--- /dev/null
+++ b/Engine/Source/Programs/Horde/HordeServer/Logs/Data/LogText.cs
@@ -0,0 +1,497 @@
+// Copyright Epic Games, Inc. All Rights Reserved.
+
+using EpicGames.Core;
+using Microsoft.VisualBasic.CompilerServices;
+using System;
+using System.Collections;
+using System.Collections.Generic;
+using System.Diagnostics;
+using System.IO;
+using System.Linq;
+using System.Text;
+using System.Text.Json;
+using System.Threading.Tasks;
+
+namespace HordeServer.Logs
+{
+	/// <summary>
+	/// Interface for log data
+	/// </summary>
+	public interface ILogText
+	{
+		/// <summary>
+		/// The raw text data. Contains a complete set of lines followed by newline characters.
+		/// </summary>
+		public ReadOnlyMemory<byte> Data { get; }
+
+		/// <summary>
+		/// Length of the data in this block
+		/// </summary>
+		public int Length { get; }
+
+		/// <summary>
+		/// Offsets of lines within the data object, including a sentinel for the end of the data (LineCount + 1 entries).
+		/// </summary>
+		public IReadOnlyList<int> LineOffsets { get; }
+
+		/// <summary>
+		/// Number of lines in the 
+		/// </summary>
+		public int LineCount { get; }
+
+		/// <summary>
+		/// Computes the number of bytes used by this object
+		/// </summary>
+		public int AllocatedSize { get; }
+	}
+
+	/// <summary>
+	/// Helper methods for parsing log text
+	/// </summary>
+	public static class LogTextHelpers
+	{
+		/// <summary>
+		/// Finds the number of newline characters in the given span
+		/// </summary>
+		/// <param name="Span">Span of utf-8 characters to scan</param>
+		/// <param name="Offset">Starting offset within the span</param>
+		/// <returns>Number of lines in the given span</returns>
+		public static int GetLineCount(ReadOnlySpan<byte> Span, int Offset)
+		{
+			int LineCount = 0;
+			for (; Offset < Span.Length; Offset++)
+			{
+				if (Span[Offset] == (byte)'\n')
+				{
+					LineCount++;
+				}
+			}
+			return LineCount;
+		}
+
+		/// <summary>
+		/// Find the offsets of the start of each line within the given span
+		/// </summary>
+		/// <param name="Span">Span to search</param>
+		/// <param name="Offset">Offset within the span to start searching</param>
+		/// <param name="LineOffsets">Receives the list of line offsets</param>
+		/// <param name="LineCount">The current line count</param>
+		public static void FindLineOffsets(ReadOnlySpan<byte> Span, int Offset, Span<int> LineOffsets, int LineCount)
+		{
+			Debug.Assert(Span.Length == 0 || Span[Span.Length - 1] == (byte)'\n');
+			for(; ; )
+			{
+				// Store the start of this line
+				LineOffsets[LineCount++] = Offset;
+
+				// Check if we're at the end
+				if(Offset >= Span.Length)
+				{
+					break;
+				}
+
+				// Move to the end of this line
+				while (Span[Offset] != '\n')
+				{
+					Offset++;
+				}
+
+				// Move to the start of the next line
+				Offset++;
+			}
+		}
+	}
+
+	/// <summary>
+	/// Immutable block of log text
+	/// </summary>
+	public class ReadOnlyLogText : ILogText
+	{
+		/// <inheritdoc/>
+		public ReadOnlyMemory<byte> Data { get; private set; }
+
+		/// <inheritdoc/>
+		public int Length
+		{
+			get { return Data.Length; }
+		}
+
+		/// <inheritdoc/>
+		public IReadOnlyList<int> LineOffsets { get; private set; }
+
+		/// <inheritdoc/>
+		public int LineCount
+		{
+			get { return LineOffsets.Count - 1; }
+		}
+
+		/// <inheritdoc/>
+		public int AllocatedSize
+		{
+			get { return Data.Length + (LineOffsets.Count * sizeof(int)); }
+		}
+
+		/// <summary>
+		/// Constructor
+		/// </summary>
+		/// <param name="Data">Data to construct from</param>
+		public ReadOnlyLogText(ReadOnlyMemory<byte> Data)
+		{
+			ReadOnlySpan<byte> Span = Data.Span;
+
+			int LineCount = LogTextHelpers.GetLineCount(Span, 0);
+
+			int[] LineOffsets = new int[LineCount + 1];
+			LogTextHelpers.FindLineOffsets(Span, 0, LineOffsets, 0);
+			Debug.Assert(LineOffsets[LineCount] == Data.Length);
+
+			this.Data = Data;
+			this.LineOffsets = LineOffsets;
+		}
+	}
+
+	/// <summary>
+	/// Mutable log text object
+	/// </summary>
+	public class LogText : ILogText
+	{
+		/// <summary>
+		/// Accessor for Data
+		/// </summary>
+		byte[] InternalData;
+
+		/// <inheritdoc/>
+		public ReadOnlyMemory<byte> Data => InternalData.AsMemory(0, Length);
+
+		/// <summary>
+		/// Hacky accessor for InternalData, for serializing to Redis
+		/// </summary>
+		public byte[] InternalDataHACK => InternalData;
+
+		/// <inheritdoc/>
+		public int Length
+		{
+			get;
+			private set;
+		}
+
+		/// <summary>
+		/// Offsets of the start of each line within the data
+		/// </summary>
+		List<int> InternalLineOffsets = new List<int> { 0 };
+
+		/// <inheritdoc/>
+		public IReadOnlyList<int> LineOffsets => InternalLineOffsets;
+
+		/// <inheritdoc/>
+		public int LineCount => InternalLineOffsets.Count - 1;
+
+		/// <summary>
+		/// Chunk Size
+		/// </summary>
+		public int MaxLength
+		{
+			get { return InternalData.Length; }
+		}
+
+		/// <summary>
+		/// Computes the size of data used by this object
+		/// </summary>
+		public int AllocatedSize
+		{
+			get { return InternalData.Length + (InternalLineOffsets.Capacity * sizeof(int)); }
+		}
+
+		/// <summary>
+		/// Constructor
+		/// </summary>
+		public LogText()
+		{
+			this.InternalData = Array.Empty<byte>();
+		}
+
+		/// <summary>
+		/// Constructor
+		/// </summary>
+		/// <param name="Data">Data to initialize this chunk with. Ownership of this array is transferred to the chunk, and its length determines the chunk size.</param>
+		/// <param name="Length">Number of valid bytes within the initial data array</param>
+		public LogText(byte[] Data, int Length)
+		{
+			this.InternalData = Data;
+			this.Length = 0; // Updated below
+
+			UpdateLength(Length);
+		}
+
+		/// <summary>
+		/// Create a new chunk data object with the given data appended. The internal buffers are reused, with the assumption that
+		/// there is no contention over writing to the same location in the chunk.
+		/// </summary>
+		/// <param name="TextData">The data to append</param>
+		/// <returns>New chunk data object</returns>
+		public void Append(ReadOnlySpan<byte> TextData)
+		{
+			TextData.CopyTo(InternalData.AsSpan(Length, TextData.Length));
+			UpdateLength(Length + TextData.Length);
+		}
+
+		/// <summary>
+		/// Updates the plain text representation of this chunk
+		/// </summary>
+		public void AppendPlainText(ILogText SrcText, int SrcLineIndex, int SrcLineCount)
+		{
+			// Convert all the lines to plain text
+			for (int Idx = 0; Idx < SrcLineCount; Idx++)
+			{
+				int LineOffset = SrcText.LineOffsets[SrcLineIndex + Idx];
+				int NextLineOffset = SrcText.LineOffsets[SrcLineIndex + Idx + 1];
+				ReadOnlySpan<byte> InputLine = SrcText.Data.Slice(LineOffset, NextLineOffset - LineOffset).Span;
+
+				// Make sure the output buffer is large enough
+				int RequiredSpace = NextLineOffset - LineOffset;
+				if (RequiredSpace > InternalData.Length - Length)
+				{
+					int MaxLineOffset = SrcText.LineOffsets[SrcLineIndex + SrcLineCount];
+					Array.Resize(ref InternalData, InternalData.Length + RequiredSpace + (MaxLineOffset - NextLineOffset) / 2);
+				}
+
+				// Convert the line to plain text
+				Length = ConvertToPlainText(InputLine, InternalData, Length);
+				InternalLineOffsets.Add(Length);
+			}
+		}
+
+		/// <summary>
+		/// Determines if the given line is empty
+		/// </summary>
+		/// <param name="Input">The input data</param>
+		/// <returns>True if the given text is empty</returns>
+		public static bool IsEmptyOrWhitespace(ReadOnlySpan<byte> Input)
+		{
+			for(int Idx = 0; Idx < Input.Length; Idx++)
+			{
+				byte Byte = Input[Idx];
+				if(Byte != (byte)'\n' && Byte != '\r' && Byte != ' ')
+				{
+					return false;
+				}
+			}
+			return true;
+		}
+
+		/// <summary>
+		/// Converts a JSON log line to plain text
+		/// </summary>
+		/// <param name="Input">The JSON data</param>
+		/// <param name="Output">Output buffer for the converted line</param>
+		/// <param name="OutputOffset">Offset within the buffer to write the converted data</param>
+		/// <returns></returns>
+		public static int ConvertToPlainText(ReadOnlySpan<byte> Input, byte[] Output, int OutputOffset)
+		{
+			if(IsEmptyOrWhitespace(Input))
+			{
+				Output[OutputOffset] = (byte)'\n';
+				return OutputOffset + 1;
+			}
+
+			Utf8JsonReader Reader = new Utf8JsonReader(Input);
+			if (Reader.Read() && Reader.TokenType == JsonTokenType.StartObject)
+			{
+				while (Reader.Read() && Reader.TokenType == JsonTokenType.PropertyName)
+				{
+					if (!Reader.ValueTextEquals("message"))
+					{
+						Reader.Skip();
+						continue;
+					}
+					if (!Reader.Read() || Reader.TokenType != JsonTokenType.String)
+					{
+						Reader.Skip();
+						continue;
+					}
+
+					int UnescapedLength = UnescapeUtf8(Reader.ValueSpan, Output.AsSpan(OutputOffset));
+					OutputOffset += UnescapedLength;
+
+					Output[OutputOffset] = (byte)'\n';
+					OutputOffset++;
+
+					break;
+				}
+			}
+			return OutputOffset;
+		}
+
+		/// <summary>
+		/// Unescape a json utf8 string
+		/// </summary>
+		/// <param name="Source">Source span of bytes</param>
+		/// <param name="Target">Target span of bytes</param>
+		/// <returns>Length of the converted data</returns>
+		static int UnescapeUtf8(ReadOnlySpan<byte> Source, Span<byte> Target)
+		{
+			int Length = 0;
+			for (; ; )
+			{
+				// Copy up to the next backslash
+				int Backslash = Source.IndexOf((byte)'\\');
+				if (Backslash == -1)
+				{
+					Source.CopyTo(Target);
+					Length += Source.Length;
+					break;
+				}
+				else if (Backslash > 0)
+				{
+					Source.Slice(0, Backslash).CopyTo(Target);
+					Source = Source.Slice(Backslash);
+					Target = Target.Slice(Backslash);
+					Length += Backslash;
+				}
+
+				// Check what the escape code is
+				if (Source[1] == 'u')
+				{
+					char[] Chars = { (char)((StringUtils.ParseHexByte(Source, 2) << 8) | StringUtils.ParseHexByte(Source, 4)) };
+					int EncodedLength = Encoding.UTF8.GetBytes(Chars.AsSpan(), Target);
+					Source = Source.Slice(6);
+					Target = Target.Slice(EncodedLength);
+					Length += EncodedLength;
+				}
+				else
+				{
+					Target[0] = Source[1] switch
+					{
+						(byte)'\"' => (byte)'\"',
+						(byte)'\\' => (byte)'\\',
+						(byte)'\b' => (byte)'\b',
+						(byte)'\f' => (byte)'\f',
+						(byte)'\n' => (byte)'\n',
+						(byte)'\r' => (byte)'\r',
+						(byte)'\t' => (byte)'\t',
+						_ => Source[1]
+					};
+					Source = Source.Slice(2);
+					Target = Target.Slice(1);
+					Length++;
+				}
+			}
+			return Length;
+		}
+
+		/// <summary>
+		/// Generates placeholder data for a missing span
+		/// </summary>
+		/// <param name="ChunkIdx">Index of the chunk</param>
+		/// <param name="HostName">Host name of the machine that was holding the data</param>
+		/// <param name="TargetLength">Desired length of the buffer</param>
+		/// <param name="TargetLineCount">Desired line count for the buffer</param>
+		public void AppendMissingDataInfo(int ChunkIdx, string? HostName, int TargetLength, int TargetLineCount)
+		{
+			int NewLength = Length;
+
+			if(InternalData.Length < TargetLength)
+			{
+				Array.Resize(ref InternalData, TargetLength);
+			}
+
+			string Suffix = $" (server: {HostName})";
+			for(int NewLineCount = LineCount; NewLineCount < TargetLineCount; NewLineCount++)
+			{
+				byte[] ErrorBytes = Encoding.ASCII.GetBytes($"{{ \"level\":\"Error\",\"message\":\"[Missing data at chunk {ChunkIdx}, line {NewLineCount}{Suffix}]\" }}\n");
+				if (NewLength + ErrorBytes.Length > TargetLength)
+				{
+					break;
+				}
+
+				ErrorBytes.AsSpan().CopyTo(InternalData.AsSpan(NewLength));
+				NewLength += ErrorBytes.Length;
+
+				Suffix = String.Empty;
+			}
+
+			if (NewLength < TargetLength)
+			{
+				NewLength = Math.Max(NewLength - 1, 0);
+				InternalData.AsSpan(NewLength, (int)((TargetLength - 1) - NewLength)).Fill((byte)' ');
+				InternalData[TargetLength - 1] = (byte)'\n';
+			}
+
+			UpdateLength((int)TargetLength);
+		}
+
+		/// <summary>
+		/// Shrinks the data allocated by this chunk to the minimum required
+		/// </summary>
+		public void Shrink()
+		{
+			if (Length < InternalData.Length)
+			{
+				Array.Resize(ref InternalData, Length);
+			}
+		}
+
+		/// <summary>
+		/// Updates the length of this chunk, computing all the newline offsets
+		/// </summary>
+		/// <param name="NewLength">New length of the chunk</param>
+		private void UpdateLength(int NewLength)
+		{
+			if (NewLength > Length)
+			{
+				// Make sure the data ends with a newline
+				if (InternalData[NewLength - 1] != '\n')
+				{
+					throw new InvalidDataException("Chunk data must end with a newline");
+				}
+
+				// Calculate the new number of newlines
+				ReadOnlySpan<byte> NewData = InternalData.AsSpan(0, NewLength);
+				for(int Idx = Length; Idx < NewLength; Idx++)
+				{
+					if(NewData[Idx] == '\n')
+					{
+						InternalLineOffsets.Add(Idx + 1);
+					}
+				}
+
+				// Update the size of the buffer
+				Length = NewLength;
+			}
+		}
+	}
+
+	/// <summary>
+	/// Extension methods for ILogText
+	/// </summary>
+	public static class LogTextExtensions
+	{
+		/// <summary>
+		/// Find the line index for a particular offset
+		/// </summary>
+		/// <param name="LogText">The text to search</param>
+		/// <param name="Offset">Offset within the text</param>
+		/// <returns>The line index</returns>
+		public static int GetLineIndexForOffset(this ILogText LogText, int Offset)
+		{
+			int LineIdx = LogText.LineOffsets.BinarySearch(Offset);
+			if(LineIdx < 0)
+			{
+				LineIdx = ~LineIdx - 1;
+			}
+			return LineIdx;
+		}
+
+		/// <summary>
+		/// Converts a log text instance to plain text
+		/// </summary>
+		/// <param name="LogText">The text to convert</param>
+		/// <returns>The plain text instance</returns>
+		public static ILogText ToPlainText(this ILogText LogText)
+		{
+			LogText Other = new LogText();
+			Other.AppendPlainText(LogText, 0, LogText.LineCount);
+			return Other;
+		}
+	}
+}
diff --git a/Engine/Source/Programs/Horde/HordeServer/Logs/Data/LogToken.cs b/Engine/Source/Programs/Horde/HordeServer/Logs/Data/LogToken.cs
new file mode 100644
--- /dev/null
+++ b/Engine/Source/Programs/Horde/HordeServer/Logs/Data/LogToken.cs
@@ -0,0 +1,207 @@
+// Copyright Epic Games, Inc. All Rights Reserved.
+
+using System;
+using System.Collections.Generic;
+using System.Linq;
+using System.Threading.Tasks;
+
+namespace HordeServer.Logs
+{
+	/// <summary>
+	/// Functionality for decomposing log text into tokens
+	/// </summary>
+	public static class LogToken
+	{
+		/// <summary>
+		/// Maximum number of bytes in each token
+		/// </summary>
+		public const int MaxTokenBytes = 4;
+
+		/// <summary>
+		/// Number of bits in each token
+		/// </summary>
+		public const int MaxTokenBits = MaxTokenBytes * 8;
+
+		/// <summary>
+		/// Lookup from input byte to token type
+		/// </summary>
+		static readonly byte[] TokenTypes = GetTokenTypes();
+
+		/// <summary>
+		/// Lookup from input byte to token char
+		/// </summary>
+		static readonly byte[] TokenChars = GetTokenChars();
+
+		/// <summary>
+		/// Get a set of tokens in a trie
+		/// </summary>
+		/// <param name="Text">Text to scan</param>
+		/// <returns></returns>
+		public static ReadOnlyTrie GetTrie(ReadOnlySpan<byte> Text)
+		{
+			ReadOnlyTrieBuilder Builder = new ReadOnlyTrieBuilder();
+			GetTokens(Text, Builder.Add);
+			return Builder.Build();
+		}
+
+		/// <summary>
+		/// Gets a single token
+		/// </summary>
+		/// <param name="Text">The text to parse</param>
+		/// <returns>The token value</returns>
+		public static ulong GetToken(ReadOnlySpan<byte> Text)
+		{
+			ulong Token = 0;
+			GetTokens(Text, x => Token = x);
+			return Token;
+		}
+
+		/// <summary>
+		/// Get a set of tokens in a trie
+		/// </summary>
+		/// <param name="Text">Text to scan</param>
+		/// <returns></returns>
+		public static HashSet<ulong> GetTokenSet(ReadOnlySpan<byte> Text)
+		{
+			HashSet<ulong> Tokens = new HashSet<ulong>();
+			GetTokens(Text, x => Tokens.Add(x));
+			return Tokens;
+		}
+
+		/// <summary>
+		/// Decompose a span of text into tokens
+		/// </summary>
+		/// <param name="Text">Text to scan</param>
+		/// <param name="AddToken">Receives a set of tokens</param>
+		public static void GetTokens(ReadOnlySpan<byte> Text, Action<ulong> AddToken)
+		{
+			if (Text.Length > 0)
+			{
+				int Type = TokenTypes[Text[0]];
+				int NumTokenBits = 8;
+				ulong Token = TokenChars[Text[0]];
+
+				for (int TextIdx = 1; TextIdx < Text.Length; TextIdx++)
+				{
+					byte NextChar = TokenChars[Text[TextIdx]];
+					int NextType = TokenTypes[NextChar];
+					if (Type != NextType || NumTokenBits + 8 > MaxTokenBits)
+					{
+						AddToken(Token << (MaxTokenBits - NumTokenBits));
+						Token = 0;
+						NumTokenBits = 0;
+						Type = NextType;
+					}
+					Token = (Token << 8) | NextChar;
+					NumTokenBits += 8;
+				}
+
+				AddToken(Token << (MaxTokenBits - NumTokenBits));
+			}
+		}
+
+		/// <summary>
+		/// Gets the length of the first token in the given span
+		/// </summary>
+		/// <param name="Text">The text to search</param>
+		/// <param name="Pos">Start position for the search</param>
+		/// <returns>Length of the first token</returns>
+		public static ReadOnlySpan<byte> GetTokenText(ReadOnlySpan<byte> Text, int Pos)
+		{
+			int Type = TokenTypes[Text[Pos]];
+			for (int End = Pos + 1; ; End++)
+			{
+				if (End == Text.Length || TokenTypes[Text[End]] != Type)
+				{
+					return Text.Slice(Pos, End - Pos);
+				}
+			}
+		}
+
+		/// <summary>
+		/// Gets the length of the first token in the given span
+		/// </summary>
+		/// <param name="Text">The text to search</param>
+		/// <param name="Offset">Offset of the window to read from the token</param>
+		/// <returns>Length of the first token</returns>
+		public static ulong GetWindowedTokenValue(ReadOnlySpan<byte> Text, int Offset)
+		{
+			ulong Token = 0;
+			for (int Idx = 0; Idx < MaxTokenBytes; Idx++)
+			{
+				Token <<= 8;
+				if (Offset >= 0 && Offset < Text.Length)
+				{
+					Token |= TokenChars[Text[Offset]];
+				}
+				Offset++;
+			}
+			return Token;
+		}
+
+		/// <summary>
+		/// Gets the length of the first token in the given span
+		/// </summary>
+		/// <param name="Text">The text to search</param>
+		/// <param name="Offset">Offset of the window to read from the token</param>
+		/// <param name="bAllowPartialMatch">Whether to allow only matching the start of the string</param>
+		/// <returns>Length of the first token</returns>
+		public static ulong GetWindowedTokenMask(ReadOnlySpan<byte> Text, int Offset, bool bAllowPartialMatch)
+		{
+			ulong Token = 0;
+			for (int Idx = 0; Idx < MaxTokenBytes; Idx++)
+			{
+				Token <<= 8;
+				if (Offset >= 0 && (Offset < Text.Length || !bAllowPartialMatch))
+				{
+					Token |= 0xff;
+				}
+				Offset++;
+			}
+			return Token;
+		}
+
+		/// <summary>
+		/// Build the lookup table for token types
+		/// </summary>
+		/// <returns>Array whose elements map from an input byte to token type</returns>
+		static byte[] GetTokenTypes()
+		{
+			byte[] CharTypes = new byte[256];
+			for (int Idx = 'a'; Idx <= 'z'; Idx++)
+			{
+				CharTypes[Idx] = 1;
+			}
+			for (int Idx = 'A'; Idx <= 'Z'; Idx++)
+			{
+				CharTypes[Idx] = 1;
+			}
+			for (int Idx = '0'; Idx <= '9'; Idx++)
+			{
+				CharTypes[Idx] = 2;
+			}
+			CharTypes[' '] = 3;
+			CharTypes['\t'] = 3;
+			CharTypes['\n'] = 4;
+			return CharTypes;
+		}
+
+		/// <summary>
+		/// Build the lookup table for token types
+		/// </summary>
+		/// <returns>Array whose elements map from an input byte to token type</returns>
+		static byte[] GetTokenChars()
+		{
+			byte[] Chars = new byte[256];
+			for(int Idx = 0; Idx < 256; Idx++)
+			{
+				Chars[Idx] = (byte)Idx;
+			}
+			for (int Idx = 'A'; Idx <= 'Z'; Idx++)
+			{
+				Chars[Idx] = (byte)('a' + Idx - 'A');
+			}
+			return Chars;
+		}
+	}
+}
diff --git a/Engine/Source/Programs/Horde/HordeServer/Logs/Impl/LogFileService.cs b/Engine/Source/Programs/Horde/HordeServer/Logs/Impl/LogFileService.cs
new file mode 100644
--- /dev/null
+++ b/Engine/Source/Programs/Horde/HordeServer/Logs/Impl/LogFileService.cs
@@ -0,0 +1,1605 @@
+// Copyright Epic Games, Inc. All Rights Reserved.
+
+using EpicGames.Core;
+using HordeCommon;
+using HordeServer.Api;
+using HordeServer.Collections;
+using HordeServer.Logs;
+using HordeServer.Models;
+using HordeServer.Storage;
+using HordeServer.Utilities;
+using Microsoft.Extensions.Caching.Memory;
+using Microsoft.Extensions.Logging;
+using MongoDB.Bson;
+using MongoDB.Driver;
+using System;
+using System.Collections.Concurrent;
+using System.Collections.Generic;
+using System.Collections.Immutable;
+using System.Diagnostics;
+using System.Diagnostics.CodeAnalysis;
+using System.Globalization;
+using System.IO;
+using System.Linq;
+using System.Runtime.CompilerServices;
+using System.Security.Claims;
+using System.Text.Json;
+using System.Text;
+using System.Threading.Tasks;
+using System.Threading;
+
+using ILogger = Microsoft.Extensions.Logging.ILogger;
+using Stream = System.IO.Stream;
+using OpenTracing.Util;
+using OpenTracing;
+
+[assembly: InternalsVisibleTo("HordeServerTests")]
+
+namespace HordeServer.Services
+{
+	/// <summary>
+	/// Metadata about a log file
+	/// </summary>
+	public class LogMetadata
+	{
+		/// <summary>
+		/// Length of the log file
+		/// </summary>
+		public long Length { get; set; }
+
+		/// <summary>
+		/// Number of lines in the log file
+		/// </summary>
+		public int MaxLineIndex { get; set; }
+	}
+
+	/// <summary>
+	/// Interface for the log file service
+	/// </summary>
+	public interface ILogFileService
+	{
+		/// <summary>
+		/// Creates a new log
+		/// </summary>
+		/// <param name="JobId">Unique id of the job that owns this log file</param>
+		/// <param name="SessionId">Agent session allowed to update the log</param>
+		/// <param name="Type">Type of events to be stored in the log</param>
+		/// <returns>The new log file document</returns>
+		Task<ILogFile> CreateLogFileAsync(ObjectId JobId, ObjectId? SessionId, LogType Type);
+
+		/// <summary>
+		/// Gets a logfile by ID
+		/// </summary>
+		/// <param name="LogFileId">Unique id of the log file</param>
+		/// <returns>The logfile document</returns>
+		Task<ILogFile?> GetLogFileAsync(ObjectId LogFileId);
+
+		/// <summary>
+		/// Gets a logfile by ID, returning a cached copy if available. This should only be used to retrieve constant properties set at creation, such as the session or job it's associated with.
+		/// </summary>
+		/// <param name="LogFileId">Unique id of the log file</param>
+		/// <returns>The logfile document</returns>
+		Task<ILogFile?> GetCachedLogFileAsync(ObjectId LogFileId);
+
+		/// <summary>
+		/// Returns a list of log files
+		/// </summary>
+		/// <param name="Index">Index of the first result to return</param>
+		/// <param name="Count">Number of results to return</param>
+		/// <returns>List of logfile documents</returns>
+		Task<List<ILogFile>> GetLogFilesAsync(int? Index = null, int? Count = null);
+
+		/// <summary>
+		/// Writes out chunk data and assigns to a file
+		/// </summary>
+		/// <param name="LogFile">The log file</param>
+		/// <param name="Offset">Offset within the file of data</param>
+		/// <param name="LineIndex">Current line index of the data (need not be the starting of the line)</param>
+		/// <param name="Data">the data to add</param>
+		/// <param name="Flush">Whether the current chunk is complete and should be flushed</param>
+		/// <param name="MaxChunkLength">The maximum chunk length. Defaults to 128kb.</param>
+		/// <param name="MaxSubChunkLineCount">Maximum number of lines in each sub-chunk.</param>
+		/// <returns></returns>
+		Task<ILogFile?> WriteLogDataAsync(ILogFile LogFile, long Offset, int LineIndex, ReadOnlyMemory<byte> Data, bool Flush, int MaxChunkLength = 256 * 1024, int MaxSubChunkLineCount = 128);
+
+		/// <summary>
+		/// Gets metadata about the log file
+		/// </summary>
+		/// <param name="LogFile">The log file to query</param>
+		/// <returns>Metadata about the log file</returns>
+		Task<LogMetadata> GetMetadataAsync(ILogFile LogFile);
+
+		/// <summary>
+		/// Creates new log events
+		/// </summary>
+		/// <param name="NewEvents">List of events</param>
+		/// <returns>Async task</returns>
+		Task CreateEventsAsync(List<NewLogEventData> NewEvents);
+
+		/// <summary>
+		/// Find events for a particular log file
+		/// </summary>
+		/// <param name="LogFile">The log file instance</param>
+		/// <param name="Index">Index of the first event to retrieve</param>
+		/// <param name="Count">Number of events to retrieve</param>
+		/// <returns>List of log events</returns>
+		Task<List<ILogEvent>> FindLogEventsAsync(ILogFile LogFile, int? Index = null, int? Count = null);
+
+		/// <summary>
+		/// Adds events to a log span
+		/// </summary>
+		/// <param name="Events">The events to add</param>
+		/// <param name="SpanId">The span id</param>
+		/// <returns>Async task</returns>
+		Task AddSpanToEventsAsync(IEnumerable<ILogEvent> Events, ObjectId SpanId);
+
+		/// <summary>
+		/// Find events for an issue
+		/// </summary>
+		/// <param name="SpanIds">The span ids</param>
+		/// <param name="LogIds">Log ids to include</param>
+		/// <param name="Index">Index within the events for results to return</param>
+		/// <param name="Count">Number of results to return</param>
+		/// <returns>Async task</returns>
+		Task<List<ILogEvent>> FindEventsForSpansAsync(IEnumerable<ObjectId> SpanIds, ObjectId[]? LogIds, int Index, int Count);
+
+		/// <summary>
+		/// Gets the data for an event
+		/// </summary>
+		/// <param name="LogFile">The log file instance</param>
+		/// <param name="LineIndex">Index of the line in the file</param>
+		/// <param name="LineCount">Number of lines in the event</param>
+		/// <returns>New event data instance</returns>
+		Task<ILogEventData> GetEventDataAsync(ILogFile LogFile, int LineIndex, int LineCount);
+
+		/// <summary>
+		/// Gets lines from the given log 
+		/// </summary>
+		/// <param name="LogFile">The log file</param>
+		/// <param name="Offset">Offset of the data to return</param>
+		/// <param name="Length">Length of the data to return</param>
+		/// <returns>Data for the requested range</returns>
+		Task<Stream> OpenRawStreamAsync(ILogFile LogFile, long Offset, long Length);
+
+		/// <summary>
+		/// Gets the offset of the given line number
+		/// </summary>
+		/// <param name="LogFile">The log file to search</param>
+		/// <param name="LineIdx">The line index to retrieve the offset for</param>
+		/// <returns>The actual clamped line number and offset</returns>
+		Task<(int, long)> GetLineOffsetAsync(ILogFile LogFile, int LineIdx);
+
+		/// <summary>
+		/// Search for the specified text in a log file
+		/// </summary>
+		/// <param name="LogFile">The log file to search</param>
+		/// <param name="Text">Text to search for</param>
+		/// <param name="FirstLine">Line to start search from</param>
+		/// <param name="Count">Number of results to return</param>
+		/// <param name="Stats">Receives stats for the search</param>
+		/// <returns>List of line numbers containing the given term</returns>
+		Task<List<int>> SearchLogDataAsync(ILogFile LogFile, string Text, int FirstLine, int Count, LogSearchStats Stats);
+	}
+
+	/// <summary>
+	/// Extension methods for dealing with log files
+	/// </summary>
+	public static class LogFileServiceExtensions
+	{
+		/// <summary>
+		/// Parses a stream of json text and outputs plain text
+		/// </summary>
+		/// <param name="LogFileService">The log file service</param>
+		/// <param name="LogFile">The log file to query</param>
+		/// <param name="Offset">Offset within the log file to copy</param>
+		/// <param name="Length">Length of the data to copy</param>
+		/// <param name="OutputStream">Output stream to receive the text data</param>
+		/// <returns>Async text</returns>
+		public static async Task CopyRawStreamAsync(this ILogFileService LogFileService, ILogFile LogFile, long Offset, long Length, Stream OutputStream)
+		{
+			using (Stream Stream = await LogFileService.OpenRawStreamAsync(LogFile, Offset, Length))
+			{
+				await Stream.CopyToAsync(OutputStream);
+			}
+		}
+
+		/// <summary>
+		/// Parses a stream of json text and outputs plain text
+		/// </summary>
+		/// <param name="LogFileService">The log file service</param>
+		/// <param name="LogFile">The log file to query</param>
+		/// <param name="Offset">Offset within the data to copy from</param>
+		/// <param name="Length">Length of the data to copy</param>
+		/// <param name="OutputStream">Output stream to receive the text data</param>
+		/// <returns>Async text</returns>
+		public static async Task CopyPlainTextStreamAsync(this ILogFileService LogFileService, ILogFile LogFile, long Offset, long Length, Stream OutputStream)
+		{
+			using (Stream Stream = await LogFileService.OpenRawStreamAsync(LogFile, 0, long.MaxValue))
+			{
+				byte[] ReadBuffer = new byte[4096];
+				int ReadBufferLength = 0;
+
+				byte[] WriteBuffer = new byte[4096];
+				int WriteBufferLength = 0;
+
+				while(Length > 0)
+				{
+					// Add more data to the buffer
+					int ReadBytes = await Stream.ReadAsync(ReadBuffer, ReadBufferLength, ReadBuffer.Length - ReadBufferLength);
+					ReadBufferLength += ReadBytes;
+
+					// Copy as many lines as possible to the output
+					int ConvertedBytes = 0;
+					for(int EndIdx = 1; EndIdx < ReadBufferLength; EndIdx++)
+					{
+						if (ReadBuffer[EndIdx] == '\n')
+						{
+							WriteBufferLength = LogText.ConvertToPlainText(ReadBuffer.AsSpan(ConvertedBytes, EndIdx - ConvertedBytes), WriteBuffer, WriteBufferLength);
+							ConvertedBytes = EndIdx + 1;
+						}
+					}
+
+					// If there's anything in the write buffer, write it out
+					if (WriteBufferLength > 0)
+					{
+						if (Offset < WriteBufferLength)
+						{
+							int WriteLength = (int)Math.Min((long)WriteBufferLength - Offset, Length);
+							await OutputStream.WriteAsync(WriteBuffer, (int)Offset, WriteLength);
+							Length -= WriteLength;
+						}
+						Offset = Math.Max(Offset - WriteBufferLength, 0);
+						WriteBufferLength = 0;
+					}
+
+					// If we were able to read something, shuffle down the rest of the buffer. Otherwise expand the read buffer.
+					if (ConvertedBytes > 0)
+					{
+						Buffer.BlockCopy(ReadBuffer, ConvertedBytes, ReadBuffer, 0, ReadBufferLength - ConvertedBytes);
+						ReadBufferLength -= ConvertedBytes;
+					}
+					else if(ReadBufferLength > 0)
+					{
+						Array.Resize(ref ReadBuffer, ReadBuffer.Length + 128);
+						WriteBuffer = new byte[ReadBuffer.Length];
+					}
+
+					// Exit if we didn't read anything in this iteration
+					if(ReadBytes == 0)
+					{
+						break;
+					}
+				}
+			}
+		}
+	}
+	
+	/// <summary>
+	/// Wraps functionality for manipulating logs
+	/// </summary>
+	public class LogFileService : TickedBackgroundService, ILogFileService
+	{
+		/// <summary>
+		/// Information Logger
+		/// </summary>
+		private readonly ILogger<LogFileService> Logger;
+
+		/// <summary>
+		/// Collection of log documents
+		/// </summary>
+		private readonly ILogFileCollection LogFiles;
+
+		/// <summary>
+		/// Collection of log events
+		/// </summary>
+		private readonly ILogEventCollection LogEvents;
+
+		/// <summary>
+		/// Interface for the log reader
+		/// </summary>
+		private readonly ILogStorage Storage;
+
+		/// <summary>
+		/// Interface for the log builder
+		/// </summary>
+		private readonly ILogBuilder Builder;
+
+		/// <summary>
+		/// Lock object for the <see cref="WriteTasks"/> and <see cref="WriteChunks"/> members
+		/// </summary>
+		private object WriteLock = new object();
+
+		/// <summary>
+		/// List of active write tasks
+		/// </summary>
+		private List<Task> WriteTasks = new List<Task>();
+
+		/// <summary>
+		/// Set of chunks which is currently being written to
+		/// </summary>
+		private HashSet<(ObjectId, long)> WriteChunks = new HashSet<(ObjectId, long)>();
+
+		/// <summary>
+		/// Cache of session ids for each log
+		/// </summary>
+		private IMemoryCache LogFileCache;
+
+		/// <summary>
+		/// Streams log data to a caller
+		/// </summary>
+		class ResponseStream : Stream
+		{
+			/// <summary>
+			/// The log file service that created this stream
+			/// </summary>
+			LogFileService LogFileService;
+
+			/// <summary>
+			/// The log file being read
+			/// </summary>
+			ILogFile LogFile;
+
+			/// <summary>
+			/// Starting offset within the file of the data to return 
+			/// </summary>
+			long ResponseOffset;
+
+			/// <summary>
+			/// Length of data to return
+			/// </summary>
+			long ResponseLength;
+
+			/// <summary>
+			/// Current offset within the stream
+			/// </summary>
+			long CurrentOffset;
+
+			/// <summary>
+			/// The current chunk index
+			/// </summary>
+			int ChunkIdx;
+
+			/// <summary>
+			/// Buffer containing a message for missing data
+			/// </summary>
+			ReadOnlyMemory<byte> SourceBuffer;
+
+			/// <summary>
+			/// Offset within the source buffer
+			/// </summary>
+			int SourcePos;
+
+			/// <summary>
+			/// Length of the source buffer being copied from
+			/// </summary>
+			int SourceEnd;
+
+			/// <summary>
+			/// Constructor
+			/// </summary>
+			/// <param name="LogFileService">The log file service, for q</param>
+			/// <param name="LogFile"></param>
+			/// <param name="Offset"></param>
+			/// <param name="Length"></param>
+			public ResponseStream(LogFileService LogFileService, ILogFile LogFile, long Offset, long Length)
+			{
+				this.LogFileService = LogFileService;
+				this.LogFile = LogFile;
+
+				this.ResponseOffset = Offset;
+				this.ResponseLength = Length;
+
+				this.CurrentOffset = Offset;
+
+				this.ChunkIdx = LogFile.Chunks.GetChunkForOffset(Offset);
+				this.SourceBuffer = null!;
+			}
+
+			/// <inheritdoc/>
+			public override bool CanRead => true;
+
+			/// <inheritdoc/>
+			public override bool CanSeek => false;
+
+			/// <inheritdoc/>
+			public override bool CanWrite => false;
+
+			/// <inheritdoc/>
+			public override long Length
+			{
+				get { return ResponseLength; }
+			}
+
+			/// <inheritdoc/>
+			public override long Position
+			{
+				get { return CurrentOffset - ResponseOffset; }
+				set => throw new NotImplementedException();
+			}
+
+			/// <inheritdoc/>
+			public override void Flush()
+			{
+			}
+
+			/// <inheritdoc/>
+			public override int Read(byte[] Buffer, int Offset, int Count)
+			{
+				return ReadAsync(Buffer, Offset, Count, CancellationToken.None).Result;
+			}
+
+			/// <inheritdoc/>
+			public override async Task<int> ReadAsync(byte[] Buffer, int Offset, int Length, CancellationToken CancellationToken)
+			{
+				int ReadBytes = 0;
+				while (ReadBytes < Length)
+				{
+					if (SourcePos < SourceEnd)
+					{
+						// Try to copy from the current buffer
+						int BlockSize = Math.Min(SourceEnd - SourcePos, Length - ReadBytes);
+						SourceBuffer.Slice(SourcePos, BlockSize).Span.CopyTo(Buffer.AsSpan(Offset + ReadBytes));
+						CurrentOffset += BlockSize;
+						ReadBytes += BlockSize;
+						SourcePos += BlockSize;
+					}
+					else if (CurrentOffset < ResponseOffset + ResponseLength)
+					{
+						// Move to the right chunk
+						while (ChunkIdx + 1 < LogFile.Chunks.Count && CurrentOffset >= LogFile.Chunks[ChunkIdx + 1].Offset)
+						{
+							ChunkIdx++;
+						}
+
+						// Get the end of this chunk
+						long NextOffset = ResponseOffset + ResponseLength;
+						if (ChunkIdx + 1 < LogFile.Chunks.Count)
+						{
+							ILogChunk NextChunk = LogFile.Chunks[ChunkIdx + 1];
+							NextOffset = Math.Min(NextOffset, NextChunk.Offset);
+						}
+
+						// Get the chunk data
+						ILogChunk Chunk = LogFile.Chunks[ChunkIdx];
+						LogChunkData ChunkData = await LogFileService.ReadChunkAsync(LogFile, ChunkIdx);
+
+						// Figure out which sub-chunk to use
+						int SubChunkIdx = ChunkData.GetSubChunkForOffsetWithinChunk((int)(CurrentOffset - Chunk.Offset));
+						LogSubChunkData SubChunkData = ChunkData.SubChunks[SubChunkIdx];
+
+						// Get the source data
+						long SubChunkOffset = Chunk.Offset + ChunkData.SubChunkOffset[SubChunkIdx];
+						SourceBuffer = SubChunkData.InflateText().Data;
+						SourcePos = (int)(CurrentOffset - SubChunkOffset);
+						SourceEnd = (int)Math.Min(SourceBuffer.Length, (ResponseOffset + ResponseLength) - SubChunkOffset);
+					}
+					else
+					{
+						// End of the log
+						break;
+					}
+				}
+				return ReadBytes;
+			}
+
+			/// <inheritdoc/>
+			public override long Seek(long offset, SeekOrigin origin) => throw new NotImplementedException();
+
+			/// <inheritdoc/>
+			public override void SetLength(long value) => throw new NotImplementedException();
+
+			/// <inheritdoc/>
+			public override void Write(byte[] buffer, int offset, int count) => throw new NotImplementedException();
+		}
+
+		/// <summary>
+		/// Constructor
+		/// </summary>
+		/// <param name="LogFiles">The log file collection</param>
+		/// <param name="LogEvents">The log events collection</param>
+		/// <param name="Builder">The log builder</param>
+		/// <param name="Storage">THe log storage hierarchy</param>
+		/// <param name="Logger">Log interface</param>
+		public LogFileService(ILogFileCollection LogFiles, ILogEventCollection LogEvents, ILogBuilder Builder, ILogStorage Storage, ILogger<LogFileService> Logger)
+			: base(TimeSpan.FromSeconds(30.0), Logger)
+		{
+			this.LogFiles = LogFiles;
+			this.LogEvents = LogEvents;
+			this.Logger = Logger;
+			this.LogFileCache = new MemoryCache(new MemoryCacheOptions());
+			this.Builder = Builder;
+			this.Storage = Storage;
+		}
+
+		/// <inheritdoc/>
+		public Task<ILogFile> CreateLogFileAsync(ObjectId JobId, ObjectId? SessionId, LogType Type)
+		{
+			return LogFiles.CreateLogFileAsync(JobId, SessionId, Type);
+		}
+
+		/// <inheritdoc/>
+		public async Task<ILogFile?> GetLogFileAsync(ObjectId LogFileId)
+		{
+			ILogFile? LogFile = await LogFiles.GetLogFileAsync(LogFileId);
+			if(LogFile != null)
+			{
+				AddCachedLogFile(LogFile);
+			}
+			return LogFile;
+		}
+
+		/// <summary>
+		/// Adds a log file to the cache
+		/// </summary>
+		/// <param name="LogFile">The log file to cache</param>
+		void AddCachedLogFile(ILogFile LogFile)
+		{
+			MemoryCacheEntryOptions Options = new MemoryCacheEntryOptions().SetSlidingExpiration(TimeSpan.FromSeconds(30));
+			LogFileCache.Set(LogFile.Id, LogFile, Options);
+		}
+
+		/// <summary>
+		/// Gets a cached log file by id
+		/// </summary>
+		/// <param name="LogFileId">The log file id</param>
+		/// <returns>New log file, or null if not found</returns>
+		public async Task<ILogFile?> GetCachedLogFileAsync(ObjectId LogFileId)
+		{
+			object? LogFile;
+			if (!LogFileCache.TryGetValue(LogFileId, out LogFile))
+			{
+				LogFile = await GetLogFileAsync(LogFileId);
+			}
+			return (ILogFile?)LogFile;
+		}
+
+		/// <inheritdoc/>
+		public Task<List<ILogFile>> GetLogFilesAsync(int? Index = null, int? Count = null)
+		{
+			return LogFiles.GetLogFilesAsync(Index, Count);
+		}
+
+		class WriteState
+		{
+			public long Offset;
+			public int LineIndex;
+			public ReadOnlyMemory<byte> Memory;
+
+			public WriteState(long Offset, int LineIndex, ReadOnlyMemory<byte> Memory)
+			{
+				this.Offset = Offset;
+				this.LineIndex = LineIndex;
+				this.Memory = Memory;
+			}
+		}
+
+		/// <inheritdoc/>
+		public async Task<ILogFile?> WriteLogDataAsync(ILogFile LogFile, long Offset, int LineIndex, ReadOnlyMemory<byte> Data, bool Flush, int MaxChunkLength, int MaxSubChunkLineCount)
+		{
+			using IScope Scope = GlobalTracer.Instance.BuildSpan("WriteLogDataAsync").StartActive();
+			Scope.Span.SetTag("LogId", LogFile.Id.ToString());
+			Scope.Span.SetTag("Offset", Offset.ToString(CultureInfo.InvariantCulture));
+			Scope.Span.SetTag("Length", Data.Length.ToString(CultureInfo.InvariantCulture));
+			Scope.Span.SetTag("LineIndex", LineIndex.ToString(CultureInfo.InvariantCulture));
+
+			// Make sure the data ends in a newline
+			if (Data.Length > 0 && Data.Span[Data.Length - 1] != '\n')
+			{
+				throw new ArgumentException("Log data must consist of a whole number of lines", nameof(Data));
+			}
+
+			// Make sure the line count is a power of two
+			if ((MaxSubChunkLineCount & (MaxSubChunkLineCount - 1)) != 0)
+			{
+				throw new ArgumentException("Maximum line count per sub-chunk must be a power of two", nameof(MaxSubChunkLineCount));
+			}
+
+			// List of the flushed chunks
+			List<long> CompleteOffsets = new List<long>();
+
+			// Add the data to new chunks
+			WriteState State = new WriteState(Offset, LineIndex, Data);
+			while (State.Memory.Length > 0)
+			{
+				// Find an existing chunk to append to
+				int ChunkIdx = LogFile.Chunks.GetChunkForOffset(State.Offset);
+				if (ChunkIdx >= 0)
+				{
+					ILogChunk Chunk = LogFile.Chunks[ChunkIdx];
+					if (await WriteLogChunkDataAsync(LogFile, Chunk, State, CompleteOffsets, MaxChunkLength, MaxSubChunkLineCount))
+					{
+						continue;
+					}
+				}
+
+				// Create a new chunk. Ensure that there's a chunk at the start of the file, even if the current write is beyond it.
+				ILogFile? NewLogFile;
+				if (LogFile.Chunks.Count == 0)
+				{
+					NewLogFile = await LogFiles.TryAddChunkAsync(LogFile, 0, 0);
+				}
+				else
+				{
+					NewLogFile = await LogFiles.TryAddChunkAsync(LogFile, State.Offset, State.LineIndex);
+				}
+
+				// Try to add a new chunk at the new location
+				if (NewLogFile == null)
+				{
+					NewLogFile = await LogFiles.GetLogFileAsync(LogFile.Id);
+					if (NewLogFile == null)
+					{
+						Logger.LogError("Unable to update log file {LogId}", LogFile.Id);
+						return null;
+					}
+					LogFile = NewLogFile;
+				}
+				else
+				{
+					// Logger.LogDebug("Added new chunk at offset {Offset} to log {LogId}", State.Offset, LogFile.Id);
+					LogFile = NewLogFile;
+				}
+			}
+
+			// Flush any pending chunks on this log file
+			if (Flush)
+			{
+				foreach(ILogChunk Chunk in LogFile.Chunks)
+				{
+					if (Chunk.Length == 0 && !CompleteOffsets.Contains(Chunk.Offset))
+					{
+						await Builder.CompleteChunkAsync(LogFile.Id, Chunk.Offset);
+						CompleteOffsets.Add(Chunk.Offset);
+					}
+				}
+			}
+
+			// Write all the chunks
+			if (CompleteOffsets.Count > 0 || Flush)
+			{
+				ILogFile? NewLogFile = await WriteCompleteChunksForLogAsync(LogFile, CompleteOffsets, Flush);
+				if (NewLogFile == null)
+				{
+					return null;
+				}
+				LogFile = NewLogFile;
+			}
+			return LogFile;
+		}
+
+		/// <summary>
+		/// Append data to an existing chunk.
+		/// </summary>
+		/// <param name="LogFile">The log file to append to</param>
+		/// <param name="Chunk">Chunk within the log file to update</param>
+		/// <param name="State">Data remaining to be written</param>
+		/// <param name="CompleteOffsets">List of complete chunks</param>
+		/// <param name="MaxChunkLength">Maximum length of each chunk</param>
+		/// <param name="MaxSubChunkLineCount">Maximum number of lines in each subchunk</param>
+		/// <returns>True if data was appended to </returns>
+		private async Task<bool> WriteLogChunkDataAsync(ILogFile LogFile, ILogChunk Chunk, WriteState State, List<long> CompleteOffsets, int MaxChunkLength, int MaxSubChunkLineCount)
+		{
+			// Don't allow data to be appended if the chunk is complete
+			if(Chunk.Length > 0)
+			{
+				return false;
+			}
+
+			// Otherwise keep appending subchunks
+			bool bResult = false;
+			for (; ; )
+			{
+				// Flush the current sub-chunk if we're on a boundary
+				if (State.LineIndex > 0 && (State.LineIndex & (MaxSubChunkLineCount - 1)) == 0)
+				{
+					Logger.LogDebug("Completing log {LogId} chunk offset {Offset} sub-chunk at line {LineIndex}", LogFile.Id, Chunk.Offset, State.LineIndex);
+					await Builder.CompleteSubChunkAsync(LogFile.Id, Chunk.Offset);
+				}
+
+				// Figure out the max length to write to the current chunk
+				int MaxLength = Math.Min((int)((Chunk.Offset + MaxChunkLength) - State.Offset), State.Memory.Length);
+
+				// Figure out the maximum line index for the current sub chunk
+				int MinLineIndex = State.LineIndex;
+				int MaxLineIndex = (MinLineIndex & ~(MaxSubChunkLineCount - 1)) + MaxSubChunkLineCount;
+
+				// Append this data
+				(int Length, int LineCount) = GetWriteLength(State.Memory.Span, MaxLength, MaxLineIndex - MinLineIndex, State.Offset == Chunk.Offset);
+				if (Length > 0)
+				{
+					// Append this data
+					ReadOnlyMemory<byte> AppendData = State.Memory.Slice(0, Length);
+					if (!await Builder.AppendAsync(LogFile.Id, Chunk.Offset, State.Offset, State.LineIndex, LineCount, AppendData, LogFile.Type))
+					{
+						break;
+					}
+
+					// Update the state
+					//Logger.LogDebug("Append to log {LogId} chunk offset {Offset} (LineIndex={LineIndex}, LineCount={LineCount}, Offset={WriteOffset}, Length={WriteLength})", LogFile.Id, Chunk.Offset, State.LineIndex, LineCount, State.Offset, Length);
+					State.Offset += Length;
+					State.LineIndex += LineCount;
+					State.Memory = State.Memory.Slice(Length);
+					bResult = true;
+
+					// If this is the end of the data, bail out
+					if(State.Memory.Length == 0)
+					{
+						break;
+					}
+				}
+
+				// Flush the sub-chunk if it's full
+				if (State.LineIndex < MaxLineIndex)
+				{
+					Logger.LogDebug("Completing chunk for log {LogId} at offset {Offset}", LogFile.Id, Chunk.Offset);
+					await Builder.CompleteChunkAsync(LogFile.Id, Chunk.Offset);
+					CompleteOffsets.Add(Chunk.Offset);
+					break;
+				}
+			}
+			return bResult;
+		}
+
+		/// <summary>
+		/// Get the amount of data to write from the given span
+		/// </summary>
+		/// <param name="Span">Data to write</param>
+		/// <param name="MaxLength">Maximum length of the data to write</param>
+		/// <param name="MaxLineCount">Maximum number of lines to write</param>
+		/// <param name="bIsEmptyChunk">Whether the current chunk is empty</param>
+		/// <returns>A tuple consisting of the amount of data to write and number of lines in it</returns>
+		private static (int, int) GetWriteLength(ReadOnlySpan<byte> Span, int MaxLength, int MaxLineCount, bool bIsEmptyChunk)
+		{
+			int Length = 0;
+			int LineCount = 0;
+			for (int Idx = 0; Idx < MaxLength || bIsEmptyChunk; Idx++)
+			{
+				if (Span[Idx] == '\n')
+				{
+					Length = Idx + 1;
+					LineCount++;
+					bIsEmptyChunk = false;
+
+					if (LineCount >= MaxLineCount)
+					{
+						break;
+					}
+				}
+			}
+			return (Length, LineCount);
+		}
+
+		/// <inheritdoc/>
+		public async Task<LogMetadata> GetMetadataAsync(ILogFile LogFile)
+		{
+			LogMetadata Metadata = new LogMetadata();
+			if (LogFile.Chunks.Count > 0)
+			{
+				ILogChunk Chunk = LogFile.Chunks[LogFile.Chunks.Count - 1];
+				if (LogFile.MaxLineIndex == null || Chunk.Length == 0)
+				{
+					LogChunkData ChunkData = await ReadChunkAsync(LogFile, LogFile.Chunks.Count - 1);
+					Metadata.Length = Chunk.Offset + ChunkData.Length;
+					Metadata.MaxLineIndex = Chunk.LineIndex + ChunkData.LineCount;
+				}
+				else
+				{
+					Metadata.Length = Chunk.Offset + Chunk.Length;
+					Metadata.MaxLineIndex = LogFile.MaxLineIndex.Value;
+				}
+			}
+			return Metadata;
+		}
+
+		/// <inheritdoc/>
+		public Task CreateEventsAsync(List<NewLogEventData> NewEvents)
+		{
+			return LogEvents.AddManyAsync(NewEvents);
+		}
+
+		/// <inheritdoc/>
+		public Task<List<ILogEvent>> FindLogEventsAsync(ILogFile LogFile, int? Index = null, int? Count = null)
+		{
+			return LogEvents.FindAsync(LogFile.Id, Index, Count);
+		}
+
+		class LogEventLine : ILogEventLine
+		{
+			LogLevel Level;
+			public EventId? EventId { get; }
+			public string Message { get; }
+			public JsonElement Data { get; }
+
+			LogLevel ILogEventLine.Level => Level;
+
+			public LogEventLine(ReadOnlySpan<byte> Data)
+				: this(JsonSerializer.Deserialize<JsonElement>(Data))
+			{
+			}
+
+			[SuppressMessage("Design", "CA1031:Do not catch general exception types")]
+			public LogEventLine(JsonElement Data)
+			{
+				this.Data = Data;
+
+				JsonElement LevelElement;
+				if (!Data.TryGetProperty("level", out LevelElement) || !Enum.TryParse(LevelElement.GetString(), out Level))
+				{
+					Level = LogLevel.Information;
+				}
+
+				JsonElement IdElement;
+				if (Data.TryGetProperty("id", out IdElement))
+				{
+					int IdValue;
+					if (IdElement.TryGetInt32(out IdValue))
+					{
+						EventId = IdValue;
+					}
+				}
+
+				JsonElement MessageElement;
+				if (Data.TryGetProperty("renderedMessage", out MessageElement) || Data.TryGetProperty("message", out MessageElement))
+				{
+					Message = MessageElement.GetString() ?? "(Invalid)";
+				}
+				else
+				{
+					Message = "(Missing message or renderedMessage field)";
+				}
+			}
+		}
+
+		class LogEventData : ILogEventData
+		{
+			public IReadOnlyList<ILogEventLine> Lines { get; }
+
+			EventId? ILogEventData.EventId => (Lines.Count > 0) ? Lines[0].EventId : null;
+			EventSeverity ILogEventData.Severity => (Lines.Count == 0) ? EventSeverity.Information : (Lines[0].Level == LogLevel.Warning) ? EventSeverity.Warning : EventSeverity.Error;
+			string ILogEventData.Message => String.Join("\n", Lines.Select(x => x.Message));
+
+			public LogEventData(IReadOnlyList<ILogEventLine> Lines)
+			{
+				this.Lines = Lines;
+			}
+		}
+
+		/// <inheritdoc/>
+		public Task AddSpanToEventsAsync(IEnumerable<ILogEvent> Events, ObjectId SpanId)
+		{
+			return LogEvents.AddSpanToEventsAsync(Events, SpanId);
+		}
+
+		/// <inheritdoc/>
+		public Task<List<ILogEvent>> FindEventsForSpansAsync(IEnumerable<ObjectId> SpanIds, ObjectId[]? LogIds, int Index, int Count)
+		{
+			return LogEvents.FindEventsForSpansAsync(SpanIds, LogIds, Index, Count);
+		}
+
+		/// <inheritdoc/>
+		public async Task<ILogEventData> GetEventDataAsync(ILogFile LogFile, int LineIndex, int LineCount)
+		{
+			using IScope Scope = GlobalTracer.Instance.BuildSpan("GetEventDataAsync").StartActive();
+			Scope.Span.SetTag("LogId", LogFile.Id.ToString());
+			Scope.Span.SetTag("LineIndex", LineIndex.ToString(CultureInfo.InvariantCulture));
+			Scope.Span.SetTag("LineCount", LineCount.ToString(CultureInfo.InvariantCulture));
+
+			(_, long MinOffset) = await GetLineOffsetAsync(LogFile, LineIndex);
+			(_, long MaxOffset) = await GetLineOffsetAsync(LogFile, LineIndex + LineCount);
+
+			byte[] Data = new byte[MaxOffset - MinOffset];
+			using (Stream Stream = await OpenRawStreamAsync(LogFile, MinOffset, MaxOffset - MinOffset))
+			{
+				int Length = await Stream.ReadAsync(Data.AsMemory());
+				if(Length != Data.Length)
+				{
+					Logger.LogWarning("Read less than expected from log stream (Expected {Expected}, Got {Got})", Data.Length, Length);
+				}
+				return ParseEventData(Data.AsSpan(0, Length));
+			}
+		}
+
+		/// <summary>
+		/// Parses event data from a buffer
+		/// </summary>
+		/// <param name="Data">Data to parse</param>
+		/// <returns>Parsed event data</returns>
+		private ILogEventData ParseEventData(ReadOnlySpan<byte> Data)
+		{
+			List<LogEventLine> Lines = new List<LogEventLine>();
+
+			ReadOnlySpan<byte> RemainingData = Data;
+			while(RemainingData.Length > 0)
+			{
+				int EndOfLine = RemainingData.IndexOf((byte)'\n');
+				if(EndOfLine == -1)
+				{
+					break;
+				}
+
+				ReadOnlySpan<byte> LineData = RemainingData.Slice(0, EndOfLine);
+				try
+				{
+					Lines.Add(new LogEventLine(LineData));
+				}
+				catch(JsonException Ex)
+				{
+					Logger.LogWarning(Ex, "Unable to parse line from log file: {Line}", Encoding.UTF8.GetString(LineData));
+				}
+				RemainingData = RemainingData.Slice(EndOfLine + 1);
+			}
+
+			return new LogEventData(Lines);
+		}
+
+		/// <inheritdoc/>
+		public async Task<Stream> OpenRawStreamAsync(ILogFile LogFile, long Offset, long Length)
+		{
+			if (LogFile.Chunks.Count == 0)
+			{
+				return new MemoryStream(Array.Empty<byte>(), false);
+			}
+			else
+			{
+				int LastChunkIdx = LogFile.Chunks.Count - 1;
+
+				// Clamp the length of the request
+				ILogChunk LastChunk = LogFile.Chunks[LastChunkIdx];
+				if (Length > LastChunk.Offset)
+				{
+					long LastChunkLength = LastChunk.Length;
+					if (LastChunkLength <= 0)
+					{
+						LogChunkData LastChunkData = await ReadChunkAsync(LogFile, LastChunkIdx);
+						LastChunkLength = LastChunkData.Length;
+					}
+					Length = Math.Min(Length, (LastChunk.Offset + LastChunkLength) - Offset);
+				}
+
+				// Create the new stream
+				return new ResponseStream(this, LogFile, Offset, Length);
+			}
+		}
+
+		/// <inheritdoc/>
+		public async Task<(int, long)> GetLineOffsetAsync(ILogFile LogFile, int LineIdx)
+		{
+			int ChunkIdx = LogFile.Chunks.GetChunkForLine(LineIdx);
+
+			ILogChunk Chunk = LogFile.Chunks[ChunkIdx];
+			LogChunkData ChunkData = await ReadChunkAsync(LogFile, ChunkIdx);
+
+			if (LineIdx < Chunk.LineIndex)
+			{
+				LineIdx = Chunk.LineIndex;
+			}
+
+			int MaxLineIndex = Chunk.LineIndex + ChunkData.LineCount;
+			if (LineIdx >= MaxLineIndex)
+			{
+				LineIdx = MaxLineIndex;
+			}
+
+			long Offset = Chunk.Offset + ChunkData.GetLineOffsetWithinChunk(LineIdx - Chunk.LineIndex);
+			return (LineIdx, Offset);
+		}
+
+		/// <summary>
+		/// Executes a background task
+		/// </summary>
+		/// <param name="StoppingToken">Cancellation token</param>
+		protected override async Task TickAsync(CancellationToken StoppingToken)
+		{
+			lock (WriteLock)
+			{
+				try
+				{
+					WriteTasks.RemoveCompleteTasks();
+				}
+				catch (Exception Ex)
+				{
+					Logger.LogError(Ex, "Exception while waiting for write tasks to complete");
+				}
+			}
+			await IncrementalFlush();
+		}
+		
+		/// <summary>
+		/// Executes a background task. Publicly exposed for tests.
+		/// </summary>
+		public async Task TickOnlyForTestingAsync()
+		{
+			await TickAsync(new CancellationToken());
+		}
+		
+		/// <summary>
+		/// Flushes complete chunks to the storage provider
+		/// </summary>
+		/// <returns>Async task</returns>
+		private async Task IncrementalFlush()
+		{
+			// Get all the chunks older than 20 minutes
+			List<(ObjectId, long)> FlushChunks = await Builder.TouchChunksAsync(TimeSpan.FromMinutes(10.0));
+			Logger.LogDebug("Performing incremental flush of log builder ({NumChunks} chunks)", FlushChunks.Count);
+
+			// Mark them all as complete
+			foreach ((ObjectId LogId, long Offset) in FlushChunks)
+			{
+				await Builder.CompleteChunkAsync(LogId, Offset);
+			}
+
+			// Add tasks for flushing all the chunks
+			WriteCompleteChunks(FlushChunks, true);
+		}
+
+		/// <summary>
+		/// Called when the service is shutting down
+		/// </summary>
+		/// <param name="CancellationToken"></param>
+		/// <returns></returns>
+		[SuppressMessage("Design", "CA1031:Do not catch general exception types")]
+		public override async Task StopAsync(CancellationToken CancellationToken)
+		{
+			Logger.LogInformation("Stopping log file service");
+			if (Builder.FlushOnShutdown)
+			{
+				await FlushAsync();
+			}
+			Logger.LogInformation("Log service stopped");
+		}
+
+		/// <summary>
+		/// Flushes the write cache
+		/// </summary>
+		/// <returns>Async task</returns>
+		public async Task FlushAsync()
+		{
+			Logger.LogInformation("Forcing flush of pending log chunks...");
+
+			// Mark everything in the cache as complete
+			List<(ObjectId, long)> WriteChunks = await Builder.TouchChunksAsync(TimeSpan.Zero);
+			WriteCompleteChunks(WriteChunks, true);
+
+			// Wait for everything to flush
+			await FlushPendingWritesAsync();
+		}
+
+		/// <summary>
+		/// Flush any writes in progress
+		/// </summary>
+		/// <returns>Async task</returns>
+		public async Task FlushPendingWritesAsync()
+		{
+			for(; ;)
+			{
+				// Capture the current contents of the WriteTasks list
+				List<Task> Tasks;
+				lock (WriteLock)
+				{
+					WriteTasks.RemoveCompleteTasks();
+					Tasks = new List<Task>(WriteTasks);
+				}
+				if (Tasks.Count == 0)
+				{
+					break;
+				}
+
+				// Also add a delay so we'll periodically refresh the list
+				Tasks.Add(Task.Delay(TimeSpan.FromSeconds(5.0)));
+				await Task.WhenAny(Tasks);
+			}
+		}
+
+		/// <summary>
+		/// Adds tasks for writing a list of complete chunks
+		/// </summary>
+		/// <param name="ChunksToWrite">List of chunks to write</param>
+		/// <param name="bCreateIndex">Create an index for the log</param>
+		private void WriteCompleteChunks(List<(ObjectId, long)> ChunksToWrite, bool bCreateIndex)
+		{
+			foreach (IGrouping<ObjectId, long> Group in ChunksToWrite.GroupBy(x => x.Item1, x => x.Item2))
+			{
+				ObjectId LogId = Group.Key;
+
+				// Find offsets of new chunks to write
+				List<long> Offsets = new List<long>();
+				lock (WriteLock)
+				{
+					foreach (long Offset in Group.OrderBy(x => x))
+					{
+						if (WriteChunks.Add((LogId, Offset)))
+						{
+							Offsets.Add(Offset);
+						}
+					}
+				}
+
+				// Create the write task
+				if (Offsets.Count > 0)
+				{
+					Task Task = Task.Run(() => WriteCompleteChunksForLogAsync(LogId, Offsets, bCreateIndex));
+					lock (WriteLock)
+					{
+						WriteTasks.Add(Task);
+					}
+				}
+			}
+		}
+
+		/// <summary>
+		/// Writes a set of chunks to the database
+		/// </summary>
+		/// <param name="LogId">Log file to update</param>
+		/// <param name="Offsets">Chunks to write</param>
+		/// <param name="bCreateIndex">Whether to create the index for this log</param>
+		/// <returns>Async task</returns>
+		private async Task<ILogFile?> WriteCompleteChunksForLogAsync(ObjectId LogId, List<long> Offsets, bool bCreateIndex)
+		{
+			ILogFile? LogFile = await LogFiles.GetLogFileAsync(LogId);
+			if(LogFile != null)
+			{
+				LogFile = await WriteCompleteChunksForLogAsync(LogFile, Offsets, bCreateIndex);
+			}
+			return LogFile;
+		}
+
+		/// <summary>
+		/// Writes a set of chunks to the database
+		/// </summary>
+		/// <param name="LogFileInterface">Log file to update</param>
+		/// <param name="Offsets">Chunks to write</param>
+		/// <param name="bCreateIndex">Whether to create the index for this log</param>
+		/// <returns>Async task</returns>
+		[SuppressMessage("Design", "CA1031:Do not catch general exception types", Justification = "<Pending>")]
+		private async Task<ILogFile?> WriteCompleteChunksForLogAsync(ILogFile LogFileInterface, List<long> Offsets, bool bCreateIndex)
+		{
+			// Write the data to the storage provider
+			List<Task<LogChunkData?>> ChunkWriteTasks = new List<Task<LogChunkData?>>();
+			foreach (long Offset in Offsets)
+			{
+				int ChunkIdx = LogFileInterface.Chunks.BinarySearch(x => x.Offset, Offset);
+				if (ChunkIdx >= 0)
+				{
+					Logger.LogDebug("Queuing write of log {LogId} chunk {ChunkIdx} offset {Offset}", LogFileInterface.Id, ChunkIdx, Offset);
+					int LineIndex = LogFileInterface.Chunks[ChunkIdx].LineIndex;
+					ChunkWriteTasks.Add(Task.Run(() => WriteChunkAsync(LogFileInterface.Id, Offset, LineIndex)));
+				}
+			}
+
+			// Wait for the tasks to complete, periodically updating the log file object
+			ILogFile? LogFile = LogFileInterface;
+			while (ChunkWriteTasks.Count > 0)
+			{
+				// Wait for all tasks to be complete OR (any task has completed AND 30 seconds has elapsed)
+				Task AllCompleteTask = Task.WhenAll(ChunkWriteTasks);
+				Task AnyCompleteTask = Task.WhenAny(ChunkWriteTasks);
+				await Task.WhenAny(AllCompleteTask, Task.WhenAll(AnyCompleteTask, Task.Delay(TimeSpan.FromSeconds(30.0))));
+
+				// Update the log file with the written chunks
+				List<LogChunkData?> WrittenChunks = ChunkWriteTasks.RemoveCompleteTasks();
+				while (LogFile != null)
+				{
+					// Update the length of any complete chunks
+					List<CompleteLogChunkUpdate> Updates = new List<CompleteLogChunkUpdate>();
+					foreach (LogChunkData? ChunkData in WrittenChunks)
+					{
+						if (ChunkData != null)
+						{
+							int ChunkIdx = LogFile.Chunks.GetChunkForOffset(ChunkData.Offset);
+							if (ChunkIdx >= 0)
+							{
+								ILogChunk Chunk = LogFile.Chunks[ChunkIdx];
+								if (Chunk.Offset == ChunkData.Offset)
+								{
+									CompleteLogChunkUpdate Update = new CompleteLogChunkUpdate(ChunkIdx, ChunkData.Length, ChunkData.LineCount);
+									Updates.Add(Update);
+								}
+							}
+						}
+					}
+
+					// Try to apply the updates
+					ILogFile? NewLogFile = await LogFiles.TryCompleteChunksAsync(LogFile, Updates);
+					if (NewLogFile != null)
+					{
+						LogFile = NewLogFile;
+						break;
+					}
+
+					// Update the log file
+					LogFile = await GetLogFileAsync(LogFile.Id);
+				}
+			}
+
+			// Create the index if necessary
+			if (bCreateIndex && LogFile != null)
+			{
+				try
+				{
+					LogFile = await CreateIndexAsync(LogFile);
+				}
+				catch(Exception Ex)
+				{
+					Logger.LogError(Ex, "Failed to create index for log {LogId}", LogFileInterface.Id);
+				}
+			}
+
+			return LogFile;
+		}
+
+		/// <summary>
+		/// Creates an index for the given log file
+		/// </summary>
+		/// <param name="LogFile">The log file object</param>
+		/// <returns>Updated log file</returns>
+		private async Task<ILogFile?> CreateIndexAsync(ILogFile LogFile)
+		{
+			if(LogFile.Chunks.Count == 0)
+			{
+				return LogFile;
+			}
+
+			// Get the new length of the log, and early out if it won't be any longer
+			ILogChunk LastChunk = LogFile.Chunks[LogFile.Chunks.Count - 1];
+			if(LastChunk.Offset + LastChunk.Length <= (LogFile.IndexLength ?? 0))
+			{
+				return LogFile;
+			}
+
+			// Save stats for the index creation
+			using IScope Scope = GlobalTracer.Instance.BuildSpan("CreateIndexAsync").StartActive();
+			Scope.Span.SetTag("LogId", LogFile.Id.ToString());
+			Scope.Span.SetTag("Length", (LastChunk.Offset + LastChunk.Length).ToString(CultureInfo.InvariantCulture));
+
+			long NewLength = 0;
+			int NewLineCount = 0;
+
+			// Read the existing index if there is one
+			List<LogIndexData> Indexes = new List<LogIndexData>();
+			if (LogFile.IndexLength != null)
+			{
+				LogIndexData? ExistingIndex = await ReadIndexAsync(LogFile, LogFile.IndexLength.Value);
+				if(ExistingIndex != null)
+				{
+					Indexes.Add(ExistingIndex);
+					NewLineCount = ExistingIndex.LineCount;
+				}
+			}
+
+			// Add all the new chunks
+			int ChunkIdx = LogFile.Chunks.GetChunkForLine(NewLineCount);
+			if (ChunkIdx < 0)
+			{
+				int FirstLine = (LogFile.Chunks.Count > 0) ? LogFile.Chunks[0].LineIndex : -1;
+				throw new Exception($"Invalid chunk index {ChunkIdx}. Index.LineCount={NewLineCount}, Chunks={LogFile.Chunks.Count}, First line={FirstLine}");
+			}
+
+			for (; ChunkIdx < LogFile.Chunks.Count; ChunkIdx++)
+			{
+				ILogChunk Chunk = LogFile.Chunks[ChunkIdx];
+				LogChunkData ChunkData = await ReadChunkAsync(LogFile, ChunkIdx);
+
+				int SubChunkIdx = ChunkData.GetSubChunkForLine(Math.Max(NewLineCount - Chunk.LineIndex, 0));
+				if(SubChunkIdx < 0)
+				{
+					throw new Exception($"Invalid subchunk index {SubChunkIdx}. Chunk {ChunkIdx}/{LogFile.Chunks.Count}. Index.LineCount={NewLineCount}, Chunk.LineIndex={Chunk.LineIndex}, First subchunk {ChunkData.SubChunkLineIndex[0]}");
+				}
+
+				for (; SubChunkIdx < ChunkData.SubChunks.Count; SubChunkIdx++)
+				{
+					LogSubChunkData SubChunkData = ChunkData.SubChunks[SubChunkIdx];
+					if (SubChunkData.LineIndex >= NewLineCount)
+					{
+						try
+						{
+							Indexes.Add(SubChunkData.BuildIndex());
+						}
+						catch (Exception Ex)
+						{
+							throw new Exception($"Failed to create index block - log {LogFile.Id}, chunk {ChunkIdx} ({LogFile.Chunks.Count}), subchunk {SubChunkIdx} ({ChunkData.SubChunks.Count}), index lines: {NewLineCount}, chunk index: {Chunk.LineIndex}, subchunk index: {Chunk.LineIndex + ChunkData.SubChunkLineIndex[SubChunkIdx]}, subchunk count: {SubChunkData.LineCount}", Ex);
+						}
+
+						NewLength = SubChunkData.Offset + SubChunkData.Length;
+						NewLineCount = SubChunkData.LineIndex + SubChunkData.LineCount;
+					}
+				}
+			}
+
+			// Try to update the log file
+			ILogFile? NewLogFile = LogFile;
+			if (NewLength > (LogFile.IndexLength ?? 0))
+			{
+				LogIndexData Index = LogIndexData.Merge(Indexes);
+				Logger.LogDebug("Writing index for log {LogId} covering {Length} (index length {IndexLength})", LogFile.Id, NewLength, Index.GetSerializedSize());
+
+				await WriteIndexAsync(LogFile.Id, NewLength, Index);
+
+				while(NewLogFile != null && NewLength > (NewLogFile.IndexLength ?? 0))
+				{
+					NewLogFile = await LogFiles.TryUpdateIndexAsync(NewLogFile, NewLength);
+					if(NewLogFile != null)
+					{
+						break;
+					}
+					NewLogFile = await LogFiles.GetLogFileAsync(LogFile.Id);
+				}
+			}
+			return NewLogFile;
+		}
+
+		/// <summary>
+		/// Reads a chunk from storage
+		/// </summary>
+		/// <param name="LogFile">Log file to read from</param>
+		/// <param name="ChunkIdx">The chunk to read</param>
+		/// <returns>Chunk data</returns>
+		[SuppressMessage("Design", "CA1031:Do not catch general exception types", Justification = "<Pending>")]
+		private async Task<LogChunkData> ReadChunkAsync(ILogFile LogFile, int ChunkIdx)
+		{
+			ILogChunk Chunk = LogFile.Chunks[ChunkIdx];
+
+			// Try to read the chunk data from storage
+			LogChunkData? ChunkData = null;
+			try
+			{
+				// If the chunk is not yet complete, query the log builder
+				if (Chunk.Length == 0)
+				{
+					ChunkData = await Builder.GetChunkAsync(LogFile.Id, Chunk.Offset, Chunk.LineIndex);
+				}
+
+				// Otherwise go directly to the log storage
+				if (ChunkData == null)
+				{
+					ChunkData = await Storage.ReadChunkAsync(LogFile.Id, Chunk.Offset, Chunk.LineIndex);
+				}
+			}
+			catch (Exception Ex)
+			{
+				Logger.LogError(Ex, "Unable to read log {LogId} at offset {Offset}", LogFile.Id, Chunk.Offset);
+			}
+
+			// Get the minimum length and line count for the chunk
+			if (ChunkIdx + 1 < LogFile.Chunks.Count)
+			{
+				ILogChunk NextChunk = LogFile.Chunks[ChunkIdx + 1];
+				ChunkData = RepairChunkData(LogFile, ChunkIdx, ChunkData, (int)(NextChunk.Offset - Chunk.Offset), NextChunk.LineIndex - Chunk.LineIndex);
+			}
+			else
+			{
+				if (LogFile.MaxLineIndex != null && Chunk.Length != 0)
+				{
+					ChunkData = RepairChunkData(LogFile, ChunkIdx, ChunkData, Chunk.Length, LogFile.MaxLineIndex.Value - Chunk.LineIndex);
+				}
+				else if(ChunkData == null)
+				{
+					ChunkData = RepairChunkData(LogFile, ChunkIdx, ChunkData, 1024, 1);
+				}
+			}
+
+			return ChunkData;
+		}
+
+		/// <summary>
+		/// Validates the given chunk data, and fix it up if necessary
+		/// </summary>
+		/// <param name="LogFile">The log file instance</param>
+		/// <param name="ChunkIdx">Index of the chunk within the logfile</param>
+		/// <param name="ChunkData">The chunk data that was read</param>
+		/// <param name="Length">Expected length of the data</param>
+		/// <param name="LineCount">Expected number of lines in the data</param>
+		/// <returns>Repaired chunk data</returns>
+		LogChunkData RepairChunkData(ILogFile LogFile, int ChunkIdx, LogChunkData? ChunkData, int Length, int LineCount)
+		{
+			int CurrentLength = 0;
+			int CurrentLineCount = 0;
+			if(ChunkData != null)
+			{
+				CurrentLength = ChunkData.Length;
+				CurrentLineCount = ChunkData.LineCount;
+			}
+
+			if (ChunkData == null || CurrentLength < Length || CurrentLineCount < LineCount)
+			{
+				Logger.LogWarning("Creating placeholder subchunk for log {LogId} chunk {ChunkIdx} (length {Length} vs expected {ExpLength}, lines {LineCount} vs expected {ExpLineCount})", LogFile.Id, ChunkIdx, CurrentLength, Length, CurrentLineCount, LineCount);
+
+				List<LogSubChunkData> SubChunks = new List<LogSubChunkData>();
+				if (ChunkData != null && ChunkData.Length < Length && ChunkData.LineCount < LineCount)
+				{
+					SubChunks.AddRange(ChunkData.SubChunks);
+				}
+
+				LogText Text = new LogText();
+				Text.AppendMissingDataInfo(ChunkIdx, LogFile.Chunks[ChunkIdx].Server, Length - CurrentLength, LineCount - CurrentLineCount);
+				SubChunks.Add(new LogSubChunkData(LogFile.Type, CurrentLength, CurrentLineCount, Text));
+
+				ILogChunk Chunk = LogFile.Chunks[ChunkIdx];
+				ChunkData = new LogChunkData(Chunk.Offset, Chunk.LineIndex, SubChunks);
+			}
+			return ChunkData;
+		}
+
+		/// <summary>
+		/// Writes a set of chunks to the database
+		/// </summary>
+		/// <param name="LogFileId">Unique id of the log file</param>
+		/// <param name="Offset">Offset of the chunk to write</param>
+		/// <param name="LineIndex">First line index of the chunk</param>
+		/// <returns>Chunk daata</returns>
+		[SuppressMessage("Design", "CA1031:Do not catch general exception types", Justification = "<Pending>")]
+		private async Task<LogChunkData?> WriteChunkAsync(ObjectId LogFileId, long Offset, int LineIndex)
+		{
+			// Write the chunk to storage
+			LogChunkData? ChunkData = await Builder.GetChunkAsync(LogFileId, Offset, LineIndex);
+			if (ChunkData == null)
+			{
+				Logger.LogDebug("Log {LogId} offset {Offset} not found in log builder", LogFileId, Offset);
+			}
+			else
+			{
+				try
+				{
+					await Storage.WriteChunkAsync(LogFileId, ChunkData.Offset, ChunkData);
+				}
+				catch (Exception Ex)
+				{
+					Logger.LogError(Ex, "Unable to write log {LogId} at offset {Offset}", LogFileId, ChunkData.Offset);
+				}
+			}
+
+			// Remove it from the log builder
+			try
+			{
+				await Builder.RemoveChunkAsync(LogFileId, Offset);
+			}
+			catch (Exception Ex)
+			{
+				Logger.LogError(Ex, "Unable to remove log {LogId} at offset {Offset} from log builder", LogFileId, Offset);
+			}
+
+			return ChunkData;
+		}
+
+		/// <summary>
+		/// Reads a chunk from storage
+		/// </summary>
+		/// <param name="LogFile">Log file to read from</param>
+		/// <param name="Length">Length of the log covered by the index</param>
+		/// <returns>Chunk data</returns>
+		[SuppressMessage("Design", "CA1031:Do not catch general exception types", Justification = "<Pending>")]
+		private async Task<LogIndexData?> ReadIndexAsync(ILogFile LogFile, long Length)
+		{
+			try
+			{
+				LogIndexData? Index = await Storage.ReadIndexAsync(LogFile.Id, Length);
+				return Index;
+			}
+			catch (Exception Ex)
+			{
+				Logger.LogError(Ex, "Unable to read log {LogId} index at length {Length}", LogFile.Id, Length);
+				return null;
+			}
+		}
+
+		/// <summary>
+		/// Writes an index to the database
+		/// </summary>
+		/// <param name="LogFileId">Unique id of the log file</param>
+		/// <param name="Length">Length of the data covered by the index</param>
+		/// <param name="Index">Index to write</param>
+		/// <returns>Async task</returns>
+		[SuppressMessage("Design", "CA1031:Do not catch general exception types", Justification = "<Pending>")]
+		private async Task WriteIndexAsync(ObjectId LogFileId, long Length, LogIndexData Index)
+		{
+			try
+			{
+				await Storage.WriteIndexAsync(LogFileId, Length, Index);
+			}
+			catch(Exception Ex)
+			{
+				Logger.LogError(Ex, "Unable to write index for log {LogId}", LogFileId);
+			}
+		}
+
+		/// <summary>
+		/// Determines if the user is authorized to perform an action on a particular template
+		/// </summary>
+		/// <param name="LogFile">The template to check</param>
+		/// <param name="User">The principal to authorize</param>
+		/// <returns>True if the action is authorized</returns>
+		public static bool AuthorizeForSession(ILogFile LogFile, ClaimsPrincipal User)
+		{
+			if(LogFile.SessionId != null)
+			{
+				return User.HasClaim(HordeClaimTypes.AgentSessionId, LogFile.SessionId.Value.ToString());
+			}
+			else
+			{
+				return false;
+			}
+		}
+
+		/// <inheritdoc/>
+		public async Task<List<int>> SearchLogDataAsync(ILogFile LogFile, string Text, int FirstLine, int Count, LogSearchStats SearchStats)
+		{
+			Stopwatch Timer = Stopwatch.StartNew();
+
+			using IScope Scope = GlobalTracer.Instance.BuildSpan("SearchLogDataAsync").StartActive();
+			Scope.Span.SetTag("LogId", LogFile.Id.ToString());
+			Scope.Span.SetTag("Text", Text);
+			Scope.Span.SetTag("Count", Count.ToString(CultureInfo.InvariantCulture));
+
+			List<int> Results = new List<int>();
+			if (Count > 0)
+			{
+				IAsyncEnumerator<int> Enumerator = SearchLogDataInternalAsync(LogFile, Text, FirstLine, SearchStats).GetAsyncEnumerator();
+				while (await Enumerator.MoveNextAsync() && Results.Count < Count)
+				{
+					Results.Add(Enumerator.Current);
+				}
+			}
+
+			Logger.LogDebug("Search for \"{SearchText}\" in log {LogId} found {NumResults}/{MaxResults} results, took {Time}ms ({@Stats})", Text, LogFile.Id, Results.Count, Count, Timer.ElapsedMilliseconds, SearchStats);
+			return Results;
+		}
+
+		/// <inheritdoc/>
+		public async IAsyncEnumerable<int> SearchLogDataInternalAsync(ILogFile LogFile, string Text, int FirstLine, LogSearchStats SearchStats)
+		{
+			SearchText SearchText = new SearchText(Text);
+
+			// Read the index for this log file
+			if (LogFile.IndexLength != null)
+			{
+				LogIndexData? IndexData = await ReadIndexAsync(LogFile, LogFile.IndexLength.Value);
+				if(IndexData != null && FirstLine < IndexData.LineCount)
+				{
+					using IScope IndexScope = GlobalTracer.Instance.BuildSpan("Indexed").StartActive();
+					IndexScope.Span.SetTag("LineCount", IndexData.LineCount.ToString(CultureInfo.InvariantCulture));
+
+					foreach(int LineIndex in IndexData.Search(FirstLine, SearchText, SearchStats))
+					{
+						yield return LineIndex;
+					}
+
+					FirstLine = IndexData.LineCount;
+				}
+			}
+
+			// Manually search through the rest of the log
+			int ChunkIdx = LogFile.Chunks.GetChunkForLine(FirstLine);
+			for (; ChunkIdx < LogFile.Chunks.Count; ChunkIdx++)
+			{
+				ILogChunk Chunk = LogFile.Chunks[ChunkIdx];
+
+				// Read the chunk data
+				LogChunkData ChunkData = await ReadChunkAsync(LogFile, ChunkIdx);
+				if (FirstLine < ChunkData.LineIndex + ChunkData.LineCount)
+				{
+					// Find the first sub-chunk we're looking for
+					int SubChunkIdx = 0;
+					if (FirstLine > Chunk.LineIndex)
+					{
+						SubChunkIdx = ChunkData.GetSubChunkForLine(FirstLine - Chunk.LineIndex);
+					}
+
+					// Search through the sub-chunks
+					for (; SubChunkIdx < ChunkData.SubChunks.Count; SubChunkIdx++)
+					{
+						LogSubChunkData SubChunkData = ChunkData.SubChunks[SubChunkIdx];
+						if (FirstLine < SubChunkData.LineIndex + SubChunkData.LineCount)
+						{
+							// Create an index containing just this sub-chunk
+							LogIndexData Index = SubChunkData.BuildIndex();
+							foreach (int LineIndex in Index.Search(FirstLine, SearchText, SearchStats))
+							{
+								yield return LineIndex;
+							}
+						}
+					}
+				}
+			}
+		}
+	}
+}
+
diff --git a/Engine/Source/Programs/Horde/HordeServer/Perforce/Change.cs b/Engine/Source/Programs/Horde/HordeServer/Perforce/Change.cs
new file mode 100644
--- /dev/null
+++ b/Engine/Source/Programs/Horde/HordeServer/Perforce/Change.cs
@@ -0,0 +1,311 @@
+// Copyright Epic Games, Inc. All Rights Reserved.
+
+using EpicGames.Core;
+using HordeServer.Api;
+using HordeServer.Utilities;
+using MongoDB.Bson;
+using MongoDB.Bson.Serialization.Attributes;
+using System;
+using System.Collections.Generic;
+using System.Diagnostics;
+using System.Diagnostics.CodeAnalysis;
+using System.Globalization;
+using System.Linq;
+using System.Text.RegularExpressions;
+using System.Threading.Tasks;
+
+namespace HordeServer.Models
+{
+	/// <summary>
+	/// Summary information for a change
+	/// </summary>
+	[DebuggerDisplay("{Number}: {Description}")]
+	public class ChangeSummary
+	{
+		/// <summary>
+		/// The changelist number
+		/// </summary>
+		public int Number { get; set; }
+
+		/// <summary>
+		/// Author of the change
+		/// </summary>
+		public string Author { get; set; }
+
+		/// <summary>
+		/// The base path for modified files
+		/// </summary>
+		public string Path { get; set; }
+
+		/// <summary>
+		/// Abbreviated changelist description
+		/// </summary>
+		public string Description { get; set; }
+
+		/// <summary>
+		/// Constructor
+		/// </summary>
+		/// <param name="Number">Changelist number</param>
+		/// <param name="Author">Author of the change</param>
+		/// <param name="Path">Base path for modified files</param>
+		/// <param name="Description">Changelist description</param>
+		public ChangeSummary(int Number, string Author, string Path, string Description)
+		{
+			this.Number = Number;
+			this.Author = Author;
+			this.Path = Path;
+			this.Description = Description;
+		}
+	}
+
+	/// <summary>
+	/// Flags identifying content of a changelist
+	/// </summary>
+	[Flags]
+	public enum ChangeContentFlags
+	{
+		/// <summary>
+		/// The change contains code
+		/// </summary>
+		ContainsCode = 1,
+
+		/// <summary>
+		/// The change contains content
+		/// </summary>
+		ContainsContent = 2,
+	}
+
+	/// <summary>
+	/// Modified file in a changelist
+	/// </summary>
+	public class ChangeFile
+	{
+		/// <summary>
+		/// Path to the file
+		/// </summary>
+		public string Path { get; set; }
+
+		/// <summary>
+		/// Path to the file within the depot
+		/// </summary>
+		public string DepotPath { get; set; }
+
+		/// <summary>
+		/// Revision of the file. A value of -1 indicates that the file was deleted.
+		/// </summary>
+		public int Revision { get; set; }
+
+		/// <summary>
+		/// Length of the file
+		/// </summary>
+		public long Length { get; set; }
+
+		/// <summary>
+		/// MD5 digest of the file
+		/// </summary>
+		public Md5Hash? Digest { get; set; }
+
+		/// <summary>
+		/// The file type
+		/// </summary>
+		public string Type { get; set; }
+
+		/// <summary>
+		/// Constructor
+		/// </summary>
+		/// <param name="Path"></param>
+		/// <param name="DepotPath"></param>
+		/// <param name="Revision"></param>
+		/// <param name="Length"></param>
+		/// <param name="Digest"></param>
+		/// <param name="Type"></param>
+		public ChangeFile(string Path, string DepotPath, int Revision, long Length, Md5Hash? Digest, string Type)
+		{
+			this.Path = Path;
+			this.DepotPath = DepotPath;
+			this.Revision = Revision;
+			this.Length = Length;
+			this.Digest = Digest;
+			this.Type = Type;
+		}
+	}
+
+	/// <summary>
+	/// Information about a commit
+	/// </summary>
+	[DebuggerDisplay("{Number}: {Description}")]
+	public class ChangeDetails
+	{
+		/// <summary>
+		/// Set of extensions to treat as code
+		/// </summary>
+		static readonly HashSet<string> CodeExtensions = new HashSet<string>
+		{
+			".c",
+			".cc",
+			".cpp",
+			".m",
+			".mm",
+			".rc",
+			".cs",
+			".csproj",
+			".h",
+			".hpp",
+			".inl",
+			".usf",
+			".ush",
+			".uproject",
+			".uplugin",
+			".sln"
+		};
+
+		/// <summary>
+		/// The source changelist number
+		/// </summary>
+		public int Number { get; set; }
+
+		/// <summary>
+		/// Name of the user that authored this change
+		/// </summary>
+		public string Author { get; set; }
+
+		/// <summary>
+		/// The base path for modified files
+		/// </summary>
+		public string Path { get; set; }
+
+		/// <summary>
+		/// The description text
+		/// </summary>
+		public string Description { get; set; }
+
+		/// <summary>
+		/// List of files that were modified, relative to the stream base
+		/// </summary>
+		public List<ChangeFile> Files { get; set; }
+
+		/// <summary>
+		/// Date that the change was submitted
+		/// </summary>
+		public DateTime Date { get; set; }
+
+		/// <summary>
+		/// Private constructor for serialization
+		/// </summary>
+		private ChangeDetails()
+		{
+			Author = null!;
+			Path = null!;
+			Description = null!;
+			Files = null!;
+		}
+
+		/// <summary>
+		/// Constructor
+		/// </summary>
+		/// <param name="Number">Changelist number</param>
+		/// <param name="Author">Author of the change</param>
+		/// <param name="Path">Base path for modified files</param>
+		/// <param name="Description">Changelist description</param>
+		/// <param name="Files">List of files modified, relative to the stream base</param>
+		/// <param name="Date">Date that the change was submitted</param>
+		public ChangeDetails(int Number, string Author, string Path, string Description, List<ChangeFile> Files, DateTime Date)
+		{
+			this.Number = Number;
+			this.Author = Author;
+			this.Path = Path;
+			this.Description = Description;
+			this.Files = Files;
+			this.Date = Date;
+		}
+
+		/// <summary>
+		/// Determines if this change is a code change
+		/// </summary>
+		/// <returns>True if this change is a code change</returns>
+		public ChangeContentFlags GetContentFlags()
+		{
+			ChangeContentFlags Scope = 0;
+
+			// Check whether the files are code or content
+			foreach (ChangeFile File in Files)
+			{
+				if (CodeExtensions.Any(Extension => File.Path.EndsWith(Extension, StringComparison.OrdinalIgnoreCase)))
+				{
+					Scope |= ChangeContentFlags.ContainsCode;
+				}
+				else
+				{
+					Scope |= ChangeContentFlags.ContainsContent;
+				}
+
+				if (Scope == (ChangeContentFlags.ContainsCode | ChangeContentFlags.ContainsContent))
+				{
+					break;
+				}
+			}
+			return Scope;
+		}
+
+		/// <summary>
+		/// Attempts to parse the Robomerge source from this commit information
+		/// </summary>
+		/// <param name="Source">On success, receives the source information</param>
+		/// <returns>True if the commit was merged from another stream</returns>
+		public bool TryParseRobomergeSource([NotNullWhen(true)] out (string, int)? Source)
+		{
+			Match Match = Regex.Match(Description, @"#ROBOMERGE-SOURCE: CL (\d+) in (//[^ ]*)/...", RegexOptions.Multiline);
+			if (Match.Success)
+			{
+				Source = (Match.Groups[2].Value, int.Parse(Match.Groups[1].Value, CultureInfo.InvariantCulture));
+				return true;
+			}
+			else
+			{
+				Source = null;
+				return false;
+			}
+		}
+	}
+
+	/// <summary>
+	/// Summary of a file in the depot
+	/// </summary>
+	public class FileSummary
+	{
+		/// <summary>
+		/// Depot path to the file
+		/// </summary>
+		public string DepotPath { get; set; }
+
+		/// <summary>
+		/// Whether the file exists
+		/// </summary>
+		public bool Exists { get; set; }
+
+		/// <summary>
+		/// Last changelist number that the file was modified
+		/// </summary>
+		public int Change { get; set; }
+
+		/// <summary>
+		/// Error about this file, if it exists
+		/// </summary>
+		public string? Error { get; set; }
+
+		/// <summary>
+		/// Information about a file
+		/// </summary>
+		/// <param name="DepotPath">Depot </param>
+		/// <param name="Exists"></param>
+		/// <param name="Change"></param>
+		/// <param name="Error"></param>
+		public FileSummary(string DepotPath, bool Exists, int Change, string? Error = null)
+		{
+			this.DepotPath = DepotPath;
+			this.Exists = Exists;
+			this.Change = Change;
+			this.Error = Error;
+		}
+	}
+}
