Horde: More file moves.

diff --git a/Engine/Source/Programs/Horde/HordeServer/Perforce/PerforceLoadBalancer.cs b/Engine/Source/Programs/Horde/HordeServer/Perforce/PerforceLoadBalancer.cs
new file mode 100644
--- /dev/null
+++ b/Engine/Source/Programs/Horde/HordeServer/Perforce/PerforceLoadBalancer.cs
@@ -0,0 +1,598 @@
+// Copyright Epic Games, Inc. All Rights Reserved.
+
+using Google.Protobuf.WellKnownTypes;
+using HordeCommon.Rpc.Tasks;
+using HordeServer.Collections;
+using HordeServer.Models;
+using HordeServer.Utilities;
+using Microsoft.Extensions.Hosting;
+using Microsoft.Extensions.Logging;
+using MongoDB.Bson;
+using MongoDB.Bson.Serialization.Attributes;
+using MongoDB.Driver;
+using System;
+using System.Collections.Concurrent;
+using System.Collections.Generic;
+using System.Globalization;
+using System.Linq;
+using System.Net;
+using System.Net.Http;
+using System.Text.Json;
+using System.Threading;
+using System.Threading.Tasks;
+
+namespace HordeServer.Services
+{
+	/// <summary>
+	/// Health of a particular server
+	/// </summary>
+	public enum PerforceServerStatus
+	{
+		/// <summary>
+		/// Server could not be reached
+		/// </summary>
+		Unknown,
+
+		/// <summary>
+		/// Bad. Do not use.
+		/// </summary>
+		Unhealthy,
+
+		/// <summary>
+		/// Degraded, but functioning
+		/// </summary>
+		Degraded,
+
+		/// <summary>
+		/// Good
+		/// </summary>
+		Healthy,
+	}
+
+	/// <summary>
+	/// Perforce server infomation
+	/// </summary>
+	public interface IPerforceServer
+	{
+		/// <summary>
+		/// The server and port
+		/// </summary>
+		public string ServerAndPort { get; }
+
+		/// <summary>
+		/// The server and port before doing DNS resolution
+		/// </summary>
+		public string BaseServerAndPort { get; }
+
+		/// <summary>
+		/// The cluster this server belongs to
+		/// </summary>
+		public string Cluster { get; }
+
+		/// <summary>
+		/// Current status
+		/// </summary>
+		public PerforceServerStatus Status { get; }
+
+		/// <summary>
+		/// Number of leases using this server
+		/// </summary>
+		public int NumLeases { get; }
+
+		/// <summary>
+		/// Error message related to this server
+		/// </summary>
+		public string? Detail { get; }
+
+		/// <summary>
+		/// Last update time for this server
+		/// </summary>
+		public DateTime? LastUpdateTime { get; }
+	}
+
+	/// <summary>
+	/// Load balancer for Perforce edge servers
+	/// </summary>
+	public sealed class PerforceLoadBalancer : IHostedService, IDisposable
+	{
+		/// <summary>
+		/// Information about a resolved Perforce server
+		/// </summary>
+		class PerforceServerEntry : IPerforceServer
+		{
+			public string ServerAndPort { get; set; } = String.Empty;
+			public string BaseServerAndPort { get; set; } = String.Empty;
+			public string? HealthCheckUrl { get; set; }
+			public string Cluster { get; set; } = String.Empty;
+			public PerforceServerStatus Status { get; set; }
+			public string? Detail { get; set; }
+			public int NumLeases { get; set; }
+			public DateTime? LastUpdateTime { get; set; }
+
+			[BsonConstructor]
+			private PerforceServerEntry()
+			{
+			}
+
+			public PerforceServerEntry(string ServerAndPort, string BaseServerAndPort, string? HealthCheckUrl, string Cluster, PerforceServerStatus Status, string? Detail, DateTime? LastUpdateTime)
+			{
+				this.ServerAndPort = ServerAndPort;
+				this.BaseServerAndPort = BaseServerAndPort;
+				this.HealthCheckUrl = HealthCheckUrl;
+				this.Cluster = Cluster;
+				this.Status = Status;
+				this.Detail = Detail;
+				this.LastUpdateTime = LastUpdateTime;
+			}
+		}
+
+		/// <summary>
+		/// 
+		/// </summary>
+		[SingletonDocument("6046aec374a9283100967ee7")]
+		class PerforceServerList : SingletonBase
+		{
+			public List<PerforceServerEntry> Servers { get; set; } = new List<PerforceServerEntry>();
+		}
+
+		DatabaseService DatabaseService;
+		ILeaseCollection LeaseCollection;
+		SingletonDocument<PerforceServerList> ServerListSingleton;
+		Random Random = new Random();
+		ILogger Logger;
+		ElectedTick Ticker;
+		Task InternalTickTask = Task.CompletedTask;
+
+		/// <summary>
+		/// Constructor
+		/// </summary>
+		/// <param name="DatabaseService"></param>
+		/// <param name="LeaseCollection"></param>
+		/// <param name="Logger"></param>
+		public PerforceLoadBalancer(DatabaseService DatabaseService, ILeaseCollection LeaseCollection, ILogger<PerforceLoadBalancer> Logger)
+		{
+			this.DatabaseService = DatabaseService;
+			this.LeaseCollection = LeaseCollection;
+			this.ServerListSingleton = new SingletonDocument<PerforceServerList>(DatabaseService);
+			this.Logger = Logger;
+			this.Ticker = new ElectedTick(DatabaseService, new ObjectId("603fb20536e999974fd6ff6b"), TickAsync, TimeSpan.FromMinutes(1.0), Logger);
+		}
+
+		/// <inheritdoc/>
+		public void Dispose()
+		{
+			Ticker.Dispose();
+		}
+
+		/// <inheritdoc/>
+		public Task StartAsync(CancellationToken cancellationToken)
+		{
+			Ticker.Start();
+			return Task.CompletedTask;
+		}
+
+		/// <inheritdoc/>
+		public async Task StopAsync(CancellationToken cancellationToken)
+		{
+			await Ticker.StopAsync();
+			await InternalTickTask;
+		}
+
+		/// <summary>
+		/// Get the current server list
+		/// </summary>
+		/// <returns></returns>
+		async Task<PerforceServerList> GetServerListAsync()
+		{
+			PerforceServerList ServerList = await ServerListSingleton.GetAsync();
+
+			DateTime MinLastUpdateTime = DateTime.UtcNow - TimeSpan.FromMinutes(2.5);
+			foreach (PerforceServerEntry Server in ServerList.Servers)
+			{
+				if (Server.Status == PerforceServerStatus.Healthy && Server.LastUpdateTime != null && Server.LastUpdateTime.Value < MinLastUpdateTime)
+				{
+					Server.Status = PerforceServerStatus.Degraded;
+					Server.Detail = "Server has not responded to health check";
+				}
+			}
+
+			return ServerList;
+		}
+
+		/// <summary>
+		/// Get the perforce servers
+		/// </summary>
+		/// <returns></returns>
+		public async Task<List<IPerforceServer>> GetServersAsync()
+		{
+			PerforceServerList ServerList = await GetServerListAsync();
+			return ServerList.Servers.ConvertAll<IPerforceServer>(x => x);
+		}
+
+		/// <summary>
+		/// Allocates a server for use by a lease
+		/// </summary>
+		/// <returns>The server to use. Null if there is no healthy server available.</returns>
+		public async Task<IPerforceServer?> GetServer(string Cluster)
+		{
+			PerforceServerList ServerList = await GetServerListAsync();
+
+			List<PerforceServerEntry> Candidates = ServerList.Servers.Where(x => x.Cluster == Cluster && x.Status >= PerforceServerStatus.Healthy).ToList();
+			if(Candidates.Count == 0)
+			{
+				int Idx = Random.Next(0, Candidates.Count);
+				return Candidates[Idx];
+			}
+			return null;
+		}
+
+		/// <summary>
+		/// Select a Perforce server to use by the Horde server
+		/// </summary>
+		/// <param name="Cluster"></param>
+		/// <returns></returns>
+		public Task<IPerforceServer?> SelectServerAsync(PerforceCluster Cluster)
+		{
+			List<string> Properties = new List<string>{ "HordeServer=1" };
+			return SelectServerAsync(Cluster, Properties);
+		}
+
+		/// <summary>
+		/// Select a Perforce server to use
+		/// </summary>
+		/// <param name="Cluster"></param>
+		/// <param name="Agent"></param>
+		/// <returns></returns>
+		public Task<IPerforceServer?> SelectServerAsync(PerforceCluster Cluster, IAgent Agent)
+		{
+			return SelectServerAsync(Cluster, Agent.Properties);
+		}
+
+		/// <summary>
+		/// Select a Perforce server to use
+		/// </summary>
+		/// <param name="Cluster"></param>
+		/// <param name="Properties"></param>
+		/// <returns></returns>
+		public async Task<IPerforceServer?> SelectServerAsync(PerforceCluster Cluster, IReadOnlyList<string> Properties)
+		{
+			// Find all the valid servers for this agent
+			List<PerforceServer> ValidServers = new List<PerforceServer>();
+			if (Properties != null)
+			{
+				ValidServers.AddRange(Cluster.Servers.Where(x => x.Properties != null && x.Properties.Count > 0 && x.Properties.All(y => Properties.Contains(y))));
+			}
+			if (ValidServers.Count == 0)
+			{
+				ValidServers.AddRange(Cluster.Servers.Where(x => x.Properties == null || x.Properties.Count == 0));
+			}
+
+			HashSet<string> ValidServerNames = new HashSet<string>(ValidServers.Select(x => x.ServerAndPort), StringComparer.OrdinalIgnoreCase);
+
+			// Find all the matching servers.
+			PerforceServerList ServerList = await GetServerListAsync();
+
+			List<PerforceServerEntry> Candidates = ServerList.Servers.Where(x => x.Cluster == Cluster.Name && ValidServerNames.Contains(x.BaseServerAndPort)).ToList();
+			if (Candidates.Count == 0)
+			{
+				foreach (PerforceServer ValidServer in ValidServers)
+				{
+					Logger.LogDebug("Fetching server info for {ServerAndPort}", ValidServer.ServerAndPort);
+					await UpdateServerAsync(Cluster, ValidServer, Candidates);
+				}
+				if (Candidates.Count == 0)
+				{
+					Logger.LogWarning("Unable to resolve any Perforce servers from valid list");
+					return null; 
+				}
+			}
+
+			// Remove any servers that are unhealthy
+			if (Candidates.Any(x => x.Status == PerforceServerStatus.Healthy))
+			{
+				Candidates.RemoveAll(x => x.Status != PerforceServerStatus.Healthy);
+			}
+			else
+			{
+				Candidates.RemoveAll(x => x.Status != PerforceServerStatus.Unknown);
+			}
+
+			if (Candidates.Count == 0)
+			{
+				Logger.LogWarning("Unable to find any healthy Perforce server in cluster {ClusterName}", Cluster.Name);
+				return null;
+			}
+
+			// Select which server to use with a weighted average of the number of active leases
+			int Index = 0;
+			if (Candidates.Count > 1)
+			{
+				const int BaseWeight = 20;
+
+				int TotalLeases = Candidates.Sum(x => x.NumLeases);
+				int TotalWeight = Candidates.Sum(x => (TotalLeases + BaseWeight) - x.NumLeases);
+
+				int Weight = Random.Next(TotalWeight);
+				for (; Index + 1 < Candidates.Count; Index++)
+				{
+					Weight -= (TotalLeases + BaseWeight) - Candidates[Index].NumLeases;
+					if(Weight < 0)
+					{
+						break;
+					}
+				}
+			}
+
+			return Candidates[Index];
+		}
+
+		/// <inheritdoc/>
+		Task TickAsync(CancellationToken StoppingToken)
+		{
+			if (InternalTickTask.IsCompleted)
+			{
+				InternalTickTask = Task.Run(() => TickInternalGuardedAsync());
+			}
+			return Task.CompletedTask;
+		}
+
+		/// <inheritdoc/>
+		async Task TickInternalGuardedAsync()
+		{
+			try
+			{
+				await TickInternalAsync();
+			}
+			catch (Exception Ex)
+			{
+				Logger.LogError(Ex, "Error updating server state");
+			}
+		}
+
+		/// <inheritdoc/>
+		async Task TickInternalAsync()
+		{
+			Globals Globals = await DatabaseService.GetGlobalsAsync();
+
+			// Set of new server entries
+			List<PerforceServerEntry> NewServers = new List<PerforceServerEntry>();
+
+			// Update the state of all the valid servers
+			foreach (PerforceCluster Cluster in Globals.PerforceClusters)
+			{
+				foreach (PerforceServer Server in Cluster.Servers)
+				{
+					try
+					{
+						await UpdateServerAsync(Cluster, Server, NewServers);
+					}
+					catch (Exception Ex)
+					{
+						Logger.LogError(Ex, "Exception while updating Perforce server status");
+					}
+				}
+			}
+
+			// Update the number of leases for each entry
+			List<PerforceServerEntry> NewEntries = NewServers.OrderBy(x => x.Cluster).ThenBy(x => x.BaseServerAndPort).ThenBy(x => x.ServerAndPort).ToList();
+			await UpdateLeaseCounts(NewEntries);
+			PerforceServerList List = await ServerListSingleton.UpdateAsync(List => MergeServerList(List, NewEntries));
+
+			// Now update the health of each entry
+			List<Task> Tasks = new List<Task>();
+			foreach (PerforceServerEntry Entry in List.Servers)
+			{
+				Tasks.Add(Task.Run(() => UpdateHealthAsync(Entry)));
+			}
+			await Task.WhenAll(Tasks);
+		}
+
+		static void MergeServerList(PerforceServerList ServerList, List<PerforceServerEntry> NewEntries)
+		{
+			Dictionary<string, PerforceServerEntry> ExistingEntries = ServerList.Servers.ToDictionary(x => x.ServerAndPort, x => x, StringComparer.OrdinalIgnoreCase);
+			foreach (PerforceServerEntry NewEntry in NewEntries)
+			{
+				PerforceServerEntry? ExistingEntry;
+				if (ExistingEntries.TryGetValue(NewEntry.ServerAndPort, out ExistingEntry))
+				{
+					NewEntry.Status = ExistingEntry.Status;
+					NewEntry.Detail = ExistingEntry.Detail;
+					NewEntry.LastUpdateTime = ExistingEntry.LastUpdateTime;
+				}
+			}
+			ServerList.Servers = NewEntries;
+		}
+
+		async Task UpdateLeaseCounts(IEnumerable<PerforceServerEntry> NewServerEntries)
+		{
+			Dictionary<string, Dictionary<string, PerforceServerEntry>> NewServerLookup = new Dictionary<string, Dictionary<string, PerforceServerEntry>>();
+			foreach (PerforceServerEntry NewServerEntry in NewServerEntries)
+			{
+				if (NewServerEntry.Cluster != null)
+				{
+					Dictionary<string, PerforceServerEntry>? ClusterServers;
+					if (!NewServerLookup.TryGetValue(NewServerEntry.Cluster, out ClusterServers))
+					{
+						ClusterServers = new Dictionary<string, PerforceServerEntry>();
+						NewServerLookup.Add(NewServerEntry.Cluster, ClusterServers);
+					}
+					ClusterServers[NewServerEntry.ServerAndPort] = NewServerEntry;
+				}
+			}
+
+			List<ILease> Leases = await LeaseCollection.FindActiveLeasesAsync();
+			foreach (ILease Lease in Leases)
+			{
+				Any Any = Any.Parser.ParseFrom(Lease.Payload.ToArray());
+
+				HashSet<(string, string)> Servers = new HashSet<(string, string)>();
+				if (Any.TryUnpack(out ConformTask ConformTask))
+				{
+					foreach (HordeCommon.Rpc.Messages.AgentWorkspace Workspace in ConformTask.Workspaces)
+					{
+						Servers.Add((Workspace.Cluster, Workspace.ServerAndPort));
+					}
+				}
+				if (Any.TryUnpack(out ExecuteJobTask ExecuteJobTask))
+				{
+					if (ExecuteJobTask.AutoSdkWorkspace != null)
+					{
+						Servers.Add((ExecuteJobTask.AutoSdkWorkspace.Cluster, ExecuteJobTask.AutoSdkWorkspace.ServerAndPort));
+					}
+					if (ExecuteJobTask.Workspace != null)
+					{
+						Servers.Add((ExecuteJobTask.Workspace.Cluster, ExecuteJobTask.Workspace.ServerAndPort));
+					}
+				}
+
+				foreach ((string Cluster, string ServerAndPort) in Servers)
+				{
+					IncrementLeaseCount(Cluster, ServerAndPort, NewServerLookup);
+				}
+			}
+		}
+
+		static void IncrementLeaseCount(string Cluster, string ServerAndPort, Dictionary<string, Dictionary<string, PerforceServerEntry>> NewServers)
+		{
+			Dictionary<string, PerforceServerEntry>? ClusterServers;
+			if (Cluster != null && NewServers.TryGetValue(Cluster, out ClusterServers))
+			{
+				PerforceServerEntry? Entry;
+				if (ClusterServers.TryGetValue(ServerAndPort, out Entry))
+				{
+					Entry.NumLeases++;
+				}
+			}
+		}
+
+		async Task UpdateServerAsync(PerforceCluster Cluster, PerforceServer Server, List<PerforceServerEntry> NewServers)
+		{
+			string InitialHostName = Server.ServerAndPort;
+			int Port = 1666;
+
+			int PortIdx = InitialHostName.LastIndexOf(':');
+			if (PortIdx != -1)
+			{
+				Port = int.Parse(InitialHostName.Substring(PortIdx + 1), NumberStyles.Integer, CultureInfo.InvariantCulture);
+				InitialHostName = InitialHostName.Substring(0, PortIdx);
+			}
+
+			List<string> HostNames = new List<string>();
+			if (!Server.ResolveDns || IPAddress.TryParse(InitialHostName, out IPAddress? HostAddress))
+			{
+				HostNames.Add(InitialHostName);
+			}
+			else
+			{
+				await ResolveServersAsync(InitialHostName, HostNames);
+			}
+
+			foreach (string HostName in HostNames)
+			{
+				string? HealthCheckUrl = null;
+				if (Server.HealthCheck)
+				{
+					HealthCheckUrl = $"http://{HostName}:5000/healthcheck";
+				}
+				NewServers.Add(new PerforceServerEntry($"{HostName}:{Port}", Server.ServerAndPort, HealthCheckUrl, Cluster.Name, PerforceServerStatus.Unknown, "", DateTime.UtcNow));
+			}
+		}
+
+		async Task ResolveServersAsync(string HostName, List<string> HostNames)
+		{
+			// Find all the addresses of the hosts
+			IPHostEntry Entry = await Dns.GetHostEntryAsync(HostName);
+			foreach (IPAddress Address in Entry.AddressList)
+			{
+				try
+				{
+					HostNames.Add(Address.ToString());
+				}
+				catch (Exception Ex)
+				{
+					Logger.LogDebug(Ex, "Unable to resolve host name for Perforce server {Address}", Address);
+				}
+			}
+		}
+
+		/// <summary>
+		/// Updates 
+		/// </summary>
+		/// <param name="Entry"></param>
+		/// <returns></returns>
+		async Task UpdateHealthAsync(PerforceServerEntry Entry)
+		{
+			DateTime? UpdateTime = null;
+			string Detail = "Health check disabled";
+
+			// Get the health of the server
+			PerforceServerStatus Health = PerforceServerStatus.Healthy;
+			if (Entry.HealthCheckUrl != null)
+			{
+				UpdateTime = DateTime.UtcNow;
+				Uri HealthCheckUrl = new Uri(Entry.HealthCheckUrl);
+				try
+				{
+					(Health, Detail) = await GetServerHealthAsync(HealthCheckUrl);
+				}
+				catch
+				{
+					(Health, Detail) = (PerforceServerStatus.Unhealthy, $"Failed to query status at {HealthCheckUrl}");
+				}
+			}
+
+			// Update the server record
+			if (Health != Entry.Status || Detail != Entry.Detail || UpdateTime != Entry.LastUpdateTime)
+			{
+				await ServerListSingleton.UpdateAsync(x => UpdateHealth(x, Entry.ServerAndPort, Health, Detail, UpdateTime));
+			}
+		}
+
+		static void UpdateHealth(PerforceServerList ServerList, string ServerAndPort, PerforceServerStatus Status, string Detail, DateTime? UpdateTime)
+		{
+			PerforceServerEntry? Entry = ServerList.Servers.FirstOrDefault(x => x.ServerAndPort.Equals(ServerAndPort, StringComparison.Ordinal));
+			if (Entry != null)
+			{
+				if (Entry.LastUpdateTime == null || UpdateTime == null || Entry.LastUpdateTime.Value < UpdateTime.Value)
+				{
+					Entry.Status = Status;
+					Entry.Detail = Detail;
+					Entry.LastUpdateTime = UpdateTime;
+				}
+			}
+		}
+
+		static async Task<(PerforceServerStatus, string)> GetServerHealthAsync(Uri HealthCheckUrl)
+		{
+			using HttpClient Client = new HttpClient();
+			HttpResponseMessage Response = await Client.GetAsync(HealthCheckUrl);
+
+			byte[] Data = await Response.Content.ReadAsByteArrayAsync();
+			JsonDocument Document = JsonDocument.Parse(Data);
+
+			foreach (JsonElement Element in Document.RootElement.GetProperty("results").EnumerateArray())
+			{
+				if (Element.TryGetProperty("checker", out JsonElement Checker) && Checker.ValueEquals("edge_traffic_lights"))
+				{
+					if (Element.TryGetProperty("output", out JsonElement Output))
+					{
+						string Status = Output.GetString() ?? String.Empty;
+						switch(Status)
+						{
+							case "green":
+								return (PerforceServerStatus.Healthy, "Server is healthy");
+							case "yellow":
+								return (PerforceServerStatus.Degraded, "Degraded service");
+							case "red":
+								return (PerforceServerStatus.Unhealthy, "Server is being drained");
+							default:
+								return (PerforceServerStatus.Unknown, $"Expected state for health check ({Status})");
+						}
+					}
+				}
+			}
+
+			return (PerforceServerStatus.Unknown, "Unable to parse health check output");
+		}
+	}
+}
diff --git a/Engine/Source/Programs/Horde/HordeServer/Projects/ScheduleService.cs b/Engine/Source/Programs/Horde/HordeServer/Projects/ScheduleService.cs
new file mode 100644
--- /dev/null
+++ b/Engine/Source/Programs/Horde/HordeServer/Projects/ScheduleService.cs
@@ -0,0 +1,436 @@
+// Copyright Epic Games, Inc. All Rights Reserved.
+
+using Amazon.S3.Model.Internal.MarshallTransformations;
+using EpicGames.Core;
+using HordeCommon;
+using HordeServer.Api;
+using HordeServer.Collections;
+using HordeServer.Controllers;
+using HordeServer.Models;
+using HordeServer.Utilities;
+using Microsoft.AspNetCore.Server.Kestrel.Core.Features;
+using Microsoft.Extensions.Hosting;
+using Microsoft.Extensions.Logging;
+using Microsoft.Extensions.Options;
+using MongoDB.Bson;
+using MongoDB.Driver;
+using OpenTracing;
+using OpenTracing.Util;
+using System;
+using System.Collections.Generic;
+using System.Diagnostics;
+using System.Linq;
+using System.Security.Claims;
+using System.Threading;
+using System.Threading.Tasks;
+using TimeZoneConverter;
+
+namespace HordeServer.Services
+{
+	using StreamId = StringId<IStream>;
+	using TemplateRefId = StringId<TemplateRef>;
+
+	/// <summary>
+	/// Manipulates schedule instances
+	/// </summary>
+	public class ScheduleService : ElectedBackgroundService
+	{
+		/// <summary>
+		/// The downtime service instance
+		/// </summary>
+		IDowntimeService DowntimeService;
+
+		/// <summary>
+		/// Collection of graph documents
+		/// </summary>
+		IGraphCollection Graphs;
+
+		/// <summary>
+		/// The perforce service singleton
+		/// </summary>
+		IPerforceService Perforce;
+
+		/// <summary>
+		/// Collection of job documents
+		/// </summary>
+		IJobCollection JobCollection;
+
+		/// <summary>
+		/// The job service
+		/// </summary>
+		JobService JobService;
+
+		/// <summary>
+		/// The stream service
+		/// </summary>
+		StreamService StreamService;
+
+		/// <summary>
+		/// The template service
+		/// </summary>
+		TemplateService TemplateService;
+
+		/// <summary>
+		/// The timezone to use for schedules
+		/// </summary>
+		TimeZoneInfo TimeZone;
+
+		/// <summary>
+		/// The clock instance
+		/// </summary>
+		IClock Clock;
+
+		/// <summary>
+		/// The logging instance
+		/// </summary>
+		ILogger Logger;
+
+		/// <summary>
+		/// Constructor
+		/// </summary>
+		/// <param name="DatabaseService">Database service instance</param>
+		/// <param name="Graphs">Collection of graph documents</param>
+		/// <param name="DowntimeService">The downtime service</param>
+		/// <param name="Perforce">The Perforce service</param>
+		/// <param name="JobCollection">Job collection singleton</param>
+		/// <param name="JobService">The job service instance</param>
+		/// <param name="StreamService">The stream service instance</param>
+		/// <param name="TemplateService">The template service instance</param>
+		/// <param name="Clock"></param>
+		/// <param name="Settings">Settings for the server</param>
+		/// <param name="Logger">Logging instance</param>
+		public ScheduleService(DatabaseService DatabaseService, IGraphCollection Graphs, IDowntimeService DowntimeService, IPerforceService Perforce, IJobCollection JobCollection, JobService JobService, StreamService StreamService, TemplateService TemplateService, IClock Clock, IOptionsMonitor<ServerSettings> Settings, ILogger<ScheduleService> Logger)
+			: base(DatabaseService, new ObjectId("6035593e6d721d80fb9efa5c"), Logger)
+		{
+			this.Graphs = Graphs;
+			this.DowntimeService = DowntimeService;
+			this.Perforce = Perforce;
+			this.JobCollection = JobCollection;
+			this.JobService = JobService;
+			this.StreamService = StreamService;
+			this.TemplateService = TemplateService;
+			this.Clock = Clock;
+
+			string? ScheduleTimeZone = Settings.CurrentValue.ScheduleTimeZone;
+			this.TimeZone = (ScheduleTimeZone == null) ? TimeZoneInfo.Local : TZConvert.GetTimeZoneInfo(ScheduleTimeZone);
+
+			this.Logger = Logger;
+		}
+
+		/// <summary>
+		/// Execute any scheduled tasks
+		/// </summary>
+		/// <param name="StoppingToken">Cancellation token for stopping the service</param>
+		/// <returns>Async task</returns>
+		protected override async Task<DateTime> TickLeaderAsync(CancellationToken StoppingToken)
+		{
+			DateTime NextTickTime = DateTime.UtcNow.AddMinutes(1.0);
+			if (!DowntimeService.IsDowntimeActive)
+			{
+				Logger.LogInformation("Updating scheduled triggers...");
+				using (IScope Scope = GlobalTracer.Instance.BuildSpan("ScheduleService Tick").StartActive())
+				{
+					NextTickTime = await UpdateSchedulesAsync();
+				}
+			}
+			return NextTickTime;
+		}
+
+		/// <summary>
+		/// Update all the schedules
+		/// </summary>
+		/// <returns>Next time that checkes need to be updated</returns>
+		[System.Diagnostics.CodeAnalysis.SuppressMessage("Design", "CA1031:Do not catch general exception types", Justification = "<Pending>")]
+		async Task<DateTime> UpdateSchedulesAsync()
+		{
+			// Find the next time at which each schedule should run
+			DateTimeOffset Now = Clock.UtcNow;
+			DateTimeOffset NextUpdateTime = Now.AddMinutes(1.0);
+
+			// Find all the enabled schedules
+			List<IStream> Streams = await StreamService.GetStreamsAsync();
+			foreach (IStream Stream in Streams)
+			{
+				if (Stream.IsPaused(Clock.UtcNow))
+				{
+					Logger.LogDebug("Skipping schedule update for stream {StreamId}. It has been paused until {PausedUntil} with comment '{PauseComment}'.", Stream.Id, Stream.PausedUntil, Stream.PauseComment);
+					continue;
+				}
+				
+				foreach ((TemplateRefId TemplateRefId, TemplateRef TemplateRef) in Stream.Templates)
+				{
+					Schedule? Schedule = TemplateRef.Schedule;
+					if (Schedule != null && Schedule.Enabled)
+					{
+						using IDisposable Scope = Logger.BeginScope("Updating schedule {StreamId}/{TemplateRefId}", Stream.Id, TemplateRefId);
+
+						DateTimeOffset? NextTriggerTime = await UpdateScheduleAsync(Stream, TemplateRefId, TemplateRef, Schedule, Now);
+						if (NextTriggerTime.HasValue && NextTriggerTime.Value < NextUpdateTime)
+						{
+							NextUpdateTime = NextTriggerTime.Value;
+						}
+					}
+				}
+			}
+
+			return NextUpdateTime.UtcDateTime;
+		}
+
+		/// <summary>
+		/// Updates a schedule
+		/// </summary>
+		/// <returns>Next time to trigger the schedule</returns>
+		[System.Diagnostics.CodeAnalysis.SuppressMessage("Design", "CA1031:Do not catch general exception types", Justification = "<Pending>")]
+		private async Task<DateTimeOffset?> UpdateScheduleAsync(IStream Stream, TemplateRefId TemplateRefId, TemplateRef TemplateRef, Schedule Schedule, DateTimeOffset Now)
+		{
+			// Check if we can run the trigger
+			DateTimeOffset? NextTriggerTime = Schedule.GetNextTriggerTime(TimeZone);
+			if (!NextTriggerTime.HasValue)
+			{
+				Logger.LogInformation("Schedule {StreamId}/{TemplateRefId} is disabled", Stream.Id, TemplateRefId);
+				return NextTriggerTime;
+			}
+			if (NextTriggerTime.Value > Now)
+			{
+				Logger.LogDebug("Schedule {StreamId}/{TemplateRefId} will next trigger at {NextTriggerTime}", Stream.Id, TemplateRefId, NextTriggerTime.Value);
+				return NextTriggerTime;
+			}
+			Logger.LogInformation("Schedule {StreamId}/{TemplateRefId} is ready to trigger (last: {Change} at {Time})", Stream.Id, TemplateRefId, Schedule.LastTriggerChange, Schedule.LastTriggerTime);
+
+			// Get a list of jobs that we need to remove
+			List<ObjectId> RemoveJobIds = new List<ObjectId>();
+			foreach (ObjectId ActiveJobId in Schedule.ActiveJobs)
+			{
+				IJob? Job = await JobService.GetJobAsync(ActiveJobId);
+				if (Job == null || Job.Batches.All(x => x.State == JobStepBatchState.Complete))
+				{
+					Logger.LogInformation("Removing active job {JobId}", ActiveJobId);
+					RemoveJobIds.Add(ActiveJobId);
+				}
+			}
+			if (RemoveJobIds.Count > 0)
+			{
+				await StreamService.UpdateScheduleTriggerAsync(Stream, TemplateRefId, Now, null, new List<ObjectId>(), RemoveJobIds);
+			}
+
+			// Trigger this schedule
+			try
+			{
+				await TriggerAsync(Stream, TemplateRefId, TemplateRef, Schedule, Schedule.ActiveJobs.Count - RemoveJobIds.Count, Now);
+			}
+			catch (Exception Ex)
+			{
+				Logger.LogError(Ex, "Failed to start schedule {StreamId}/{TemplateRefId}", Stream.Id, TemplateRefId);
+			}
+
+			// Return the next update time
+			return Schedule.GetNextTriggerTime(Now, TimeZone);
+		}
+
+		/// <summary>
+		/// Trigger a schedule to run
+		/// </summary>
+		/// <param name="Stream">Stream for the schedule</param>
+		/// <param name="TemplateRefId"></param>
+		/// <param name="TemplateRef"></param>
+		/// <param name="Schedule"></param>
+		/// <param name="NumActiveJobs"></param>
+		/// <param name="Now"></param>
+		/// <returns>Async task</returns>
+		private async Task TriggerAsync(IStream Stream, TemplateRefId TemplateRefId, TemplateRef TemplateRef, Schedule Schedule, int NumActiveJobs, DateTimeOffset Now)
+		{
+			// Check we're not already at the maximum number of allowed jobs
+			if (Schedule.MaxActive != 0 && NumActiveJobs >= Schedule.MaxActive)
+			{
+				Logger.LogInformation("Skipping trigger of {StreamId} template {TemplateRefId} - already have maximum number of jobs running ({NumJobs})", Stream.Id, TemplateRefId, Schedule.MaxActive);
+				foreach (ObjectId JobId in Schedule.ActiveJobs)
+				{
+					Logger.LogInformation("Active job for {StreamId} template {TemplateRefId}: {JobId}", Stream.Id, TemplateRefId, JobId);
+				}
+				return;
+			}
+
+			// Minimum changelist number, inclusive
+			int MinChangeNumber = Schedule.LastTriggerChange;
+			if (MinChangeNumber > 0 && !Schedule.RequireSubmittedChange)
+			{
+				MinChangeNumber--;
+			}
+
+			// Maximum changelist number, exclusive
+			int? MaxChangeNumber = null;
+
+			// Get the maximum number of changes to trigger
+			int MaxNewChanges = 1;
+			if (Schedule.MaxChanges != 0)
+			{
+				MaxNewChanges = Schedule.MaxChanges;
+			}
+			if (Schedule.MaxActive != 0)
+			{
+				MaxNewChanges = Math.Min(MaxNewChanges, Schedule.MaxActive - NumActiveJobs);
+			}
+
+			// Create a timer to limit the amount we look back through P4 history
+			Stopwatch Timer = Stopwatch.StartNew();
+			ChangeContentFlags? FilterFlags = Schedule.GetFilterFlags();
+
+			// Start as many jobs as possible
+			List<(int Change, int CodeChange)> TriggerChanges = new List<(int, int)>();
+			while (TriggerChanges.Count < MaxNewChanges)
+			{
+				// Get the next valid change
+				int Change, CodeChange;
+				if (Schedule.Gate != null)
+				{
+					(Change, CodeChange) = await GetNextChangeForGateAsync(Stream.Id, TemplateRefId, Schedule.Gate, MinChangeNumber, MaxChangeNumber);
+				}
+				else
+				{
+					(Change, CodeChange) = await GetNextChangeAsync(Stream.ClusterName, Stream.Name, MinChangeNumber, MaxChangeNumber);
+				}
+
+				// Quit if we didn't find anything
+				if (Change < MinChangeNumber)
+				{
+					break;
+				}
+				if (Change == MinChangeNumber && (Schedule.RequireSubmittedChange || TriggerChanges.Count > 0))
+				{
+					break;
+				}
+
+				// Adjust the changelist for the desired filter
+				if(await ShouldBuildChangeAsync(Stream.ClusterName, Stream.Name, Change, FilterFlags))
+				{
+					TriggerChanges.Add((Change, CodeChange));
+				}
+
+				// Check we haven't exceeded the time limit
+				if (Timer.Elapsed > TimeSpan.FromMinutes(2.0))
+				{
+					Logger.LogError("Querying for changes to trigger has taken {Time}. Aborting.", Timer.Elapsed);
+					break;
+				}
+
+				// Start the next build before this change
+				if (FilterFlags != null && (FilterFlags.Value & ChangeContentFlags.ContainsContent) == 0)
+				{
+					MaxChangeNumber = Math.Min(CodeChange, Change - 1);
+				}
+				else
+				{
+					MaxChangeNumber = Change - 1;
+				}
+			}
+
+			// Early out if there's nothing to do
+			if (TriggerChanges.Count == 0)
+			{
+				Logger.LogInformation("Skipping trigger of {StreamName} template {TemplateRefId} - no candidate changes after CL {LastTriggerChange}", Stream.Id, TemplateRefId, Schedule.LastTriggerChange);
+				return;
+			}
+
+			// Get the matching template
+			ITemplate? Template = await TemplateService.GetTemplateAsync(TemplateRef.Hash);
+			if (Template == null)
+			{
+				Logger.LogWarning("Unable to find template '{TemplateHash}' for '{TemplateRefId}'", TemplateRef.Hash, TemplateRefId);
+				return;
+			}
+
+			// Register the graph for it
+			IGraph Graph = await Graphs.AddAsync(Template);
+
+			// We may need to submit a new change for any new jobs. This only makes sense if there's one change.
+			if (Template.SubmitNewChange != null)
+			{
+				int NewChange = await Perforce.CreateNewChangeForTemplateAsync(Stream, Template);
+				int NewCodeChange = await Perforce.GetCodeChangeAsync(Stream.ClusterName, Stream.Name, NewChange);
+				TriggerChanges = new List<(int, int)> { (NewChange, NewCodeChange) };
+			}
+
+			// Try to start all the new jobs
+			foreach ((int Change, int CodeChange) in TriggerChanges.OrderBy(x => x.Change))
+			{
+				List<string> DefaultArguments = Template.GetDefaultArguments();
+				IJob NewJob = await JobService.CreateJobAsync(null, Stream, TemplateRefId, Template.Id, Graph, Template.Name, Change, CodeChange, null, null, null, null, Template.Priority, null, null, TemplateRef.ChainedJobs, TemplateRef.ShowUgsBadges, TemplateRef.ShowUgsAlerts, TemplateRef.NotificationChannel, TemplateRef.NotificationChannelFilter, null, Template.Counters, DefaultArguments);
+				Logger.LogInformation("Started new job for {StreamName} template {TemplateName} at CL {Change} (Code CL {CodeChange}): {JobId}", Stream.Id, TemplateRef.Name, Change, CodeChange, NewJob.Id);
+				await StreamService.UpdateScheduleTriggerAsync(Stream, TemplateRefId, Now, Change, new List<ObjectId> { NewJob.Id }, new List<ObjectId>());
+			}
+		}
+
+		/// <summary>
+		/// Tests whether a schedule should build a particular change, based on its requested change filters
+		/// </summary>
+		/// <param name="ClusterName"></param>
+		/// <param name="StreamName"></param>
+		/// <param name="Change"></param>
+		/// <param name="FilterFlags"></param>
+		/// <returns></returns>
+		private async Task<bool> ShouldBuildChangeAsync(string ClusterName, string StreamName, int Change, ChangeContentFlags? FilterFlags)
+		{
+			if (FilterFlags != null && FilterFlags.Value != 0)
+			{
+				ChangeDetails Details = await Perforce.GetChangeDetailsAsync(ClusterName, StreamName, Change, null);
+				if ((Details.GetContentFlags() & FilterFlags.Value) == 0)
+				{
+					Logger.LogDebug("Not building change {Change} ({ChangeFlags}) due to filter flags ({FilterFlags})", Change, Details.GetContentFlags().ToString(), FilterFlags.Value.ToString());
+					return false;
+				}
+			}
+			return true;
+		}
+
+		/// <summary>
+		/// Get the next change to build in this stream, between a given range
+		/// </summary>
+		/// <returns>The new change and code change to build, or 0.</returns>
+		private async Task<(int Change, int CodeChange)> GetNextChangeAsync(string ClusterName, string StreamName, int? MinChange, int? MaxChange)
+		{
+			List<ChangeSummary> Changes = await Perforce.GetChangesAsync(ClusterName, StreamName, MinChange, MaxChange, 1, null);
+			if (Changes.Count > 0)
+			{
+				int Change = Changes[0].Number;
+				int CodeChange = await Perforce.GetCodeChangeAsync(ClusterName, StreamName, Change);
+				return (Change, CodeChange);
+			}
+			return default;
+		}
+
+		/// <summary>
+		/// Gets the next change to build for a schedule on a gate
+		/// </summary>
+		/// <returns></returns>
+		private async Task<(int Change, int CodeChange)> GetNextChangeForGateAsync(StreamId StreamId, TemplateRefId TemplateRefId, ScheduleGate Gate, int? MinChange, int? MaxChange)
+		{
+			for (; ; )
+			{
+				List<IJob> Jobs = await JobCollection.FindAsync(StreamId: StreamId, Templates: new[] { Gate.TemplateRefId }, MinChange: MinChange, MaxChange: MaxChange, Count: 1);
+				if (Jobs.Count == 0)
+				{
+					return default;
+				}
+
+				IJob Job = Jobs[0];
+
+				IGraph? Graph = await Graphs.GetAsync(Job.GraphHash);
+				if (Graph != null)
+				{
+					(JobStepState, JobStepOutcome)? State = Job.GetTargetState(Graph, Gate.Target);
+					if (State != null && State.Value.Item1 == JobStepState.Completed)
+					{
+						JobStepOutcome Outcome = State.Value.Item2;
+						if (Outcome == JobStepOutcome.Success || Outcome == JobStepOutcome.Warnings)
+						{
+							return (Job.Change, Job.CodeChange);
+						}
+						Logger.LogInformation("Skipping trigger of {StreamName} template {TemplateRefId} - last {OtherTemplateRefId} job ({JobId}) ended with errors", StreamId, TemplateRefId, Gate.TemplateRefId, Job.Id);
+					}
+				}
+
+				MaxChange = Job.Change - 1;
+			}
+		}
+	}
+}
diff --git a/Engine/Source/Programs/Horde/HordeServer/Projects/StreamTab.cs b/Engine/Source/Programs/Horde/HordeServer/Projects/StreamTab.cs
new file mode 100644
--- /dev/null
+++ b/Engine/Source/Programs/Horde/HordeServer/Projects/StreamTab.cs
@@ -0,0 +1,269 @@
+// Copyright Epic Games, Inc. All Rights Reserved.
+
+using HordeServer.Api;
+using HordeServer.Utilities;
+using MongoDB.Bson.Serialization.Attributes;
+using System;
+using System.Collections.Generic;
+using System.Linq;
+using System.Threading.Tasks;
+
+namespace HordeServer.Models
+{
+	using TemplateRefId = StringId<TemplateRef>;
+
+	/// <summary>
+	/// Information about a page to display in the dashboard for a stream
+	/// </summary>
+	[BsonKnownTypes(typeof(JobsTab))]
+	public abstract class StreamTab
+	{
+		/// <summary>
+		/// Title of this page
+		/// </summary>
+		[BsonRequired]
+		public string Title { get; set; }
+
+		/// <summary>
+		/// Private constructor for serialization
+		/// </summary>
+		private StreamTab()
+		{
+			Title = null!;
+		}
+
+		/// <summary>
+		/// Constructor
+		/// </summary>
+		/// <param name="Title">Title of this page</param>
+		public StreamTab(string Title)
+		{
+			this.Title = Title;
+		}
+
+		/// <summary>
+		/// Constructor
+		/// </summary>
+		/// <param name="Request">Public request object</param>
+		public StreamTab(CreateStreamTabRequest Request)
+		{
+			Title = Request.Title;
+		}
+
+		/// <summary>
+		/// Creates an instance from a request
+		/// </summary>
+		/// <param name="Request">The request object</param>
+		/// <returns>New tab instance</returns>
+		public static StreamTab FromRequest(CreateStreamTabRequest Request)
+		{
+			CreateJobsTabRequest? JobsRequest = Request as CreateJobsTabRequest;
+			if (JobsRequest != null)
+			{
+				return new JobsTab(JobsRequest);
+			}
+
+			throw new Exception($"Unknown tab request type '{Request.GetType()}'");
+		}
+
+		/// <summary>
+		/// Creates a response object
+		/// </summary>
+		/// <returns>Response object</returns>
+		public abstract GetStreamTabResponse ToResponse();
+	}
+
+	/// <summary>
+	/// Describes a column to display on the jobs page
+	/// </summary>
+	[BsonKnownTypes(typeof(JobsTabLabelColumn), typeof(JobsTabParameterColumn))]
+	public abstract class JobsTabColumn
+	{
+		/// <summary>
+		/// Heading for this column
+		/// </summary>
+		public string Heading { get; set; }
+
+		/// <summary>
+		/// Relative width of this column.
+		/// </summary>
+		public int? RelativeWidth { get; set; }
+
+		/// <summary>
+		/// Private constructor for serialization
+		/// </summary>
+		protected JobsTabColumn()
+		{
+			Heading = null!;
+			RelativeWidth = 1;
+		}
+
+		/// <summary>
+		/// Constructor
+		/// </summary>
+		/// <param name="Heading">Heading for this column</param>
+		/// <param name="RelativeWidth">Relative width of this column.</param>
+		public JobsTabColumn(string Heading, int? RelativeWidth)
+		{
+			this.Heading = Heading;
+			this.RelativeWidth = RelativeWidth;
+		}
+
+		/// <summary>
+		/// Constructor
+		/// </summary>
+		/// <param name="Request">Public request object</param>
+		public JobsTabColumn(CreateJobsTabColumnRequest Request)
+		{
+			this.Heading = Request.Heading;
+			this.RelativeWidth = Request.RelativeWidth;
+		}
+
+		/// <summary>
+		/// Converts this object to a JSON response object
+		/// </summary>
+		/// <returns>Response object</returns>
+		public abstract GetJobsTabColumnResponse ToResponse();
+	}
+
+	/// <summary>
+	/// Column that contains a set of labels
+	/// </summary>
+	public class JobsTabLabelColumn : JobsTabColumn
+	{
+		/// <summary>
+		/// Category of labels to display in this column. If null, includes any label not matched by another column.
+		/// </summary>
+		public string? Category { get; set; }
+
+		/// <summary>
+		/// Private constructor for serialization
+		/// </summary>
+		protected JobsTabLabelColumn()
+		{
+		}
+
+		/// <summary>
+		/// Constructor
+		/// </summary>
+		/// <param name="Heading">Heading for this column</param>
+		/// <param name="Category">Category of labels to display in this column</param>
+		/// <param name="RelativeWidth">Relative width of this column.</param>
+		public JobsTabLabelColumn(string Heading, string? Category, int? RelativeWidth)
+			: base(Heading, RelativeWidth)
+		{
+			this.Category = Category;
+		}
+
+		/// <summary>
+		/// Converts this object to a JSON response object
+		/// </summary>
+		/// <returns>Response object</returns>
+		public override GetJobsTabColumnResponse ToResponse()
+		{
+			return new GetJobsTabLabelColumnResponse(Heading, Category, RelativeWidth);
+		}
+	}
+
+	/// <summary>
+	/// Include the value for a parameter on the jobs tab
+	/// </summary>
+	public class JobsTabParameterColumn : JobsTabColumn
+	{
+		/// <summary>
+		/// Name of a parameter to show in this column. Should be in the form of a prefix, eg. "-set:Foo="
+		/// </summary>
+		public string? Parameter { get; set; }
+
+		/// <summary>
+		/// Private constructor for serialization
+		/// </summary>
+		protected JobsTabParameterColumn()
+		{
+		}
+
+		/// <summary>
+		/// Constructor
+		/// </summary>
+		/// <param name="Heading">Heading for this column</param>
+		/// <param name="Parameter">Category of labels to display in this column</param>
+		/// <param name="RelativeWidth">Relative width of this column.</param>
+		public JobsTabParameterColumn(string Heading, string? Parameter, int? RelativeWidth)
+			: base(Heading, RelativeWidth)
+		{
+			this.Parameter = Parameter;
+		}
+
+		/// <summary>
+		/// Converts this object to a JSON response object
+		/// </summary>
+		/// <returns>Response object</returns>
+		public override GetJobsTabColumnResponse ToResponse()
+		{
+			return new GetJobsTabParameterColumnResponse(Heading, Parameter, RelativeWidth);
+		}
+	}
+
+	/// <summary>
+	/// Describes a job page
+	/// </summary>
+	[BsonDiscriminator("Jobs")]
+	public class JobsTab : StreamTab
+	{
+		/// <summary>
+		/// Whether to show job names
+		/// </summary>
+		public bool ShowNames { get; set; }
+
+		/// <summary>
+		/// Names of jobs to include on this page. If there is only one name specified, the name column does not need to be displayed.
+		/// </summary>
+		public List<string>? JobNames { get; set; }
+
+		/// <summary>
+		/// List of templates to display on this tab.
+		/// </summary>
+		public List<TemplateRefId>? Templates { get; set; }
+
+		/// <summary>
+		/// Columns to display in the table
+		/// </summary>
+		public List<JobsTabColumn>? Columns { get; set; }
+
+		/// <summary>
+		/// Constructor
+		/// </summary>
+		/// <param name="Title">Title for this page</param>
+		/// <param name="ShowNames">Show names of jobs on the page</param>
+		/// <param name="Templates">Template names to include on this page</param>
+		/// <param name="JobNames">Names of jobs to include on this page</param>
+		/// <param name="Columns">Columns ot display in the table</param>
+		public JobsTab(string Title, bool ShowNames, List<TemplateRefId>? Templates, List<string>? JobNames, List<JobsTabColumn>? Columns)
+			: base(Title)
+		{
+			this.ShowNames = ShowNames;
+			this.Templates = Templates;
+			this.JobNames = JobNames;
+			this.Columns = Columns;
+		}
+
+		/// <summary>
+		/// Constructor
+		/// </summary>
+		/// <param name="Request">Public request object</param>
+		public JobsTab(CreateJobsTabRequest Request)
+			: base(Request)
+		{
+			ShowNames = Request.ShowNames;
+			Templates = Request.Templates?.ConvertAll(x => new TemplateRefId(x));
+			JobNames = Request.JobNames;
+			Columns = Request.Columns?.ConvertAll(x => x.ToModel());
+		}
+
+		/// <inheritdoc/>
+		public override GetStreamTabResponse ToResponse()
+		{
+			return new GetJobsTabResponse(Title, ShowNames, Templates?.ConvertAll(x => x.ToString()), JobNames, Columns?.ConvertAll(x => x.ToResponse()));
+		}
+	}
+}
diff --git a/Engine/Source/Programs/Horde/HordeServer/Server/DatabaseService.cs b/Engine/Source/Programs/Horde/HordeServer/Server/DatabaseService.cs
new file mode 100644
--- /dev/null
+++ b/Engine/Source/Programs/Horde/HordeServer/Server/DatabaseService.cs
@@ -0,0 +1,662 @@
+// Copyright Epic Games, Inc. All Rights Reserved.
+
+using EpicGames.Core;
+using HordeServer.Api;
+using HordeCommon;
+using HordeServer.Models;
+using HordeServer.Utilities;
+using Microsoft.Extensions.Logging;
+using Microsoft.Extensions.Options;
+using Microsoft.IdentityModel.Tokens;
+using MongoDB.Bson;
+using MongoDB.Bson.Serialization;
+using MongoDB.Driver;
+using System;
+using System.Collections.Concurrent;
+using System.Collections.Generic;
+using System.IO;
+using System.Linq;
+using System.Linq.Expressions;
+using System.Net;
+using System.Net.Security;
+using System.Runtime.InteropServices;
+using System.Security.Cryptography.X509Certificates;
+using System.Text;
+using System.Threading.Tasks;
+using System.Net.Sockets;
+using MongoDB.Driver.Core.Events;
+using System.Threading;
+using HordeServer.Collections;
+using System.Diagnostics;
+using System.Text.RegularExpressions;
+using System.Net.NetworkInformation;
+using OpenTracing.Util;
+using OpenTracing;
+using HordeServer.Utiltiies;
+
+namespace HordeServer.Services
+{
+	using IStream = HordeServer.Models.IStream;
+	using PoolId = StringId<IPool>;
+	using ProjectId = StringId<IProject>;
+	using TemplateRefId = StringId<TemplateRef>;
+
+	/// <summary>
+	/// Singleton for accessing the database
+	/// </summary>
+	public sealed class DatabaseService : IDisposable
+	{
+		/// <summary>
+		/// The database instance
+		/// </summary>
+		public IMongoDatabase Database { get; private set; }
+
+		/// <summary>
+		/// Collection of credential documents
+		/// </summary>
+		public IMongoCollection<Credential> Credentials { get; }
+
+		/// <summary>
+		/// Collection of artifact documents
+		/// </summary>
+		public IMongoCollection<Artifact> Artifacts { get; }
+
+		/// <summary>
+		/// Collection of singleton documents
+		/// </summary>
+		public IMongoCollection<BsonDocument> Singletons { get; }
+
+		/// <summary>
+		/// Settings for the application
+		/// </summary>
+		ServerSettings Settings { get; }
+
+		/// <summary>
+		/// Name of the JWT issuer
+		/// </summary>
+		public string JwtIssuer { get; }
+
+		/// <summary>
+		/// The signing key for this instance
+		/// </summary>
+		public SymmetricSecurityKey JwtSigningKey { get; }
+
+		/// <summary>
+		/// Logger for this instance
+		/// </summary>
+		ILogger<DatabaseService> Logger;
+		
+		/// <summary>
+		/// Access the database in a read-only mode (don't create indices or modify content)
+		/// </summary>
+		public bool ReadOnlyMode { get; }
+
+		/// <summary>
+		/// The mongo process group
+		/// </summary>
+		ManagedProcessGroup? MongoProcessGroup;
+
+		/// <summary>
+		/// The mongo process
+		/// </summary>
+		ManagedProcess? MongoProcess;
+
+		/// <summary>
+		/// Task to read from the mongo process
+		/// </summary>
+		Task? MongoOutputTask;
+
+		/// <summary>
+		/// Factory for creating logger instances
+		/// </summary>
+		ILoggerFactory LoggerFactory;
+
+		/// <summary>
+		/// Default port for MongoDB connections
+		/// </summary>
+		const int DefaultMongoPort = 27017;
+
+		/// <summary>
+		/// Constructor
+		/// </summary>
+		/// <param name="SettingsSnapshot">The settings instance</param>
+		/// <param name="LoggerFactory">Instance of the logger for this service</param>
+		public DatabaseService(IOptions<ServerSettings> SettingsSnapshot, ILoggerFactory LoggerFactory)
+		{
+			this.Settings = SettingsSnapshot.Value;
+			this.Logger = LoggerFactory.CreateLogger<DatabaseService>();
+			this.LoggerFactory = LoggerFactory;
+
+			try
+			{
+				ReadOnlyMode = Settings.DatabaseReadOnlyMode;
+				if (Settings.DatabasePublicCert != null)
+				{
+					X509Store LocalTrustStore = new X509Store(StoreName.Root);
+					try
+					{
+						LocalTrustStore.Open(OpenFlags.ReadWrite);
+
+						X509Certificate2Collection Collection = ImportCertificateBundle(Settings.DatabasePublicCert);
+						foreach (X509Certificate2 Certificate in Collection)
+						{
+							Logger.LogInformation("Importing certificate for {Subject}", Certificate.Subject);
+						}
+
+						LocalTrustStore.AddRange(Collection);
+					}
+					finally
+					{
+						LocalTrustStore.Close();
+					}
+				}
+
+				string? ConnectionString = Settings.DatabaseConnectionString;
+				if (ConnectionString == null)
+				{
+					if (IsPortInUse(DefaultMongoPort))
+					{
+						ConnectionString = "mongodb://localhost:27017";
+					}
+					else if (TryStartMongoServer(Logger))
+					{
+						ConnectionString = "mongodb://localhost:27017/?readPreference=primary&appname=Horde&ssl=false";
+					}
+					else
+					{
+						throw new Exception($"Unable to connect to MongoDB server. Setup a MongoDB server and set the connection string in {Program.UserConfigFile}");
+					}
+				}
+
+				MongoClientSettings MongoSettings = MongoClientSettings.FromConnectionString(ConnectionString);
+				MongoSettings.ClusterConfigurator = ClusterBuilder =>
+				{
+					if (Logger.IsEnabled(LogLevel.Trace))
+					{
+						ClusterBuilder.Subscribe<CommandStartedEvent>(Event => TraceMongoCommand(Event.Command));
+					}
+				};
+				
+				MongoSettings.SslSettings = new SslSettings();
+				MongoSettings.SslSettings.ServerCertificateValidationCallback = CertificateValidationCallBack;
+				MongoSettings.MaxConnectionPoolSize = 300; // Default is 100
+
+				//TestSslConnection(MongoSettings.Server.Host, MongoSettings.Server.Port, Logger);
+
+				MongoClient Client = new MongoClient(MongoSettings);
+				Database = Client.GetDatabase(Settings.DatabaseName);
+
+				Singletons = GetCollection<BsonDocument>("Singletons");
+
+				Globals Globals = GetGlobalsAsync().Result;
+				if (Globals.ForceReset)
+				{
+					Logger.LogInformation("Performing DB reset...");
+
+					// Drop all the collections
+					List<string> CollectionNames = Database.ListCollectionNames().ToList();
+					foreach (string CollectionName in CollectionNames)
+					{
+						Logger.LogInformation("Dropping {CollectionName}", CollectionName);
+						Database.DropCollection(CollectionName);
+					}
+
+					// Create the new singletons collection
+					Singletons = GetCollection<BsonDocument>("Singletons");
+
+					// Create the new globals object
+					Globals NewGlobals = new Globals();
+					NewGlobals.JwtSigningKey = Globals.JwtSigningKey; // Ensure agents can still register
+					Singletons.InsertOne(NewGlobals.ToBsonDocument());
+					Globals = NewGlobals;
+				}
+
+				while (Globals.JwtSigningKey == null)
+				{
+					Globals.RotateSigningKey();
+					if (!TryUpdateSingletonAsync(Globals).Result)
+					{
+						Globals = GetGlobalsAsync().Result;
+					}
+				}
+
+				JwtIssuer = Settings.JwtIssuer ?? Dns.GetHostName();
+				if (String.IsNullOrEmpty(Settings.JwtSecret))
+				{
+					JwtSigningKey = new SymmetricSecurityKey(Globals.JwtSigningKey);
+				}
+				else
+				{
+					JwtSigningKey = new SymmetricSecurityKey(Convert.FromBase64String(Settings.JwtSecret));
+				}
+
+				Credentials = GetCollection<Credential>("Credentials");
+				Artifacts = GetCollection<Artifact>("Artifacts");
+			}
+			catch (Exception Ex)
+			{
+				Logger.LogError(Ex, "Exception while initializing DatabaseService");
+				throw;
+			}
+		}
+
+		internal const int CTRL_C_EVENT = 0;
+
+		[DllImport("kernel32.dll")]
+		internal static extern bool GenerateConsoleCtrlEvent(int Event, int ProcessGroupId);
+
+		/// <inheritdoc/>
+		public void Dispose()
+		{
+			if (MongoProcess != null)
+			{
+				GenerateConsoleCtrlEvent(CTRL_C_EVENT, MongoProcess.Id);
+
+				MongoOutputTask?.Wait();
+				MongoOutputTask = null;
+
+				MongoProcess.WaitForExit();
+				MongoProcess.Dispose();
+				MongoProcess = null;
+			}
+			if(MongoProcessGroup != null)
+			{
+				MongoProcessGroup?.Dispose();
+				MongoProcessGroup = null;
+			}
+		}
+
+		/// <summary>
+		/// Checks if the given port is in use
+		/// </summary>
+		/// <param name="PortNumber"></param>
+		/// <returns></returns>
+		static bool IsPortInUse(int PortNumber)
+		{
+			IPGlobalProperties IpGlobalProperties = IPGlobalProperties.GetIPGlobalProperties();
+
+			IPEndPoint[] Listeners = IpGlobalProperties.GetActiveTcpListeners();
+			if (Listeners.Any(x => x.Port == PortNumber))
+			{
+				return true;
+			}
+
+			return false;
+		}
+
+		/// <summary>
+		/// Attempts to start a local instance of MongoDB
+		/// </summary>
+		/// <param name="Logger"></param>
+		/// <returns></returns>
+		bool TryStartMongoServer(ILogger Logger)
+		{
+			if (!RuntimeInformation.IsOSPlatform(OSPlatform.Windows))
+			{
+				return false;
+			}
+
+			FileReference MongoExe = FileReference.Combine(Program.AppDir, "ThirdParty", "Mongo", "mongod.exe");
+			if (!FileReference.Exists(MongoExe))
+			{
+				Logger.LogWarning("Unable to find Mongo executable.");
+				return false;
+			}
+
+			DirectoryReference MongoDir = DirectoryReference.Combine(Program.DataDir, "Mongo");
+
+			DirectoryReference MongoDataDir = DirectoryReference.Combine(MongoDir, "Data");
+			DirectoryReference.CreateDirectory(MongoDataDir);
+
+			FileReference MongoLogFile = FileReference.Combine(MongoDir, "mongod.log");
+
+			FileReference ConfigFile = FileReference.Combine(MongoDir, "mongod.conf");
+			if(!FileReference.Exists(ConfigFile))
+			{
+				DirectoryReference.CreateDirectory(ConfigFile.Directory);
+				using (StreamWriter Writer = new StreamWriter(ConfigFile.FullName))
+				{
+					Writer.WriteLine("# mongod.conf");
+					Writer.WriteLine();
+					Writer.WriteLine("# for documentation of all options, see:");
+					Writer.WriteLine("# http://docs.mongodb.org/manual/reference/configuration-options/");
+					Writer.WriteLine();
+					Writer.WriteLine("storage:");
+					Writer.WriteLine("    dbPath: {0}", MongoDataDir.FullName);
+					Writer.WriteLine();
+					Writer.WriteLine("net:");
+					Writer.WriteLine("    port: {0}", DefaultMongoPort);
+					Writer.WriteLine("    bindIp: 127.0.0.1");
+				}
+			}
+
+			MongoProcessGroup = new ManagedProcessGroup();
+			try
+			{
+				MongoProcess = new ManagedProcess(MongoProcessGroup, MongoExe.FullName, $"--config \"{ConfigFile}\"", null, null, ProcessPriorityClass.Normal);
+				MongoProcess.StdIn.Close();
+				MongoOutputTask = Task.Run(() => RelayMongoOutput());
+				return true;
+			}
+			catch (Exception Ex)
+			{
+				Logger.LogWarning(Ex, "Unable to start Mongo server process");
+				return false;
+			}
+		}
+
+		/// <summary>
+		/// Copies output from the mongo process to the logger
+		/// </summary>
+		/// <returns></returns>
+		async Task RelayMongoOutput()
+		{
+			ILogger MongoLogger = LoggerFactory.CreateLogger("MongoDB");
+
+			Dictionary<string, ILogger> ChannelLoggers = new Dictionary<string, ILogger>();
+			for (; ; )
+			{
+				string? Line = await MongoProcess!.ReadLineAsync();
+				if (Line == null)
+				{
+					break;
+				}
+				if (Line.Length > 0)
+				{
+					Match Match = Regex.Match(Line, @"^\s*[^\s]+\s+([A-Z])\s+([^\s]+)\s+(.*)");
+					if (Match.Success)
+					{
+						ILogger? ChannelLogger;
+						if (!ChannelLoggers.TryGetValue(Match.Groups[2].Value, out ChannelLogger))
+						{
+							ChannelLogger = LoggerFactory.CreateLogger($"MongoDB.{Match.Groups[2].Value}");
+							ChannelLoggers.Add(Match.Groups[2].Value, ChannelLogger);
+						}
+						ChannelLogger.Log(ParseMongoLogLevel(Match.Groups[1].Value), Match.Groups[3].Value.TrimEnd());
+					}
+					else
+					{
+						MongoLogger.Log(LogLevel.Information, Line);
+					}
+				}
+			}
+			MongoLogger.LogInformation("Exit code {ExitCode}", MongoProcess.ExitCode);
+		}
+
+		static LogLevel ParseMongoLogLevel(string Text)
+		{
+			if (Text.Equals("I", StringComparison.Ordinal))
+			{
+				return LogLevel.Information;
+			}
+			else if (Text.Equals("E", StringComparison.Ordinal))
+			{
+				return LogLevel.Error;
+			}
+			else
+			{
+				return LogLevel.Warning;
+			}
+		}
+
+		/// <summary>
+		/// Logs a mongodb command, removing any fields we don't care about
+		/// </summary>
+		/// <param name="Command">The command document</param>
+		void TraceMongoCommand(BsonDocument Command)
+		{
+			List<string> Params = new List<string>();
+			List<string> Values = new List<string>();
+
+			foreach(BsonElement Element in Command)
+			{
+				if (Element.Value != null && !Element.Name.Equals("$db", StringComparison.Ordinal) && !Element.Name.Equals("lsid", StringComparison.Ordinal))
+				{
+					Params.Add($"{Element.Name}: {{{Values.Count}}}");
+					Values.Add(Element.Value.ToString()!);
+				}
+			}
+
+			Logger.LogTrace($"MongoDB: {String.Join(", ", Params)}", Values.ToArray());
+		}
+
+		/// <summary>
+		/// Imports one or more certificates from a single PEM file
+		/// </summary>
+		/// <param name="FileName">File to import</param>
+		/// <returns>Collection of certificates</returns>
+		static X509Certificate2Collection ImportCertificateBundle(string FileName)
+		{
+			X509Certificate2Collection Collection = new X509Certificate2Collection();
+
+			string Text = File.ReadAllText(FileName);
+			for (int Offset = 0; Offset < Text.Length;)
+			{
+				int NextOffset = Text.IndexOf("-----BEGIN CERTIFICATE-----", Offset + 1, StringComparison.Ordinal);
+				if (NextOffset == -1)
+				{
+					NextOffset = Text.Length;
+				}
+
+				string CertificateText = Text.Substring(Offset, NextOffset - Offset);
+				Collection.Add(new X509Certificate2(Encoding.UTF8.GetBytes(CertificateText)));
+
+				Offset = NextOffset;
+			}
+
+			return Collection;
+		}
+
+		/// <summary>
+		/// Tests connection to the given server
+		/// </summary>
+		/// <param name="Host">Host name</param>
+		/// <param name="Port">Port number to connect on</param>
+		/// <param name="Logger">Logger for diagnostic messages</param>
+		/// <returns>True if the connection was valid</returns>
+		static bool TestSslConnection(string Host, int Port, ILogger Logger)
+		{
+#pragma warning disable CA1031 // Do not catch general exception types
+			using (TcpClient Client = new TcpClient())
+			{
+				try
+				{
+					Client.Connect(Host, Port);
+					Logger.LogInformation("Successfully connected to {Host} on port {Port}", Host, Port);
+				}
+				catch (Exception Ex)
+				{
+					Logger.LogError(Ex, "Unable to connect to {Host} on port {Port}", Host, Port);
+					return false;
+				}
+
+				using (SslStream Stream = new SslStream(Client.GetStream()))
+				{
+					try
+					{
+						Stream.AuthenticateAsClient(Host);
+						Logger.LogInformation("Successfully authenticated host {Host}", Host);
+					}
+					catch (Exception Ex)
+					{
+						Logger.LogError(Ex, "Unable to authenticate as client");
+						return false;
+					}
+				}
+			}
+			return true;
+#pragma warning restore CA1031 // Do not catch general exception types
+		}
+
+		/// <summary>
+		/// Provides additional diagnostic information for SSL certificate validation
+		/// </summary>
+		/// <param name="Sender"></param>
+		/// <param name="Certificate"></param>
+		/// <param name="Chain"></param>
+		/// <param name="SslPolicyErrors"></param>
+		/// <returns>True if the certificate is allowed, false otherwise</returns>
+		bool CertificateValidationCallBack(object Sender, X509Certificate? Certificate, X509Chain? Chain, SslPolicyErrors SslPolicyErrors)
+		{
+			// If the certificate is a valid, signed certificate, return true.
+			if (SslPolicyErrors == SslPolicyErrors.None)
+			{
+				return true;
+			}
+
+			// Generate diagnostic information
+			StringBuilder Builder = new StringBuilder();
+			if (Sender != null)
+			{
+				string SenderInfo = StringUtils.Indent(Sender.ToString() ?? String.Empty, "    ");
+				Builder.Append($"\nSender:\n{SenderInfo}");
+			}
+			if (Certificate != null)
+			{
+				Builder.Append($"\nCertificate: {Certificate.Subject}");
+			}
+			if (Chain != null)
+			{
+				if (Chain.ChainStatus != null && Chain.ChainStatus.Length > 0)
+				{
+					Builder.Append("\nChain status:");
+					foreach (X509ChainStatus Status in Chain.ChainStatus)
+					{
+						Builder.Append($"\n  {Status.StatusInformation}");
+					}
+				}
+				if (Chain.ChainElements != null)
+				{
+					Builder.Append("\nChain elements:");
+					for (int Idx = 0; Idx < Chain.ChainElements.Count; Idx++)
+					{
+						X509ChainElement Element = Chain.ChainElements[Idx];
+						Builder.Append($"\n  {Idx,4} - Certificate: {Element.Certificate.Subject}");
+						if (Element.ChainElementStatus != null && Element.ChainElementStatus.Length > 0)
+						{
+							foreach (X509ChainStatus Status in Element.ChainElementStatus)
+							{
+								Builder.Append($"\n         Status: {Status.StatusInformation} ({Status.Status})");
+							}
+						}
+						if (!String.IsNullOrEmpty(Element.Information))
+						{
+							Builder.Append($"\n         Info: {Element.Information}");
+						}
+					}
+				}
+			}
+
+			// Print out additional diagnostic information
+			Logger.LogError("TLS certificate validation failed ({Errors}).{AdditionalInfo}", SslPolicyErrors, StringUtils.Indent(Builder.ToString(), "    "));
+			return false;
+		}
+
+		/// <summary>
+		/// Get a MongoDB collection from database
+		/// </summary>
+		/// <param name="Name">Name of collection</param>
+		/// <typeparam name="T">A MongoDB document</typeparam>
+		/// <returns></returns>
+		public IMongoCollection<T> GetCollection<T>(string Name)
+		{
+			return new MongoTracingCollection<T>(Database.GetCollection<T>(Name));
+		}
+
+		/// <summary>
+		/// Gets a singleton document by id
+		/// </summary>
+		/// <returns>The globals document</returns>
+		public Task<Globals> GetGlobalsAsync()
+		{
+			return GetSingletonAsync<Globals>(Globals.StaticId, () => new Globals() { SchemaVersion = UpgradeService.LatestSchemaVersion });
+		}
+
+		/// <summary>
+		/// Gets a singleton document by id
+		/// </summary>
+		/// <param name="Id">Id of the singleton document</param>
+		/// <returns>The document</returns>
+		public async Task<T> GetSingletonAsync<T>(ObjectId Id) where T : SingletonBase, new()
+		{
+			T Singleton = await GetSingletonAsync(Id, () => new T());
+			Singleton.PostLoad();
+			return Singleton;
+		}
+
+		/// <summary>
+		/// Gets a singleton document by id
+		/// </summary>
+		/// <param name="Id">Id of the singleton document</param>
+		/// <param name="Constructor">Method to use to construct a new object</param>
+		/// <returns>The document</returns>
+		public async Task<T> GetSingletonAsync<T>(ObjectId Id, Func<T> Constructor) where T : SingletonBase, new()
+		{
+			FilterDefinition<BsonDocument> Filter = new BsonDocument(new BsonElement("_id", Id));
+			for (; ; )
+			{
+				BsonDocument? Object = await Singletons.Find<BsonDocument>(Filter).FirstOrDefaultAsync();
+				if (Object != null)
+				{
+					T Item = BsonSerializer.Deserialize<T>(Object);
+					Item.PostLoad();
+					return Item;
+				}
+
+				T NewItem = Constructor();
+				NewItem.Id = Id;
+				await Singletons.InsertOneAsync(NewItem.ToBsonDocument());
+			}
+		}
+
+		/// <summary>
+		/// Updates a singleton
+		/// </summary>
+		/// <typeparam name="T"></typeparam>
+		/// <param name="Id"></param>
+		/// <param name="Updater"></param>
+		/// <returns></returns>
+		public async Task UpdateSingletonAsync<T>(ObjectId Id, Action<T> Updater) where T : SingletonBase, new()
+		{
+			for (; ; )
+			{
+				T Object = await GetSingletonAsync(Id, () => new T());
+				Updater(Object);
+				
+				if (await TryUpdateSingletonAsync(Object))
+				{
+					break;
+				}
+			}
+		}
+
+		/// <summary>
+		/// Attempts to update a singleton object
+		/// </summary>
+		/// <param name="SingletonObject">The singleton object</param>
+		/// <returns>True if the singleton document was updated</returns>
+		public async Task<bool> TryUpdateSingletonAsync<T>(T SingletonObject) where T : SingletonBase
+		{
+			int PrevRevision = SingletonObject.Revision++;
+
+			BsonDocument Filter = new BsonDocument { new BsonElement("_id", SingletonObject.Id), new BsonElement(nameof(SingletonBase.Revision), PrevRevision) };
+			try
+			{
+				ReplaceOneResult Result = await Singletons.ReplaceOneAsync(Filter, SingletonObject.ToBsonDocument(), new ReplaceOptions { IsUpsert = true });
+				return Result.MatchedCount > 0;
+			}
+			catch (MongoWriteException Ex)
+			{
+				// Duplicate key error occurs if filter fails to match because revision is not the same.
+				if (Ex.WriteError != null && Ex.WriteError.Category == ServerErrorCategory.DuplicateKey)
+				{
+					return false;
+				}
+				else
+				{
+					throw;
+				}
+			}
+		}
+	}
+}
diff --git a/Engine/Source/Programs/Horde/HordeServer/Server/DowntimeService.cs b/Engine/Source/Programs/Horde/HordeServer/Server/DowntimeService.cs
new file mode 100644
--- /dev/null
+++ b/Engine/Source/Programs/Horde/HordeServer/Server/DowntimeService.cs
@@ -0,0 +1,107 @@
+// Copyright Epic Games, Inc. All Rights Reserved.
+
+using Amazon.Runtime.Internal.Util;
+using HordeServer.Models;
+using HordeServer.Utilities;
+using Microsoft.Extensions.Logging;
+using Microsoft.Extensions.Options;
+using System;
+using System.Collections.Generic;
+using System.Linq;
+using System.Security.Cryptography.X509Certificates;
+using System.Threading;
+using System.Threading.Tasks;
+
+namespace HordeServer.Services
+{
+	/// <summary>
+	/// Interface for a service which keeps track of whether we're during downtime
+	/// </summary>
+	public interface IDowntimeService
+	{
+		/// <summary>
+		/// Returns true if downtime is currently active
+		/// </summary>
+		public bool IsDowntimeActive
+		{
+			get;
+		}
+	}
+
+	/// <summary>
+	/// Service which manages the downtime schedule
+	/// </summary>
+	public class DowntimeService : TickedBackgroundService, IDowntimeService
+	{
+		/// <summary>
+		/// Whether the server is currently in downtime
+		/// </summary>
+		public bool IsDowntimeActive
+		{
+			get;
+			private set;
+		}
+
+		DatabaseService DatabaseService;
+		IOptionsMonitor<ServerSettings> Settings;
+		ILogger<DowntimeService> Logger;
+
+		/// <summary>
+		/// Constructor
+		/// </summary>
+		/// <param name="DatabaseService">The database service instance</param>
+		/// <param name="Settings">The server settings</param>
+		/// <param name="Logger">Logger instance</param>
+		public DowntimeService(DatabaseService DatabaseService, IOptionsMonitor<ServerSettings> Settings, ILogger<DowntimeService> Logger)
+			: base(TimeSpan.FromMinutes(1.0), Logger)
+		{
+			this.DatabaseService = DatabaseService;
+			this.Settings = Settings;
+			this.Logger = Logger;
+
+			// Ensure the initial value to be correct
+			TickAsync(CancellationToken.None).Wait();
+		}
+
+		/// <summary>
+		/// Periodically called tick function
+		/// </summary>
+		/// <param name="StoppingToken">Token indicating that the service should stop</param>
+		/// <returns>Async task</returns>
+		protected override async Task TickAsync(CancellationToken StoppingToken)
+		{
+			Globals Globals = await DatabaseService.GetGlobalsAsync();
+
+			DateTimeOffset Now = TimeZoneInfo.ConvertTime(DateTimeOffset.Now, Settings.CurrentValue.TimeZoneInfo);
+			bool bIsActive = Globals.ScheduledDowntime.Any(x => x.IsActive(Now));
+
+			DateTimeOffset? Next = null;
+			foreach (ScheduledDowntime Schedule in Globals.ScheduledDowntime)
+			{
+				DateTimeOffset Start = Schedule.GetNext(Now).StartTime;
+				if(Next == null || Start < Next)
+				{
+					Next = Start;
+				}
+			}
+
+			if (Next != null)
+			{
+				Logger.LogInformation("Server time: {Time}. Downtime: {Downtime}. Next: {Next}.", Now, bIsActive, Next.Value);
+			}
+
+			if (bIsActive != IsDowntimeActive)
+			{
+				IsDowntimeActive = bIsActive;
+				if (IsDowntimeActive)
+				{
+					Logger.LogInformation("Entering downtime");
+				}
+				else
+				{
+					Logger.LogInformation("Leaving downtime");
+				}
+			}
+		}
+	}
+}
diff --git a/Engine/Source/Programs/Horde/HordeServer/Server/MetricService.cs b/Engine/Source/Programs/Horde/HordeServer/Server/MetricService.cs
new file mode 100644
--- /dev/null
+++ b/Engine/Source/Programs/Horde/HordeServer/Server/MetricService.cs
@@ -0,0 +1,77 @@
+// Copyright Epic Games, Inc. All Rights Reserved.
+
+using HordeServer.Utilities;
+using Microsoft.Extensions.Logging;
+using System;
+using System.Threading;
+using System.Threading.Tasks;
+using StatsdClient;
+
+namespace HordeServer.Services
+{
+	/// <summary>
+	/// Periodically send metrics for the CLR and other services that cannot be collected on a per-request basis
+	/// </summary>
+	public class MetricService : TickedBackgroundService
+	{
+		/// <summary>
+		/// DogStatsd instance
+		/// </summary>
+		private readonly IDogStatsd DogStatsd;
+
+		/// <summary>
+		/// Constructor
+		/// </summary>
+		/// <param name="DogStatsd"></param>
+		/// <param name="Logger">Logger instance</param>
+		public MetricService(IDogStatsd DogStatsd, ILogger<MetricService> Logger) : base(TimeSpan.FromSeconds(20.0), Logger)
+		{
+			this.DogStatsd = DogStatsd;
+		}
+
+		/// <summary>
+		/// Execute the background task
+		/// </summary>
+		/// <param name="StoppingToken">Cancellation token for the async task</param>
+		/// <returns>Async task</returns>
+		protected override Task TickAsync(CancellationToken StoppingToken)
+		{
+			ReportGcMetrics();
+			ReportThreadMetrics();
+			return Task.CompletedTask;
+		}
+
+		private void ReportThreadMetrics()
+		{
+			ThreadPool.GetMaxThreads(out int MaxWorkerThreads, out int MaxIoThreads);
+			ThreadPool.GetAvailableThreads(out int FreeWorkerThreads, out int FreeIoThreads);
+			ThreadPool.GetMinThreads(out int MinWorkerThreads, out int MinIoThreads);
+
+			int BusyIoThreads = MaxIoThreads - FreeIoThreads;
+			int BusyWorkerThreads = MaxWorkerThreads - FreeWorkerThreads;
+			
+			DogStatsd.Gauge("horde.clr.threadpool.io.max", MaxIoThreads);
+			DogStatsd.Gauge("horde.clr.threadpool.io.min", MinIoThreads);
+			DogStatsd.Gauge("horde.clr.threadpool.io.free", FreeIoThreads);
+			DogStatsd.Gauge("horde.clr.threadpool.io.busy", BusyIoThreads);
+			
+			DogStatsd.Gauge("horde.clr.threadpool.worker.max", MaxWorkerThreads);
+			DogStatsd.Gauge("horde.clr.threadpool.worker.min", MinWorkerThreads);
+			DogStatsd.Gauge("horde.clr.threadpool.worker.free", FreeWorkerThreads);
+			DogStatsd.Gauge("horde.clr.threadpool.worker.busy", BusyWorkerThreads);
+		}
+
+		private void ReportGcMetrics()
+		{
+			GCMemoryInfo GcMemoryInfo = GC.GetGCMemoryInfo();
+
+			DogStatsd.Gauge("horde.clr.gc.totalMemory", GC.GetTotalMemory(false));
+			DogStatsd.Gauge("horde.clr.gc.totalAllocated", GC.GetTotalAllocatedBytes());
+			DogStatsd.Gauge("horde.clr.gc.heapSize", GcMemoryInfo.HeapSizeBytes);
+			DogStatsd.Gauge("horde.clr.gc.fragmented", GcMemoryInfo.FragmentedBytes);
+			DogStatsd.Gauge("horde.clr.gc.memoryLoad", GcMemoryInfo.MemoryLoadBytes);
+			DogStatsd.Gauge("horde.clr.gc.totalAvailableMemory", GcMemoryInfo.TotalAvailableMemoryBytes);
+			DogStatsd.Gauge("horde.clr.gc.highMemoryLoadThreshold", GcMemoryInfo.HighMemoryLoadThresholdBytes);
+		}
+	}
+}
diff --git a/Engine/Source/Programs/Horde/HordeServer/Server/RequestTrackerService.cs b/Engine/Source/Programs/Horde/HordeServer/Server/RequestTrackerService.cs
new file mode 100644
--- /dev/null
+++ b/Engine/Source/Programs/Horde/HordeServer/Server/RequestTrackerService.cs
@@ -0,0 +1,173 @@
+﻿// Copyright Epic Games, Inc. All Rights Reserved.
+
+using System;
+using System.Collections.Concurrent;
+using System.Collections.Generic;
+using System.Linq;
+using System.Text;
+using System.Threading.Tasks;
+using Microsoft.AspNetCore.Http;
+using Microsoft.Extensions.Logging;
+
+namespace HordeServer.Utilities
+{
+	/// <summary>
+	/// Tracks open HTTP requests to ASP.NET
+	/// Will block a pending shutdown until all requests in progress are finished (or graceful timeout is reached)
+	/// This avoids interrupting long running requests such as artifact uploads.
+	/// </summary>
+	public class RequestTrackerService
+	{
+		/// <summary>
+		/// Writer for log output
+		/// </summary>
+		private readonly ILogger<RequestTrackerService> Logger;
+
+		readonly ConcurrentDictionary<string, TrackedRequest> RequestsInProgress = new ConcurrentDictionary<string, TrackedRequest>();
+
+		/// <summary>
+		/// Constructor
+		/// </summary>
+		/// <param name="Logger">Logger</param>
+		public RequestTrackerService(ILogger<RequestTrackerService> Logger)
+		{
+			this.Logger = Logger;
+		}
+
+		/// <summary>
+		/// Called by the middleware when a request is started
+		/// </summary>
+		/// <param name="Context">HTTP Context</param>
+		public void RequestStarted(HttpContext Context)
+		{
+			RequestsInProgress[Context.TraceIdentifier] = new TrackedRequest(Context.Request);
+		}
+		
+		/// <summary>
+		/// Called by the middleware when a request is finished (no matter if an exception occurred or not)
+		/// </summary>
+		/// <param name="Context">HTTP Context</param>
+		public void RequestFinished(HttpContext Context)
+		{
+			RequestsInProgress.Remove(Context.TraceIdentifier, out _);
+		}
+
+		/// <summary>
+		/// Get current requests in progress
+		/// </summary>
+		/// <returns>The requests in progress</returns>
+		public IReadOnlyDictionary<string, TrackedRequest> GetRequestsInProgress()
+		{
+			return RequestsInProgress;
+		}
+
+		private string GetRequestsInProgressAsString()
+		{
+			List<KeyValuePair<string, TrackedRequest>> Requests = GetRequestsInProgress().ToList();
+			Requests.Sort((A, B) => A.Value.StartedAt.CompareTo(B.Value.StartedAt));
+			StringBuilder Content = new StringBuilder();
+			foreach (KeyValuePair<string,TrackedRequest> Pair in Requests)
+			{
+				int AgeInMs = Pair.Value.GetTimeSinceStartInMs();
+				string Path = Pair.Value.Request.Path;
+				Content.AppendLine($"{AgeInMs,9}  {Path}");
+			}
+
+			return Content.ToString();
+		}
+
+		/// <summary>
+		/// Log all requests currently in progress
+		/// </summary>
+		public void LogRequestsInProgress()
+		{
+			if (GetRequestsInProgress().Count == 0)
+			{
+				Logger.LogInformation("There are no requests in progress!");
+			}
+			else
+			{
+				Logger.LogInformation("Current open requests are:\n{RequestsInProgress}", GetRequestsInProgressAsString());
+			}
+		}
+	}
+	
+	/// <summary>
+	/// Value object for tracked requests
+	/// </summary>
+	public class TrackedRequest
+	{
+		/// <summary>
+		/// When the request was received
+		/// </summary>
+		public DateTime StartedAt { get; }
+			
+		/// <summary>
+		/// The HTTP request being tracked
+		/// </summary>
+		public HttpRequest Request { get; }
+
+		/// <summary>
+		/// Constructor
+		/// </summary>
+		/// <param name="Request">HTTP request being tracked</param>
+		public TrackedRequest(HttpRequest Request)
+		{
+			StartedAt = DateTime.UtcNow;
+			this.Request = Request;
+		}
+
+		/// <summary>
+		/// How long the request has been running
+		/// </summary>
+		/// <returns>Time elapsed in milliseconds since request started</returns>
+		public int GetTimeSinceStartInMs()
+		{
+			return (int) (DateTime.UtcNow - StartedAt).TotalMilliseconds;
+		}
+	}
+	
+	/// <summary>
+	/// ASP.NET Middleware to track open requests
+	/// </summary>
+	public class RequestTrackerMiddleware
+	{
+		private readonly RequestDelegate Next;
+	
+		/// <summary>
+		/// Constructor
+		/// </summary>
+		/// <param name="Next">Next middleware to call</param>
+		public RequestTrackerMiddleware(RequestDelegate Next)
+		{
+			this.Next = Next;
+		}
+	
+		/// <summary>
+		/// Invoked by ASP.NET framework itself
+		/// </summary>
+		/// <param name="Context">HTTP Context</param>
+		/// <param name="Service">The RequestTrackerService singleton</param>
+		/// <returns></returns>
+		public async Task Invoke(HttpContext Context, RequestTrackerService Service)
+		{
+			if (!Context.Request.Path.StartsWithSegments("/health", StringComparison.Ordinal))
+			{
+				try
+				{
+					Service.RequestStarted(Context);
+					await Next(Context);
+				}
+				finally
+				{
+					Service.RequestFinished(Context);
+				}
+			}
+			else
+			{
+				// Ignore requests to /health/*
+				await Next(Context);
+			}
+		}
+	}
+}
\ No newline at end of file
diff --git a/Engine/Source/Programs/Horde/HordeServer/Users/Impl/ServiceAccountCollection.cs b/Engine/Source/Programs/Horde/HordeServer/Users/Impl/ServiceAccountCollection.cs
new file mode 100644
--- /dev/null
+++ b/Engine/Source/Programs/Horde/HordeServer/Users/Impl/ServiceAccountCollection.cs
@@ -0,0 +1,149 @@
+// Copyright Epic Games, Inc. All Rights Reserved.
+
+using System;
+using HordeServer.Models;
+using HordeServer.Services;
+using MongoDB.Bson;
+using MongoDB.Bson.Serialization.Attributes;
+using MongoDB.Driver;
+using System.Collections.Generic;
+using System.Linq;
+using System.Threading.Tasks;
+
+namespace HordeServer.Collections.Impl
+{
+	/// <summary>
+	/// Collection of service account documents
+	/// </summary>
+	public class ServiceAccountCollection : IServiceAccountCollection
+	{
+		/// <summary>
+		/// Concrete implementation of IServiceAccount
+		/// </summary>
+		class ServiceAccountDocument : IServiceAccount
+		{
+			public const string ClaimSeparator = "###";
+			
+			[BsonRequired, BsonId]
+			public ObjectId Id { get; set; }
+
+			[BsonRequired]
+			public string SecretToken { get; set; } = "empty";
+
+			[BsonRequired]
+			public List<string> Claims { get; set; } = new List<string>();
+			
+			public bool Enabled { get; set; }
+			public string Description { get; set; } = "empty";
+
+			[BsonConstructor]
+			private ServiceAccountDocument()
+			{
+			}
+
+			public ServiceAccountDocument(ObjectId Id, string SecretToken, List<string> Claims, bool Enabled, string Description)
+			{
+				this.Id = Id;
+				this.SecretToken = SecretToken;
+				this.Claims = Claims;
+				this.Enabled = Enabled;
+				this.Description = Description;
+			}
+			
+			/// <inheritdoc/>
+			public void AddClaim(string Type, string Value)
+			{
+				Claims.Add(Type + ClaimSeparator + Value);
+			}
+
+			/// <inheritdoc/>
+			public IReadOnlyList<(string Type, string Value)> GetClaims()
+			{
+				return Claims.Select(x =>
+				{
+					string[] Split = x.Split(ClaimSeparator);
+					return (Split[0], Split[1]);
+				}).ToList();
+			}
+
+			protected bool Equals(ServiceAccountDocument other)
+			{
+				bool AreClaimsEqual = !Claims.Except(other.Claims).Any();
+
+                return Id.Equals(other.Id) && SecretToken == other.SecretToken && AreClaimsEqual && Enabled == other.Enabled && Description == other.Description;
+			}
+
+			public override bool Equals(object? obj)
+			{
+				if (ReferenceEquals(null, obj)) return false;
+				if (ReferenceEquals(this, obj)) return true;
+				if (obj.GetType() != this.GetType()) return false;
+				return Equals((ServiceAccountDocument) obj);
+			}
+
+			public override int GetHashCode()
+			{
+				return HashCode.Combine(Id, SecretToken, Claims, Enabled, Description);
+			}
+		}
+
+		/// <summary>
+		/// Collection of session documents
+		/// </summary>
+		private readonly IMongoCollection<ServiceAccountDocument> ServiceAccounts;
+
+		/// <summary>
+		/// Constructor
+		/// </summary>
+		/// <param name="DatabaseService">The database service</param>
+		public ServiceAccountCollection(DatabaseService DatabaseService)
+		{
+			ServiceAccounts = DatabaseService.GetCollection<ServiceAccountDocument>("ServiceAccounts");
+
+			if (!DatabaseService.ReadOnlyMode)
+			{
+				ServiceAccounts.Indexes.CreateOne(new CreateIndexModel<ServiceAccountDocument>(Builders<ServiceAccountDocument>.IndexKeys.Ascending(x => x.SecretToken)));
+			}
+		}
+
+		/// <inheritdoc/>
+		public async Task<IServiceAccount> AddAsync(string SecretToken, List<string> Claims, string Description)
+		{
+			ServiceAccountDocument NewSession = new ServiceAccountDocument(ObjectId.GenerateNewId(), SecretToken, Claims, true, Description);
+			await ServiceAccounts.InsertOneAsync(NewSession);
+			return NewSession;
+		}
+
+		/// <inheritdoc/>
+		public async Task<IServiceAccount?> GetAsync(ObjectId Id)
+		{
+			return await ServiceAccounts.Find(x => x.Id == Id).FirstOrDefaultAsync();
+		}
+
+		/// <inheritdoc/>
+		public async Task<IServiceAccount?> GetBySecretTokenAsync(string SecretToken)
+		{
+			return await ServiceAccounts.Find(x => x.SecretToken == SecretToken).FirstOrDefaultAsync();
+		}
+
+		/// <inheritdoc/>
+		public Task UpdateAsync(ObjectId Id, string? SecretToken, List<string>? Claims, bool? Enabled, string? Description)
+		{
+			UpdateDefinitionBuilder<ServiceAccountDocument> Update = Builders<ServiceAccountDocument>.Update;
+			List<UpdateDefinition<ServiceAccountDocument>> Updates = new List<UpdateDefinition<ServiceAccountDocument>>();
+			
+			if (SecretToken != null) Updates.Add(Update.Set(x => x.SecretToken, SecretToken));
+			if (Claims != null) Updates.Add(Update.Set(x => x.Claims, Claims));
+			if (Enabled != null) Updates.Add(Update.Set(x => x.Enabled, Enabled));
+			if (Description != null) Updates.Add(Update.Set(x => x.Description, Description));
+
+			return ServiceAccounts.FindOneAndUpdateAsync(x => x.Id == Id, Update.Combine(Updates));
+		}
+
+		/// <inheritdoc/>
+		public Task DeleteAsync(ObjectId SessionId)
+		{
+			return ServiceAccounts.DeleteOneAsync(x => x.Id == SessionId);
+		}
+	}
+}
