Horde: More renames.

diff --git a/Engine/Source/Programs/Horde/HordeServer.Compute.Tests/TestSetup.cs b/Engine/Source/Programs/Horde/HordeServer.Compute.Tests/TestSetup.cs
--- a/Engine/Source/Programs/Horde/HordeServer.Compute.Tests/TestSetup.cs
+++ b/Engine/Source/Programs/Horde/HordeServer.Compute.Tests/TestSetup.cs
@@ -158,7 +158,6 @@ namespace HordeServerTests
 			Services.AddSingleton<IProjectCollection, ProjectCollection>();
 			Services.AddSingleton<ISessionCollection, SessionCollection>();
 			Services.AddSingleton<ISubscriptionCollection, SubscriptionCollection>();
-			Services.AddSingleton<ISoftwareCollection, SoftwareCollection>();
 			Services.AddSingleton<IStreamCollection, StreamCollection>();
 			Services.AddSingleton<ITemplateCollection, TemplateCollection>();
 			Services.AddSingleton<ITestDataCollection, TestDataCollection>();
@@ -196,7 +195,6 @@ namespace HordeServerTests
 			Services.AddSingleton<ProjectService>();
 			Services.AddSingleton<RpcService>();
 			Services.AddSingleton<ScheduleService>();
-			Services.AddSingleton<SoftwareService>();
 			Services.AddSingleton<StreamService>();
 			Services.AddSingleton<TemplateService>();
 			Services.AddSingleton<UpgradeService>();
diff --git a/Engine/Source/Programs/Horde/HordeServer/Acls/Acl.cs b/Engine/Source/Programs/Horde/HordeServer/Acls/Acl.cs
new file mode 100644
--- /dev/null
+++ b/Engine/Source/Programs/Horde/HordeServer/Acls/Acl.cs
@@ -0,0 +1,805 @@
+// Copyright Epic Games, Inc. All Rights Reserved.
+
+using HordeServer.Api;
+using HordeServer.Services;
+using HordeServer.Utilities;
+using MongoDB.Bson;
+using MongoDB.Bson.Serialization;
+using MongoDB.Bson.Serialization.Attributes;
+using MongoDB.Driver;
+using System;
+using System.Collections.Generic;
+using System.Linq;
+using System.Linq.Expressions;
+using System.Security.Claims;
+
+namespace HordeServer.Models
+{
+	/// <summary>
+	/// Set of actions that can be performed by a user. NOTE: This enum is sensitive to ordering. Do not change values.
+	/// </summary>
+	public enum AclAction
+	{
+		//// PROJECTS ////
+
+		/// <summary>
+		/// Allows the creation of new projects
+		/// </summary>
+		CreateProject,
+
+		/// <summary>
+		/// Allows deletion of projects.
+		/// </summary>
+		DeleteProject,
+		
+		/// <summary>
+		/// Modify attributes of a project (name, categories, etc...)
+		/// </summary>
+		UpdateProject,
+
+		/// <summary>
+		/// View information about a project
+		/// </summary>
+		ViewProject,
+
+
+
+
+		//// STREAMS ////
+
+		/// <summary>
+		/// Allows the creation of new streams within a project
+		/// </summary>
+		CreateStream,
+
+		/// <summary>
+		/// Allows updating a stream (agent types, templates, schedules)
+		/// </summary>
+		UpdateStream,
+
+		/// <summary>
+		/// Allows deleting a stream
+		/// </summary>
+		DeleteStream,
+
+		/// <summary>
+		/// Ability to view a stream
+		/// </summary>
+		ViewStream,
+
+		/// <summary>
+		/// View changes submitted to a stream. NOTE: this returns responses from the server's Perforce account, which may be a priviledged user.
+		/// </summary>
+		ViewChanges,
+
+		/// <summary>
+		/// Update records in commit queues for a stream
+		/// </summary>
+		UpdateCommitQueues,
+
+		/// <summary>
+		/// View commit queues for a stream
+		/// </summary>
+		ViewCommitQueues,
+
+
+
+
+
+
+
+
+
+
+		//// JOBS ////
+
+		/// <summary>
+		/// Ability to start new jobs
+		/// </summary>
+		CreateJob,
+
+		/// <summary>
+		/// Rename a job, modify its priority, etc...
+		/// </summary>
+		UpdateJob,
+
+		/// <summary>
+		/// Delete a job properties
+		/// </summary>
+		DeleteJob,
+
+		/// <summary>
+		/// Allows updating a job metadata (name, changelist number, step properties, new groups, job states, etc...). Typically granted to agents. Not user facing.
+		/// </summary>
+		ExecuteJob,
+
+		/// <summary>
+		/// Ability to retry a failed job step
+		/// </summary>
+		RetryJobStep,
+
+		/// <summary>
+		/// Ability to view a job
+		/// </summary>
+		ViewJob,
+
+
+
+
+
+
+		//// EVENTS ////
+		
+		/// <summary>
+		/// Ability to create events
+		/// </summary>
+		CreateEvent,
+
+		/// <summary>
+		/// Ability to view events
+		/// </summary>
+		ViewEvent,
+
+
+
+		//// AGENTS ////
+
+		/// <summary>
+		/// Ability to create an agent. This may be done explicitly, or granted to agents to allow them to self-register.
+		/// </summary>
+		CreateAgent,
+
+		/// <summary>
+		/// Update an agent's name, pools, etc...
+		/// </summary>
+		UpdateAgent,
+
+		/// <summary>
+		/// Soft-delete an agent
+		/// </summary>
+		DeleteAgent,
+
+		/// <summary>
+		/// View an available agents
+		/// </summary>
+		ViewAgent,
+
+		/// <summary>
+		/// List the available agents
+		/// </summary>
+		ListAgents,
+
+
+
+
+		//// POOLS ////
+
+		/// <summary>
+		/// Create a global pool of agents
+		/// </summary>
+		CreatePool,
+
+		/// <summary>
+		/// Modify an agent pool
+		/// </summary>
+		UpdatePool,
+
+		/// <summary>
+		/// Delete an agent pool
+		/// </summary>
+		DeletePool,
+
+		/// <summary>
+		/// Ability to view a pool
+		/// </summary>
+		ViewPool,
+
+		/// <summary>
+		/// View all the available agent pools
+		/// </summary>
+		ListPools,
+
+
+
+
+		//// SESSIONS ////
+
+		/// <summary>
+		/// Granted to agents to call CreateSession, which returns a bearer token identifying themselves valid to call UpdateSesssion via gRPC.
+		/// </summary>
+		CreateSession,
+
+		/// <summary>
+		/// Allows viewing information about an agent session
+		/// </summary>
+		ViewSession,
+
+
+
+
+		//// CREDENTIALS ////
+
+		/// <summary>
+		/// Create a new credential
+		/// </summary>
+		CreateCredential,
+
+		/// <summary>
+		/// Delete a credential
+		/// </summary>
+		DeleteCredential,
+
+		/// <summary>
+		/// Modify an existing credential
+		/// </summary>
+		UpdateCredential,
+
+		/// <summary>
+		/// Enumerates all the available credentials
+		/// </summary>
+		ListCredentials,
+
+		/// <summary>
+		/// View a credential
+		/// </summary>
+		ViewCredential,
+
+
+
+
+
+
+
+		//// LEASES ////
+
+		/// <summary>
+		/// View all the leases that an agent has worked on
+		/// </summary>
+		ViewLeases,
+
+
+
+
+
+		//// COMMITS ////
+
+		/// <summary>
+		/// Add a new commit
+		/// </summary>
+		AddCommit,
+
+		/// <summary>
+		/// List all the commits that have been added
+		/// </summary>
+		FindCommits,
+
+		/// <summary>
+		/// View the commits for a particular stream
+		/// </summary>
+		ViewCommits,
+
+
+
+
+
+
+		/// <summary>
+		/// View a build health issue
+		/// </summary>
+		ViewIssue,
+
+
+
+
+
+		//// TEMPLATES ////
+
+		/// <summary>
+		/// View template associated with a stream
+		/// </summary>
+		ViewTemplate,
+
+
+
+
+
+
+		//// LOGS ////
+
+		/// <summary>
+		/// Ability to create a log. Implicitly granted to agents.
+		/// </summary>
+		CreateLog,
+
+		/// <summary>
+		/// Ability to update log metadata
+		/// </summary>
+		UpdateLog,
+
+		/// <summary>
+		/// Ability to view a log contents
+		/// </summary>
+		ViewLog,
+
+		/// <summary>
+		/// Ability to write log data
+		/// </summary>
+		WriteLogData,
+
+
+
+
+		//// ARTIFACTS ////
+
+		/// <summary>
+		/// Ability to create an artifact. Typically just for debugging; agents have this access for a particular session.
+		/// </summary>
+		UploadArtifact,
+
+		/// <summary>
+		/// Ability to download an artifact
+		/// </summary>
+		DownloadArtifact,
+
+
+
+
+		//// SOFTWARE ////
+
+		/// <summary>
+		/// Ability to upload new versions of the agent software
+		/// </summary>
+		UploadSoftware,
+
+		/// <summary>
+		/// Ability to download the agent software
+		/// </summary>
+		DownloadSoftware,
+
+		/// <summary>
+		/// Ability to delete agent software
+		/// </summary>
+		DeleteSoftware,
+
+
+
+
+		//// ADMIN ////
+
+		/// <summary>
+		/// Ability to read any data from the server. Always inherited.
+		/// </summary>
+		AdminRead,
+
+		/// <summary>
+		/// Ability to write any data to the server.
+		/// </summary>
+		AdminWrite,
+
+		/// <summary>
+		/// Ability to impersonate another user
+		/// </summary>
+		Impersonate,
+
+
+
+
+
+		//// PERMISSIONS ////
+		
+		/// <summary>
+		/// Ability to view permissions on an object
+		/// </summary>
+		ViewPermissions,
+
+		/// <summary>
+		/// Ability to change permissions on an object
+		/// </summary>
+		ChangePermissions,
+
+		/// <summary>
+		/// Issue bearer token for the current user
+		/// </summary>
+		IssueBearerToken,
+
+
+
+
+
+		//// NOTIFICATIONS ////
+
+		/// <summary>
+		/// Ability to subscribe to notifications
+		/// </summary>
+		CreateSubscription,
+
+
+
+
+
+		//// DEVICES ////
+		
+		/// <summary>
+		/// Ability to read devices
+		/// </summary>
+		DeviceRead,
+
+		/// <summary>
+		/// Ability to write devices
+		/// </summary>
+		DeviceWrite,
+
+
+
+
+
+		//// STORAGE ////
+
+		/// <summary>
+		/// Ability to read blobs from the storage service
+		/// </summary>
+		ReadBlobs,
+
+		/// <summary>
+		/// Ability to write blobs to the storage service
+		/// </summary>
+		WriteBlobs,
+
+		/// <summary>
+		/// Ability to read refs from the storage service
+		/// </summary>
+		ReadRefs,
+
+		/// <summary>
+		/// Ability to write refs to the storage service
+		/// </summary>
+		WriteRefs, 
+
+		/// <summary>
+		/// Ability to delete refs
+		/// </summary>
+		DeleteRefs,
+	}
+
+	/// <summary>
+	/// Stores information about a claim
+	/// </summary>
+	public class AclClaim
+	{
+		/// <summary>
+		/// The claim type, typically a URI
+		/// </summary>
+		public string Type { get; set; }
+
+		/// <summary>
+		/// The claim value
+		/// </summary>
+		public string Value { get; set; }
+
+		/// <summary>
+		/// Constructor
+		/// </summary>
+		/// <param name="Claim">The claim object</param>
+		public AclClaim(Claim Claim)
+			: this(Claim.Type, Claim.Value)
+		{
+		}
+
+		/// <summary>
+		/// Constructor
+		/// </summary>
+		/// <param name="Type">The claim type</param>
+		/// <param name="Value">The claim value</param>
+		public AclClaim(string Type, string Value)
+		{
+			this.Type = Type;
+			this.Value = Value;
+		}
+
+		/// <summary>
+		/// Constructs a claim from a request object
+		/// </summary>
+		/// <param name="Request">The request object</param>
+		public AclClaim(CreateAclClaimRequest Request)
+			: this(Request.Type, Request.Value)
+		{
+		}
+	}
+
+	/// <summary>
+	/// Describes an entry in the ACL for a particular claim
+	/// </summary>
+	public class AclEntry
+	{
+		/// <summary>
+		/// Claim for this entry
+		/// </summary>
+		public AclClaim Claim { get; set; }
+
+		/// <summary>
+		/// List of allowed operations
+		/// </summary>
+		[BsonSerializer(typeof(AclActionSetSerializer))]
+		public HashSet<AclAction> Actions { get; set; }
+
+		/// <summary>
+		/// Private constructor for serialization
+		/// </summary>
+		[BsonConstructor]
+		private AclEntry()
+		{
+			Claim = new AclClaim(String.Empty, String.Empty);
+			Actions = new HashSet<AclAction>();
+		}
+
+		/// <summary>
+		/// Constructor
+		/// </summary>
+		/// <param name="Claim">The claim this entry applies to</param>
+		/// <param name="Actions">List of allowed operations</param>
+		internal AclEntry(AclClaim Claim, IEnumerable<AclAction> Actions)
+		{
+			this.Claim = Claim;
+			this.Actions = new HashSet<AclAction>(Actions);
+		}
+
+		/// <summary>
+		/// Constructs an ACL entry from a request
+		/// </summary>
+		/// <param name="Request">Request instance</param>
+		/// <returns>New ACL entry</returns>
+		public static AclEntry FromRequest(CreateAclEntryRequest Request)
+		{
+			return new AclEntry(new AclClaim(Request.Claim.Type, Request.Claim.Value), ParseActionNames(Request.Actions).ToArray());
+		}
+
+		/// <summary>
+		/// Parses a list of names as an AclAction bitmask
+		/// </summary>
+		/// <param name="ActionNames">Array of names</param>
+		/// <returns>Action bitmask</returns>
+		public static List<AclAction> ParseActionNames(string[]? ActionNames)
+		{
+			List<AclAction> Actions = new List<AclAction>();
+			if (ActionNames != null)
+			{
+				foreach (string Name in ActionNames)
+				{
+					Actions.Add(Enum.Parse<AclAction>(Name, true));
+				}
+			}
+			return Actions;
+		}
+
+		/// <summary>
+		/// Build a list of action names from an array of flags
+		/// </summary>
+		/// <param name="Actions"></param>
+		/// <returns></returns>
+		public static List<string> GetActionNames(IEnumerable<AclAction> Actions)
+		{
+			List<string> ActionNames = new List<string>();
+			foreach(AclAction Action in Actions)
+			{
+				string Name = Enum.GetName(typeof(AclAction), Action)!;
+				ActionNames.Add(Name);
+			}
+			return ActionNames;
+		}
+	}
+
+	/// <summary>
+	/// Represents an access control list for an object in the database
+	/// </summary>
+	public class Acl
+	{
+		/// <summary>
+		/// List of entries for this ACL
+		/// </summary>
+		public List<AclEntry> Entries { get; set; }
+
+		/// <summary>
+		/// Whether to inherit permissions from the parent ACL by default
+		/// </summary>
+		public bool Inherit { get; set; } = true;
+
+		/// <summary>
+		/// Specifies a list of exceptions to the inheritance setting
+		/// </summary>
+		public List<AclAction>? Exceptions { get; set; }
+
+
+		/// <summary>
+		/// Default constructor
+		/// </summary>
+		public Acl()
+		{
+			Entries = new List<AclEntry>();
+			Inherit = true;
+		}
+
+		/// <summary>
+		/// Constructor
+		/// </summary>
+		/// <param name="Entries">List of entries for this ACL</param>
+		/// <param name="bInheritPermissions">Whether to inherit permissions from the parent object by default</param>
+		public Acl(List<AclEntry> Entries, bool bInheritPermissions)
+		{
+			this.Entries = Entries;
+			this.Inherit = bInheritPermissions;
+		}
+
+		/// <summary>
+		/// Constructor
+		/// </summary>
+		/// <param name="Entries">List of entries for this ACL</param>
+		public Acl(params AclEntry[] Entries)
+			: this(Entries.ToList(), true)
+		{
+		}
+
+		/// <summary>
+		/// Tests whether a user is authorized to perform the given actions
+		/// </summary>
+		/// <param name="Action">Action that is being performed. This should be a single flag.</param>
+		/// <param name="User">The principal to authorize</param>
+		/// <returns>True/false if the action is allowed or denied, null if there is no specific setting for this user</returns>
+		public bool? Authorize(AclAction Action, ClaimsPrincipal User)
+		{
+			// Check if there's a specific entry for this action
+			foreach (AclEntry Entry in Entries)
+			{
+				if(Entry.Actions.Contains(Action) && User.HasClaim(Entry.Claim.Type, Entry.Claim.Value))
+				{
+					return true;
+				}
+			}
+
+			// Otherwise check if we're prevented from inheriting permissions
+			if(Inherit)
+			{
+				if(Exceptions != null && Exceptions.Contains(Action))
+				{
+					return false;
+				}
+			}
+			else
+			{
+				if (Exceptions == null || !Exceptions.Contains(Action))
+				{
+					return false;
+				}
+			}
+
+			// Otherwise allow to propagate up the hierarchy
+			return null;
+		}
+
+		/// <summary>
+		/// Merge new settings into an ACL
+		/// </summary>
+		/// <param name="BaseAcl">The current acl</param>
+		/// <param name="Update">The update to apply</param>
+		/// <returns>The new ACL value. Null if the ACL has all default settings.</returns>
+		public static Acl? Merge(Acl? BaseAcl, UpdateAclRequest? Update)
+		{
+			Acl? NewAcl = null;
+			if(Update != null)
+			{
+				NewAcl = new Acl();
+
+				if(Update.Entries != null)
+				{
+					NewAcl.Entries = Update.Entries.ConvertAll(x => AclEntry.FromRequest(x));
+				}
+				else if(BaseAcl != null)
+				{
+					NewAcl.Entries = BaseAcl.Entries;
+				}
+
+				if(Update.Inherit != null)
+				{
+					NewAcl.Inherit = Update.Inherit.Value;
+				}
+				else if(BaseAcl != null)
+				{
+					NewAcl.Inherit = BaseAcl.Inherit;
+				}
+			}
+			return NewAcl;
+		}
+
+		/// <summary>
+		/// Creates an update definition for the given ACL. Clears the ACL property if it's null.
+		/// </summary>
+		/// <typeparam name="T">Type of document containing the ACL</typeparam>
+		/// <param name="Field">Selector for the ACL property</param>
+		/// <param name="NewAcl">The new ACL value</param>
+		/// <returns>Update definition for the document</returns>
+		public static UpdateDefinition<T> CreateUpdate<T>(Expression<Func<T, object>> Field, Acl NewAcl)
+		{
+			if (NewAcl.Entries.Count == 0 && NewAcl.Inherit)
+			{
+				return Builders<T>.Update.Unset(Field);
+			}
+			else
+			{
+				return Builders<T>.Update.Set(Field, NewAcl);
+			}
+		}
+	}
+
+	/// <summary>
+	/// Serializer for JobStepRefId objects
+	/// </summary>
+	public sealed class AclActionSetSerializer : IBsonSerializer<HashSet<AclAction>>
+	{
+		/// <inheritdoc/>
+		public Type ValueType
+		{
+			get { return typeof(HashSet<AclAction>); }
+		}
+
+		/// <inheritdoc/>
+		void IBsonSerializer.Serialize(BsonSerializationContext Context, BsonSerializationArgs Args, object Value)
+		{
+			Serialize(Context, Args, (HashSet<AclAction>)Value);
+		}
+
+		/// <inheritdoc/>
+		object IBsonSerializer.Deserialize(BsonDeserializationContext Context, BsonDeserializationArgs Args)
+		{
+			return ((IBsonSerializer<HashSet<AclAction>>)this).Deserialize(Context, Args);
+		}
+
+		/// <inheritdoc/>
+		public HashSet<AclAction> Deserialize(BsonDeserializationContext Context, BsonDeserializationArgs Args)
+		{
+			HashSet<AclAction> Values = new HashSet<AclAction>();
+
+			Context.Reader.ReadStartArray();
+			for(; ;)
+			{
+				BsonType Type = Context.Reader.ReadBsonType();
+				if(Type == BsonType.EndOfDocument)
+				{
+					break;
+				}
+				else if (Type == BsonType.Int32)
+				{
+					Values.Add((AclAction)Context.Reader.ReadInt32());
+				}
+				else
+				{
+					Values.Add((AclAction)Enum.Parse(typeof(AclAction), Context.Reader.ReadString()));
+				}
+			}
+			Context.Reader.ReadEndArray();
+
+			return Values;
+		}
+
+		/// <inheritdoc/>
+		public void Serialize(BsonSerializationContext Context, BsonSerializationArgs Args, HashSet<AclAction> Values)
+		{
+			Context.Writer.WriteStartArray();
+			foreach (AclAction Value in Values)
+			{
+				Context.Writer.WriteString(Value.ToString());
+			}
+			Context.Writer.WriteEndArray();
+		}
+	}
+
+	/// <summary>
+	/// Represents an access control list for an object in the database
+	/// </summary>
+	[SingletonDocument("5e39b51941b875626e600193")]
+	public class GlobalPermissions : SingletonBase
+	{
+		/// <summary>
+		/// The global ACL object
+		/// </summary>
+		public Acl Acl { get; set; } = new Acl();
+	}
+}
diff --git a/Engine/Source/Programs/Horde/HordeServer/Acls/AclService.cs b/Engine/Source/Programs/Horde/HordeServer/Acls/AclService.cs
new file mode 100644
--- /dev/null
+++ b/Engine/Source/Programs/Horde/HordeServer/Acls/AclService.cs
@@ -0,0 +1,269 @@
+// Copyright Epic Games, Inc. All Rights Reserved.
+
+using HordeServer.Api;
+using HordeServer.Collections;
+using HordeServer.Models;
+using HordeServer.Utilities;
+using Microsoft.AspNetCore.Authentication.OAuth.Claims;
+using Microsoft.Extensions.Options;
+using Microsoft.IdentityModel.Tokens;
+using MongoDB.Bson;
+using MongoDB.Driver;
+using System;
+using System.Collections.Generic;
+using System.IdentityModel.Tokens.Jwt;
+using System.Linq;
+using System.Linq.Expressions;
+using System.Security.Claims;
+using System.Security.Cryptography;
+using System.Threading.Tasks;
+
+namespace HordeServer.Services
+{
+	/// <summary>
+	/// Cache of global ACL
+	/// </summary>
+	public class GlobalPermissionsCache
+	{
+		/// <summary>
+		/// The global permissions list
+		/// </summary>
+		public GlobalPermissions? GlobalPermissions { get; set; }
+	}
+
+	/// <summary>
+	/// Wraps functionality for manipulating permissions
+	/// </summary>
+	public class AclService
+	{
+		/// <summary>
+		/// Name of the role that can be used to administer agents
+		/// </summary>
+		public static AclClaim AgentRegistrationClaim { get; } = new AclClaim(HordeClaimTypes.Role, "agent-registration");
+
+		/// <summary>
+		/// Name of the role used to upload software
+		/// </summary>
+		public static AclClaim DownloadSoftwareClaim { get; } = new AclClaim(HordeClaimTypes.Role, "download-software");
+
+		/// <summary>
+		/// Name of the role used to upload software
+		/// </summary>
+		public static AclClaim UploadSoftwareClaim { get; } = new AclClaim(HordeClaimTypes.Role, "upload-software");
+
+		/// <summary>
+		/// Name of the role used to upload software
+		/// </summary>
+		public static AclClaim ConfigureProjectsClaim { get; } = new AclClaim(HordeClaimTypes.Role, "configure-projects");
+
+		/// <summary>
+		/// Name of the role used to upload software
+		/// </summary>
+		public static AclClaim StartChainedJobClaim { get; } = new AclClaim(HordeClaimTypes.Role, "start-chained-job");
+
+		/// <summary>
+		/// Role for all agents
+		/// </summary>
+		public static AclClaim AgentClaim { get; } = new AclClaim(HordeClaimTypes.Role, "agent");
+
+		/// <summary>
+		/// The default permissions
+		/// </summary>
+		Acl DefaultAcl = new Acl();
+
+		/// <summary>
+		/// The database service
+		/// </summary>
+		DatabaseService DatabaseService;
+
+		/// <summary>
+		/// The global permissions singleton
+		/// </summary>
+		ISingletonDocument<GlobalPermissions> GlobalPermissions;
+
+		/// <summary>
+		/// The settings object
+		/// </summary>
+		ServerSettings Settings;
+
+		/// <summary>
+		/// Constructor
+		/// </summary>
+		/// <param name="DatabaseService">The database service instance</param>
+		/// <param name="GlobalPermissions">The global permissions instance</param>
+		/// <param name="Settings">The settings object</param>
+		public AclService(DatabaseService DatabaseService, ISingletonDocument<GlobalPermissions> GlobalPermissions, IOptionsMonitor<ServerSettings> Settings)
+		{
+			this.DatabaseService = DatabaseService;
+			this.GlobalPermissions = GlobalPermissions;
+			this.Settings = Settings.CurrentValue;
+
+			List<AclAction> AdminActions = new List<AclAction>();
+			foreach (AclAction? Action in Enum.GetValues(typeof(AclAction)))
+			{
+				AdminActions.Add(Action!.Value);
+			}
+
+			DefaultAcl = new Acl();
+			DefaultAcl.Entries.Add(new AclEntry(new AclClaim(ClaimTypes.Role, "internal:AgentRegistration"), new[] { AclAction.CreateAgent, AclAction.CreateSession }));
+			DefaultAcl.Entries.Add(new AclEntry(AgentRegistrationClaim, new[] { AclAction.CreateAgent, AclAction.CreateSession, AclAction.UpdateAgent, AclAction.DownloadSoftware }));
+			DefaultAcl.Entries.Add(new AclEntry(AgentClaim, new[] { AclAction.ViewProject, AclAction.ViewStream, AclAction.CreateEvent, AclAction.DownloadSoftware }));
+			DefaultAcl.Entries.Add(new AclEntry(DownloadSoftwareClaim, new[] { AclAction.DownloadSoftware }));
+			DefaultAcl.Entries.Add(new AclEntry(UploadSoftwareClaim, new[] { AclAction.UploadSoftware }));
+			DefaultAcl.Entries.Add(new AclEntry(ConfigureProjectsClaim, new[] { AclAction.CreateProject, AclAction.UpdateProject, AclAction.ViewProject, AclAction.CreateStream, AclAction.UpdateStream, AclAction.ViewStream, AclAction.ChangePermissions }));
+			DefaultAcl.Entries.Add(new AclEntry(StartChainedJobClaim, new[] { AclAction.CreateJob, AclAction.ExecuteJob, AclAction.UpdateJob, AclAction.ViewJob, AclAction.ViewTemplate, AclAction.ViewStream }));
+
+			ServerSettings SettingsValue = Settings.CurrentValue;
+			if (SettingsValue.AdminClaimType != null && SettingsValue.AdminClaimValue != null)
+			{
+				DefaultAcl.Entries.Add(new AclEntry(new AclClaim(SettingsValue.AdminClaimType, SettingsValue.AdminClaimValue), AdminActions.ToArray()));
+			}
+		}
+
+		/// <summary>
+		/// Gets the root ACL scope table
+		/// </summary>
+		/// <returns>Scopes instance</returns>
+		public Task<GlobalPermissions> GetGlobalPermissionsAsync()
+		{
+			return GlobalPermissions.GetAsync();
+		}
+
+		/// <summary>
+		/// Updates the root ACL scopes
+		/// </summary>
+		/// <param name="NewPermissions">New permissions</param>
+		/// <returns>Async task</returns>
+		public Task<bool> TryUpdateGlobalPermissionsAsync(GlobalPermissions NewPermissions)
+		{
+			return GlobalPermissions.TryUpdateAsync(NewPermissions);
+		}
+
+		/// <summary>
+		/// Authorizes a user against a given scope
+		/// </summary>
+		/// <param name="Action">The action being performed</param>
+		/// <param name="User">The principal to validate</param>
+		/// <param name="Cache">The ACL scope cache</param>
+		/// <returns>Async task</returns>
+		public async Task<bool> AuthorizeAsync(AclAction Action, ClaimsPrincipal User, GlobalPermissionsCache? Cache = null)
+		{
+			GlobalPermissions? GlobalPermissions;
+			if(Cache == null)
+			{
+				GlobalPermissions = await GetGlobalPermissionsAsync();
+			}
+			else if(Cache.GlobalPermissions == null)
+			{
+				GlobalPermissions = Cache.GlobalPermissions = await GetGlobalPermissionsAsync();
+			}
+			else
+			{
+				GlobalPermissions = Cache.GlobalPermissions;
+			}
+			return GlobalPermissions.Acl.Authorize(Action, User) ?? DefaultAcl.Authorize(Action, User) ?? false;
+		}
+
+		/// <summary>
+		/// Issues a bearer token with the given roles
+		/// </summary>
+		/// <param name="Claims">List of claims to include</param>
+		/// <param name="Expiry">Time that the token expires</param>
+		/// <returns>JWT security token with a claim for creating new agents</returns>
+		public string IssueBearerToken(IEnumerable<AclClaim> Claims, TimeSpan? Expiry)
+		{
+			return IssueBearerToken(Claims.Select(x => new Claim(x.Type, x.Value)), Expiry);
+		}
+
+		/// <summary>
+		/// Issues a bearer token with the given claims
+		/// </summary>
+		/// <param name="Claims">List of claims to include</param>
+		/// <param name="Expiry">Time that the token expires</param>
+		/// <returns>JWT security token with a claim for creating new agents</returns>
+		public string IssueBearerToken(IEnumerable<Claim> Claims, TimeSpan? Expiry)
+		{
+			SigningCredentials SigningCredentials = new SigningCredentials(DatabaseService.JwtSigningKey, SecurityAlgorithms.HmacSha256);
+
+			JwtSecurityToken Token = new JwtSecurityToken(DatabaseService.JwtIssuer, null, Claims, null, DateTime.UtcNow + Expiry, SigningCredentials);
+			return new JwtSecurityTokenHandler().WriteToken(Token);
+		}
+
+		/// <summary>
+		/// Get the roles for the given user
+		/// </summary>
+		/// <param name="User">The user to query roles for</param>
+		/// <returns>Collection of roles</returns>
+		public static HashSet<string> GetRoles(ClaimsPrincipal User)
+		{
+			return new HashSet<string>(User.Claims.Where(x => x.Type == ClaimTypes.Role).Select(x => x.Value));
+		}
+
+		/// <summary>
+		/// Gets the user name from the given principal
+		/// </summary>
+		/// <param name="User">The principal to check</param>
+		/// <returns></returns>
+		public static string? GetUserName(ClaimsPrincipal User)
+		{
+			return (User.Claims.FirstOrDefault(x => x.Type == HordeClaimTypes.User) ?? User.Claims.FirstOrDefault(x => x.Type == ClaimTypes.Name))?.Value ?? "Anonymous";
+		}
+
+		/// <summary>
+		/// Gets the role for a specific user
+		/// </summary>
+		/// <param name="SessionId">The session id</param>
+		/// <returns>New claim instance</returns>
+		public static AclClaim GetSessionClaim(ObjectId SessionId)
+		{
+			return new AclClaim(HordeClaimTypes.AgentSessionId, SessionId.ToString());
+		}
+
+		/// <summary>
+		/// Determines whether the given user can masquerade as a given user
+		/// </summary>
+		/// <param name="User"></param>
+		/// <param name="UserId"></param>
+		/// <returns></returns>
+		public Task<bool> AuthorizeAsUserAsync(ClaimsPrincipal User, ObjectId UserId)
+		{
+			ObjectId? CurrentUserId = User.GetUserId();
+			if (CurrentUserId != null && CurrentUserId.Value == UserId)
+			{
+				return Task.FromResult(true);
+			}
+			else
+			{
+				return AuthorizeAsync(AclAction.Impersonate, User);
+			}
+		}
+	}
+
+	static class ClaimExtensions
+	{
+		public static bool HasSessionClaim(this ClaimsPrincipal User, ObjectId SessionId)
+		{
+			return User.HasClaim(HordeClaimTypes.AgentSessionId, SessionId.ToString());
+		}
+
+		public static ObjectId? GetSessionClaim(this ClaimsPrincipal User)
+		{
+			Claim? Claim = User.FindFirst(HordeClaimTypes.AgentSessionId);
+			if (Claim == null || !ObjectId.TryParse(Claim.Value, out ObjectId SessionId))
+			{
+				return null;
+			}
+			else
+			{
+				return SessionId;
+			}
+		}
+		
+		public static string GetSessionClaimsAsString(this ClaimsPrincipal User)
+		{
+			return String.Join(",", User.Claims
+				.Where(c => c.Type == HordeClaimTypes.AgentSessionId)
+				.Select(c => c.Value));
+		}
+	}
+}
diff --git a/Engine/Source/Programs/Horde/HordeServer/Agents/AgentSoftware.cs b/Engine/Source/Programs/Horde/HordeServer/Agents/AgentSoftware.cs
deleted file mode 100644
diff --git a/Engine/Source/Programs/Horde/HordeServer/Agents/IAgentSoftwareChannel.cs b/Engine/Source/Programs/Horde/HordeServer/Agents/IAgentSoftwareChannel.cs
new file mode 100644
--- /dev/null
+++ b/Engine/Source/Programs/Horde/HordeServer/Agents/IAgentSoftwareChannel.cs
@@ -0,0 +1,47 @@
+// Copyright Epic Games, Inc. All Rights Reserved.
+
+using HordeServer.Api;
+using HordeServer.Collections;
+using HordeServer.Services;
+using HordeServer.Utilities;
+using MongoDB.Bson;
+using MongoDB.Bson.Serialization.Attributes;
+using System;
+using System.Collections.Generic;
+using System.IO;
+using System.IO.Compression;
+using System.Linq;
+using System.Text.Json;
+using System.Threading.Tasks;
+
+namespace HordeServer.Models
+{
+	using AgentSoftwareVersion = StringId<IAgentSoftwareCollection>;
+	using AgentSoftwareChannelName = StringId<AgentSoftwareChannels>;
+
+	/// <summary>
+	/// A software channel
+	/// </summary>
+	public interface IAgentSoftwareChannel
+	{
+		/// <summary>
+		/// The channel id
+		/// </summary>
+		public AgentSoftwareChannelName Name { get; set; }
+
+		/// <summary>
+		/// Name of the user that made the last modification
+		/// </summary>
+		public string? ModifiedBy { get; set; }
+
+		/// <summary>
+		/// Last modification time
+		/// </summary>
+		public DateTime ModifiedTime { get; set; }
+
+		/// <summary>
+		/// The software version number
+		/// </summary>
+		public string Version { get; set; }
+	}
+}
diff --git a/Engine/Source/Programs/Horde/HordeServer/Agents/Impl/ConformList.cs b/Engine/Source/Programs/Horde/HordeServer/Agents/Impl/ConformList.cs
new file mode 100644
--- /dev/null
+++ b/Engine/Source/Programs/Horde/HordeServer/Agents/Impl/ConformList.cs
@@ -0,0 +1,74 @@
+// Copyright Epic Games, Inc. All Rights Reserved.
+
+using HordeServer.Utilities;
+using MongoDB.Bson;
+using MongoDB.Bson.Serialization.Attributes;
+using System;
+using System.Collections.Generic;
+using System.Linq;
+using System.Threading.Tasks;
+
+namespace HordeServer.Models
+{
+	/// <summary>
+	/// Entry to the conform list
+	/// </summary>
+	public class ConformListEntry
+	{
+		/// <summary>
+		/// The agent id
+		/// </summary>
+		[BsonElement("a")]
+		public AgentId AgentId { get; set; }
+
+		/// <summary>
+		/// The lease id
+		/// </summary>
+		[BsonElement("l")]
+		public ObjectId LeaseId { get; set; }
+
+		/// <summary>
+		/// Last timestamp that the lease was checked for validity 
+		/// </summary>
+		[BsonElement("t"), BsonIgnoreIfNull]
+		public DateTime? LastCheckTimeUtc { get; set; }
+	}
+
+	/// <summary>
+	/// List of entries for a particular server
+	/// </summary>
+	public class ConformListServer
+	{
+		/// <summary>
+		/// The perforce cluster
+		/// </summary>
+		public string Cluster { get; set; } = String.Empty;
+
+		/// <summary>
+		/// The server and port
+		/// </summary>
+		public string ServerAndPort { get; set; } = String.Empty;
+
+		/// <summary>
+		/// List of entries
+		/// </summary>
+		public List<ConformListEntry> Entries { get; set; } = new List<ConformListEntry>();
+	}
+
+	/// <summary>
+	/// List of machines that are currently conforming
+	/// </summary>
+	[SingletonDocument("60afc737f0d2a70754229300")]
+	public class ConformList : SingletonBase
+	{
+		/// <summary>
+		/// List of conforming servers [DEPRECATED]
+		/// </summary>
+		public List<ConformListEntry> Entries { get; set; } = new List<ConformListEntry>();
+
+		/// <summary>
+		/// List of conforming servers
+		/// </summary>
+		public List<ConformListServer> Servers { get; set; } = new List<ConformListServer>();
+	}
+}
diff --git a/Engine/Source/Programs/Horde/HordeServer/Api/Software.cs b/Engine/Source/Programs/Horde/HordeServer/Api/Software.cs
deleted file mode 100644
diff --git a/Engine/Source/Programs/Horde/HordeServer/Auditing/IAuditLog.cs b/Engine/Source/Programs/Horde/HordeServer/Auditing/IAuditLog.cs
new file mode 100644
--- /dev/null
+++ b/Engine/Source/Programs/Horde/HordeServer/Auditing/IAuditLog.cs
@@ -0,0 +1,121 @@
+// Copyright Epic Games, Inc. All Rights Reserved.
+
+using Google.Protobuf.WellKnownTypes;
+using HordeServer.Api;
+using HordeCommon;
+using HordeServer.Models;
+using HordeCommon.Rpc.Tasks;
+using HordeServer.Services;
+using HordeServer.Utilities;
+using MongoDB.Bson;
+using MongoDB.Bson.Serialization.Attributes;
+using MongoDB.Driver;
+using System;
+using System.Collections.Generic;
+using System.Linq;
+using System.Linq.Expressions;
+using System.Threading.Tasks;
+using Microsoft.Extensions.Logging;
+using EpicGames.Core;
+using System.Buffers;
+using System.Text.Json;
+using System.Text;
+
+namespace HordeServer.Collections
+{
+	/// <summary>
+	/// Message from an audit log
+	/// </summary>
+	public interface IAuditLogMessage
+	{
+		/// <summary>
+		/// Timestamp for the event
+		/// </summary>
+		public DateTime TimeUtc { get; }
+
+		/// <summary>
+		/// Severity of the message
+		/// </summary>
+		public LogLevel Level { get; }
+
+		/// <summary>
+		/// The message payload. Should be an encoded JSON object, with format/properties fields.
+		/// </summary>
+		public string Data { get; }
+	}
+
+	/// <summary>
+	/// Channel for a particular entity
+	/// </summary>
+	public interface IAuditLogChannel : ILogger
+	{
+		/// <summary>
+		/// Finds messages matching certain criteria
+		/// </summary>
+		/// <param name="MinTime"></param>
+		/// <param name="MaxTime"></param>
+		/// <param name="Index"></param>
+		/// <param name="Count"></param>
+		/// <returns></returns>
+		IAsyncEnumerable<IAuditLogMessage> FindAsync(DateTime? MinTime = null, DateTime? MaxTime = null, int? Index = null, int? Count = null);
+
+		/// <summary>
+		/// Deletes messages between a given time range for a particular object
+		/// </summary>
+		/// <param name="MinTime">Minimum time to remove</param>
+		/// <param name="MaxTime">Maximum time to remove</param>
+		/// <returns>Async task</returns>
+		Task<long> DeleteAsync(DateTime? MinTime = null, DateTime? MaxTime = null);
+	}
+
+	/// <summary>
+	/// Channel for a particular entity
+	/// </summary>
+	public interface IAuditLogChannel<TSubject> : IAuditLogChannel
+	{
+		/// <summary>
+		/// Identifier for the subject of this channel
+		/// </summary>
+		TSubject Subject { get; }
+	}
+
+	/// <summary>
+	/// Message in an audit log
+	/// </summary>
+	/// <typeparam name="TTargetId">Type of entity that the log is for</typeparam>
+	public interface IAuditLogMessage<TTargetId> : IAuditLogMessage
+	{
+		/// <summary>
+		/// Unique id for the entity
+		/// </summary>
+		public TTargetId Subject { get; }
+	}
+
+	/// <summary>
+	/// Interface for a collection of log messages for a particular document type
+	/// </summary>
+	public interface IAuditLog<TSubject>
+	{
+		/// <summary>
+		/// Get the channel for a particular subject
+		/// </summary>
+		/// <param name="Subject"></param>
+		/// <returns></returns>
+		IAuditLogChannel<TSubject> this[TSubject Subject] { get; }
+	}
+
+	/// <summary>
+	/// Factory for instantiating audit log instances
+	/// </summary>
+	/// <typeparam name="TSubject"></typeparam>
+	public interface IAuditLogFactory<TSubject>
+	{
+		/// <summary>
+		/// Create a new audit log instance, with the given database name
+		/// </summary>
+		/// <param name="CollectionName"></param>
+		/// <param name="SubjectProperty"></param>
+		/// <returns></returns>
+		IAuditLog<TSubject> Create(string CollectionName, string SubjectProperty);
+	}
+}
diff --git a/Engine/Source/Programs/Horde/HordeServer/Auditing/Impl/AuditLog.cs b/Engine/Source/Programs/Horde/HordeServer/Auditing/Impl/AuditLog.cs
new file mode 100644
--- /dev/null
+++ b/Engine/Source/Programs/Horde/HordeServer/Auditing/Impl/AuditLog.cs
@@ -0,0 +1,200 @@
+// Copyright Epic Games, Inc. All Rights Reserved.
+
+using EpicGames.Core;
+using Google.Protobuf.WellKnownTypes;
+using HordeCommon;
+using HordeCommon.Rpc.Tasks;
+using HordeServer.Models;
+using HordeServer.Services;
+using HordeServer.Utilities;
+using Microsoft.Extensions.Logging;
+using MongoDB.Bson;
+using MongoDB.Bson.Serialization.Attributes;
+using MongoDB.Driver;
+using System;
+using System.Buffers;
+using System.Collections.Generic;
+using System.Linq;
+using System.Linq.Expressions;
+using System.Text.Json;
+using System.Threading.Channels;
+using System.Threading.Tasks;
+
+namespace HordeServer.Collections.Impl
+{
+	class AuditLog<TSubject> : IAuditLog<TSubject>, IAsyncDisposable, IDisposable
+	{
+		class AuditLogMessage : IAuditLogMessage<TSubject>
+		{
+			public ObjectId Id { get; set; }
+
+			[BsonElement("s")]
+			public TSubject Subject { get; set; }
+
+			[BsonElement("t")]
+			public DateTime TimeUtc { get; set; }
+
+			[BsonElement("l")]
+			public LogLevel Level { get; set; }
+
+			[BsonElement("d")]
+			public string Data { get; set; }
+
+			public AuditLogMessage()
+			{
+				Subject = default!;
+				Data = String.Empty;
+			}
+
+			public AuditLogMessage(TSubject Subject, DateTime TimeUtc, LogLevel Level, string Data)
+			{
+				this.Id = ObjectId.GenerateNewId();
+				this.Subject = Subject;
+				this.TimeUtc = TimeUtc;
+				this.Level = Level;
+				this.Data = Data;
+			}
+		}
+
+		class AuditLogChannel : IAuditLogChannel<TSubject>
+		{
+			public readonly AuditLog<TSubject> Outer;
+			public TSubject Subject { get; }
+
+			public AuditLogChannel(AuditLog<TSubject> Outer, TSubject Subject)
+			{
+				this.Outer = Outer;
+				this.Subject = Subject;
+			}
+
+			public IDisposable? BeginScope<TState>(TState state) => null;
+
+			public bool IsEnabled(LogLevel logLevel) => true;
+
+			public void Log<TState>(LogLevel LogLevel, EventId EventId, TState State, Exception Exception, Func<TState, Exception, string> Formatter)
+			{
+				DateTime Time = DateTime.UtcNow;
+
+				string Data = MessageTemplate.Serialize(LogLevel, EventId, State, Exception, Formatter);
+				AuditLogMessage Message = new AuditLogMessage(Subject, Time, LogLevel, Data);
+				Outer.MessageChannel.Writer.TryWrite(Message);
+
+				using (IDisposable _ = Outer.Logger.BeginScope($"Subject: {{{Outer.SubjectProperty}}}", Subject))
+				{
+					Outer.Logger.Log(LogLevel, EventId, State, Exception, Formatter);
+				}
+			}
+
+			public IAsyncEnumerable<IAuditLogMessage> FindAsync(DateTime? MinTime, DateTime? MaxTime, int? Index, int? Count) => Outer.FindAsync(Subject, MinTime, MaxTime, Index, Count);
+
+			public Task<long> DeleteAsync(DateTime? MinTime, DateTime? MaxTime) => Outer.DeleteAsync(Subject, MinTime, MaxTime);
+		}
+
+		IMongoCollection<AuditLogMessage> Messages;
+		Channel<AuditLogMessage> MessageChannel;
+		string SubjectProperty;
+		ILogger Logger;
+		Task BackgroundTask;
+
+		public IAuditLogChannel<TSubject> this[TSubject Subject] => new AuditLogChannel(this, Subject);
+
+		public AuditLog(DatabaseService DatabaseService, string CollectionName, string SubjectProperty, ILogger Logger)
+		{
+			this.Messages = DatabaseService.GetCollection<AuditLogMessage>(CollectionName);
+			this.MessageChannel = Channel.CreateUnbounded<AuditLogMessage>();
+			this.SubjectProperty = SubjectProperty;
+			this.Logger = Logger;
+
+			if (!DatabaseService.ReadOnlyMode)
+			{
+				Messages.Indexes.CreateOne(new CreateIndexModel<AuditLogMessage>(Builders<AuditLogMessage>.IndexKeys.Ascending(x => x.Id).Descending(x => x.TimeUtc)));
+			}
+
+			BackgroundTask = Task.Run(() => WriteMessagesAsync());
+		}
+
+		public async ValueTask DisposeAsync()
+		{
+			MessageChannel.Writer.TryComplete();
+			await BackgroundTask;
+		}
+
+		public void Dispose()
+		{
+			DisposeAsync().AsTask().Wait();
+		}
+
+		async Task WriteMessagesAsync()
+		{
+			while (await MessageChannel.Reader.WaitToReadAsync())
+			{
+				List<AuditLogMessage> NewMessages = new List<AuditLogMessage>();
+				while (MessageChannel.Reader.TryRead(out AuditLogMessage? NewMessage))
+				{
+					NewMessages.Add(NewMessage);
+				}
+				if (NewMessages.Count > 0)
+				{
+					await Messages.InsertManyAsync(NewMessages);
+				}
+			}
+		}
+
+		async IAsyncEnumerable<IAuditLogMessage<TSubject>> FindAsync(TSubject Subject, DateTime? MinTime = null, DateTime? MaxTime = null, int? Index = null, int? Count = null)
+		{
+			FilterDefinition<AuditLogMessage> Filter = Builders<AuditLogMessage>.Filter.Eq(x => x.Subject, Subject);
+			if (MinTime != null)
+			{
+				Filter &= Builders<AuditLogMessage>.Filter.Gte(x => x.TimeUtc, MinTime.Value);
+			}
+			if (MaxTime != null)
+			{
+				Filter &= Builders<AuditLogMessage>.Filter.Lte(x => x.TimeUtc, MaxTime.Value);
+			}
+
+			using (IAsyncCursor<AuditLogMessage> Cursor = await Messages.Find(Filter).SortByDescending(x => x.TimeUtc).Range(Index, Count).ToCursorAsync())
+			{
+				while (await Cursor.MoveNextAsync())
+				{
+					foreach (AuditLogMessage Message in Cursor.Current)
+					{
+						yield return Message;
+					}
+				}
+			}
+		}
+
+		async Task<long> DeleteAsync(TSubject Subject, DateTime? MinTime = null, DateTime? MaxTime = null)
+		{
+			FilterDefinition<AuditLogMessage> Filter = Builders<AuditLogMessage>.Filter.Eq(x => x.Subject, Subject);
+			if (MinTime != null)
+			{
+				Filter &= Builders<AuditLogMessage>.Filter.Gte(x => x.TimeUtc, MinTime.Value);
+			}
+			if (MaxTime != null)
+			{
+				Filter &= Builders<AuditLogMessage>.Filter.Lte(x => x.TimeUtc, MaxTime.Value);
+			}
+
+			DeleteResult Result = await Messages.DeleteManyAsync(Filter);
+			return Result.DeletedCount;
+		}
+	}
+
+	class AuditLogFactory<TSubject> : IAuditLogFactory<TSubject>
+	{
+		DatabaseService DatabaseService;
+		ILogger<AuditLog<TSubject>> Logger;
+
+		public AuditLogFactory(DatabaseService DatabaseService, ILogger<AuditLog<TSubject>> Logger)
+		{
+			this.DatabaseService = DatabaseService;
+			this.Logger = Logger;
+		}
+
+		public IAuditLog<TSubject> Create(string CollectionName, string SubjectProperty)
+		{
+			return new AuditLog<TSubject>(DatabaseService, CollectionName, SubjectProperty, Logger);
+		}
+	}
+}
diff --git a/Engine/Source/Programs/Horde/HordeServer/Collections/IAuditLog.cs b/Engine/Source/Programs/Horde/HordeServer/Collections/IAuditLog.cs
deleted file mode 100644
diff --git a/Engine/Source/Programs/Horde/HordeServer/Collections/ICounterCollection.cs b/Engine/Source/Programs/Horde/HordeServer/Collections/ICounterCollection.cs
deleted file mode 100644
diff --git a/Engine/Source/Programs/Horde/HordeServer/Collections/IUgsMetadataCollection.cs b/Engine/Source/Programs/Horde/HordeServer/Collections/IUgsMetadataCollection.cs
deleted file mode 100644
diff --git a/Engine/Source/Programs/Horde/HordeServer/Collections/Impl/AuditLog.cs b/Engine/Source/Programs/Horde/HordeServer/Collections/Impl/AuditLog.cs
deleted file mode 100644
diff --git a/Engine/Source/Programs/Horde/HordeServer/Collections/Impl/CounterCollection.cs b/Engine/Source/Programs/Horde/HordeServer/Collections/Impl/CounterCollection.cs
deleted file mode 100644
diff --git a/Engine/Source/Programs/Horde/HordeServer/Collections/Impl/TestDataCollection.cs b/Engine/Source/Programs/Horde/HordeServer/Collections/Impl/TestDataCollection.cs
deleted file mode 100644
diff --git a/Engine/Source/Programs/Horde/HordeServer/Collections/Impl/UgsMetadataCollection.cs b/Engine/Source/Programs/Horde/HordeServer/Collections/Impl/UgsMetadataCollection.cs
deleted file mode 100644
diff --git a/Engine/Source/Programs/Horde/HordeServer/Controllers/SoftwareController.cs b/Engine/Source/Programs/Horde/HordeServer/Controllers/SoftwareController.cs
deleted file mode 100644
diff --git a/Engine/Source/Programs/Horde/HordeServer/Fleet/Telemetry.cs b/Engine/Source/Programs/Horde/HordeServer/Fleet/Telemetry.cs
new file mode 100644
--- /dev/null
+++ b/Engine/Source/Programs/Horde/HordeServer/Fleet/Telemetry.cs
@@ -0,0 +1,213 @@
+// Copyright Epic Games, Inc. All Rights Reserved.
+
+using HordeServer.Utilities;
+using System;
+using System.Collections.Generic;
+using System.Linq;
+using System.Threading.Tasks;
+
+namespace HordeServer.Models
+{
+	using PoolId = StringId<IPool>;
+	using StreamId = StringId<IStream>;
+
+	/// <summary>
+	/// Information about the utilization of a pool
+	/// </summary>
+	public interface IUtilizationTelemetry
+	{
+		/// <summary>
+		/// Start time for the bucket
+		/// </summary>
+		public DateTime StartTime { get; }
+
+		/// <summary>
+		/// Finish time for this bucket
+		/// </summary>
+		public DateTime FinishTime { get; }
+
+		/// <summary>
+		/// Number of agents in total
+		/// </summary>
+		public int NumAgents { get; }
+
+		/// <summary>
+		/// Breakdown of utilization by pool
+		/// </summary>
+		public IReadOnlyList<IPoolUtilizationTelemetry> Pools { get; }
+
+		/// <summary>
+		/// Total time spent running administrative tasks (conform, etc...)
+		/// </summary>
+		public double AdminTime { get; }
+	}
+
+	/// <summary>
+	/// Concrete implementation of <see cref="IUtilizationTelemetry"/>
+	/// </summary>
+	public sealed class NewUtilizationTelemetry : IUtilizationTelemetry
+	{
+		/// <inheritdoc/>
+		public DateTime StartTime { get; set; }
+
+		/// <inheritdoc/>
+		public DateTime FinishTime { get; set; }
+
+		/// <inheritdoc/>
+		public int NumAgents { get; set; }
+
+		/// <inheritdoc/>
+		List<NewPoolUtilizationTelemetry> Pools { get; set; } = new List<NewPoolUtilizationTelemetry>();
+		IReadOnlyList<IPoolUtilizationTelemetry> IUtilizationTelemetry.Pools => Pools;
+
+		Dictionary<PoolId, NewPoolUtilizationTelemetry> PoolsLookup { get; set; } = new Dictionary<PoolId, NewPoolUtilizationTelemetry>();
+
+		/// <inheritdoc/>
+		public double AdminTime { get; set; }
+
+		/// <summary>
+		/// Constructor
+		/// </summary>
+		/// <param name="StartTime"></param>
+		/// <param name="FinishTime"></param>
+		public NewUtilizationTelemetry(DateTime StartTime, DateTime FinishTime)
+		{
+			this.StartTime = StartTime;
+			this.FinishTime = FinishTime;
+		}
+
+		/// <summary>
+		/// Adds a new pool
+		/// </summary>
+		/// <param name="PoolId">The pool id</param>
+		/// <returns>Telemetry for the given pool</returns>
+		public NewPoolUtilizationTelemetry FindOrAddPool(PoolId PoolId)
+		{
+			NewPoolUtilizationTelemetry? Pool;
+			if (!PoolsLookup.TryGetValue(PoolId, out Pool))
+			{
+				Pool = new NewPoolUtilizationTelemetry(PoolId);
+				Pools.Add(Pool);
+				PoolsLookup.Add(PoolId, Pool);
+			}
+			return Pool;
+		}
+	}
+
+	/// <summary>
+	/// Information about the utilization of a pool
+	/// </summary>
+	public interface IPoolUtilizationTelemetry
+	{
+		/// <summary>
+		/// Pool containing the work to execute
+		/// </summary>
+		public PoolId PoolId { get; }
+
+		/// <summary>
+		/// Number of agents in this pool
+		/// </summary>
+		public int NumAgents { get; }
+
+		/// <summary>
+		/// The stream executing work. If this is null, the time accounts for the machine executing work in another pool.
+		/// </summary>
+		public IReadOnlyList<IStreamUtilizationTelemetry> Streams { get; }
+
+		/// <summary>
+		/// Admount of time spent running 
+		/// </summary>
+		public double AdminTime { get; }
+
+		/// <summary>
+		/// Amount of time spent by agents in this pool servicing other pools
+		/// </summary>
+		public double OtherTime { get; }
+	}
+
+	/// <summary>
+	/// Concrete implementation of <see cref="IPoolUtilizationTelemetry"/>
+	/// </summary>
+	public sealed class NewPoolUtilizationTelemetry : IPoolUtilizationTelemetry
+	{
+		/// <inheritdoc/>
+		public PoolId PoolId { get; set; }
+
+		/// <inheritdoc/>
+		public int NumAgents { get; set; }
+
+		/// <inheritdoc/>
+		List<NewPoolUtilizationTelemetryStream> Streams { get; set; } = new List<NewPoolUtilizationTelemetryStream>();
+		IReadOnlyList<IStreamUtilizationTelemetry> IPoolUtilizationTelemetry.Streams => Streams;
+
+		Dictionary<StreamId, NewPoolUtilizationTelemetryStream> StreamLookup { get; set; } = new Dictionary<StreamId, NewPoolUtilizationTelemetryStream>();
+
+		/// <inheritdoc/>
+		public double AdminTime { get; set; }
+
+		/// <inheritdoc/>
+		public double OtherTime { get; set; }
+
+		/// <summary>
+		/// Constructor
+		/// </summary>
+		/// <param name="PoolId"></param>
+		public NewPoolUtilizationTelemetry(PoolId PoolId)
+		{
+			this.PoolId = PoolId;
+		}
+
+		/// <summary>
+		/// Adds a stream to the object
+		/// </summary>
+		/// <param name="StreamId"></param>
+		public NewPoolUtilizationTelemetryStream FindOrAddStream(StreamId StreamId)
+		{
+			NewPoolUtilizationTelemetryStream? Stream;
+			if (!StreamLookup.TryGetValue(StreamId, out Stream))
+			{
+				Stream = new NewPoolUtilizationTelemetryStream(StreamId);
+				Streams.Add(Stream);
+				StreamLookup.Add(StreamId, Stream);
+			}
+			return Stream;
+		}
+	}
+
+	/// <summary>
+	/// Utilization of a pool for a particular stream
+	/// </summary>
+	public interface IStreamUtilizationTelemetry
+	{
+		/// <summary>
+		/// The stream id
+		/// </summary>
+		public StreamId StreamId { get; }
+
+		/// <summary>
+		/// Number of machine hours spent executing work for this stream
+		/// </summary>
+		public double Time { get; }
+	}
+
+	/// <summary>
+	/// Concrete implementation of <see cref="IStreamUtilizationTelemetry"/>
+	/// </summary>
+	public sealed class NewPoolUtilizationTelemetryStream : IStreamUtilizationTelemetry
+	{
+		/// <inheritdoc/>
+		public StreamId StreamId { get; set; }
+
+		/// <inheritdoc/>
+		public double Time { get; set; }
+
+		/// <summary>
+		/// Constructor
+		/// </summary>
+		/// <param name="StreamId"></param>
+		public NewPoolUtilizationTelemetryStream(StreamId StreamId)
+		{
+			this.StreamId = StreamId;
+		}
+	}
+}
diff --git a/Engine/Source/Programs/Horde/HordeServer/Jobs/ICounterCollection.cs b/Engine/Source/Programs/Horde/HordeServer/Jobs/ICounterCollection.cs
new file mode 100644
--- /dev/null
+++ b/Engine/Source/Programs/Horde/HordeServer/Jobs/ICounterCollection.cs
@@ -0,0 +1,22 @@
+// Copyright Epic Games, Inc. All Rights Reserved.
+
+using System;
+using System.Collections.Generic;
+using System.Linq;
+using System.Threading.Tasks;
+
+namespace HordeServer.Collections
+{
+	/// <summary>
+	/// Collection of counter objects. Counters are named values stored in the database which can be atomically incremented.
+	/// </summary>
+	public interface ICounterCollection
+	{
+		/// <summary>
+		/// Increments a counter with the given name. The first value returned is 1.
+		/// </summary>
+		/// <param name="Name">Name of the counter</param>
+		/// <returns></returns>
+		Task<int> IncrementAsync(string Name);
+	}
+}
diff --git a/Engine/Source/Programs/Horde/HordeServer/Jobs/ITestData.cs b/Engine/Source/Programs/Horde/HordeServer/Jobs/ITestData.cs
new file mode 100644
--- /dev/null
+++ b/Engine/Source/Programs/Horde/HordeServer/Jobs/ITestData.cs
@@ -0,0 +1,60 @@
+// Copyright Epic Games, Inc. All Rights Reserved.
+
+using HordeServer.Utilities;
+using MongoDB.Bson;
+using System;
+using System.Collections.Generic;
+using System.Linq;
+using System.Threading.Tasks;
+
+namespace HordeServer.Models
+{
+	using StreamId = StringId<IStream>;
+	using TemplateRefId = StringId<TemplateRef>;
+
+	/// <summary>
+	/// Stores information about the results of a test
+	/// </summary>
+	public interface ITestData
+	{
+		/// <summary>
+		/// Unique id of the test data
+		/// </summary>
+		ObjectId Id { get; }
+
+		/// <summary>
+		/// Stream that generated the test data
+		/// </summary>
+		StreamId StreamId { get; }
+
+		/// <summary>
+		/// The template reference id
+		/// </summary>
+		TemplateRefId TemplateRefId { get; }
+
+		/// <summary>
+		/// The job which produced the data
+		/// </summary>
+		ObjectId JobId { get; }
+
+		/// <summary>
+		/// The step that ran
+		/// </summary>
+		SubResourceId StepId { get; }
+
+		/// <summary>
+		/// The changelist number that contained the data
+		/// </summary>
+		int Change { get; }
+
+		/// <summary>
+		/// Key used to identify the particular data
+		/// </summary>
+		string Key { get; }
+
+		/// <summary>
+		/// The data stored for this test
+		/// </summary>
+		BsonDocument Data { get; }
+	}
+}
diff --git a/Engine/Source/Programs/Horde/HordeServer/Jobs/Impl/CounterCollection.cs b/Engine/Source/Programs/Horde/HordeServer/Jobs/Impl/CounterCollection.cs
new file mode 100644
--- /dev/null
+++ b/Engine/Source/Programs/Horde/HordeServer/Jobs/Impl/CounterCollection.cs
@@ -0,0 +1,60 @@
+// Copyright Epic Games, Inc. All Rights Reserved.
+
+using HordeServer.Services;
+using HordeServer.Utilities;
+using MongoDB.Driver;
+using System;
+using System.Collections.Generic;
+using System.Linq;
+using System.Threading.Tasks;
+
+namespace HordeServer.Collections
+{
+	/// <summary>
+	/// Collection of counter objects. Counters are named values stored in the database which can be atomically incremented.
+	/// </summary>
+	public class CounterCollection : ICounterCollection
+	{
+		class Counter
+		{
+			public string Id { get; set; } = String.Empty;
+			public int Value { get; set; }
+		}
+
+		IMongoCollection<Counter> Counters;
+
+		/// <summary>
+		/// Constructor
+		/// </summary>
+		/// <param name="DatabaseService">Database service instance</param>
+		public CounterCollection(DatabaseService DatabaseService)
+		{
+			Counters = DatabaseService.GetCollection<Counter>("Counters");
+		}
+
+		/// <inheritdoc/>
+		public async Task<int> IncrementAsync(string Name)
+		{
+			for (; ; )
+			{
+				Counter? Current = await Counters.Find(x => x.Id == Name).FirstOrDefaultAsync();
+				if (Current == null)
+				{
+					Counter NewCounter = new Counter { Id = Name, Value = 1 };
+					if (await Counters.InsertOneIgnoreDuplicatesAsync(NewCounter))
+					{
+						return NewCounter.Value;
+					}
+				}
+				else
+				{
+					UpdateResult Result = await Counters.UpdateOneAsync(x => x.Id == Name && x.Value == Current.Value, Builders<Counter>.Update.Set(x => x.Value, Current.Value + 1));
+					if (Result.ModifiedCount > 0)
+					{
+						return Current.Value + 1;
+					}
+				}
+			}
+		}
+	}
+}
diff --git a/Engine/Source/Programs/Horde/HordeServer/Jobs/Impl/TestDataCollection.cs b/Engine/Source/Programs/Horde/HordeServer/Jobs/Impl/TestDataCollection.cs
new file mode 100644
--- /dev/null
+++ b/Engine/Source/Programs/Horde/HordeServer/Jobs/Impl/TestDataCollection.cs
@@ -0,0 +1,134 @@
+// Copyright Epic Games, Inc. All Rights Reserved.
+
+using HordeServer.Models;
+using HordeServer.Services;
+using HordeServer.Utilities;
+using MongoDB.Bson;
+using MongoDB.Bson.Serialization.Attributes;
+using MongoDB.Driver;
+using System;
+using System.Collections.Generic;
+using System.Diagnostics.CodeAnalysis;
+using System.Linq;
+using System.Threading.Tasks;
+
+namespace HordeServer.Collections.Impl
+{
+	using ProjectId = StringId<IProject>;
+	using StreamId = StringId<IStream>;
+	using TemplateRefId = StringId<TemplateRef>;
+
+	/// <summary>
+	/// Collection of test data documents
+	/// </summary>
+	public class TestDataCollection : ITestDataCollection
+	{
+		/// <summary>
+		/// Information about a test data document
+		/// </summary>
+		class TestDataDocument : ITestData
+		{
+			public ObjectId Id { get; set; }
+			public StreamId StreamId { get; set; }
+			public TemplateRefId TemplateRefId { get; set; }
+			public ObjectId JobId { get; set; }
+			public SubResourceId StepId { get; set; }
+			public int Change { get; set; }
+			public string Key { get; set; }
+			public BsonDocument Data { get; set; }
+
+			private TestDataDocument()
+			{
+				this.Key = String.Empty;
+				this.Data = new BsonDocument();
+			}
+
+			public TestDataDocument(IJob Job, IJobStep JobStep, string Key, BsonDocument Value)
+			{
+				this.Id = ObjectId.GenerateNewId();
+				this.StreamId = Job.StreamId;
+				this.TemplateRefId = Job.TemplateId;
+				this.JobId = Job.Id;
+				this.StepId = JobStep.Id;
+				this.Change = Job.Change;
+				this.Key = Key;
+				this.Data = Value;
+			}
+		}
+
+		/// <summary>
+		/// The stream collection
+		/// </summary>
+		IMongoCollection<TestDataDocument> TestDataDocuments;
+
+		/// <summary>
+		/// Constructor
+		/// </summary>
+		/// <param name="DatabaseService">The database service instance</param>
+		public TestDataCollection(DatabaseService DatabaseService)
+		{
+			TestDataDocuments = DatabaseService.GetCollection<TestDataDocument>("TestData");
+
+			if (!DatabaseService.ReadOnlyMode)
+			{
+				TestDataDocuments.Indexes.CreateOne(new CreateIndexModel<TestDataDocument>(Builders<TestDataDocument>.IndexKeys.Ascending(x => x.StreamId).Ascending(x => x.Change).Ascending(x => x.Key)));
+				TestDataDocuments.Indexes.CreateOne(new CreateIndexModel<TestDataDocument>(Builders<TestDataDocument>.IndexKeys.Ascending(x => x.JobId).Ascending(x => x.StepId).Ascending(x => x.Key), new CreateIndexOptions { Unique = true }));
+			}
+		}
+
+		/// <inheritdoc/>
+		public async Task<ITestData> AddAsync(IJob Job, IJobStep Step, string Key, BsonDocument Value)
+		{
+			TestDataDocument NewTestData = new TestDataDocument(Job, Step, Key, Value);
+			await TestDataDocuments.InsertOneAsync(NewTestData);
+			return NewTestData;
+		}
+
+		/// <inheritdoc/>
+		public async Task<ITestData?> GetAsync(ObjectId Id)
+		{
+			return await TestDataDocuments.Find<TestDataDocument>(x => x.Id == Id).FirstOrDefaultAsync();
+		}
+
+		/// <inheritdoc/>
+		public async Task<List<ITestData>> FindAsync(StreamId? StreamId, int? MinChange, int? MaxChange, ObjectId? JobId, SubResourceId? StepId, string? Key = null, int Index = 0, int Count = 10)
+		{
+			FilterDefinition<TestDataDocument> Filter = FilterDefinition<TestDataDocument>.Empty;
+			if (StreamId != null)
+			{
+				Filter &= Builders<TestDataDocument>.Filter.Eq(x => x.StreamId, StreamId.Value);
+				if (MinChange != null)
+				{
+					Filter &= Builders<TestDataDocument>.Filter.Gte(x => x.Change, MinChange.Value);
+				}
+				if (MaxChange != null)
+				{
+					Filter &= Builders<TestDataDocument>.Filter.Lte(x => x.Change, MaxChange.Value);
+				}
+			}
+			if (JobId != null)
+			{
+				Filter &= Builders<TestDataDocument>.Filter.Eq(x => x.JobId, JobId.Value);
+				if (StepId != null)
+				{
+					Filter &= Builders<TestDataDocument>.Filter.Eq(x => x.StepId, StepId.Value);
+				}
+			}
+			if (Key != null)
+			{
+				Filter &= Builders<TestDataDocument>.Filter.Eq(x => x.Key, Key);
+			}
+
+			SortDefinition<TestDataDocument> Sort = Builders<TestDataDocument>.Sort.Ascending(x => x.StreamId).Descending(x => x.Change);
+
+			List<TestDataDocument> Results = await TestDataDocuments.Find(Filter).Sort(Sort).Skip(Index).Limit(Count).ToListAsync();
+			return Results.ConvertAll<ITestData>(x => x);
+		}
+
+		/// <inheritdoc/>
+		public async Task DeleteAsync(ObjectId Id)
+		{
+			await TestDataDocuments.DeleteOneAsync<TestDataDocument>(x => x.Id == Id);
+		}
+	}
+}
diff --git a/Engine/Source/Programs/Horde/HordeServer/Models/Acl.cs b/Engine/Source/Programs/Horde/HordeServer/Models/Acl.cs
deleted file mode 100644
diff --git a/Engine/Source/Programs/Horde/HordeServer/Models/Config.cs b/Engine/Source/Programs/Horde/HordeServer/Models/Config.cs
deleted file mode 100644
diff --git a/Engine/Source/Programs/Horde/HordeServer/Models/ConformList.cs b/Engine/Source/Programs/Horde/HordeServer/Models/ConformList.cs
deleted file mode 100644
diff --git a/Engine/Source/Programs/Horde/HordeServer/Models/Credential.cs b/Engine/Source/Programs/Horde/HordeServer/Models/Credential.cs
deleted file mode 100644
diff --git a/Engine/Source/Programs/Horde/HordeServer/Models/ITestData.cs b/Engine/Source/Programs/Horde/HordeServer/Models/ITestData.cs
deleted file mode 100644
diff --git a/Engine/Source/Programs/Horde/HordeServer/Models/Schedule.cs b/Engine/Source/Programs/Horde/HordeServer/Models/Schedule.cs
deleted file mode 100644
diff --git a/Engine/Source/Programs/Horde/HordeServer/Models/Telemetry.cs b/Engine/Source/Programs/Horde/HordeServer/Models/Telemetry.cs
deleted file mode 100644
diff --git a/Engine/Source/Programs/Horde/HordeServer/Models/UgsMetadata.cs b/Engine/Source/Programs/Horde/HordeServer/Models/UgsMetadata.cs
deleted file mode 100644
diff --git a/Engine/Source/Programs/Horde/HordeServer/Projects/ISchedule.cs b/Engine/Source/Programs/Horde/HordeServer/Projects/ISchedule.cs
new file mode 100644
--- /dev/null
+++ b/Engine/Source/Programs/Horde/HordeServer/Projects/ISchedule.cs
@@ -0,0 +1,306 @@
+// Copyright Epic Games, Inc. All Rights Reserved.
+
+using EpicGames.Core;
+using HordeServer.Api;
+using HordeServer.Utilities;
+using MongoDB.Bson;
+using MongoDB.Bson.Serialization.Attributes;
+using System;
+using System.Collections.Generic;
+using System.Linq;
+using System.Threading.Tasks;
+
+namespace HordeServer.Models
+{
+	using TemplateRefId = StringId<TemplateRef>;
+
+	/// <summary>
+	/// Specifies a pattern of times that this schedule should run. Each schedule may have multiple patterns.
+	/// </summary>
+	public class SchedulePattern
+	{
+		/// <summary>
+		/// Which days of the week the schedule should run
+		/// </summary>
+		public List<DayOfWeek>? DaysOfWeek { get; set; }
+
+		/// <summary>
+		/// Time during the day for the first schedule to trigger. Measured in minutes from midnight.
+		/// </summary>
+		public int MinTime { get; set; }
+
+		/// <summary>
+		/// Time during the day for the last schedule to trigger. Measured in minutes from midnight.
+		/// </summary>
+		public int? MaxTime { get; set; }
+
+		/// <summary>
+		/// Interval between each schedule triggering
+		/// </summary>
+		public int? Interval { get; set; }
+
+		/// <summary>
+		/// Private constructor for serialization
+		/// </summary>
+		private SchedulePattern()
+		{
+		}
+
+		/// <summary>
+		/// Constructor
+		/// </summary>
+		/// <param name="DaysOfWeek">Which days of the week the schedule should run</param>
+		/// <param name="MinTime">Time during the day for the first schedule to trigger. Measured in minutes from midnight.</param>
+		/// <param name="MaxTime">Time during the day for the last schedule to trigger. Measured in minutes from midnight.</param>
+		/// <param name="Interval">Interval between each schedule triggering</param>
+		public SchedulePattern(List<DayOfWeek>? DaysOfWeek, int MinTime, int? MaxTime, int? Interval)
+		{
+			this.DaysOfWeek = DaysOfWeek;
+			this.MinTime = MinTime;
+			this.MaxTime = MaxTime;
+			this.Interval = Interval;
+		}
+
+		/// <summary>
+		/// Calculates the trigger index based on the given time in minutes
+		/// </summary>
+		/// <param name="LastTime">Time in minutes for the last trigger</param>
+		/// <param name="TimeZone">The timezone for running the schedule</param>
+		/// <returns>Index of the trigger</returns>
+		public DateTimeOffset GetNextTriggerTime(DateTimeOffset LastTime, TimeZoneInfo TimeZone)
+		{
+			// Convert last time into the correct timezone for running the scheule
+			LastTime = TimeZoneInfo.ConvertTime(LastTime, TimeZone);
+
+			// Get the base time (ie. the start of this day) for anchoring the schedule
+			DateTimeOffset BaseTime = new DateTimeOffset(LastTime.Year, LastTime.Month, LastTime.Day, 0, 0, 0, LastTime.Offset);
+			for (; ; )
+			{
+				if (DaysOfWeek == null || DaysOfWeek.Contains(BaseTime.DayOfWeek))
+				{
+					// Get the last time in minutes from the start of this day
+					int LastTimeMinutes = (int)(LastTime - BaseTime).TotalMinutes;
+
+					// Get the time of the first trigger of this day. If the last time is less than this, this is the next trigger.
+					if (LastTimeMinutes < MinTime)
+					{
+						return BaseTime.AddMinutes(MinTime);
+					}
+
+					// Otherwise, get the time for the last trigger in the day.
+					if (Interval.HasValue && Interval.Value > 0)
+					{
+						int ActualMaxTime = MaxTime ?? ((24 * 60) - 1);
+						if (LastTimeMinutes < ActualMaxTime)
+						{
+							int LastIndex = (LastTimeMinutes - MinTime) / Interval.Value;
+							int NextIndex = LastIndex + 1;
+
+							int NextTimeMinutes = MinTime + (NextIndex * Interval.Value);
+							if (NextTimeMinutes <= ActualMaxTime)
+							{
+								return BaseTime.AddMinutes(NextTimeMinutes);
+							}
+						}
+					}
+				}
+				BaseTime = BaseTime.AddDays(1.0);
+			}
+		}
+	}
+
+	/// <summary>
+	/// Gate allowing a schedule to trigger.
+	/// </summary>
+	public class ScheduleGate
+	{
+		/// <summary>
+		/// The template containing the dependency
+		/// </summary>
+		public TemplateRefId TemplateRefId { get; set; }
+
+		/// <summary>
+		/// Target to wait for
+		/// </summary>
+		public string Target { get; set; }
+
+		/// <summary>
+		/// Private constructor for serializtaion
+		/// </summary>
+		[BsonConstructor]
+		private ScheduleGate()
+		{
+			Target = String.Empty;
+		}
+
+		/// <summary>
+		/// Constructor
+		/// </summary>
+		/// <param name="TemplateRefId">The template containing the dependency</param>
+		/// <param name="Target">Target to wait for</param>
+		public ScheduleGate(TemplateRefId TemplateRefId, string Target)
+		{
+			this.TemplateRefId = TemplateRefId;
+			this.Target = Target;
+		}
+	}
+
+	/// <summary>
+	/// Represents a schedule
+	/// </summary>
+	public class Schedule
+	{
+		/// <summary>
+		/// Whether this schedule is currently enabled
+		/// </summary>
+		public bool Enabled { get; set; }
+
+		/// <summary>
+		/// Maximum number of builds triggered by this schedule active at once. Set to zero for unlimited.
+		/// </summary>
+		public int MaxActive { get; set; }
+
+		/// <summary>
+		/// Maximum number of changes the schedule can fall behind head revision. If greater than zero, builds will be triggered for every submitted changelist until the backlog is this size.
+		/// </summary>
+		public int MaxChanges { get; set; }
+
+		/// <summary>
+		/// Whether the build requires a change to be submitted
+		/// </summary>
+		[BsonIgnoreIfDefault, BsonDefaultValue(true)]
+		public bool RequireSubmittedChange { get; set; } = true;
+
+		/// <summary>
+		/// Reference to another job/target which much succeed for this schedule to trigger
+		/// </summary>
+		public ScheduleGate? Gate { get; set; }
+
+		/// <summary>
+		/// Whether this build requires a code change to trigger
+		/// </summary>
+		public List<ChangeContentFlags>? Filter { get; set; }
+
+		/// <summary>
+		/// Parameters for the template
+		/// </summary>
+		public Dictionary<string, string> TemplateParameters { get; set; }
+
+		/// <summary>
+		/// List of patterns for this schedule
+		/// </summary>
+		public List<SchedulePattern> Patterns { get; set; }
+
+		/// <summary>
+		/// Last changelist number that this was triggered for
+		/// </summary>
+		public int LastTriggerChange { get; set; }
+
+		/// <summary>
+		/// Last time that the schedule was triggered
+		/// </summary>
+		public DateTimeOffset LastTriggerTime { get; set; }
+
+		/// <summary>
+		/// List of jobs that are currently active
+		/// </summary>
+		public List<ObjectId> ActiveJobs { get; set; } = new List<ObjectId>();
+
+		/// <summary>
+		/// Private constructor for serializtaion
+		/// </summary>
+		[BsonConstructor]
+		private Schedule()
+		{
+			TemplateParameters = new Dictionary<string, string>();
+			Patterns = new List<SchedulePattern>();
+		}
+
+		/// <summary>
+		/// Constructor
+		/// </summary>
+		/// <param name="Enabled">Whether the schedule is currently enabled</param>
+		/// <param name="MaxActive">Maximum number of builds that may be active at once</param>
+		/// <param name="MaxChanges">Maximum number of changes the schedule can fall behind head revision</param>
+		/// <param name="RequireSubmittedChange">Whether a change has to be submitted for the schedule to trigger</param>
+		/// <param name="Gate">Reference to another job/target which much succeed for this schedule to trigger</param>
+		/// <param name="Filter">Filter for changes to consider</param>
+		/// <param name="TemplateParameters">Parameters for the template to run</param>
+		/// <param name="Patterns">List of patterns for the schedule</param>
+		public Schedule(bool Enabled = true, int MaxActive = 0, int MaxChanges = 0, bool RequireSubmittedChange = true, ScheduleGate? Gate = null, List<ChangeContentFlags>? Filter = null, Dictionary<string, string>? TemplateParameters = null, List<SchedulePattern>? Patterns = null)
+		{
+			this.Enabled = Enabled;
+			this.MaxActive = MaxActive;
+			this.MaxChanges = MaxChanges;
+			this.RequireSubmittedChange = RequireSubmittedChange;
+			this.Gate = Gate;
+			this.Filter = Filter;
+			this.TemplateParameters = TemplateParameters ?? new Dictionary<string, string>();
+			this.Patterns = Patterns ?? new List<SchedulePattern>();
+			this.LastTriggerTime = DateTimeOffset.Now;
+		}
+
+		/// <summary>
+		/// Copies the state fields from another schedule object
+		/// </summary>
+		/// <param name="Other">The schedule object to copy from</param>
+		public void CopyState(Schedule Other)
+		{
+			LastTriggerChange = Other.LastTriggerChange;
+			LastTriggerTime = Other.LastTriggerTime;
+			ActiveJobs.AddRange(Other.ActiveJobs);
+		}
+
+		/// <summary>
+		/// Gets the flags to filter changes by
+		/// </summary>
+		/// <returns>Set of filter flags</returns>
+		public ChangeContentFlags? GetFilterFlags()
+		{
+			if(Filter == null)
+			{
+				return null;
+			}
+
+			ChangeContentFlags Flags = 0;
+			foreach (ChangeContentFlags Flag in Filter)
+			{
+				Flags |= Flag;
+			}
+			return Flags;
+		}
+
+		/// <summary>
+		/// Get the next time that the schedule will trigger
+		/// </summary>
+		/// <param name="TimeZone">Timezone to evaluate the trigger</param>
+		/// <returns>Next time at which the schedule will trigger</returns>
+		public DateTimeOffset? GetNextTriggerTime(TimeZoneInfo TimeZone)
+		{
+			return GetNextTriggerTime(LastTriggerTime, TimeZone);
+		}
+
+		/// <summary>
+		/// Get the next time that the schedule will trigger
+		/// </summary>
+		/// <param name="LastTime">Last time at which the schedule triggered</param>
+		/// <param name="TimeZone">Timezone to evaluate the trigger</param>
+		/// <returns>Next time at which the schedule will trigger</returns>
+		public DateTimeOffset? GetNextTriggerTime(DateTimeOffset LastTime, TimeZoneInfo TimeZone)
+		{
+			DateTimeOffset? NextTriggerTime = null;
+			if (Enabled)
+			{
+				foreach (SchedulePattern Pattern in Patterns)
+				{
+					DateTimeOffset PatternTriggerTime = Pattern.GetNextTriggerTime(LastTime, TimeZone);
+					if (NextTriggerTime == null || NextTriggerTime < PatternTriggerTime)
+					{
+						NextTriggerTime = PatternTriggerTime;
+					}
+				}
+			}
+			return NextTriggerTime;
+		}
+	}
+}
diff --git a/Engine/Source/Programs/Horde/HordeServer/Secrets/Credential.cs b/Engine/Source/Programs/Horde/HordeServer/Secrets/Credential.cs
new file mode 100644
--- /dev/null
+++ b/Engine/Source/Programs/Horde/HordeServer/Secrets/Credential.cs
@@ -0,0 +1,78 @@
+// Copyright Epic Games, Inc. All Rights Reserved.
+
+using HordeServer.Api;
+using HordeServer.Utilities;
+using MongoDB.Bson;
+using MongoDB.Bson.Serialization.Attributes;
+using System;
+using System.Collections.Generic;
+using System.Linq;
+using System.Threading.Tasks;
+
+namespace HordeServer.Models
+{
+	/// <summary>
+	/// Stores information about a credential
+	/// </summary>
+	public class Credential
+	{
+		/// <summary>
+		/// Unique id for this credential
+		/// </summary>
+		[BsonId]
+		public ObjectId Id { get; set; }
+
+		/// <summary>
+		/// Name of this credential
+		/// </summary>
+		[BsonRequired]
+		public string Name { get; set; }
+
+		/// <summary>
+		/// The normalized name of this credential
+		/// </summary>
+		[BsonRequired]
+		public string NormalizedName { get; set; }
+
+		/// <summary>
+		/// Properties for this credential
+		/// </summary>
+		public Dictionary<string, string> Properties { get; set; } = new Dictionary<string, string>();
+
+		/// <summary>
+		/// The ACL for this credential
+		/// </summary>
+		public Acl? Acl { get; set; }
+
+		/// <summary>
+		/// Private constructor for serialization
+		/// </summary>
+		[BsonConstructor]
+		private Credential()
+		{
+			this.Name = null!;
+			this.NormalizedName = null!;
+		}
+
+		/// <summary>
+		/// Constructor
+		/// </summary>
+		/// <param name="Name">Name of this credential</param>
+		public Credential(string Name)
+		{
+			this.Id = ObjectId.GenerateNewId();
+			this.Name = Name;
+			this.NormalizedName = GetNormalizedName(Name);
+		}
+
+		/// <summary>
+		/// Gets the normalized form of a name, used for case-insensitive comparisons
+		/// </summary>
+		/// <param name="Name">The name to normalize</param>
+		/// <returns>The normalized name</returns>
+		public static string GetNormalizedName(string Name)
+		{
+			return Name.ToUpperInvariant();
+		}
+	}
+}
diff --git a/Engine/Source/Programs/Horde/HordeServer/Secrets/CredentialService.cs b/Engine/Source/Programs/Horde/HordeServer/Secrets/CredentialService.cs
new file mode 100644
--- /dev/null
+++ b/Engine/Source/Programs/Horde/HordeServer/Secrets/CredentialService.cs
@@ -0,0 +1,178 @@
+// Copyright Epic Games, Inc. All Rights Reserved.
+
+using HordeServer.Api;
+using HordeServer.Models;
+using MongoDB.Bson;
+using MongoDB.Driver;
+using System;
+using System.Collections.Generic;
+using System.Linq;
+using System.Security.Claims;
+using System.Threading.Tasks;
+
+namespace HordeServer.Services
+{
+	/// <summary>
+	/// Service for accessing and modifying credentials
+	/// </summary>
+	public class CredentialService
+	{
+		/// <summary>
+		/// The ACL service instance
+		/// </summary>
+		AclService AclService;
+
+		/// <summary>
+		/// Collection of credential documents
+		/// </summary>
+		IMongoCollection<Credential> Credentials; 
+
+		/// <summary>
+		/// Constructor
+		/// </summary>
+		/// <param name="AclService">The ACL service instance</param>
+		/// <param name="DatabaseService">The database service instance</param>
+		public CredentialService(AclService AclService, DatabaseService DatabaseService)
+		{
+			this.AclService = AclService;
+
+			Credentials = DatabaseService.Credentials;
+
+			if (!DatabaseService.ReadOnlyMode)
+			{
+				Credentials.Indexes.CreateOne(new CreateIndexModel<Credential>(Builders<Credential>.IndexKeys.Ascending(x => x.NormalizedName), new CreateIndexOptions { Unique = true }));
+			}
+		}
+
+		/// <summary>
+		/// Creates a new credential
+		/// </summary>
+		/// <param name="Name">Name of the new credential</param>
+		/// <param name="Properties">Properties for the new credential</param>
+		/// <returns>The new credential document</returns>
+		public async Task<Credential> CreateCredentialAsync(string Name, Dictionary<string, string>? Properties)
+		{
+			Credential NewCredential = new Credential(Name);
+			if (Properties != null)
+			{
+				NewCredential.Properties = new Dictionary<string, string>(Properties, NewCredential.Properties.Comparer);
+			}
+
+			await Credentials.InsertOneAsync(NewCredential);
+			return NewCredential;
+		}
+
+		/// <summary>
+		/// Updates an existing credential
+		/// </summary>
+		/// <param name="Id">Unique id of the credential</param>
+		/// <param name="NewName">The new name for the credential</param>
+		/// <param name="NewProperties">Properties on the credential to update. Any properties with a value of null will be removed.</param>
+		/// <param name="NewAcl">The new ACL for the credential</param>
+		/// <returns>Async task object</returns>
+		public async Task UpdateCredentialAsync(ObjectId Id, string? NewName, Dictionary<string, string?>? NewProperties, Acl? NewAcl)
+		{
+			UpdateDefinitionBuilder<Credential> UpdateBuilder = Builders<Credential>.Update;
+
+			List<UpdateDefinition<Credential>> Updates = new List<UpdateDefinition<Credential>>();
+			if (NewName != null)
+			{
+				Updates.Add(UpdateBuilder.Set(x => x.Name, NewName));
+				Updates.Add(UpdateBuilder.Set(x => x.NormalizedName, Credential.GetNormalizedName(NewName)));
+			}
+			if (NewProperties != null)
+			{
+				foreach (KeyValuePair<string, string?> Pair in NewProperties)
+				{
+					if (Pair.Value == null)
+					{
+						Updates.Add(UpdateBuilder.Unset(x => x.Properties[Pair.Key]));
+					}
+					else
+					{
+						Updates.Add(UpdateBuilder.Set(x => x.Properties[Pair.Key], Pair.Value));
+					}
+				}
+			}
+			if (NewAcl != null)
+			{
+				Updates.Add(Acl.CreateUpdate<Credential>(x => x.Acl!, NewAcl));
+			}
+
+			if (Updates.Count > 0)
+			{
+				await Credentials.FindOneAndUpdateAsync<Credential>(x => x.Id == Id, UpdateBuilder.Combine(Updates));
+			}
+		}
+
+		/// <summary>
+		/// Gets all the available credentials
+		/// </summary>
+		/// <returns>List of project documents</returns>
+		public Task<List<Credential>> FindCredentialsAsync(string? Name)
+		{
+			FilterDefinitionBuilder<Credential> FilterBuilder = Builders<Credential>.Filter;
+
+			FilterDefinition<Credential> Filter = FilterDefinition<Credential>.Empty;
+			if (Name != null)
+			{
+				Filter &= FilterBuilder.Eq(x => x.NormalizedName, Credential.GetNormalizedName(Name));
+			}
+
+			return Credentials.Find(Filter).ToListAsync();
+		}
+
+		/// <summary>
+		/// Gets a credential by ID
+		/// </summary>
+		/// <param name="Id">Unique id of the credential</param>
+		/// <returns>The credential document</returns>
+		public async Task<Credential?> GetCredentialAsync(ObjectId Id)
+		{
+			return await Credentials.Find<Credential>(x => x.Id == Id).FirstOrDefaultAsync();
+		}
+
+		/// <summary>
+		/// Gets a credential by name
+		/// </summary>
+		/// <param name="Name">Name of the credential</param>
+		/// <returns>The credential document</returns>
+		public async Task<Credential?> GetCredentialAsync(string Name)
+		{
+			string NormalizedName = Credential.GetNormalizedName(Name);
+			return await Credentials.Find<Credential>(x => x.NormalizedName == NormalizedName).FirstOrDefaultAsync();
+		}
+
+		/// <summary>
+		/// Deletes a credential by id
+		/// </summary>
+		/// <param name="Id">Unique id of the credential</param>
+		/// <returns>True if the credential was deleted</returns>
+		public async Task<bool> DeleteCredentialAsync(ObjectId Id)
+		{
+			DeleteResult Result = await Credentials.DeleteOneAsync<Credential>(x => x.Id == Id);
+			return Result.DeletedCount > 0;
+		}
+
+		/// <summary>
+		/// Determines if the user is authorized to perform an action on a particular credential
+		/// </summary>
+		/// <param name="Credential">The credential to check</param>
+		/// <param name="Action">The action being performed</param>
+		/// <param name="User">The principal to authorize</param>
+		/// <param name="Cache">Permissions cache</param>
+		/// <returns>True if the action is authorized</returns>
+		public Task<bool> AuthorizeAsync(Credential Credential, AclAction Action, ClaimsPrincipal User, GlobalPermissionsCache? Cache)
+		{
+			bool? Result = Credential.Acl?.Authorize(Action, User);
+			if (Result == null)
+			{
+				return AclService.AuthorizeAsync(Action, User, Cache);
+			}
+			else
+			{
+				return Task.FromResult(Result.Value);
+			}
+		}
+	}
+}
diff --git a/Engine/Source/Programs/Horde/HordeServer/Server/Config.cs b/Engine/Source/Programs/Horde/HordeServer/Server/Config.cs
new file mode 100644
--- /dev/null
+++ b/Engine/Source/Programs/Horde/HordeServer/Server/Config.cs
@@ -0,0 +1,257 @@
+// Copyright Epic Games, Inc. All Rights Reserved.
+
+using EpicGames.Core;
+using HordeServer.Api;
+using HordeServer.Utilities;
+using System;
+using System.Collections.Generic;
+using System.ComponentModel.DataAnnotations;
+using System.Linq;
+using System.Threading.Tasks;
+
+namespace HordeServer.Models
+{
+	using ProjectId = StringId<IProject>;
+	using StreamId = StringId<IStream>;
+
+	/// <summary>
+	/// Global configuration
+	/// </summary>
+	[JsonSchema("https://unrealengine.com/horde/global")]
+	[JsonSchemaCatalog("Horde Globals", "Horde global configuration file", "globals.json")]
+	public class GlobalConfig
+	{
+		/// <summary>
+		/// List of projects
+		/// </summary>
+		public List<ProjectConfigRef> Projects { get; set; } = new List<ProjectConfigRef>();
+
+		/// <summary>
+		/// Manually added status messages
+		/// </summary>
+		public List<Notice> Notices { get; set; } = new List<Notice>();
+
+		/// <summary>
+		/// List of scheduled downtime
+		/// </summary>
+		public List<ScheduledDowntime> Downtime { get; set; } = new List<ScheduledDowntime>();
+
+		/// <summary>
+		/// List of Perforce clusters
+		/// </summary>
+		public List<PerforceCluster> PerforceClusters { get; set; } = new List<PerforceCluster>();
+
+		/// <summary>
+		/// Maximum number of conforms to run at once
+		/// </summary>
+		public int MaxConformCount { get; set; }
+
+		/// <summary>
+		/// List of storage namespaces
+		/// </summary>
+		public StorageConfig? Storage { get; set; }
+
+		/// <summary>
+		/// Access control list
+		/// </summary>
+		public UpdateAclRequest? Acl { get; set; }
+	}
+
+	/// <summary>
+	/// References a project configuration
+	/// </summary>
+	public class ProjectConfigRef
+	{
+		/// <summary>
+		/// Unique id for the project
+		/// </summary>
+		[Required]
+		public ProjectId Id { get; set; } = ProjectId.Empty;
+
+		/// <summary>
+		/// Config path for the project
+		/// </summary>
+		[Required]
+		public string Path { get; set; } = null!;
+	}
+
+	/// <summary>
+	/// Stores configuration for a project
+	/// </summary>
+	[JsonSchema("https://unrealengine.com/horde/project")]
+	[JsonSchemaCatalog("Horde Project", "Horde project configuration file", new[] { "*.project.json", "Projects/*.json" })]
+	public class ProjectConfig
+	{
+		/// <summary>
+		/// Name for the new project
+		/// </summary>
+		[Required]
+		public string Name { get; set; } = null!;
+
+		/// <summary>
+		/// Path to the project logo
+		/// </summary>
+		public string? Logo { get; set; }
+
+		/// <summary>
+		/// Categories to include in this project
+		/// </summary>
+		public List<CreateProjectCategoryRequest> Categories { get; set; } = new List<CreateProjectCategoryRequest>();
+
+		/// <summary>
+		/// List of streams
+		/// </summary>
+		public List<StreamConfigRef> Streams { get; set; } = new List<StreamConfigRef>();
+
+		/// <summary>
+		/// Acl entries
+		/// </summary>
+		public UpdateAclRequest? Acl { get; set; }
+	}
+
+	/// <summary>
+	/// Reference to configuration for a stream
+	/// </summary>
+	public class StreamConfigRef
+	{
+		/// <summary>
+		/// Unique id for the stream
+		/// </summary>
+		[Required]
+		public StreamId Id { get; set; } = StreamId.Empty;
+
+		/// <summary>
+		/// Path to the configuration file
+		/// </summary>
+		[Required]
+		public string Path { get; set; } = null!;
+	}
+
+	/// <summary>
+	/// Config for a stream
+	/// </summary>
+	[JsonSchema("https://unrealengine.com/horde/stream")]
+	[JsonSchemaCatalog("Horde Stream", "Horde stream configuration file", new[] { "*.stream.json", "Streams/*.json" })]
+	public class StreamConfig
+	{
+		/// <summary>
+		/// Name of the stream
+		/// </summary>
+		[Required]
+		public string Name { get; set; } = String.Empty;
+
+		/// <summary>
+		/// The perforce cluster containing the stream
+		/// </summary>
+		public string? ClusterName { get; set; }
+
+		/// <summary>
+		/// Order for this stream
+		/// </summary>
+		public int? Order { get; set; }
+
+		/// <summary>
+		/// Notification channel for all jobs in this stream
+		/// </summary>
+		public string? NotificationChannel { get; set; }
+
+		/// <summary>
+		/// Notification channel filter for this template. Can be Success|Failure|Warnings
+		/// </summary>
+		public string? NotificationChannelFilter { get; set; }
+
+		/// <summary>
+		/// Channel to post issue triage notifications
+		/// </summary>
+		public string? TriageChannel { get; set; }
+
+		/// <summary>
+		/// Default template for running preflights
+		/// </summary>
+		public string? DefaultPreflightTemplate { get; set; }
+
+		/// <summary>
+		/// Default template for running preflights
+		/// </summary>
+		public DefaultPreflightRequest? DefaultPreflight { get; set; }
+
+		/// <summary>
+		/// List of tabs to show for the new stream
+		/// </summary>
+		public List<CreateStreamTabRequest> Tabs { get; set; } = new List<CreateStreamTabRequest>();
+
+		/// <summary>
+		/// Map of agent name to type
+		/// </summary>
+		public Dictionary<string, CreateAgentTypeRequest> AgentTypes { get; set; } = new Dictionary<string, CreateAgentTypeRequest>();
+
+		/// <summary>
+		/// Map of workspace name to type
+		/// </summary>
+		public Dictionary<string, CreateWorkspaceTypeRequest> WorkspaceTypes { get; set; } = new Dictionary<string, CreateWorkspaceTypeRequest>();
+
+		/// <summary>
+		/// List of templates to create
+		/// </summary>
+		public List<CreateTemplateRefRequest> Templates { get; set; } = new List<CreateTemplateRefRequest>();
+
+		/// <summary>
+		/// Custom permissions for this object
+		/// </summary>
+		public UpdateAclRequest? Acl { get; set; }
+
+		/// <summary>
+		/// Pause stream builds until specified date
+		/// </summary>
+		public DateTime? PausedUntil { get; set; }
+
+		/// <summary>
+		/// Reason for pausing builds of the stream
+		/// </summary>
+		public string? PauseComment { get; set; }
+	}
+
+	/// <summary>
+	/// Configuration for storage system
+	/// </summary>
+	public class StorageConfig
+	{
+		/// <summary>
+		/// List of storage namespaces
+		/// </summary>
+		public List<NamespaceConfig> Namespaces { get; set; } = new List<NamespaceConfig>();
+	}
+
+	/// <summary>
+	/// Configuration for a storage namespace
+	/// </summary>
+	public class NamespaceConfig
+	{
+		/// <summary>
+		/// Identifier for this namespace
+		/// </summary>
+		public string Id { get; set; } = String.Empty;
+
+		/// <summary>
+		/// Buckets within this namespace
+		/// </summary>
+		public List<BucketConfig> Buckets { get; set; } = new List<BucketConfig>();
+
+		/// <summary>
+		/// Access control for this namespace
+		/// </summary>
+		public UpdateAclRequest? Acl { get; set; }
+	}
+
+	/// <summary>
+	/// Configuration for a bucket
+	/// </summary>
+	public class BucketConfig
+	{
+		/// <summary>
+		/// Identifier for the bucket
+		/// </summary>
+		public string Id { get; set; } = String.Empty;
+	}
+}
+
diff --git a/Engine/Source/Programs/Horde/HordeServer/Server/LifetimeService.cs b/Engine/Source/Programs/Horde/HordeServer/Server/LifetimeService.cs
new file mode 100644
--- /dev/null
+++ b/Engine/Source/Programs/Horde/HordeServer/Server/LifetimeService.cs
@@ -0,0 +1,285 @@
+// Copyright Epic Games, Inc. All Rights Reserved.
+
+using Microsoft.AspNetCore.Hosting;
+using Microsoft.Extensions.Hosting;
+using Microsoft.Extensions.Logging;
+using System;
+using System.Collections.Generic;
+using System.Diagnostics.CodeAnalysis;
+using System.Linq;
+using System.Threading;
+using System.Threading.Tasks;
+using Grpc.Core;
+using Grpc.Core.Interceptors;
+using HordeServer.Utilities;
+using StackExchange.Redis;
+using System.Diagnostics;
+
+namespace HordeServer.Services
+{
+	/// <summary>
+	/// Service containing an async task that allows long polling operations to complete early if the server is shutting down
+	/// </summary>
+	public sealed class LifetimeService : IDisposable
+	{
+		/// <summary>
+		/// Writer for log output
+		/// </summary>
+		ILogger<LifetimeService> Logger;
+
+		/// <summary>
+		/// Task source for the server stopping
+		/// </summary>
+		TaskCompletionSource<bool> StoppingTaskCompletionSource;
+
+		/// <summary>
+		/// Task source for the server stopping
+		/// </summary>
+		TaskCompletionSource<bool> PreStoppingTaskCompletionSource;
+
+		/// <summary>
+		/// Registration token for the stopping event
+		/// </summary>
+		CancellationTokenRegistration Registration;
+		
+		/// <summary>
+		/// Singleton instance of the database service
+		/// </summary>
+		DatabaseService DatabaseService;
+		
+		/// <summary>
+		/// The Redis database
+		/// </summary>
+		IDatabase RedisDb;
+		
+		/// <summary>
+		/// Singleton instance of the request tracker service
+		/// </summary>
+		RequestTrackerService RequestTrackerService;
+		
+		/*
+		/// <summary>
+		/// Max time to wait for any outstanding requests to finish
+		/// </summary>
+		readonly TimeSpan RequestGracefulTimeout = TimeSpan.FromMinutes(5);
+		
+		/// <summary>
+		/// Initial delay before attempting the shutdown. This to ensure any load balancers/ingress will detect
+		/// the server is unavailable to serve new requests in the event of no outstanding requests to wait for.
+		/// </summary>
+		readonly TimeSpan InitialStoppingDelay = TimeSpan.FromSeconds(35);
+		*/
+
+		/// <summary>
+		/// Constructor
+		/// </summary>
+		/// <param name="Lifetime">Application lifetime interface</param>
+		/// <param name="DatabaseService">Database singleton service</param>
+		/// <param name="RedisDb">Redis singleton service</param>
+		/// <param name="RequestTrackerService">Request tracker singleton service</param>
+		/// <param name="Logger">Logging interface</param>
+		public LifetimeService(IHostApplicationLifetime Lifetime, DatabaseService DatabaseService, IDatabase RedisDb, RequestTrackerService RequestTrackerService, ILogger<LifetimeService> Logger)
+		{
+			this.DatabaseService = DatabaseService;
+			this.RedisDb = RedisDb;
+			this.RequestTrackerService = RequestTrackerService;
+			this.Logger = Logger;
+			this.StoppingTaskCompletionSource = new TaskCompletionSource<bool>();
+			this.PreStoppingTaskCompletionSource = new TaskCompletionSource<bool>();
+			this.Registration = Lifetime.ApplicationStopping.Register(ApplicationStopping);
+		}
+
+		/// <inheritdoc/>
+		public void Dispose()
+		{
+			Registration.Dispose();
+		}
+
+		/// <summary>
+		/// Callback for the application stopping
+		/// </summary>
+		void ApplicationStopping()
+		{
+			Logger.LogInformation("SIGTERM signal received");
+			IsPreStopping = true;
+			IsStopping = true;
+
+			PreStoppingTaskCompletionSource.TrySetResult(true);
+			StoppingTaskCompletionSource.TrySetResult(true);
+
+			int ShutdownDelayMs = 30 * 1000;
+			Logger.LogInformation($"Delaying shutdown by sleeping {ShutdownDelayMs} ms...");
+			Thread.Sleep(ShutdownDelayMs);
+			Logger.LogInformation($"Server process now shutting down...");
+			
+			/*
+			if (PreStoppingTaskCompletionSource.TrySetResult(true))
+			{
+				
+				WaitAndTriggerStoppingTask = Task.Run(() => ExecStoppingTask());
+				
+				Logger.LogInformation("App is stopping. Waiting an initial {InitialDelay} secs before waiting on any requests...", (int)InitialStoppingDelay.TotalSeconds);
+				Thread.Sleep(InitialStoppingDelay);
+				Logger.LogInformation("Blocking shutdown for up to {MaxGraceTimeout} secs until all request have finished...", (int)RequestGracefulTimeout.TotalSeconds);
+				
+				DateTime StartTime = DateTime.UtcNow;
+				do
+				{
+					RequestTrackerService.LogRequestsInProgress();
+					Thread.Sleep(5000);
+				}
+				while (DateTime.UtcNow < StartTime + RequestGracefulTimeout && RequestTrackerService.GetRequestsInProgress().Count > 0);
+
+				if (RequestTrackerService.GetRequestsInProgress().Count == 0)
+				{
+					Logger.LogInformation("All open requests finished gracefully after {TimeTaken} secs", (DateTime.UtcNow - StartTime).TotalSeconds);	
+				}
+				else
+				{
+					Logger.LogInformation("One or more requests did not finish within the grace period of {TimeTaken} secs. Shutdown will now resume with risk of interrupting those requests!", (DateTime.UtcNow - StartTime).TotalSeconds);
+					RequestTrackerService.LogRequestsInProgress();
+				}
+			}
+			*/
+		}
+
+		async Task ExecStoppingTask()
+		{
+			await Task.Delay(TimeSpan.FromSeconds(15.0));
+
+			Logger.LogInformation("Setting stopping event");
+			IsStopping = true;
+			StoppingTaskCompletionSource.TrySetResult(true);
+		}
+
+		/// <summary>
+		/// Returns true if the server is stopping
+		/// </summary>
+		public bool IsStopping { get; private set; }
+
+		/// <summary>
+		/// Returns true if the server is stopping, but may not be removed from the load balancer yet
+		/// </summary>
+		public bool IsPreStopping { get; private set; }
+
+		/// <summary>
+		/// Gets an awaitable task for the server stopping
+		/// </summary>
+		public Task StoppingTask
+		{
+			get { return StoppingTaskCompletionSource.Task; }
+		}
+		
+		/// <summary>
+		/// Check if MongoDB can be reached
+		/// </summary>
+		/// <returns>True if communication works</returns>
+		[SuppressMessage("Design", "CA1031:Do not catch general exception types")]
+		public async Task<bool> IsMongoDbConnectionHealthy()
+		{
+			using CancellationTokenSource CancelSource = new CancellationTokenSource(10000);
+			bool IsHealthy = false;
+			try
+			{
+				await DatabaseService.Database.ListCollectionNamesAsync(null, CancelSource.Token);
+				IsHealthy = true;
+			}
+			catch (Exception e)
+			{
+				Logger.LogError("MongoDB call failed during health check", e);
+			}
+			
+			return IsHealthy;
+		}
+		
+		/// <summary>
+		/// Check if Redis can be reached
+		/// </summary>
+		/// <returns>True if communication works</returns>
+		[SuppressMessage("Design", "CA1031:Do not catch general exception types")]
+		public async Task<bool> IsRedisConnectionHealthy()
+		{
+			using CancellationTokenSource CancelSource = new CancellationTokenSource(10000);
+			bool IsHealthy = false;
+			try
+			{
+				string Key = "HordeLifetimeService-Health-Check";
+				await RedisDb.StringSetAsync(Key, "ok");
+				await RedisDb.StringGetAsync(Key);
+				IsHealthy = true;
+			}
+			catch (Exception e)
+			{
+				Logger.LogError("Redis call failed during health check", e);
+			}
+			
+			return IsHealthy;
+		}
+	}
+
+	/// <summary>
+	/// Intercept incoming gRPC calls and raise an abort exception if the server is shutting down
+	///
+	/// This is to avoid clients keeping connections to a server that is about to shut down.
+	/// The client must handle the RPC exception raised and simply retry the call on a different gRPC channel/connection.
+	/// By retrying the call on a new channel, it should hit a new valid server as the load balancer will not direct traffic
+	/// to a server process shutting down.
+	/// </summary>
+	public class LifetimeGrpcInterceptor : Interceptor
+	{
+		private readonly LifetimeService LifetimeService;
+		private readonly ILogger<LifetimeGrpcInterceptor> Logger;
+
+		/// <summary>
+		/// Constructor
+		/// </summary>
+		/// <param name="LifetimeService"></param>
+		/// <param name="Logger"></param>
+		public LifetimeGrpcInterceptor(LifetimeService LifetimeService, ILogger<LifetimeGrpcInterceptor> Logger)
+		{
+			this.LifetimeService = LifetimeService;
+			this.Logger = Logger;
+		}
+
+		[SuppressMessage("Performance", "CA1822:Mark members as static")]
+		[SuppressMessage("Usage", "CA1801:Review unused parameters")]
+		private void CheckIfServerIsBeingStopped(ServerCallContext Context)
+		{
+			/*
+			if (LifetimeService.IsStopping)
+			{
+				Logger.LogDebug("Refusing to serve {Method} as server process is being shut down.", Context.Method);
+				throw new RpcException(new Status(StatusCode.Aborted, "Server is shutting down."));
+			}
+			*/
+		}
+
+		/// <inheritdoc />
+		public override Task<TResponse> UnaryServerHandler<TRequest, TResponse>(TRequest Request, ServerCallContext Context, UnaryServerMethod<TRequest, TResponse> Continuation)
+		{
+			CheckIfServerIsBeingStopped(Context);
+			return base.UnaryServerHandler(Request, Context, Continuation);
+		}
+
+		/// <inheritdoc />
+		public override Task<TResponse> ClientStreamingServerHandler<TRequest, TResponse>(IAsyncStreamReader<TRequest> RequestStream, ServerCallContext Context, ClientStreamingServerMethod<TRequest, TResponse> Continuation)
+		{
+			CheckIfServerIsBeingStopped(Context);
+			return base.ClientStreamingServerHandler(RequestStream, Context, Continuation);
+		}
+
+		/// <inheritdoc />
+		public override Task ServerStreamingServerHandler<TRequest, TResponse>(TRequest Request, IServerStreamWriter<TResponse> ResponseStream, ServerCallContext Context, ServerStreamingServerMethod<TRequest, TResponse> Continuation)
+		{
+			CheckIfServerIsBeingStopped(Context);
+			return base.ServerStreamingServerHandler(Request, ResponseStream, Context, Continuation);
+		}
+
+		/// <inheritdoc />
+		public override Task DuplexStreamingServerHandler<TRequest, TResponse>(IAsyncStreamReader<TRequest> RequestStream, IServerStreamWriter<TResponse> ResponseStream, ServerCallContext Context, DuplexStreamingServerMethod<TRequest, TResponse> Continuation)
+		{
+			CheckIfServerIsBeingStopped(Context);
+			return base.DuplexStreamingServerHandler(RequestStream, ResponseStream, Context, Continuation);
+		}
+	}
+}
\ No newline at end of file
diff --git a/Engine/Source/Programs/Horde/HordeServer/Services/AclService.cs b/Engine/Source/Programs/Horde/HordeServer/Services/AclService.cs
deleted file mode 100644
diff --git a/Engine/Source/Programs/Horde/HordeServer/Services/CredentialService.cs b/Engine/Source/Programs/Horde/HordeServer/Services/CredentialService.cs
deleted file mode 100644
diff --git a/Engine/Source/Programs/Horde/HordeServer/Services/LifetimeService.cs b/Engine/Source/Programs/Horde/HordeServer/Services/LifetimeService.cs
deleted file mode 100644
diff --git a/Engine/Source/Programs/Horde/HordeServer/Startup.cs b/Engine/Source/Programs/Horde/HordeServer/Startup.cs
--- a/Engine/Source/Programs/Horde/HordeServer/Startup.cs
+++ b/Engine/Source/Programs/Horde/HordeServer/Startup.cs
@@ -331,7 +331,6 @@ namespace HordeServer
 			Services.AddSingleton<ISessionCollection, SessionCollection>();
 			Services.AddSingleton<IServiceAccountCollection, ServiceAccountCollection>();
 			Services.AddSingleton<ISubscriptionCollection, SubscriptionCollection>();
-			Services.AddSingleton<ISoftwareCollection, SoftwareCollection>();
 			Services.AddSingleton<IStreamCollection, StreamCollection>();
 			Services.AddSingleton<ITemplateCollection, TemplateCollection>();
 			Services.AddSingleton<ITestDataCollection, TestDataCollection>();
@@ -408,7 +407,6 @@ namespace HordeServer
 			Services.AddSingleton<SlackNotificationSink>();
 			Services.AddSingleton<IAvatarService, SlackNotificationSink>(SP => SP.GetRequiredService<SlackNotificationSink>());
 			Services.AddSingleton<INotificationSink, SlackNotificationSink>(SP => SP.GetRequiredService<SlackNotificationSink>());
-			Services.AddSingleton<SoftwareService>();
 			Services.AddSingleton<StreamService>();
 			Services.AddSingleton<TemplateService>();
 			Services.AddSingleton<UpgradeService>();
diff --git a/Engine/Source/Programs/Horde/HordeServer/Ugs/IUgsMetadataCollection.cs b/Engine/Source/Programs/Horde/HordeServer/Ugs/IUgsMetadataCollection.cs
new file mode 100644
--- /dev/null
+++ b/Engine/Source/Programs/Horde/HordeServer/Ugs/IUgsMetadataCollection.cs
@@ -0,0 +1,66 @@
+// Copyright Epic Games, Inc. All Rights Reserved.
+
+using HordeServer.Api;
+using HordeServer.Models;
+using HordeServer.Services;
+using HordeServer.Utilities;
+using MongoDB.Bson;
+using MongoDB.Bson.Serialization.Attributes;
+using MongoDB.Driver;
+using MongoDB.Driver.Core.Authentication;
+using System;
+using System.Collections.Generic;
+using System.Diagnostics.CodeAnalysis;
+using System.Linq;
+using System.Threading.Tasks;
+
+namespace HordeServer.Collections
+{
+	/// <summary>
+	/// Collection of stream documents
+	/// </summary>
+	public interface IUgsMetadataCollection
+	{
+		/// <summary>
+		/// Finds or adds metadata for a given changelist
+		/// </summary>
+		/// <param name="Stream">The stream containing the change</param>
+		/// <param name="Change">The changelist number</param>
+		/// <param name="Project">Project identifier</param>
+		/// <returns>The metadata instance</returns>
+		Task<IUgsMetadata> FindOrAddAsync(string Stream, int Change, string? Project);
+
+		/// <summary>
+		/// Adds information to a change
+		/// </summary>
+		/// <param name="Metadata">The existing metadata</param>
+		/// <param name="UserName">Name of the user posting the change</param>
+		/// <param name="Synced">Time that the change was synced</param>
+		/// <param name="Vote">Vote for this change</param>
+		/// <param name="Investigating">Whether the user is investigating the change</param>
+		/// <param name="Starred">Whether the change is starred</param>
+		/// <param name="Comment">New comment from this user</param>
+		/// <returns>Async task task</returns>
+		Task<IUgsMetadata> UpdateUserAsync(IUgsMetadata Metadata, string UserName, bool? Synced, UgsUserVote? Vote, bool? Investigating, bool? Starred, string? Comment);
+
+		/// <summary>
+		/// Updates the state of a badge
+		/// </summary>
+		/// <param name="Metadata">The existing metadata</param>
+		/// <param name="Name">Name of the badge</param>
+		/// <param name="Url">Url to link to for the badge</param>
+		/// <param name="State">State of the badge</param>
+		/// <returns>Async task</returns>
+		Task<IUgsMetadata> UpdateBadgeAsync(IUgsMetadata Metadata, string Name, Uri? Url, UgsBadgeState State);
+
+		/// <summary>
+		/// Searches for metadata updates
+		/// </summary>
+		/// <param name="Stream">The stream to search</param>
+		/// <param name="MinChange">Minimum changelist number</param>
+		/// <param name="MaxChange">Maximum changelist number</param>
+		/// <param name="AfterTicks">Last query time</param>
+		/// <returns>List of metadata updates</returns>
+		Task<List<IUgsMetadata>> FindAsync(string Stream, int MinChange, int? MaxChange = null, long? AfterTicks = null);
+	}
+}
diff --git a/Engine/Source/Programs/Horde/HordeServer/Ugs/Impl/UgsMetadataCollection.cs b/Engine/Source/Programs/Horde/HordeServer/Ugs/Impl/UgsMetadataCollection.cs
new file mode 100644
--- /dev/null
+++ b/Engine/Source/Programs/Horde/HordeServer/Ugs/Impl/UgsMetadataCollection.cs
@@ -0,0 +1,364 @@
+// Copyright Epic Games, Inc. All Rights Reserved.
+
+using HordeServer.Api;
+using HordeServer.Models;
+using HordeServer.Services;
+using HordeServer.Utilities;
+using MongoDB.Bson;
+using MongoDB.Bson.Serialization.Attributes;
+using MongoDB.Driver;
+using MongoDB.Driver.Core.Authentication;
+using System;
+using System.Collections.Generic;
+using System.Diagnostics.CodeAnalysis;
+using System.Linq;
+using System.Threading.Tasks;
+
+namespace HordeServer.Collections.Impl
+{
+	/// <summary>
+	/// Concrete implementation of IUgsMetadataCollection
+	/// </summary>
+	class UgsMetadataCollection : IUgsMetadataCollection
+	{
+		class UgsUserDataDocument : IUgsUserData
+		{
+			public string User { get; set; }
+
+			[BsonIgnoreIfNull]
+			public long? SyncTime { get; set; }
+
+			[BsonIgnoreIfNull]
+			public UgsUserVote? Vote { get; set; }
+
+			[BsonIgnoreIfNull]
+			public bool? Starred { get; set; }
+
+			[BsonIgnoreIfNull]
+			public bool? Investigating { get; set; }
+
+			[BsonIgnoreIfNull]
+			public string? Comment { get; set; }
+
+			long? IUgsUserData.SyncTime => SyncTime;
+			UgsUserVote IUgsUserData.Vote => Vote ?? UgsUserVote.None;
+
+			public UgsUserDataDocument(string User)
+			{
+				this.User = User;
+			}
+		}
+
+		class UgsBadgeDataDocument : IUgsBadgeData
+		{
+			public string Name { get; set; } = String.Empty;
+
+			[BsonIgnoreIfNull]
+			public Uri? Url { get; set; }
+
+			public UgsBadgeState State { get; set; }
+		}
+
+		class UgsMetadataDocument : IUgsMetadata
+		{
+			public ObjectId Id { get; set; }
+			public string Stream { get; set; }
+			public int Change { get; set; }
+			public string Project { get; set; } = String.Empty;
+			public List<UgsUserDataDocument> Users { get; set; } = new List<UgsUserDataDocument>();
+			public List<UgsBadgeDataDocument> Badges { get; set; } = new List<UgsBadgeDataDocument>();
+			public int UpdateIndex { get; set; }
+			public long UpdateTicks { get; set; }
+
+			IReadOnlyList<IUgsUserData>? IUgsMetadata.Users => (Users.Count > 0)? Users : null;
+			IReadOnlyList<IUgsBadgeData>? IUgsMetadata.Badges => (Badges.Count > 0) ? Badges : null;
+
+			[BsonConstructor]
+			private UgsMetadataDocument()
+			{
+				this.Stream = String.Empty;
+			}
+
+			public UgsMetadataDocument(string Stream, int Change, string Project)
+			{
+				this.Id = ObjectId.GenerateNewId();
+				this.Stream = Stream;
+				this.Change = Change;
+				this.Project = Project;
+			}
+		}
+
+		IMongoCollection<UgsMetadataDocument> Collection;
+
+		/// <summary>
+		/// Constructor
+		/// </summary>
+		/// <param name="DatabaseService">Database service instance</param>
+		public UgsMetadataCollection(DatabaseService DatabaseService)
+		{
+			Collection = DatabaseService.GetCollection<UgsMetadataDocument>("UgsMetadata");
+			
+			if (!DatabaseService.ReadOnlyMode)
+			{
+				void CreateIndex()
+				{
+					Collection.Indexes.CreateOne(new CreateIndexModel<UgsMetadataDocument>(Builders<UgsMetadataDocument>.IndexKeys.Ascending(x => x.Stream).Descending(x => x.Change).Ascending(x => x.Project), new CreateIndexOptions { Unique = true }));
+					Collection.Indexes.CreateOne(new CreateIndexModel<UgsMetadataDocument>(Builders<UgsMetadataDocument>.IndexKeys.Ascending(x => x.Stream).Descending(x => x.Change).Descending(x => x.UpdateTicks)));
+				}
+				
+				try
+				{
+					CreateIndex();	
+				}
+				catch (MongoCommandException)
+				{
+					// If index creation fails, drop all indexes and retry.
+					Collection.Indexes.DropAll();
+					CreateIndex();
+				}
+			}
+		}
+
+		/// <summary>
+		/// Find or add a document for the given change
+		/// </summary>
+		/// <param name="Stream">Stream containing the change</param>
+		/// <param name="Change">The changelist number to add a document for</param>
+		/// <param name="Project">Arbitrary identifier for this project</param>
+		/// <returns>The metadata document</returns>
+		public async Task<IUgsMetadata> FindOrAddAsync(string Stream, int Change, string? Project)
+		{
+			string NormalizedStream = GetNormalizedStream(Stream);
+			string NormalizedProject = GetNormalizedProject(Project);
+			for (; ; )
+			{
+				// Find an existing document
+				UgsMetadataDocument? Existing = await Collection.Find(x => x.Stream == NormalizedStream && x.Change == Change && x.Project == NormalizedProject).FirstOrDefaultAsync();
+				if (Existing != null)
+				{
+					return Existing;
+				}
+
+				// Try to insert a new document
+				try
+				{
+					UgsMetadataDocument NewDocument = new UgsMetadataDocument(NormalizedStream, Change, NormalizedProject);
+					await Collection.InsertOneAsync(NewDocument);
+					return NewDocument;
+				}
+				catch (MongoWriteException Ex)
+				{
+					if (Ex.WriteError.Category != ServerErrorCategory.DuplicateKey)
+					{
+						throw;
+					}
+				}
+			}
+		}
+
+		/// <inheritdoc/>
+		public async Task<IUgsMetadata> UpdateUserAsync(IUgsMetadata Metadata, string UserName, bool? Synced, UgsUserVote? Vote, bool? Investigating, bool? Starred, string? Comment)
+		{
+			UpdateDefinitionBuilder<UgsMetadataDocument> UpdateBuilder = Builders<UgsMetadataDocument>.Update;
+			for (; ; )
+			{
+				UgsMetadataDocument Document = (UgsMetadataDocument)Metadata;
+
+				int UserIdx = Document.Users.FindIndex(x => x.User.Equals(UserName, StringComparison.OrdinalIgnoreCase));
+				if (UserIdx == -1)
+				{
+					// Create a new user entry
+					UgsUserDataDocument UserData = new UgsUserDataDocument(UserName);
+					UserData.SyncTime = (Synced == true) ? (long?)DateTime.UtcNow.Ticks : null;
+					UserData.Vote = Vote;
+					UserData.Investigating = (Investigating == true) ? Investigating : null;
+					UserData.Starred = (Starred == true) ? Starred : null;
+					UserData.Comment = Comment;
+
+					if (await TryUpdateAsync(Document, UpdateBuilder.Push(x => x.Users, UserData)))
+					{
+						Document.Users.Add(UserData);
+						return Metadata;
+					}
+				}
+				else
+				{
+					// Update an existing entry
+					UgsUserDataDocument UserData = Document.Users[UserIdx];
+
+					long? NewSyncTime = Synced.HasValue ? (Synced.Value ? (long?)DateTime.UtcNow.Ticks : null) : UserData.SyncTime;
+					UgsUserVote? NewVote = Vote.HasValue ? ((Vote.Value != UgsUserVote.None) ? Vote : null) : UserData.Vote;
+					bool? NewInvestigating = Investigating.HasValue ? (Investigating.Value ? Investigating : null) : UserData.Investigating;
+					bool? NewStarred = Starred.HasValue ? (Starred.Value ? Starred : null) : UserData.Starred;
+
+					List<UpdateDefinition<UgsMetadataDocument>> Updates = new List<UpdateDefinition<UgsMetadataDocument>>();
+					if (NewSyncTime != UserData.SyncTime)
+					{
+						Updates.Add(UpdateBuilder.Set(x => x.Users![UserIdx].SyncTime, NewSyncTime));
+					}
+					if (NewVote != UserData.Vote)
+					{
+						Updates.Add(UpdateBuilder.Set(x => x.Users![UserIdx].Vote, NewVote));
+					}
+					if (NewInvestigating != UserData.Investigating)
+					{
+						Updates.Add(UpdateBuilder.SetOrUnsetNull(x => x.Users![UserIdx].Investigating, NewInvestigating));
+					}
+					if (NewStarred != UserData.Starred)
+					{
+						Updates.Add(UpdateBuilder.SetOrUnsetNull(x => x.Users![UserIdx].Starred, NewStarred));
+					}
+					if (Comment != null)
+					{
+						Updates.Add(UpdateBuilder.Set(x => x.Users![UserIdx].Comment, Comment));
+					}
+
+					if (Updates.Count == 0 || await TryUpdateAsync(Document, UpdateBuilder.Combine(Updates)))
+					{
+						UserData.SyncTime = NewSyncTime;
+						UserData.Vote = Vote;
+						UserData.Investigating = NewInvestigating;
+						UserData.Starred = NewStarred;
+						UserData.Comment = Comment;
+						return Metadata;
+					}
+				}
+
+				// Update the document and try again
+				Metadata = await FindOrAddAsync(Metadata.Stream, Metadata.Change, Metadata.Project);
+			}
+		}
+
+		/// <inheritdoc/>
+		public async Task<IUgsMetadata> UpdateBadgeAsync(IUgsMetadata Metadata, string Name, Uri? Url, UgsBadgeState State)
+		{
+			UpdateDefinitionBuilder<UgsMetadataDocument> UpdateBuilder = Builders<UgsMetadataDocument>.Update;
+			for (; ; )
+			{
+				UgsMetadataDocument Document = (UgsMetadataDocument)Metadata;
+
+				// Update the document
+				int BadgeIdx = Document.Badges.FindIndex(x => x.Name.Equals(Name, StringComparison.OrdinalIgnoreCase));
+				if (BadgeIdx == -1)
+				{
+					// Create a new badge
+					UgsBadgeDataDocument NewBadge = new UgsBadgeDataDocument();
+					NewBadge.Name = Name;
+					NewBadge.Url = Url;
+					NewBadge.State = State;
+
+					if (await TryUpdateAsync(Document, UpdateBuilder.Push(x => x.Badges, NewBadge)))
+					{
+						Document.Badges.Add(NewBadge);
+						return Metadata;
+					}
+				}
+				else
+				{
+					// Update an existing badge
+					List<UpdateDefinition<UgsMetadataDocument>> Updates = new List<UpdateDefinition<UgsMetadataDocument>>();
+
+					UgsBadgeDataDocument BadgeData = Document.Badges[BadgeIdx];
+					if (Url != BadgeData.Url)
+					{
+						Updates.Add(UpdateBuilder.Set(x => x.Badges![BadgeIdx].Url, Url));
+					}
+					if (State != BadgeData.State)
+					{
+						Updates.Add(UpdateBuilder.Set(x => x.Badges![BadgeIdx].State, State));
+					}
+
+					if (Updates.Count == 0 || await TryUpdateAsync(Document, UpdateBuilder.Combine(Updates)))
+					{
+						BadgeData.Url = Url;
+						BadgeData.State = State;
+						return Metadata;
+					}
+				}
+
+				// Update the document and try again
+				Metadata = await FindOrAddAsync(Metadata.Stream, Metadata.Change, Metadata.Project);
+			}
+		}
+
+		/// <inheritdoc/>
+		public async Task<List<IUgsMetadata>> FindAsync(string Stream, int MinChange, int? MaxChange = null, long? MinTime = null)
+		{
+			FilterDefinitionBuilder<UgsMetadataDocument> FilterBuilder = Builders<UgsMetadataDocument>.Filter;
+
+			string NormalizedStream = GetNormalizedStream(Stream);
+
+			FilterDefinition<UgsMetadataDocument> Filter = FilterBuilder.Eq(x => x.Stream, NormalizedStream) & FilterBuilder.Gte(x => x.Change, MinChange);
+			if (MaxChange != null)
+			{
+				Filter &= FilterBuilder.Lte(x => x.Change, MaxChange.Value);
+			}
+			if (MinTime != null)
+			{
+				Filter &= FilterBuilder.Gt(x => x.UpdateTicks, MinTime.Value);
+			}
+
+			List<UgsMetadataDocument> Documents = await Collection.Find(Filter).ToListAsync();
+			return Documents.ConvertAll<IUgsMetadata>(x =>
+			{
+				// Remove polluting null entries. Need to find the source of these.
+				x.Users = x.Users.Where(y => y != null).ToList();
+				x.Badges = x.Badges.Where(y => y != null).ToList();
+				return x;
+			});
+		}
+
+		/// <summary>
+		/// Normalize a stream argument
+		/// </summary>
+		/// <param name="Stream">The stream name</param>
+		/// <returns>Normalized name</returns>
+		[SuppressMessage("Globalization", "CA1308:Normalize strings to uppercase")]
+		private static string GetNormalizedStream(string Stream)
+		{
+			return Stream.ToLowerInvariant().TrimEnd('/');
+		}
+
+		/// <summary>
+		/// Normalize a project argument
+		/// </summary>
+		/// <param name="Project">The project name</param>
+		/// <returns>Normalized name</returns>
+		[SuppressMessage("Globalization", "CA1308:Normalize strings to uppercase")]
+		private static string GetNormalizedProject(string? Project)
+		{
+			if(Project == null)
+			{
+				return String.Empty;
+			}
+			else
+			{
+				return Project.ToLowerInvariant();
+			}
+		}
+
+		/// <summary>
+		/// Try to update a metadata document
+		/// </summary>
+		/// <param name="Document">The document to update</param>
+		/// <param name="Update">Update definition</param>
+		/// <returns>True if the update succeeded</returns>
+		private async Task<bool> TryUpdateAsync(UgsMetadataDocument Document, UpdateDefinition<UgsMetadataDocument> Update)
+		{
+			int NewUpdateIndex = Document.UpdateIndex + 1;
+			Update = Update.Set(x => x.UpdateIndex, NewUpdateIndex);
+
+			long NewUpdateTicks = DateTime.UtcNow.Ticks;
+			Update = Update.Set(x => x.UpdateTicks, NewUpdateTicks);
+
+			UpdateResult Result = await Collection.UpdateOneAsync(x => x.Change == Document.Change && x.UpdateIndex == Document.UpdateIndex, Update);
+			if (Result.ModifiedCount > 0)
+			{
+				Document.UpdateIndex = NewUpdateIndex;
+				Document.UpdateTicks = NewUpdateTicks;
+				return true;
+			}
+			return false;
+		}
+	}
+}
diff --git a/Engine/Source/Programs/Horde/HordeServer/Ugs/UgsMetadata.cs b/Engine/Source/Programs/Horde/HordeServer/Ugs/UgsMetadata.cs
new file mode 100644
--- /dev/null
+++ b/Engine/Source/Programs/Horde/HordeServer/Ugs/UgsMetadata.cs
@@ -0,0 +1,105 @@
+// Copyright Epic Games, Inc. All Rights Reserved.
+
+using Amazon.S3.Model;
+using HordeServer.Api;
+using MongoDB.Bson;
+using System;
+using System.Collections.Generic;
+using System.Linq;
+using System.Threading.Tasks;
+
+namespace HordeServer.Models
+{
+	/// <summary>
+	/// Metadata for a particular user
+	/// </summary>
+	public interface IUgsUserData
+	{
+		/// <summary>
+		/// The user name
+		/// </summary>
+		string User { get; }
+
+		/// <summary>
+		/// Time that this change was synced
+		/// </summary>
+		public long? SyncTime { get; }
+
+		/// <summary>
+		/// State of the user
+		/// </summary>
+		public UgsUserVote Vote { get; }
+
+		/// <summary>
+		/// Whether the change is starred
+		/// </summary>
+		bool? Starred { get; }
+
+		/// <summary>
+		/// Whether the change is under investigation
+		/// </summary>
+		bool? Investigating { get; }
+
+		/// <summary>
+		/// Comment text
+		/// </summary>
+		string? Comment { get; }
+	}
+
+	/// <summary>
+	/// Metadata for a badge
+	/// </summary>
+	public interface IUgsBadgeData
+	{
+		/// <summary>
+		/// The badge name
+		/// </summary>
+		string Name { get; }
+
+		/// <summary>
+		/// The URL to open when the badge is clicked
+		/// </summary>
+		Uri? Url { get; }
+
+		/// <summary>
+		/// State of the badge
+		/// </summary>
+		UgsBadgeState State { get; }
+	}
+
+	/// <summary>
+	/// Metadata for a particular changelist and project
+	/// </summary>
+	public interface IUgsMetadata
+	{
+		/// <summary>
+		/// The changelist number
+		/// </summary>
+		int Change { get; }
+
+		/// <summary>
+		/// Stream for the metadata
+		/// </summary>
+		string Stream { get; }
+
+		/// <summary>
+		/// The project identifier
+		/// </summary>
+		string Project { get; }
+
+		/// <summary>
+		/// List of user states
+		/// </summary>
+		IReadOnlyList<IUgsUserData>? Users { get; }
+
+		/// <summary>
+		/// List of badge states
+		/// </summary>
+		IReadOnlyList<IUgsBadgeData>? Badges { get; }
+
+		/// <summary>
+		/// Time that this object was last updated
+		/// </summary>
+		long UpdateTicks { get; }
+	}
+}
