Horde: Build fix.

diff --git a/Engine/Source/Programs/Horde/HordeServer/Controllers/AgentSoftwareController.cs b/Engine/Source/Programs/Horde/HordeServer/Controllers/AgentSoftwareController.cs
--- a/Engine/Source/Programs/Horde/HordeServer/Controllers/AgentSoftwareController.cs
+++ b/Engine/Source/Programs/Horde/HordeServer/Controllers/AgentSoftwareController.cs
@@ -134,7 +134,7 @@ namespace HordeServer.Controllers
 		/// <returns>Http response</returns>
 		[HttpGet]
 		[Route("/api/v1/agentsoftware/{Name}/zip")]
-		public async Task<ActionResult<GetSoftwareResponse>> GetArchiveAsync(string Name)
+		public async Task<ActionResult> GetArchiveAsync(string Name)
 		{
 			if (!await AclService.AuthorizeAsync(AclAction.DownloadSoftware, User))
 			{
