Horde: Allow triggering schedules based on certain paths being modified. The "Files" property on a schedule accepts a list of Perforce-style wildcards relative to the root of a stream, specifying the files to trigger new builds.

diff --git a/Engine/Source/Programs/Horde/HordeServer.Compute.Tests/SchedulerTests.cs b/Engine/Source/Programs/Horde/HordeServer.Compute.Tests/SchedulerTests.cs
--- a/Engine/Source/Programs/Horde/HordeServer.Compute.Tests/SchedulerTests.cs
+++ b/Engine/Source/Programs/Horde/HordeServer.Compute.Tests/SchedulerTests.cs
@@ -56,7 +56,7 @@ namespace HordeServerTests
 			PerforceService.AddChange("//UE5/Main", 102, "Bob", "", new[] { "content.uasset" });
 		}
 
-		async Task<IStream> SetSchedule(CreateScheduleRequest Schedule)
+		async Task<IStream> SetScheduleAsync(CreateScheduleRequest Schedule)
 		{
 			IStream? Stream = await TestSetup.StreamService.GetStreamAsync(StreamId);
 
@@ -69,8 +69,64 @@ namespace HordeServerTests
 			return (await StreamService.StreamCollection.TryCreateOrReplaceAsync(StreamId, Stream, "", "", ProjectId, Config))!; 
 		}
 
+		public async Task<List<IJob>> FileTestHelperAsync(params string[] Files)
+		{
+			PerforceService.Changes.Clear();
+			PerforceService.AddChange("//UE5/Main", 100, "Bob", "", new[] { "code.cpp" });
+			PerforceService.AddChange("//UE5/Main", 101, "Bob", "", new[] { "content.uasset" });
+			PerforceService.AddChange("//UE5/Main", 102, "Bob", "", new[] { "content.uasset" });
+			PerforceService.AddChange("//UE5/Main", 103, "Bob", "", new[] { "foo/code.cpp" });
+			PerforceService.AddChange("//UE5/Main", 104, "Bob", "", new[] { "bar/code.cpp" });
+			PerforceService.AddChange("//UE5/Main", 105, "Bob", "", new[] { "foo/bar/content.uasset" });
+			PerforceService.AddChange("//UE5/Main", 106, "Bob", "", new[] { "bar/foo/content.uasset" });
+
+			DateTime StartTime = new DateTime(2021, 1, 1, 12, 0, 0, DateTimeKind.Local); // Friday Jan 1, 2021 
+			TestSetup.Clock.UtcNow = StartTime;
+
+			CreateScheduleRequest Schedule = new CreateScheduleRequest();
+			Schedule.Enabled = true;
+			Schedule.MaxChanges = 10;
+			Schedule.Patterns.Add(new CreateSchedulePatternRequest { Interval = 1 });
+			Schedule.Files = Files.ToList();
+			await SetScheduleAsync(Schedule);
+
+			TestSetup.Clock.Advance(TimeSpan.FromHours(1.25));
+			await ScheduleService.TickSharedOnlyForTestingAsync();
+
+			return await GetNewJobs();
+		}
+
+		[TestMethod]
+		public async Task FileTest1Async()
+		{
+			List<IJob> Jobs = await FileTestHelperAsync("....cpp");
+			Assert.AreEqual(3, Jobs.Count);
+			Assert.AreEqual(100, Jobs[0].Change);
+			Assert.AreEqual(103, Jobs[1].Change);
+			Assert.AreEqual(104, Jobs[2].Change);
+		}
+
+		[TestMethod]
+		public async Task FileTest2Async()
+		{
+			List<IJob> Jobs = await FileTestHelperAsync("/foo/...");
+			Assert.AreEqual(2, Jobs.Count);
+			Assert.AreEqual(103, Jobs[0].Change);
+			Assert.AreEqual(105, Jobs[1].Change);
+		}
+
+		[TestMethod]
+		public async Task FileTest3Async()
+		{
+			List<IJob> Jobs = await FileTestHelperAsync("....uasset", "-/bar/...");
+			Assert.AreEqual(3, Jobs.Count);
+			Assert.AreEqual(101, Jobs[0].Change);
+			Assert.AreEqual(102, Jobs[1].Change);
+			Assert.AreEqual(105, Jobs[2].Change);
+		}
+
 		[TestMethod]
-		public void DayScheduleTestAsync()
+		public void DayScheduleTest()
 		{
 			DateTime StartTime = new DateTime(2021, 1, 1, 12, 0, 0, DateTimeKind.Utc); // Friday Jan 1, 2021 
 			TestSetup.Clock.UtcNow = StartTime;
@@ -92,7 +148,7 @@ namespace HordeServerTests
 		}
 
 		[TestMethod]
-		public void MulticheduleTestAsync()
+		public void MulticheduleTest()
 		{
 			DateTime StartTime = new DateTime(2021, 1, 1, 12, 0, 0, DateTimeKind.Utc); // Friday Jan 1, 2021 
 			TestSetup.Clock.UtcNow = StartTime;
@@ -129,7 +185,7 @@ namespace HordeServerTests
 			Schedule.Enabled = true;
 			Schedule.Patterns.Add(new CreateSchedulePatternRequest { MinTime = 13 * 60, MaxTime = 14 * 60, Interval = 15 });
 //			Schedule.LastTriggerTime = StartTime;
-			await SetSchedule(Schedule);
+			await SetScheduleAsync(Schedule);
 
 			// Initial tick
 			await ScheduleService.TickSharedOnlyForTestingAsync();
@@ -156,7 +212,7 @@ namespace HordeServerTests
 			CreateScheduleRequest Schedule = new CreateScheduleRequest();
 			Schedule.Enabled = true;
 			Schedule.Patterns.Add(new CreateSchedulePatternRequest { MinTime = 13 * 60, MaxTime = 14 * 60, Interval = 15 });
-			IStream Stream = await SetSchedule(Schedule);
+			IStream Stream = await SetScheduleAsync(Schedule);
 
 			// Initial tick
 			await ScheduleService.TickSharedOnlyForTestingAsync();
@@ -197,7 +253,7 @@ namespace HordeServerTests
 			Schedule.Patterns.Add(new CreateSchedulePatternRequest { MinTime = 13 * 60, MaxTime = 14 * 60, Interval = 15 });
 			Schedule.MaxChanges = 2;
 			Schedule.Filter = new List<ChangeContentFlags> { ChangeContentFlags.ContainsCode };
-			await SetSchedule(Schedule);
+			await SetScheduleAsync(Schedule);
 
 			// Initial tick
 			await ScheduleService.TickSharedOnlyForTestingAsync();
@@ -233,7 +289,7 @@ namespace HordeServerTests
 			Schedule.RequireSubmittedChange = false;
 			Schedule.Patterns.Add(new CreateSchedulePatternRequest { MinTime = 13 * 60, MaxTime = 14 * 60, Interval = 15 });
 			Schedule.MaxActive = 1;
-			await SetSchedule(Schedule);
+			await SetScheduleAsync(Schedule);
 
 			// Trigger a job
 			TestSetup.Clock.Advance(TimeSpan.FromHours(1.25));
@@ -273,7 +329,7 @@ namespace HordeServerTests
 			CreateScheduleRequest Schedule = new CreateScheduleRequest();
 			Schedule.Enabled = true;
 			Schedule.Patterns.Add(new CreateSchedulePatternRequest { MinTime = 13 * 60, MaxTime = 14 * 60, Interval = 15 });
-			IStream Stream = await SetSchedule(Schedule);
+			IStream Stream = await SetScheduleAsync(Schedule);
 
 			// Trigger a job
 			TestSetup.Clock.Advance(TimeSpan.FromHours(1.25));
@@ -367,7 +423,7 @@ namespace HordeServerTests
 			Schedule.RequireSubmittedChange = false;
 			Schedule.Patterns.Add(new CreateSchedulePatternRequest { MinTime = 13 * 60, MaxTime = 14 * 60, Interval = 15 });
 			Schedule.MaxActive = 2;
-			await SetSchedule(Schedule);
+			await SetScheduleAsync(Schedule);
 
 			TestSetup.Clock.Advance(TimeSpan.FromHours(1.25));
 			await ScheduleService.TickSharedOnlyForTestingAsync();
@@ -384,7 +440,7 @@ namespace HordeServerTests
 			Assert.AreEqual(Jobs1[0].Id, TemplateRef1.Schedule!.ActiveJobs[0]);
 
 			// Test that another job does not trigger
-			await SetSchedule(Schedule);
+			await SetScheduleAsync(Schedule);
 
 			// Make sure the job is still registered
 			IStream? Stream2 = await TestSetup.StreamService.GetStreamAsync(StreamId);
@@ -402,7 +458,7 @@ namespace HordeServerTests
 			CreateScheduleRequest Schedule = new CreateScheduleRequest();
 			Schedule.Enabled = true;
 			Schedule.Patterns.Add(new CreateSchedulePatternRequest { MinTime = 13 * 60, MaxTime = 14 * 60, Interval = 15 });
-			IStream Stream = await SetSchedule(Schedule);
+			IStream Stream = await SetScheduleAsync(Schedule);
 
 			IStreamCollection StreamCollection = ServiceProvider.GetRequiredService<IStreamCollection>();
 			await StreamCollection.TryUpdatePauseStateAsync(Stream, NewPausedUntil: StartTime.AddHours(5), NewPauseComment: "testing");
diff --git a/Engine/Source/Programs/Horde/HordeServer/Projects/ISchedule.cs b/Engine/Source/Programs/Horde/HordeServer/Projects/ISchedule.cs
--- a/Engine/Source/Programs/Horde/HordeServer/Projects/ISchedule.cs
+++ b/Engine/Source/Programs/Horde/HordeServer/Projects/ISchedule.cs
@@ -181,6 +181,11 @@ namespace HordeServer.Models
 		/// </summary>
 		public List<ChangeContentFlags>? Filter { get; set; }
 
+		/// <summary>
+		/// File paths which should trigger this schedule
+		/// </summary>
+		public List<string>? Files { get; set; }
+
 		/// <summary>
 		/// Parameters for the template
 		/// </summary>
@@ -227,7 +232,7 @@ namespace HordeServer.Models
 		/// <param name="Filter">Filter for changes to consider</param>
 		/// <param name="TemplateParameters">Parameters for the template to run</param>
 		/// <param name="Patterns">List of patterns for the schedule</param>
-		public Schedule(bool Enabled = true, int MaxActive = 0, int MaxChanges = 0, bool RequireSubmittedChange = true, ScheduleGate? Gate = null, List<ChangeContentFlags>? Filter = null, Dictionary<string, string>? TemplateParameters = null, List<SchedulePattern>? Patterns = null)
+		public Schedule(bool Enabled = true, int MaxActive = 0, int MaxChanges = 0, bool RequireSubmittedChange = true, ScheduleGate? Gate = null, List<ChangeContentFlags>? Filter = null, List<string>? Files = null, Dictionary<string, string>? TemplateParameters = null, List<SchedulePattern>? Patterns = null)
 		{
 			this.Enabled = Enabled;
 			this.MaxActive = MaxActive;
@@ -235,6 +240,7 @@ namespace HordeServer.Models
 			this.RequireSubmittedChange = RequireSubmittedChange;
 			this.Gate = Gate;
 			this.Filter = Filter;
+			this.Files = Files;
 			this.TemplateParameters = TemplateParameters ?? new Dictionary<string, string>();
 			this.Patterns = Patterns ?? new List<SchedulePattern>();
 			this.LastTriggerTime = DateTimeOffset.Now;
diff --git a/Engine/Source/Programs/Horde/HordeServer/Projects/ScheduleService.cs b/Engine/Source/Programs/Horde/HordeServer/Projects/ScheduleService.cs
--- a/Engine/Source/Programs/Horde/HordeServer/Projects/ScheduleService.cs
+++ b/Engine/Source/Programs/Horde/HordeServer/Projects/ScheduleService.cs
@@ -275,6 +275,13 @@ namespace HordeServer.Services
 			Stopwatch Timer = Stopwatch.StartNew();
 			ChangeContentFlags? FilterFlags = Schedule.GetFilterFlags();
 
+			// Create a file filter
+			FileFilter? FileFilter = null;
+			if (Schedule.Files != null)
+			{
+				FileFilter = new FileFilter(Schedule.Files);
+			}
+
 			// Start as many jobs as possible
 			List<(int Change, int CodeChange)> TriggerChanges = new List<(int, int)>();
 			while (TriggerChanges.Count < MaxNewChanges)
@@ -301,7 +308,7 @@ namespace HordeServer.Services
 				}
 
 				// Adjust the changelist for the desired filter
-				if(await ShouldBuildChangeAsync(Stream.ClusterName, Stream.Name, Change, FilterFlags))
+				if (await ShouldBuildChangeAsync(Stream.ClusterName, Stream.Name, Change, FilterFlags, FileFilter))
 				{
 					TriggerChanges.Add((Change, CodeChange));
 				}
@@ -367,8 +374,9 @@ namespace HordeServer.Services
 		/// <param name="StreamName"></param>
 		/// <param name="Change"></param>
 		/// <param name="FilterFlags"></param>
+		/// <param name="FileFilter">Filter for the files to trigger a build</param>
 		/// <returns></returns>
-		private async Task<bool> ShouldBuildChangeAsync(string ClusterName, string StreamName, int Change, ChangeContentFlags? FilterFlags)
+		private async Task<bool> ShouldBuildChangeAsync(string ClusterName, string StreamName, int Change, ChangeContentFlags? FilterFlags, FileFilter? FileFilter)
 		{
 			if (FilterFlags != null && FilterFlags.Value != 0)
 			{
@@ -379,6 +387,15 @@ namespace HordeServer.Services
 					return false;
 				}
 			}
+			if (FileFilter != null)
+			{
+				ChangeDetails Details = await Perforce.GetChangeDetailsAsync(ClusterName, StreamName, Change, null);
+				if (!Details.Files.Any(x => FileFilter.Matches(x.Path)))
+				{
+					Logger.LogDebug("Not building change {Change} due to file filter", Change);
+					return false;
+				}
+			}
 			return true;
 		}
 
