Horde: Add an interface for the artifact collection.

diff --git a/Engine/Source/Programs/Horde/HordeServer.Compute.Tests/ControllerIntegrationTest.cs b/Engine/Source/Programs/Horde/HordeServer.Compute.Tests/ControllerIntegrationTest.cs
--- a/Engine/Source/Programs/Horde/HordeServer.Compute.Tests/ControllerIntegrationTest.cs
+++ b/Engine/Source/Programs/Horde/HordeServer.Compute.Tests/ControllerIntegrationTest.cs
@@ -8,6 +8,7 @@ using System.Runtime.CompilerServices;
 using System.Threading.Tasks;
 using HordeServer;
 using HordeServer.Collections.Impl;
+using HordeServer.Jobs;
 using HordeServer.Services;
 using Microsoft.AspNetCore.Builder;
 using Microsoft.AspNetCore.Hosting;
@@ -57,14 +58,14 @@ namespace HordeServerTests
             DatabaseService DatabaseService = Services.GetRequiredService<DatabaseService>();
             TemplateService TemplateService = Services.GetRequiredService<TemplateService>();
             JobService JobService = Services.GetRequiredService<JobService>();
-            ArtifactService ArtifactService = Services.GetRequiredService<ArtifactService>();
+            IArtifactCollection ArtifactCollection = Services.GetRequiredService<IArtifactCollection>();
             StreamService StreamService = Services.GetRequiredService<StreamService>();
             AgentService AgentService = Services.GetRequiredService<AgentService>();
             IPerforceService PerforceService = Services.GetRequiredService<IPerforceService>();
             GraphCollection GraphCollection = new GraphCollection(DatabaseService);
 
             _fixture = new Fixture();
-            _fixture = await Fixture.Create(false, GraphCollection, TemplateService, JobService, ArtifactService, StreamService, AgentService, PerforceService);
+            _fixture = await Fixture.Create(false, GraphCollection, TemplateService, JobService, ArtifactCollection, StreamService, AgentService, PerforceService);
             return _fixture;
         }
 
diff --git a/Engine/Source/Programs/Horde/HordeServer.Compute.Tests/Fixture.cs b/Engine/Source/Programs/Horde/HordeServer.Compute.Tests/Fixture.cs
--- a/Engine/Source/Programs/Horde/HordeServer.Compute.Tests/Fixture.cs
+++ b/Engine/Source/Programs/Horde/HordeServer.Compute.Tests/Fixture.cs
@@ -15,6 +15,7 @@ using StreamId = HordeServer.Utilities.StringId<HordeServer.Models.IStream>;
 using TemplateRefId = HordeServer.Utilities.StringId<HordeServer.Models.TemplateRef>;
 using ProjectId = HordeServer.Utilities.StringId<HordeServer.Models.IProject>;
 using HordeServer.Api;
+using HordeServer.Jobs;
 
 namespace HordeServerTests
 {
@@ -27,15 +28,15 @@ namespace HordeServerTests
 		public IStream? Stream { get; private set; }
 		public TemplateRefId TemplateRefId1 { get; private set; }
 		public TemplateRefId TemplateRefId2 { get; private set; }
-		public Artifact Job1Artifact { get; private set; } = null!;
+		public IArtifact Job1Artifact { get; private set; } = null!;
 		public string Job1ArtifactData { get; private set; } = null!;
 		public IAgent Agent1 { get; private set; } = null!;
 		public string Agent1Name { get; private set; } = null!;
 
-		public static async Task<Fixture> Create(bool ForceNewFixture, IGraphCollection GraphCollection, TemplateService TemplateService, JobService JobService, ArtifactService ArtifactService, StreamService StreamService, AgentService AgentService, IPerforceService PerforceService)
+		public static async Task<Fixture> Create(bool ForceNewFixture, IGraphCollection GraphCollection, TemplateService TemplateService, JobService JobService, IArtifactCollection ArtifactCollection, StreamService StreamService, AgentService AgentService, IPerforceService PerforceService)
 		{
 			Fixture _fixture = new Fixture();
-			await _fixture.Populate(GraphCollection, TemplateService, JobService, ArtifactService, StreamService, AgentService, PerforceService);
+			await _fixture.Populate(GraphCollection, TemplateService, JobService, ArtifactCollection, StreamService, AgentService, PerforceService);
 			
 			(PerforceService as PerforceServiceStub)?.AddChange("//UE5/Main", 112233, "leet.coder", "Did stuff", new []{"file.cpp"});
 			(PerforceService as PerforceServiceStub)?.AddChange("//UE5/Main", 1111, "swarm", "A shelved CL here", new []{"renderer.cpp"});
@@ -43,7 +44,7 @@ namespace HordeServerTests
 			return _fixture;
 		}
 
-		private async Task Populate(IGraphCollection GraphCollection, TemplateService TemplateService, JobService JobService, ArtifactService ArtifactService, StreamService StreamService, AgentService AgentService, IPerforceService PerforceService)
+		private async Task Populate(IGraphCollection GraphCollection, TemplateService TemplateService, JobService JobService, IArtifactCollection ArtifactCollection, StreamService StreamService, AgentService AgentService, IPerforceService PerforceService)
 		{
 			var Fg = new FixtureGraph();
 			Fg.Id = ContentHash.Empty;
@@ -131,7 +132,7 @@ namespace HordeServerTests
 			Job2 = (await JobService.GetJobAsync(Job2.Id))!;
 
 			Job1ArtifactData = "For The Horde!";
-			Job1Artifact = await ArtifactService.CreateArtifactAsync(Job1.Id, SubResourceId.Parse("22"), "myFile.txt",
+			Job1Artifact = await ArtifactCollection.CreateArtifactAsync(Job1.Id, SubResourceId.Parse("22"), "myFile.txt",
 				"text/plain", new MemoryStream(Encoding.UTF8.GetBytes(Job1ArtifactData)));
 
 			Agent1Name = "testAgent1";
diff --git a/Engine/Source/Programs/Horde/HordeServer.Compute.Tests/RpcServiceTest.cs b/Engine/Source/Programs/Horde/HordeServer.Compute.Tests/RpcServiceTest.cs
--- a/Engine/Source/Programs/Horde/HordeServer.Compute.Tests/RpcServiceTest.cs
+++ b/Engine/Source/Programs/Horde/HordeServer.Compute.Tests/RpcServiceTest.cs
@@ -432,9 +432,9 @@ namespace HordeServerTests
 			UploadArtifactResponse Res = await Call;
 
 
-			Artifact? Artifact = await TestSetup.ArtifactService.GetArtifactAsync(ObjectId.Parse(Res.Id));
+			IArtifact? Artifact = await TestSetup.ArtifactCollection.GetArtifactAsync(ObjectId.Parse(Res.Id));
 			Assert.IsNotNull(Artifact);
-			Stream Stream = await TestSetup.ArtifactService.OpenArtifactReadStreamAsync(Artifact!);
+			Stream Stream = await TestSetup.ArtifactCollection.OpenArtifactReadStreamAsync(Artifact!);
 			StreamReader Reader = new StreamReader(Stream);
 			string text = Reader.ReadToEnd();
 			Assert.AreEqual(DataStr, text);
diff --git a/Engine/Source/Programs/Horde/HordeServer.Compute.Tests/TestSetup.cs b/Engine/Source/Programs/Horde/HordeServer.Compute.Tests/TestSetup.cs
--- a/Engine/Source/Programs/Horde/HordeServer.Compute.Tests/TestSetup.cs
+++ b/Engine/Source/Programs/Horde/HordeServer.Compute.Tests/TestSetup.cs
@@ -12,6 +12,7 @@ using HordeServer;
 using HordeServer.Collections;
 using HordeServer.Collections.Impl;
 using HordeServer.Controllers;
+using HordeServer.Jobs;
 using HordeServer.Logs;
 using HordeServer.Logs.Builder;
 using HordeServer.Logs.Readers;
@@ -85,7 +86,7 @@ namespace HordeServerTests
 		public IIssueService IssueService => ServiceProvider.GetRequiredService<IIssueService>();
 		public JobTaskSource JobTaskSource => ServiceProvider.GetRequiredService<JobTaskSource>();
 		public JobService JobService => ServiceProvider.GetRequiredService<JobService>();
-		public ArtifactService ArtifactService => ServiceProvider.GetRequiredService<ArtifactService>();
+		public IArtifactCollection ArtifactCollection => ServiceProvider.GetRequiredService<IArtifactCollection>();
 		public IDowntimeService DowntimeService => ServiceProvider.GetRequiredService<IDowntimeService>();
 		public RpcService RpcService => ServiceProvider.GetRequiredService<RpcService>();
 		public CredentialService CredentialService => ServiceProvider.GetRequiredService<CredentialService>();
@@ -144,6 +145,7 @@ namespace HordeServerTests
 
 			Services.AddSingleton<IAgentCollection, AgentCollection>();
 			Services.AddSingleton<IAgentSoftwareCollection, AgentSoftwareCollection>();
+			Services.AddSingleton<IArtifactCollection, ArtifactService>();
 			Services.AddSingleton<ICounterCollection, CounterCollection>();
 			Services.AddSingleton<IGraphCollection, GraphCollection>();
 			Services.AddSingleton<IIssueCollection, IssueCollection>();
@@ -172,7 +174,6 @@ namespace HordeServerTests
 			Services.AddSingleton<AclService>();
 			Services.AddSingleton<AgentService>();
 			Services.AddSingleton<AgentSoftwareService>();
-			Services.AddSingleton<ArtifactService>();
 			Services.AddSingleton<AutoscaleService>();
 			Services.AddSingleton<AwsFleetManager, AwsFleetManager>();
 			Services.AddSingleton<ConsistencyService>();
@@ -222,7 +223,7 @@ namespace HordeServerTests
 		{
 			if (!DatabaseService.ReadOnlyMode)
 			{
-				Fixture = await Fixture.Create(ForceNewFixture, GraphCollection, TemplateService, JobService, ArtifactService, StreamService, AgentService, PerforceService);
+				Fixture = await Fixture.Create(ForceNewFixture, GraphCollection, TemplateService, JobService, ArtifactCollection, StreamService, AgentService, PerforceService);
 			}
 		}
 
@@ -230,7 +231,7 @@ namespace HordeServerTests
         {
 			ILogger<JobsController> Logger = ServiceProvider.GetRequiredService<ILogger<JobsController>>();
 			JobsController JobsCtrl = new JobsController(AclService, GraphCollection, PerforceService, StreamService, JobService,
-		        TemplateService, ArtifactService, NotificationService, Logger);
+		        TemplateService, ArtifactCollection, NotificationService, Logger);
 	        JobsCtrl.ControllerContext = GetControllerContext();
 	        return JobsCtrl;
         }
diff --git a/Engine/Source/Programs/Horde/HordeServer/Api/Artifacts.cs b/Engine/Source/Programs/Horde/HordeServer/Api/Artifacts.cs
--- a/Engine/Source/Programs/Horde/HordeServer/Api/Artifacts.cs
+++ b/Engine/Source/Programs/Horde/HordeServer/Api/Artifacts.cs
@@ -105,7 +105,7 @@ namespace HordeServer.Api
 		/// </summary>
 		/// <param name="Artifact">The artifact to construct from</param>
 		/// <param name="Code">The direct download code</param>
-		public GetArtifactResponse(Artifact Artifact, string? Code)
+		public GetArtifactResponse(IArtifact Artifact, string? Code)
 		{
 			this.Id = Artifact.Id.ToString();
 			this.Name = Artifact.Name;
diff --git a/Engine/Source/Programs/Horde/HordeServer/Controllers/ArtifactsController.cs b/Engine/Source/Programs/Horde/HordeServer/Controllers/ArtifactsController.cs
--- a/Engine/Source/Programs/Horde/HordeServer/Controllers/ArtifactsController.cs
+++ b/Engine/Source/Programs/Horde/HordeServer/Controllers/ArtifactsController.cs
@@ -2,6 +2,7 @@
 
 using Amazon.S3.Transfer;
 using HordeServer.Api;
+using HordeServer.Jobs;
 using HordeServer.Models;
 using HordeServer.Services;
 using HordeServer.Utilities;
@@ -39,9 +40,9 @@ namespace HordeServer.Controllers
 		private readonly DatabaseService DatabaseService;
 
 		/// <summary>
-		/// Instance of the Artifact service
+		/// Instance of the artifact collection
 		/// </summary>
-		private readonly ArtifactService ArtifactService;
+		private readonly IArtifactCollection ArtifactCollection;
 
 		/// <summary>
 		/// Instance of the ACL service
@@ -56,14 +57,10 @@ namespace HordeServer.Controllers
 		/// <summary>
 		/// Constructor
 		/// </summary>
-		/// <param name="DatabaseSerivce">The database service</param>
-		/// <param name="ArtifactService">The Artifact service</param>
-		/// <param name="AclService">The ACL service</param>
-		/// <param name="JobService">The Job service</param>
-		public ArtifactsController(DatabaseService DatabaseSerivce, ArtifactService ArtifactService, AclService AclService, JobService JobService)
+		public ArtifactsController(DatabaseService DatabaseSerivce, IArtifactCollection ArtifactCollection, AclService AclService, JobService JobService)
 		{
 			this.DatabaseService = DatabaseSerivce;
-			this.ArtifactService = ArtifactService;
+			this.ArtifactCollection = ArtifactCollection;
 			this.AclService = AclService;
 			this.JobService = JobService;
 		}
@@ -107,8 +104,8 @@ namespace HordeServer.Controllers
 					return NotFound();
 				}
 			}
-		
-			Artifact NewArtifact = await ArtifactService.CreateArtifactAsync(Job.Id, Step?.Id, File.FileName, File.ContentType ?? "horde-mime/unknown", File.OpenReadStream());
+
+			IArtifact NewArtifact = await ArtifactCollection.CreateArtifactAsync(Job.Id, Step?.Id, File.FileName, File.ContentType ?? "horde-mime/unknown", File.OpenReadStream());
 			return new CreateArtifactResponse(NewArtifact.Id.ToString());
 		}
 
@@ -123,7 +120,7 @@ namespace HordeServer.Controllers
 		[Route("/api/v1/artifacts/{ArtifactId}")]
 		public async Task<ActionResult<CreateArtifactResponse>> UpdateArtifact(string ArtifactId, IFormFile File)
 		{
-			Artifact? Artifact = await ArtifactService.GetArtifactAsync(ArtifactId.ToObjectId());
+			IArtifact? Artifact = await ArtifactCollection.GetArtifactAsync(ArtifactId.ToObjectId());
 			if (Artifact == null)
 			{
 				return NotFound();
@@ -133,7 +130,7 @@ namespace HordeServer.Controllers
 				return Forbid();
 			}
 
-			await ArtifactService.UpdateArtifactAsync(Artifact, File.ContentType ?? "horde-mime/unknown", File.OpenReadStream());
+			await ArtifactCollection.UpdateArtifactAsync(Artifact, File.ContentType ?? "horde-mime/unknown", File.OpenReadStream());
 			return Ok();
 		}
 
@@ -159,7 +156,7 @@ namespace HordeServer.Controllers
 
 			string? DownloadCode = Code ? (string?)GetDirectDownloadCodeForJob(JobIdValue) : null;
 
-			List<Artifact> Artifacts = await ArtifactService.GetArtifactsAsync(JobIdValue, StepId?.ToSubResourceId(), null);
+			List<IArtifact> Artifacts = await ArtifactCollection.GetArtifactsAsync(JobIdValue, StepId?.ToSubResourceId(), null);
 			return Artifacts.ConvertAll(x => new GetArtifactResponse(x, DownloadCode).ApplyFilter(Filter));
 		}
 
@@ -197,7 +194,7 @@ namespace HordeServer.Controllers
 		[ProducesResponseType(typeof(GetArtifactResponse), 200)]
 		public async Task<ActionResult<object>> GetArtifact(string ArtifactId, bool Code = false, [FromQuery] PropertyFilter? Filter = null)
 		{
-			Artifact? Artifact = await ArtifactService.GetArtifactAsync(ArtifactId.ToObjectId());
+			IArtifact? Artifact = await ArtifactCollection.GetArtifactAsync(ArtifactId.ToObjectId());
 			if (Artifact == null)
 			{
 				return NotFound();
@@ -221,7 +218,7 @@ namespace HordeServer.Controllers
 		[Route("/api/v1/artifacts/{ArtifactId}/data")]
 		public async Task<ActionResult> GetArtifactData(string ArtifactId)
 		{
-			Artifact? Artifact = await ArtifactService.GetArtifactAsync(ArtifactId.ToObjectId());
+			IArtifact? Artifact = await ArtifactCollection.GetArtifactAsync(ArtifactId.ToObjectId());
 			if (Artifact == null)
 			{
 				return NotFound();
@@ -232,7 +229,7 @@ namespace HordeServer.Controllers
 			}
 
 			// Fun, filestream result automatically closes the stream!
-			return new FileStreamResult(await ArtifactService.OpenArtifactReadStreamAsync(Artifact), Artifact.MimeType);
+			return new FileStreamResult(await ArtifactCollection.OpenArtifactReadStreamAsync(Artifact), Artifact.MimeType);
 		}
 		
 		/// <summary>
@@ -254,14 +251,14 @@ namespace HordeServer.Controllers
 				return Forbid();
 			}
 			
-			List<Artifact> Artifacts = await ArtifactService.GetArtifactsAsync(JobIdValue, StepIdValue, Filename);
+			List<IArtifact> Artifacts = await ArtifactCollection.GetArtifactsAsync(JobIdValue, StepIdValue, Filename);
 			if (Artifacts.Count == 0)
 			{
 				return NotFound();
 			}
 
-			Artifact Artifact = Artifacts[0];
-			return new FileStreamResult(await ArtifactService.OpenArtifactReadStreamAsync(Artifact), Artifact.MimeType);
+			IArtifact Artifact = Artifacts[0];
+			return new FileStreamResult(await ArtifactCollection.OpenArtifactReadStreamAsync(Artifact), Artifact.MimeType);
 		}
 
 		/// <summary>
@@ -322,7 +319,7 @@ namespace HordeServer.Controllers
 			JwtSecurityTokenHandler Handler = new JwtSecurityTokenHandler();
 			ClaimsPrincipal Principal = Handler.ValidateToken(Code, Parameters, out Token);
 
-			Artifact? Artifact = await ArtifactService.GetArtifactAsync(ArtifactId.ToObjectId());
+			IArtifact? Artifact = await ArtifactCollection.GetArtifactAsync(ArtifactId.ToObjectId());
 			if (Artifact == null)
 			{
 				return NotFound();
@@ -334,7 +331,7 @@ namespace HordeServer.Controllers
 				return Forbid();
 			}
 
-			return new InlineFileStreamResult(await ArtifactService.OpenArtifactReadStreamAsync(Artifact), Artifact.MimeType, Path.GetFileName(Artifact.Name));
+			return new InlineFileStreamResult(await ArtifactCollection.OpenArtifactReadStreamAsync(Artifact), Artifact.MimeType, Path.GetFileName(Artifact.Name));
 		}
 
 		/// <summary>
@@ -362,21 +359,21 @@ namespace HordeServer.Controllers
 				return Forbid();
 			}
 
-			List<Artifact> Artifacts = await ArtifactService.GetArtifactsAsync(Job.Id, ArtifactZipRequest.StepId?.ToSubResourceId(), null);
+			List<IArtifact> Artifacts = await ArtifactCollection.GetArtifactsAsync(Job.Id, ArtifactZipRequest.StepId?.ToSubResourceId(), null);
 
-			Dictionary<ObjectId, Artifact> IdToArtifact = Artifacts.ToDictionary(x => x.Id, x => x);
+			Dictionary<ObjectId, IArtifact> IdToArtifact = Artifacts.ToDictionary(x => x.Id, x => x);
 
-			List<Artifact> ZipArtifacts;
+			List<IArtifact> ZipArtifacts;
 			if (ArtifactZipRequest.ArtifactIds == null)
 			{
 				ZipArtifacts = Artifacts;
 			}
 			else
 			{
-				ZipArtifacts = new List<Artifact>();
+				ZipArtifacts = new List<IArtifact>();
 				foreach (string ArtifactId in ArtifactZipRequest.ArtifactIds)
 				{
-					Artifact? Artifact;
+					IArtifact? Artifact;
 					if (IdToArtifact.TryGetValue(ArtifactId.ToObjectId(), out Artifact))
 					{
 						ZipArtifacts.Add(Artifact);
@@ -398,9 +395,9 @@ namespace HordeServer.Controllers
 					// Keep the stream open after dispose so we can write the EOF bits.
 					using (ZipArchive ZipArchive = new ZipArchive(ZipOutputStream, ZipArchiveMode.Create, true))
 					{
-						foreach (Artifact Artifact in ZipArtifacts)
+						foreach (IArtifact Artifact in ZipArtifacts)
 						{
-							await using (System.IO.Stream ArtifactStream = await ArtifactService.OpenArtifactReadStreamAsync(Artifact))
+							await using (System.IO.Stream ArtifactStream = await ArtifactCollection.OpenArtifactReadStreamAsync(Artifact))
 							{
 								// tack on the step name into the directory if it exists
 								string StepName = string.Empty;
diff --git a/Engine/Source/Programs/Horde/HordeServer/Controllers/JobsController.cs b/Engine/Source/Programs/Horde/HordeServer/Controllers/JobsController.cs
--- a/Engine/Source/Programs/Horde/HordeServer/Controllers/JobsController.cs
+++ b/Engine/Source/Programs/Horde/HordeServer/Controllers/JobsController.cs
@@ -23,6 +23,7 @@ using Microsoft.Extensions.Logging;
 using HordeServer.Notifications;
 using OpenTracing.Util;
 using OpenTracing;
+using HordeServer.Jobs;
 
 namespace HordeServer.Controllers
 {
@@ -70,7 +71,7 @@ namespace HordeServer.Controllers
 		/// <summary>
 		/// 
 		/// </summary>
-		private readonly ArtifactService ArtifactService;
+		private readonly IArtifactCollection ArtifactCollection;
 
 		/// <summary>
 		/// Instance of the notification service singleton
@@ -85,16 +86,7 @@ namespace HordeServer.Controllers
 		/// <summary>
 		/// Constructor
 		/// </summary>
-		/// <param name="AclService">Instance of the ACL service</param>
-		/// <param name="Graphs">Collection of graph documents</param>
-		/// <param name="Perforce">The Perforce service instance</param>
-		/// <param name="StreamService">Instance of the stream service</param>
-		/// <param name="JobService">Instance of the JobService singleton</param>
-		/// <param name="TemplateService">Instance of the TemplateService singleton</param>
-		/// <param name="ArtifactService">Instance of the ArtifactService singleton</param>
-		/// <param name="NotificationService">Instance of the NotificationService singelton</param>
-		/// <param name="Logger">The logger instance</param>
-		public JobsController(AclService AclService, IGraphCollection Graphs, IPerforceService Perforce, StreamService StreamService, JobService JobService, TemplateService TemplateService, ArtifactService ArtifactService, INotificationService NotificationService, ILogger<JobsController> Logger)
+		public JobsController(AclService AclService, IGraphCollection Graphs, IPerforceService Perforce, StreamService StreamService, JobService JobService, TemplateService TemplateService, IArtifactCollection ArtifactCollection, INotificationService NotificationService, ILogger<JobsController> Logger)
 		{
 			this.AclService = AclService;
 			this.Graphs = Graphs;
@@ -102,7 +94,7 @@ namespace HordeServer.Controllers
 			this.StreamService = StreamService;
 			this.JobService = JobService;
 			this.TemplateService = TemplateService;
-			this.ArtifactService = ArtifactService;
+			this.ArtifactCollection = ArtifactCollection;
 			this.NotificationService = NotificationService;
 			this.Logger = Logger;
 		}
@@ -1184,14 +1176,14 @@ namespace HordeServer.Controllers
 				return NotFound();
 			}
 
-			List<Artifact> Artifacts = await ArtifactService.GetArtifactsAsync(JobIdValue, StepIdValue, Name);
+			List<IArtifact> Artifacts = await ArtifactCollection.GetArtifactsAsync(JobIdValue, StepIdValue, Name);
 			if (Artifacts.Count == 0)
 			{
 				return NotFound();
 			}
 
-			Artifact Artifact = Artifacts[0];
-			return new FileStreamResult(await ArtifactService.OpenArtifactReadStreamAsync(Artifact), Artifact.MimeType);
+			IArtifact Artifact = Artifacts[0];
+			return new FileStreamResult(await ArtifactCollection.OpenArtifactReadStreamAsync(Artifact), Artifact.MimeType);
 		}
 
 		/// <summary>
@@ -1223,12 +1215,12 @@ namespace HordeServer.Controllers
 				return NotFound();
 			}
 
-			List<Artifact> Artifacts = await ArtifactService.GetArtifactsAsync(JobIdValue, StepIdValue, null);
-			foreach (Artifact Artifact in Artifacts)
+			List<IArtifact> Artifacts = await ArtifactCollection.GetArtifactsAsync(JobIdValue, StepIdValue, null);
+			foreach (IArtifact Artifact in Artifacts)
 			{
 				if (Artifact.Name.Equals("trace.json", StringComparison.OrdinalIgnoreCase))
 				{
-					return new FileStreamResult(await ArtifactService.OpenArtifactReadStreamAsync(Artifact), "text/json");
+					return new FileStreamResult(await ArtifactCollection.OpenArtifactReadStreamAsync(Artifact), "text/json");
 				}
 			}
 			return NotFound();
diff --git a/Engine/Source/Programs/Horde/HordeServer/Jobs/IArtifact.cs b/Engine/Source/Programs/Horde/HordeServer/Jobs/IArtifact.cs
--- a/Engine/Source/Programs/Horde/HordeServer/Jobs/IArtifact.cs
+++ b/Engine/Source/Programs/Horde/HordeServer/Jobs/IArtifact.cs
@@ -10,74 +10,36 @@ namespace HordeServer.Models
 	/// <summary>
 	/// Information about an artifact
 	/// </summary>
-	public class Artifact
+	public interface IArtifact
 	{
 		/// <summary>
 		/// Identifier for the Artifact. Randomly generated.
 		/// </summary>
-		[BsonRequired, BsonId]
-		public ObjectId Id { get; set; }
+		public ObjectId Id { get; }
 
 		/// <summary>
 		/// Unique id of the job containing this artifact
 		/// </summary>
-		[BsonRequired]
-		public ObjectId JobId { get; set; }
+		public ObjectId JobId { get; }
 
 		/// <summary>
 		/// Unique id of the job containing this artifact
 		/// </summary>
-		[BsonRequired]
-		public string Name { get; set; }
+		public string Name { get; }
 
 		/// <summary>
 		/// Unique id of the step containing this artifact
 		/// </summary>
-		public SubResourceId? StepId { get; set; }
+		public SubResourceId? StepId { get; }
 
 		/// <summary>
 		/// Total size of the file
 		/// </summary>
-		[BsonRequired]
-		public long Length { get; set; }
+		public long Length { get; }
 
 		/// <summary>
 		/// Type of artifact
 		/// </summary>
-		public string MimeType { get; set; }
-
-		/// <summary>
-		/// Version of the file in the db
-		/// </summary>
-		[BsonRequired]
-		public int UpdateIndex { get; set; }
-
-		/// <summary>
-		/// Private constructor for BSON serializer
-		/// </summary>
-		[BsonConstructor]
-		private Artifact()
-		{
-			this.Name = null!;
-			this.MimeType = null!;
-		}
-
-		/// <summary>
-		/// Constructor
-		/// </summary>
-		/// <param name="JobId">Unique id for the job containing this artifact</param>
-		/// <param name="StepId">Optional unique id for the step containing this artifact</param>
-		/// <param name="Name">Name</param>
-		/// <param name="Length">File Length</param>
-		/// <param name="MimeType">Type of artifact</param>
-		public Artifact(ObjectId JobId, SubResourceId? StepId, string Name, long Length, string MimeType)
-		{
-			this.Id = ObjectId.GenerateNewId();
-			this.JobId = JobId;
-			this.StepId = StepId;
-			this.Name = Name;
-			this.Length = Length;
-			this.MimeType = MimeType;
-		}
+		public string MimeType { get; }
 	}
 }
diff --git a/Engine/Source/Programs/Horde/HordeServer/Jobs/IArtifactCollection.cs b/Engine/Source/Programs/Horde/HordeServer/Jobs/IArtifactCollection.cs
new file mode 100644
--- /dev/null
+++ b/Engine/Source/Programs/Horde/HordeServer/Jobs/IArtifactCollection.cs
@@ -0,0 +1,65 @@
+// Copyright Epic Games, Inc. All Rights Reserved.
+
+using HordeServer.Models;
+using MongoDB.Bson;
+using System.Collections.Generic;
+using System.Threading.Tasks;
+
+namespace HordeServer.Jobs
+{
+	/// <summary>
+	/// Interface for a collection of artifacts
+	/// </summary>
+	public interface IArtifactCollection
+	{
+		/// <summary>
+		/// Creates a new artifact
+		/// </summary>
+		/// <param name="JobId">Unique id of the job that owns this artifact</param>
+		/// <param name="StepId">Optional Step id</param>
+		/// <param name="Name">Name of artifact</param>
+		/// <param name="MimeType">Type of artifact</param>
+		/// <param name="Data">The data to write</param>
+		/// <returns>The new log file document</returns>
+		Task<IArtifact> CreateArtifactAsync(ObjectId JobId, SubResourceId? StepId, string Name, string MimeType, System.IO.Stream Data);
+
+		/// <summary>
+		/// Gets all the available artifacts for a job
+		/// </summary>
+		/// <param name="JobId">Unique id of the job to query</param>
+		/// <param name="StepId">Unique id of the Step to query</param>
+		/// <param name="Name">Name of the artifact</param>
+		/// <returns>List of artifact documents</returns>
+		Task<List<IArtifact>> GetArtifactsAsync(ObjectId? JobId, SubResourceId? StepId, string? Name);
+
+		/// <summary>
+		/// Gets a specific list of artifacts based on id
+		/// </summary>
+		/// <param name="ArtifactIds">The list of artifact Ids</param>
+		/// <returns>List of artifact documents</returns>
+		Task<List<IArtifact>> GetArtifactsAsync(IEnumerable<ObjectId> ArtifactIds);
+
+		/// <summary>
+		/// Gets an artifact by ID
+		/// </summary>
+		/// <param name="ArtifactId">Unique id of the artifact</param>
+		/// <returns>The artifact document</returns>
+		Task<IArtifact?> GetArtifactAsync(ObjectId ArtifactId);
+
+		/// <summary>
+		/// Updates an artifact
+		/// </summary>
+		/// <param name="Artifact">The artifact</param>
+		/// <param name="NewMimeType">New mime type</param>
+		/// <param name="NewData">New data</param>
+		/// <returns>Async task</returns>
+		Task<bool> UpdateArtifactAsync(IArtifact? Artifact, string NewMimeType, System.IO.Stream NewData);
+
+		/// <summary>
+		/// gets artifact data
+		/// </summary>
+		/// <param name="Artifact">The artifact</param>
+		/// <returns>The chunk data</returns>
+		Task<System.IO.Stream> OpenArtifactReadStreamAsync(IArtifact Artifact);
+	}
+}
diff --git a/Engine/Source/Programs/Horde/HordeServer/Jobs/Impl/ArtifactService.cs b/Engine/Source/Programs/Horde/HordeServer/Jobs/Impl/ArtifactService.cs
--- a/Engine/Source/Programs/Horde/HordeServer/Jobs/Impl/ArtifactService.cs
+++ b/Engine/Source/Programs/Horde/HordeServer/Jobs/Impl/ArtifactService.cs
@@ -5,6 +5,7 @@ using Amazon.Runtime;
 using Amazon.S3;
 using Amazon.S3.Model;
 using HordeServer.Api;
+using HordeServer.Jobs;
 using HordeServer.Models;
 using HordeServer.Storage;
 using HordeServer.Utilities;
@@ -13,6 +14,7 @@ using Microsoft.Extensions.Hosting;
 using Microsoft.Extensions.Logging;
 using Microsoft.Extensions.Options;
 using MongoDB.Bson;
+using MongoDB.Bson.Serialization.Attributes;
 using MongoDB.Driver;
 using System;
 using System.Collections.Generic;
@@ -28,21 +30,49 @@ namespace HordeServer.Services
 	/// <summary>
 	/// Wraps functionality for manipulating artifacts
 	/// </summary>
-	public class ArtifactService
+	public class ArtifactService : IArtifactCollection
 	{
-		/// <summary>
-		/// Backend storage container
-		/// </summary>
-		private readonly IStorageBackend StorageBackend;
+		class Artifact : IArtifact
+		{
+			[BsonRequired, BsonId]
+			public ObjectId Id { get; set; }
 
-		/// <summary>
-		/// Information Logger
-		/// </summary>
-		private readonly ILogger<ArtifactService> Logger;
+			[BsonRequired]
+			public ObjectId JobId { get; set; }
 
-		/// <summary>
-		/// Collection of log documents
-		/// </summary>
+			[BsonRequired]
+			public string Name { get; set; }
+
+			public SubResourceId? StepId { get; set; }
+
+			[BsonRequired]
+			public long Length { get; set; }
+
+			public string MimeType { get; set; }
+
+			[BsonRequired]
+			public int UpdateIndex { get; set; }
+
+			[BsonConstructor]
+			private Artifact()
+			{
+				this.Name = null!;
+				this.MimeType = null!;
+			}
+
+			public Artifact(ObjectId JobId, SubResourceId? StepId, string Name, long Length, string MimeType)
+			{
+				this.Id = ObjectId.GenerateNewId();
+				this.JobId = JobId;
+				this.StepId = StepId;
+				this.Name = Name;
+				this.Length = Length;
+				this.MimeType = MimeType;
+			}
+		}
+
+		private readonly IStorageBackend StorageBackend;
+		private readonly ILogger<ArtifactService> Logger;
 		private readonly IMongoCollection<Artifact> Artifacts;
 
 		/// <summary>
@@ -57,7 +87,7 @@ namespace HordeServer.Services
 			this.Logger = Logger;
 
 			// Initialize Artifacts table
-			Artifacts = DatabaseService.Artifacts;
+			Artifacts = DatabaseService.GetCollection<Artifact>("Artifacts");
 			if (!DatabaseService.ReadOnlyMode)
 			{
 				Artifacts.Indexes.CreateOne(new CreateIndexModel<Artifact>(Builders<Artifact>.IndexKeys.Ascending(x => x.JobId)));
@@ -73,7 +103,7 @@ namespace HordeServer.Services
 		/// <param name="MimeType">Type of artifact</param>
 		/// <param name="Data">The data to write</param>
 		/// <returns>The new log file document</returns>
-		public async Task<Artifact> CreateArtifactAsync(ObjectId JobId, SubResourceId? StepId, string Name, string MimeType, System.IO.Stream Data)
+		public async Task<IArtifact> CreateArtifactAsync(ObjectId JobId, SubResourceId? StepId, string Name, string MimeType, System.IO.Stream Data)
 		{
 			// upload first
 			string ArtifactName = ValidateName(Name);
@@ -92,7 +122,7 @@ namespace HordeServer.Services
 		/// <param name="StepId">Unique id of the Step to query</param>
 		/// <param name="Name">Name of the artifact</param>
 		/// <returns>List of artifact documents</returns>
-		public async Task<List<Artifact>> GetArtifactsAsync(ObjectId? JobId, SubResourceId? StepId, string? Name)
+		public async Task<List<IArtifact>> GetArtifactsAsync(ObjectId? JobId, SubResourceId? StepId, string? Name)
 		{
 			FilterDefinitionBuilder<Artifact> Builder = Builders<Artifact>.Filter;
 
@@ -110,7 +140,7 @@ namespace HordeServer.Services
 				Filter &= Builder.Eq(x => x.Name, Name);
 			}
 
-			return await Artifacts.Find(Filter).ToListAsync();
+			return await Artifacts.Find(Filter).ToListAsync<Artifact, IArtifact>();
 		}
 
 		/// <summary>
@@ -118,14 +148,14 @@ namespace HordeServer.Services
 		/// </summary>
 		/// <param name="ArtifactIds">The list of artifact Ids</param>
 		/// <returns>List of artifact documents</returns>
-		public async Task<List<Artifact>> GetArtifactsAsync(IEnumerable<ObjectId> ArtifactIds)
+		public async Task<List<IArtifact>> GetArtifactsAsync(IEnumerable<ObjectId> ArtifactIds)
 		{
 			FilterDefinitionBuilder<Artifact> Builder = Builders<Artifact>.Filter;
 
 			FilterDefinition<Artifact> Filter = FilterDefinition<Artifact>.Empty;
 			Filter &= Builder.In(x => x.Id, ArtifactIds);
 			
-			return await Artifacts.Find(Filter).ToListAsync();
+			return await Artifacts.Find(Filter).ToListAsync<Artifact, IArtifact>();
 		}
 
 		/// <summary>
@@ -133,7 +163,7 @@ namespace HordeServer.Services
 		/// </summary>
 		/// <param name="ArtifactId">Unique id of the artifact</param>
 		/// <returns>The artifact document</returns>
-		public async Task<Artifact?> GetArtifactAsync(ObjectId ArtifactId)
+		public async Task<IArtifact?> GetArtifactAsync(ObjectId ArtifactId)
 		{
 			return await Artifacts.Find<Artifact>(x => x.Id == ArtifactId).FirstOrDefaultAsync();
 		}
@@ -157,7 +187,7 @@ namespace HordeServer.Services
 		/// <param name="NewMimeType">New mime type</param>
 		/// <param name="NewData">New data</param>
 		/// <returns>Async task</returns>
-		public async Task<bool> UpdateArtifactAsync(Artifact? Artifact, string NewMimeType, System.IO.Stream NewData)
+		public async Task<bool> UpdateArtifactAsync(IArtifact? Artifact, string NewMimeType, System.IO.Stream NewData)
 		{
 			while (Artifact != null)
 			{
@@ -171,7 +201,7 @@ namespace HordeServer.Services
 				string ArtifactName = ValidateName(Artifact.Name);
 				await StorageBackend.WriteAsync(GetPath(Artifact.JobId, Artifact.StepId, ArtifactName), NewData);
 
-				if (await TryUpdateArtifactAsync(Artifact, UpdateBuilder.Combine(Updates)))
+				if (await TryUpdateArtifactAsync((Artifact)Artifact, UpdateBuilder.Combine(Updates)))
 				{
 					return true;
 				}
@@ -186,7 +216,7 @@ namespace HordeServer.Services
 		/// </summary>
 		/// <param name="Artifact">The artifact</param>
 		/// <returns>The chunk data</returns>
-		public async Task<System.IO.Stream> OpenArtifactReadStreamAsync(Artifact Artifact)
+		public async Task<System.IO.Stream> OpenArtifactReadStreamAsync(IArtifact Artifact)
 		{
 			System.IO.Stream? Stream = await StorageBackend.ReadAsync(GetPath(Artifact.JobId, Artifact.StepId, Artifact.Name));
 			if (Stream == null)
diff --git a/Engine/Source/Programs/Horde/HordeServer/Projects/ISchedule.cs b/Engine/Source/Programs/Horde/HordeServer/Projects/ISchedule.cs
--- a/Engine/Source/Programs/Horde/HordeServer/Projects/ISchedule.cs
+++ b/Engine/Source/Programs/Horde/HordeServer/Projects/ISchedule.cs
@@ -230,6 +230,7 @@ namespace HordeServer.Models
 		/// <param name="RequireSubmittedChange">Whether a change has to be submitted for the schedule to trigger</param>
 		/// <param name="Gate">Reference to another job/target which much succeed for this schedule to trigger</param>
 		/// <param name="Filter">Filter for changes to consider</param>
+		/// <param name="Files">Files that should trigger the schedule</param>
 		/// <param name="TemplateParameters">Parameters for the template to run</param>
 		/// <param name="Patterns">List of patterns for the schedule</param>
 		public Schedule(bool Enabled = true, int MaxActive = 0, int MaxChanges = 0, bool RequireSubmittedChange = true, ScheduleGate? Gate = null, List<ChangeContentFlags>? Filter = null, List<string>? Files = null, Dictionary<string, string>? TemplateParameters = null, List<SchedulePattern>? Patterns = null)
diff --git a/Engine/Source/Programs/Horde/HordeServer/Server/DatabaseService.cs b/Engine/Source/Programs/Horde/HordeServer/Server/DatabaseService.cs
--- a/Engine/Source/Programs/Horde/HordeServer/Server/DatabaseService.cs
+++ b/Engine/Source/Programs/Horde/HordeServer/Server/DatabaseService.cs
@@ -56,11 +56,6 @@ namespace HordeServer.Services
 		/// </summary>
 		public IMongoCollection<Credential> Credentials { get; }
 
-		/// <summary>
-		/// Collection of artifact documents
-		/// </summary>
-		public IMongoCollection<Artifact> Artifacts { get; }
-
 		/// <summary>
 		/// Collection of singleton documents
 		/// </summary>
@@ -231,7 +226,6 @@ namespace HordeServer.Services
 				}
 
 				Credentials = GetCollection<Credential>("Credentials");
-				Artifacts = GetCollection<Artifact>("Artifacts");
 			}
 			catch (Exception Ex)
 			{
diff --git a/Engine/Source/Programs/Horde/HordeServer/Services/RpcService.cs b/Engine/Source/Programs/Horde/HordeServer/Services/RpcService.cs
--- a/Engine/Source/Programs/Horde/HordeServer/Services/RpcService.cs
+++ b/Engine/Source/Programs/Horde/HordeServer/Services/RpcService.cs
@@ -23,6 +23,7 @@ using System.Text;
 using System.Threading.Tasks;
 using HordeServer.Tasks.Impl;
 using System.Threading;
+using HordeServer.Jobs;
 
 [assembly: InternalsVisibleTo("HordeAgentTests")]
 
@@ -79,7 +80,7 @@ namespace HordeServer.Services
 		/// <summary>
 		/// The artifact service instance
 		/// </summary>
-		ArtifactService ArtifactService;
+		IArtifactCollection ArtifactCollection;
 
 		/// <summary>
 		/// The log file service instance
@@ -131,7 +132,7 @@ namespace HordeServer.Services
 		/// <param name="StreamService">Instance of the StreamService singleton</param>
 		/// <param name="LogFileService">Instance of the LogFileService singleton</param>
 		/// <param name="AgentSoftwareService">Instance of the JobService singleton</param>
-		/// <param name="ArtifactService">Instance of the ArtifactService singleton</param>
+		/// <param name="ArtifactCollection">Instance of the ArtifactService singleton</param>
 		/// <param name="CredentialService">Instance of the CredentialsService singleton</param>
 		/// <param name="PoolService">Instance of the PoolService singleton</param>
 		/// <param name="LifetimeService">The application lifetime</param>
@@ -139,7 +140,7 @@ namespace HordeServer.Services
 		/// <param name="TestData">Collection of testdata</param>
 		/// <param name="ConformTaskSource"></param>
 		/// <param name="Logger">Log writer</param>
-		public RpcService(DatabaseService DatabaseService, AclService AclService, AgentService AgentService, StreamService StreamService, JobService JobService, AgentSoftwareService AgentSoftwareService, ArtifactService ArtifactService, ILogFileService LogFileService, CredentialService CredentialService, PoolService PoolService, LifetimeService LifetimeService, IGraphCollection Graphs, ITestDataCollection TestData, ConformTaskSource ConformTaskSource, ILogger<RpcService> Logger)
+		public RpcService(DatabaseService DatabaseService, AclService AclService, AgentService AgentService, StreamService StreamService, JobService JobService, AgentSoftwareService AgentSoftwareService, IArtifactCollection ArtifactCollection, ILogFileService LogFileService, CredentialService CredentialService, PoolService PoolService, LifetimeService LifetimeService, IGraphCollection Graphs, ITestDataCollection TestData, ConformTaskSource ConformTaskSource, ILogger<RpcService> Logger)
 		{
 			this.DatabaseService = DatabaseService;
 			this.AclService = AclService;
@@ -147,7 +148,7 @@ namespace HordeServer.Services
 			this.StreamService = StreamService;
 			this.JobService = JobService;
 			this.AgentSoftwareService = AgentSoftwareService;
-			this.ArtifactService = ArtifactService;
+			this.ArtifactCollection = ArtifactCollection;
 			this.LogFileService = LogFileService;
 			this.CredentialService = CredentialService;
 			this.PoolService = PoolService;
@@ -998,7 +999,7 @@ namespace HordeServer.Services
 			// Upload the stream
 			using (ArtifactChunkStream InputStream = new ArtifactChunkStream(Reader, Metadata.Length))
 			{
-				Artifact Artifact = await ArtifactService.CreateArtifactAsync(Job.Id, Step.Id, Metadata.Name, Metadata.MimeType, InputStream);
+				IArtifact Artifact = await ArtifactCollection.CreateArtifactAsync(Job.Id, Step.Id, Metadata.Name, Metadata.MimeType, InputStream);
 
 				UploadArtifactResponse Response = new UploadArtifactResponse();
 				Response.Id = Artifact.Id.ToString();
diff --git a/Engine/Source/Programs/Horde/HordeServer/Startup.cs b/Engine/Source/Programs/Horde/HordeServer/Startup.cs
--- a/Engine/Source/Programs/Horde/HordeServer/Startup.cs
+++ b/Engine/Source/Programs/Horde/HordeServer/Startup.cs
@@ -82,6 +82,7 @@ using HordeServer.Compute.Impl;
 using HordeServer.Compute;
 using System.Net.Http.Headers;
 using Serilog.Events;
+using HordeServer.Jobs;
 
 namespace HordeServer
 {
@@ -315,6 +316,7 @@ namespace HordeServer
 
 			Services.AddSingleton<IAgentCollection, AgentCollection>();
 			Services.AddSingleton<IAgentSoftwareCollection, AgentSoftwareCollection>();
+			Services.AddSingleton<IArtifactCollection, ArtifactService>();
 			Services.AddSingleton<ICommitCollection, CommitCollection>();
 			Services.AddSingleton<ICounterCollection, CounterCollection>();
 			Services.AddSingleton<IGraphCollection, GraphCollection>();
@@ -366,7 +368,6 @@ namespace HordeServer
 			Services.AddSingleton<AclService>();
 			Services.AddSingleton<AgentService>();			
 			Services.AddSingleton<AgentSoftwareService>();
-			Services.AddSingleton<ArtifactService>();
 			Services.AddSingleton<ConsistencyService>();
 			Services.AddSingleton<RequestTrackerService>();
 			Services.AddSingleton<CredentialService>();
