Horde: Rename ArtifactService to ArtifactCollection.

diff --git a/Engine/Source/Programs/Horde/HordeServer.Compute.Tests/KeyValueTreeTests.cs b/Engine/Source/Programs/Horde/HordeServer.Compute.Tests/KeyValueTreeTests.cs
--- a/Engine/Source/Programs/Horde/HordeServer.Compute.Tests/KeyValueTreeTests.cs
+++ b/Engine/Source/Programs/Horde/HordeServer.Compute.Tests/KeyValueTreeTests.cs
@@ -67,7 +67,7 @@ namespace HordeServerTests
 		[TestMethod]
 		public async Task TestBasicInsertion()
 		{
-			IStorageBackend<ArtifactService> Backend = new TransientStorageBackend().ForType<ArtifactService>();
+			IStorageBackend<ArtifactCollection> Backend = new TransientStorageBackend().ForType<ArtifactCollection>();
 			SimpleStorageService Storage = new SimpleStorageService(Backend);
 
 			BlobRef<TestTreeMax2> Tree = await Storage.CreateKeyValueTreeAsync<TestTreeMax2>();
@@ -122,7 +122,7 @@ namespace HordeServerTests
 		{
 			Random Random = new Random(0x12345678);
 
-			IStorageBackend<ArtifactService> Backend = new TransientStorageBackend().ForType<ArtifactService>();
+			IStorageBackend<ArtifactCollection> Backend = new TransientStorageBackend().ForType<ArtifactCollection>();
 			SimpleStorageService Storage = new SimpleStorageService(Backend);
 
 			BlobRef<TestTreeMax256> Tree = await Storage.CreateKeyValueTreeAsync<TestTreeMax256>();
diff --git a/Engine/Source/Programs/Horde/HordeServer.Compute.Tests/TestSetup.cs b/Engine/Source/Programs/Horde/HordeServer.Compute.Tests/TestSetup.cs
--- a/Engine/Source/Programs/Horde/HordeServer.Compute.Tests/TestSetup.cs
+++ b/Engine/Source/Programs/Horde/HordeServer.Compute.Tests/TestSetup.cs
@@ -145,7 +145,7 @@ namespace HordeServerTests
 
 			Services.AddSingleton<IAgentCollection, AgentCollection>();
 			Services.AddSingleton<IAgentSoftwareCollection, AgentSoftwareCollection>();
-			Services.AddSingleton<IArtifactCollection, ArtifactService>();
+			Services.AddSingleton<IArtifactCollection, ArtifactCollection>();
 			Services.AddSingleton<ICounterCollection, CounterCollection>();
 			Services.AddSingleton<IGraphCollection, GraphCollection>();
 			Services.AddSingleton<IIssueCollection, IssueCollection>();
@@ -209,7 +209,7 @@ namespace HordeServerTests
 			Services.AddSingleton<IBucketCollection, BucketCollection>();
 
 			Services.AddSingleton(new Startup.StorageBackendSettings<PersistentLogStorage> { Type = StorageProviderType.Transient });
-			Services.AddSingleton(new Startup.StorageBackendSettings<ArtifactService> { Type = StorageProviderType.Transient });
+			Services.AddSingleton(new Startup.StorageBackendSettings<ArtifactCollection> { Type = StorageProviderType.Transient });
 			Services.AddSingleton(new Startup.StorageBackendSettings<BlobCollection> { Type = StorageProviderType.Transient });
 			Services.AddSingleton(typeof(IStorageBackend<>), typeof(Startup.StorageBackendFactory<>));
 
diff --git a/Engine/Source/Programs/Horde/HordeServer/Api/Schedules.cs b/Engine/Source/Programs/Horde/HordeServer/Api/Schedules.cs
--- a/Engine/Source/Programs/Horde/HordeServer/Api/Schedules.cs
+++ b/Engine/Source/Programs/Horde/HordeServer/Api/Schedules.cs
@@ -118,6 +118,11 @@ namespace HordeServer.Api
 		/// </summary>
 		public List<ChangeContentFlags>? Filter { get; set; }
 
+		/// <summary>
+		/// Files that should cause the job to trigger
+		/// </summary>
+		public List<string>? Files { get; set; }
+
 		/// <summary>
 		/// Parameters for the template
 		/// </summary>
@@ -134,7 +139,7 @@ namespace HordeServer.Api
 		/// <returns>New model object</returns>
 		public Schedule ToModel()
 		{
-			return new Schedule(Enabled, MaxActive, MaxChanges, RequireSubmittedChange, Gate?.ToModel(), Filter, TemplateParameters, Patterns.ConvertAll(x => x.ToModel()));
+			return new Schedule(Enabled, MaxActive, MaxChanges, RequireSubmittedChange, Gate?.ToModel(), Filter, Files, TemplateParameters, Patterns.ConvertAll(x => x.ToModel()));
 		}
 	}
 
diff --git a/Engine/Source/Programs/Horde/HordeServer/Startup.cs b/Engine/Source/Programs/Horde/HordeServer/Startup.cs
--- a/Engine/Source/Programs/Horde/HordeServer/Startup.cs
+++ b/Engine/Source/Programs/Horde/HordeServer/Startup.cs
@@ -316,7 +316,7 @@ namespace HordeServer
 
 			Services.AddSingleton<IAgentCollection, AgentCollection>();
 			Services.AddSingleton<IAgentSoftwareCollection, AgentSoftwareCollection>();
-			Services.AddSingleton<IArtifactCollection, ArtifactService>();
+			Services.AddSingleton<IArtifactCollection, ArtifactCollection>();
 			Services.AddSingleton<ICommitCollection, CommitCollection>();
 			Services.AddSingleton<ICounterCollection, CounterCollection>();
 			Services.AddSingleton<IGraphCollection, GraphCollection>();
@@ -432,7 +432,7 @@ namespace HordeServer
 			Services.AddSingleton(AwsOptions);
 
 			Services.AddSingleton(new StorageBackendSettings<PersistentLogStorage> { Type = Settings.ExternalStorageProviderType, BaseDir = Settings.LocalLogsDir, BucketName = Settings.S3LogBucketName });
-			Services.AddSingleton(new StorageBackendSettings<ArtifactService> { Type = Settings.ExternalStorageProviderType, BaseDir = Settings.LocalArtifactsDir, BucketName = Settings.S3ArtifactBucketName });
+			Services.AddSingleton(new StorageBackendSettings<ArtifactCollection> { Type = Settings.ExternalStorageProviderType, BaseDir = Settings.LocalArtifactsDir, BucketName = Settings.S3ArtifactBucketName });
 			Services.AddSingleton(new StorageBackendSettings<BlobCollection> { Type = Settings.ExternalStorageProviderType, BaseDir = Settings.LocalBlobsDir, BucketName = Settings.S3LogBucketName, BucketPath = "blobs/" });
 			Services.AddSingleton(typeof(IStorageBackend<>), typeof(StorageBackendFactory<>));
 
diff --git a/Engine/Source/Programs/Horde/HordeServer/Storage/Services/SimpleStorageService.cs b/Engine/Source/Programs/Horde/HordeServer/Storage/Services/SimpleStorageService.cs
--- a/Engine/Source/Programs/Horde/HordeServer/Storage/Services/SimpleStorageService.cs
+++ b/Engine/Source/Programs/Horde/HordeServer/Storage/Services/SimpleStorageService.cs
@@ -37,7 +37,7 @@ namespace HordeServer.Storage.Services
 		/// Constructor
 		/// </summary>
 		/// <param name="Backend">The storage backend</param>
-		public SimpleStorageService(IStorageBackend<ArtifactService> Backend)
+		public SimpleStorageService(IStorageBackend<ArtifactCollection> Backend)
 		{
 			this.Backend = Backend;
 		}
