Horde: Remove the "counter" collection. We've never ended up using this, and there's a lot of boilerplate code to support it.

diff --git a/Engine/Source/Programs/Horde/HordeServer.Compute.Tests/Fixture.cs b/Engine/Source/Programs/Horde/HordeServer.Compute.Tests/Fixture.cs
--- a/Engine/Source/Programs/Horde/HordeServer.Compute.Tests/Fixture.cs
+++ b/Engine/Source/Programs/Horde/HordeServer.Compute.Tests/Fixture.cs
@@ -54,7 +54,7 @@ namespace HordeServerTests
 			Fg.Labels = new List<ILabel>();
 
 			Template = await TemplateService.CreateTemplateAsync("Test template", null, false, null, null,
-				new List<TemplateCounter>(), new List<string>(), new List<Parameter>());
+				new List<string>(), new List<Parameter>());
 			Graph = await GraphCollection.AddAsync(Template);
 
 			TemplateRefId1 = new TemplateRefId("template1");
@@ -99,7 +99,6 @@ namespace HordeServerTests
 				null,
 				null,
 				null,
-				Counters: Template.Counters,
 				Arguments: new List<string>()
 			);
 			Job1 = (await JobService.GetJobAsync(Job1.Id))!;
@@ -126,7 +125,6 @@ namespace HordeServerTests
 				null,
 				null,
 				null,
-				Template.Counters,
 				Arguments: new List<string>()
 			);
 			Job2 = (await JobService.GetJobAsync(Job2.Id))!;
diff --git a/Engine/Source/Programs/Horde/HordeServer.Compute.Tests/JobServiceTests.cs b/Engine/Source/Programs/Horde/HordeServer.Compute.Tests/JobServiceTests.cs
--- a/Engine/Source/Programs/Horde/HordeServer.Compute.Tests/JobServiceTests.cs
+++ b/Engine/Source/Programs/Horde/HordeServer.Compute.Tests/JobServiceTests.cs
@@ -28,7 +28,7 @@ namespace HordeServerTests
 			IProject? Project = await TestSetup.ProjectService.Collection.AddOrUpdateAsync(ProjectId, "", "", 0, new ProjectConfig { Name = "UE5" });
 			Assert.IsNotNull(Project);
 
-			ITemplate Template = await TestSetup.TemplateService.CreateTemplateAsync("Test template", null, false, null, null, new List<TemplateCounter>(), new List<string>(), new List<Parameter>());
+			ITemplate Template = await TestSetup.TemplateService.CreateTemplateAsync("Test template", null, false, null, null, new List<string>(), new List<Parameter>());
 			IGraph Graph = await TestSetup.GraphCollection.AddAsync(Template);
 
 			TemplateRefId TemplateRefId1 = new TemplateRefId("template1");
@@ -43,7 +43,7 @@ namespace HordeServerTests
 			IStream? Stream = await TestSetup.StreamService.GetStreamAsync(StreamId);
 			Stream = await TestSetup.StreamService.StreamCollection.TryCreateOrReplaceAsync(new StreamId("ue5-main"), Stream, String.Empty, String.Empty, ProjectId, StreamConfig);
 
-			IJob Job = await TestSetup.JobService.CreateJobAsync(null, Stream!, TemplateRefId1, Template.Id, Graph, "Hello", 1234, 1233, 999, null, null, "joe", null, null, null, Stream!.Templates[TemplateRefId1].ChainedJobs, true, true, null, null, null, Template.Counters, new List<string>());
+			IJob Job = await TestSetup.JobService.CreateJobAsync(null, Stream!, TemplateRefId1, Template.Id, Graph, "Hello", 1234, 1233, 999, null, null, "joe", null, null, null, Stream!.Templates[TemplateRefId1].ChainedJobs, true, true, null, null, null, new List<string>());
 			Assert.AreEqual(1, Job.ChainedJobs.Count);
 
 			Assert.IsTrue(await TestSetup.JobService.UpdateBatchAsync(Job, Job.Batches[0].Id, ObjectId.GenerateNewId(), JobStepBatchState.Running));
@@ -114,7 +114,6 @@ namespace HordeServerTests
 				null,
 				null,
 				null,
-				TestSetup.Fixture!.Template.Counters,
 				Arguments: new List<string>(Arguments)
 			);
 		}
@@ -136,7 +135,7 @@ namespace HordeServerTests
 		// 	Assert.AreEqual(JobState.Running, Job1.GetState());
 		// 	Assert.AreEqual(JobState.Complete, Job2.GetState());
 		// 	
-		// 	IJob Job3 = await TestSetup.JobService.CreateJobAsync(null, Fix.Stream!.Id, Fix.TemplateRefId1, Fix.Template.Id, Fix.Graph, "Hello", Job1.Change, Job1.CodeChange, null, "joe", null, null, true, true, null, null, Fix.Template.Counters, new List<string>());
+		// 	IJob Job3 = await TestSetup.JobService.CreateJobAsync(null, Fix.Stream!.Id, Fix.TemplateRefId1, Fix.Template.Id, Fix.Graph, "Hello", Job1.Change, Job1.CodeChange, null, "joe", null, null, true, true, null, null, new List<string>());
 		// 	await TestSetup.DispatchService.TickOnlyForTestingAsync();
 		// 	Job1 = (await TestSetup.JobService.GetJobAsync(Job1.Id))!;
 		// 	Assert.AreEqual(JobState.Complete, Job1.GetState());
@@ -155,7 +154,7 @@ namespace HordeServerTests
 			IStream? Stream = await TestSetup.StreamCollection.GetAsync(StreamId);
 			Stream = await TestSetup.StreamCollection.TryCreateOrReplaceAsync(StreamId, Stream, "", "", Project!.Id, new StreamConfig { Name = "//UE5/Main" });
 
-			ITemplate Template = await TestSetup.TemplateService.CreateTemplateAsync("Test template", null, false, null, null, new List<TemplateCounter>(), new List<string>(), new List<Parameter>());
+			ITemplate Template = await TestSetup.TemplateService.CreateTemplateAsync("Test template", null, false, null, null, new List<string>(), new List<Parameter>());
 			IGraph Graph = await TestSetup.GraphCollection.AddAsync(Template);
 
 			NewGroup GroupA = new NewGroup("win", new List<NewNode>());
@@ -168,7 +167,7 @@ namespace HordeServerTests
 
 			Graph = await TestSetup.GraphCollection.AppendAsync(Graph, new List<NewGroup> { GroupA, GroupB });
 
-			IJob Job = await TestSetup.JobService.CreateJobAsync(null, Stream!, new TemplateRefId("temp"), Template.Id, Graph, "Hello", 1234, 1233, 999, null, null, "joe", null, null, null, null, true, true, null, null, null, Template.Counters, new List<string> { "-Target=Pak" });
+			IJob Job = await TestSetup.JobService.CreateJobAsync(null, Stream!, new TemplateRefId("temp"), Template.Id, Graph, "Hello", 1234, 1233, 999, null, null, "joe", null, null, null, null, true, true, null, null, null, new List<string> { "-Target=Pak" });
 
 			Assert.IsTrue(await TestSetup.JobService.UpdateBatchAsync(Job, Job.Batches[0].Id, ObjectId.GenerateNewId(), JobStepBatchState.Running));
 			Assert.IsNotNull(await TestSetup.JobService.UpdateStepAsync(Job, Job.Batches[0].Id, Job.Batches[0].Steps[0].Id, JobStepState.Running));
diff --git a/Engine/Source/Programs/Horde/HordeServer.Compute.Tests/SchedulerTests.cs b/Engine/Source/Programs/Horde/HordeServer.Compute.Tests/SchedulerTests.cs
--- a/Engine/Source/Programs/Horde/HordeServer.Compute.Tests/SchedulerTests.cs
+++ b/Engine/Source/Programs/Horde/HordeServer.Compute.Tests/SchedulerTests.cs
@@ -42,7 +42,7 @@ namespace HordeServerTests
 			IProject ? Project = TestSetup.ProjectService.Collection.AddOrUpdateAsync(ProjectId, "", "", 0, new ProjectConfig { Name = "UE5" }).Result;
 			Assert.IsNotNull(Project);
 
-			Template = TestSetup.TemplateService.CreateTemplateAsync("Test template", null, false, null, null, new List<TemplateCounter>(), new List<string>(), new List<Parameter>()).Result;
+			Template = TestSetup.TemplateService.CreateTemplateAsync("Test template", null, false, null, null, new List<string>(), new List<Parameter>()).Result;
 
 			InitialJobIds = new HashSet<ObjectId>(TestSetup.JobCollection.FindAsync().Result.Select(x => x.Id));
 
@@ -387,7 +387,7 @@ namespace HordeServerTests
 			Assert.AreEqual(0, Jobs2.Count);
 
 			// Create a job and fail it
-			IJob Job1 = await TestSetup.JobService.CreateJobAsync(null, Stream, NewTemplateRefId1, Template.Id, GraphA, "Hello", 1234, 1233, 999, null, null, "joe", null, null, null, null, true, true, null, null, null, Template.Counters, new List<string> { "-Target=TriggerNext" });
+			IJob Job1 = await TestSetup.JobService.CreateJobAsync(null, Stream, NewTemplateRefId1, Template.Id, GraphA, "Hello", 1234, 1233, 999, null, null, "joe", null, null, null, null, true, true, null, null, null, new List<string> { "-Target=TriggerNext" });
 			Assert.IsTrue(await TestSetup.JobService.UpdateBatchAsync(Job1, Job1.Batches[0].Id, ObjectId.GenerateNewId(), JobStepBatchState.Running));
 			Assert.IsNotNull(await TestSetup.JobService.UpdateStepAsync(Job1, Job1.Batches[0].Id, Job1.Batches[0].Steps[0].Id, JobStepState.Completed, JobStepOutcome.Failure));
 			Assert.IsTrue(await TestSetup.JobService.UpdateBatchAsync(Job1, Job1.Batches[0].Id, ObjectId.GenerateNewId(), JobStepBatchState.Complete));
@@ -400,7 +400,7 @@ namespace HordeServerTests
 			Assert.AreEqual(0, Jobs3.Count);
 
 			// Create a job and make it succeed
-			IJob Job2 = await TestSetup.JobService.CreateJobAsync(null, Stream, NewTemplateRefId1, Template.Id, GraphA, "Hello", 1234, 1233, 999, null, null, "joe", null, null, null, null, true, true, null, null, null, Template.Counters, new List<string> { "-Target=TriggerNext" });
+			IJob Job2 = await TestSetup.JobService.CreateJobAsync(null, Stream, NewTemplateRefId1, Template.Id, GraphA, "Hello", 1234, 1233, 999, null, null, "joe", null, null, null, null, true, true, null, null, null, new List<string> { "-Target=TriggerNext" });
 			Assert.IsTrue(await TestSetup.JobService.UpdateBatchAsync(Job2, Job2.Batches[0].Id, ObjectId.GenerateNewId(), JobStepBatchState.Running));
 			Assert.IsNotNull(await TestSetup.JobService.UpdateStepAsync(Job2, Job2.Batches[0].Id, Job2.Batches[0].Steps[0].Id, JobStepState.Completed, JobStepOutcome.Success));
 
diff --git a/Engine/Source/Programs/Horde/HordeServer.Compute.Tests/TestSetup.cs b/Engine/Source/Programs/Horde/HordeServer.Compute.Tests/TestSetup.cs
--- a/Engine/Source/Programs/Horde/HordeServer.Compute.Tests/TestSetup.cs
+++ b/Engine/Source/Programs/Horde/HordeServer.Compute.Tests/TestSetup.cs
@@ -63,7 +63,6 @@ namespace HordeServerTests
 		public IJobStepRefCollection JobStepRefCollection => ServiceProvider.GetRequiredService<IJobStepRefCollection>();
 		public IJobTimingCollection JobTimingCollection => ServiceProvider.GetRequiredService<IJobTimingCollection>();
 		public IUgsMetadataCollection UgsMetadataCollection => ServiceProvider.GetRequiredService<IUgsMetadataCollection>();
-		public ICounterCollection CounterCollection => ServiceProvider.GetRequiredService <ICounterCollection>();
 		public IIssueCollection IssueCollection => ServiceProvider.GetRequiredService <IIssueCollection>();
 		public IPoolCollection PoolCollection => ServiceProvider.GetRequiredService <IPoolCollection>();
 		public ILeaseCollection LeaseCollection => ServiceProvider.GetRequiredService <ILeaseCollection>();
@@ -146,7 +145,6 @@ namespace HordeServerTests
 			Services.AddSingleton<IAgentCollection, AgentCollection>();
 			Services.AddSingleton<IAgentSoftwareCollection, AgentSoftwareCollection>();
 			Services.AddSingleton<IArtifactCollection, ArtifactCollection>();
-			Services.AddSingleton<ICounterCollection, CounterCollection>();
 			Services.AddSingleton<IGraphCollection, GraphCollection>();
 			Services.AddSingleton<IIssueCollection, IssueCollection>();
 			Services.AddSingleton<IJobCollection, JobCollection>();
diff --git a/Engine/Source/Programs/Horde/HordeServer/Api/Templates.cs b/Engine/Source/Programs/Horde/HordeServer/Api/Templates.cs
--- a/Engine/Source/Programs/Horde/HordeServer/Api/Templates.cs
+++ b/Engine/Source/Programs/Horde/HordeServer/Api/Templates.cs
@@ -388,43 +388,6 @@ namespace HordeServer.Api
 			return new BoolParameter(Label, ArgumentIfEnabled, ArgumentIfDisabled, Default, ToolTip);
 		}
 	}
-	
-	/// <summary>
-	/// Counter for executing a job based on a template
-	/// </summary>
-	public class TemplateCounter
-	{
-		/// <summary>
-		/// The property name
-		/// </summary>
-		[Required]
-		public string PropertyName { get; set; }
-
-		/// <summary>
-		/// The counter name
-		/// </summary>
-		[Required]
-		public string CounterName { get; set; }
-
-		/// <summary>
-		/// Private constructor for serialization
-		/// </summary>
-		private TemplateCounter()
-		{
-			PropertyName = null!;
-			CounterName = null!;
-		}
-
-		/// <summary>
-		/// Constructor
-		/// </summary>
-		/// <param name="Counter">Counter to construct from</param>
-		public TemplateCounter(ITemplateCounter Counter)
-		{
-			this.PropertyName = Counter.PropertyName;
-			this.CounterName = Counter.CounterName;
-		}
-	}
 
 	/// <summary>
 	/// Parameters to create a new template
@@ -457,11 +420,6 @@ namespace HordeServer.Api
 		/// </summary>
 		public string? SubmitNewChange { get; set; }
 
-		/// <summary>
-		/// Counters for this template
-		/// </summary>
-		public List<TemplateCounter> Counters { get; set; } = new List<TemplateCounter>();
-
 		/// <summary>
 		/// Fixed arguments for the new job
 		/// </summary>
@@ -503,11 +461,6 @@ namespace HordeServer.Api
 		/// </summary>
 		public string? SubmitNewChange { get; }
 
-		/// <summary>
-		/// List of template counters
-		/// </summary>
-		public List<TemplateCounter> Counters { get; set; }
-
 		/// <summary>
 		/// Parameters for the job.
 		/// </summary>
@@ -525,7 +478,6 @@ namespace HordeServer.Api
 		{
 			this.Name = null!;
 			this.AllowPreflights = true;
-			this.Counters = new List<TemplateCounter>();
 			this.Arguments = new List<string>();
 			this.Parameters = new List<ParameterData>();
 		}
@@ -541,7 +493,6 @@ namespace HordeServer.Api
 			this.AllowPreflights = Template.AllowPreflights;
 			this.InitialAgentType = Template.InitialAgentType;
 			this.SubmitNewChange = Template.SubmitNewChange;
-			this.Counters = Template.Counters.ConvertAll(x => new TemplateCounter(x));
 			this.Arguments = new List<string>(Template.Arguments);
 			this.Parameters = Template.Parameters.ConvertAll(x => x.ToData());
 		}
diff --git a/Engine/Source/Programs/Horde/HordeServer/Controllers/JobsController.cs b/Engine/Source/Programs/Horde/HordeServer/Controllers/JobsController.cs
--- a/Engine/Source/Programs/Horde/HordeServer/Controllers/JobsController.cs
+++ b/Engine/Source/Programs/Horde/HordeServer/Controllers/JobsController.cs
@@ -189,7 +189,7 @@ namespace HordeServer.Controllers
 			List<string> Arguments = Create.Arguments ?? Template.GetDefaultArguments();
 
 			// Create the job
-			IJob Job = await JobService.CreateJobAsync(null, Stream, TemplateRefId, Template.Id, Graph, Name, Change, CodeChange, Create.PreflightChange, null, User.GetUserId(), User.GetUserName(), Priority, Create.AutoSubmit, Create.UpdateIssues, TemplateRef.ChainedJobs, TemplateRef.ShowUgsBadges, TemplateRef.ShowUgsAlerts, TemplateRef.NotificationChannel, TemplateRef.NotificationChannelFilter, null, Template.Counters, Arguments);
+			IJob Job = await JobService.CreateJobAsync(null, Stream, TemplateRefId, Template.Id, Graph, Name, Change, CodeChange, Create.PreflightChange, null, User.GetUserId(), User.GetUserName(), Priority, Create.AutoSubmit, Create.UpdateIssues, TemplateRef.ChainedJobs, TemplateRef.ShowUgsBadges, TemplateRef.ShowUgsAlerts, TemplateRef.NotificationChannel, TemplateRef.NotificationChannelFilter, null, Arguments);
 			await UpdateNotificationsAsync(Job.Id.ToString(), new UpdateNotificationsRequest { Slack = true });
 			return new CreateJobResponse(Job.Id.ToString());
 		}
diff --git a/Engine/Source/Programs/Horde/HordeServer/Controllers/StreamsController.cs b/Engine/Source/Programs/Horde/HordeServer/Controllers/StreamsController.cs
--- a/Engine/Source/Programs/Horde/HordeServer/Controllers/StreamsController.cs
+++ b/Engine/Source/Programs/Horde/HordeServer/Controllers/StreamsController.cs
@@ -344,7 +344,7 @@ namespace HordeServer.Controllers
 			foreach (CreateTemplateRefRequest Request in Requests)
 			{
 				// Create the template
-				ITemplate NewTemplate = await TemplateService.CreateTemplateAsync(Request.Name, Request.Priority, Request.AllowPreflights, Request.InitialAgentType, Request.SubmitNewChange, Request.Counters, Request.Arguments, Request.Parameters.ConvertAll(x => x.ToModel()));
+				ITemplate NewTemplate = await TemplateService.CreateTemplateAsync(Request.Name, Request.Priority, Request.AllowPreflights, Request.InitialAgentType, Request.SubmitNewChange, Request.Arguments, Request.Parameters.ConvertAll(x => x.ToModel()));
 
 				// Get an identifier for the new template ref
 				TemplateRefId NewTemplateRefId;
diff --git a/Engine/Source/Programs/Horde/HordeServer/Jobs/ICounterCollection.cs b/Engine/Source/Programs/Horde/HordeServer/Jobs/ICounterCollection.cs
deleted file mode 100644
diff --git a/Engine/Source/Programs/Horde/HordeServer/Jobs/ITemplate.cs b/Engine/Source/Programs/Horde/HordeServer/Jobs/ITemplate.cs
--- a/Engine/Source/Programs/Horde/HordeServer/Jobs/ITemplate.cs
+++ b/Engine/Source/Programs/Horde/HordeServer/Jobs/ITemplate.cs
@@ -416,22 +416,6 @@ namespace HordeServer.Models
 		}
 	}
 
-	/// <summary>
-	/// Counter to be updated when executing a template
-	/// </summary>
-	public interface ITemplateCounter
-	{
-		/// <summary>
-		/// Name of the property to set to the counter value
-		/// </summary>
-		string PropertyName { get; }
-
-		/// <summary>
-		/// Name of the counter
-		/// </summary>
-		string CounterName { get; }
-	}
-
 	/// <summary>
 	/// Document describing a job template. These objects are considered immutable once created and uniquely referenced by hash, in order to de-duplicate across all job runs.
 	/// </summary>
@@ -467,11 +451,6 @@ namespace HordeServer.Models
 		/// </summary>
 		public string? SubmitNewChange { get; }
 
-		/// <summary>
-		/// List of property counters
-		/// </summary>
-		public IReadOnlyList<ITemplateCounter> Counters { get; }
-
 		/// <summary>
 		/// Optional predefined user-defined properties for this job
 		/// </summary>
diff --git a/Engine/Source/Programs/Horde/HordeServer/Jobs/ITemplateCollection.cs b/Engine/Source/Programs/Horde/HordeServer/Jobs/ITemplateCollection.cs
--- a/Engine/Source/Programs/Horde/HordeServer/Jobs/ITemplateCollection.cs
+++ b/Engine/Source/Programs/Horde/HordeServer/Jobs/ITemplateCollection.cs
@@ -30,10 +30,9 @@ namespace HordeServer.Collections
 		/// <param name="bAllowPreflights">Whether to allow preflights of this job</param>
 		/// <param name="InitialAgentType">The agent type to parse the buildgraph script</param>
 		/// <param name="SubmitNewChange">Path to a file within the stream to submit to generate a new changelist for jobs</param>
-		/// <param name="Counters">List of counters for new jobs</param>
 		/// <param name="Arguments">List of arguments which are always specified</param>
 		/// <param name="Parameters">List of template parameters</param>
-		Task<ITemplate> AddAsync(string Name, Priority? Priority = null, bool bAllowPreflights = true, string? InitialAgentType = null, string? SubmitNewChange = null, List<TemplateCounter>? Counters = null, List<string>? Arguments = null, List<Parameter>? Parameters = null);
+		Task<ITemplate> AddAsync(string Name, Priority? Priority = null, bool bAllowPreflights = true, string? InitialAgentType = null, string? SubmitNewChange = null, List<string>? Arguments = null, List<Parameter>? Parameters = null);
 
 		/// <summary>
 		/// Gets all the available templates
diff --git a/Engine/Source/Programs/Horde/HordeServer/Jobs/Impl/CounterCollection.cs b/Engine/Source/Programs/Horde/HordeServer/Jobs/Impl/CounterCollection.cs
deleted file mode 100644
diff --git a/Engine/Source/Programs/Horde/HordeServer/Jobs/Impl/JobService.cs b/Engine/Source/Programs/Horde/HordeServer/Jobs/Impl/JobService.cs
--- a/Engine/Source/Programs/Horde/HordeServer/Jobs/Impl/JobService.cs
+++ b/Engine/Source/Programs/Horde/HordeServer/Jobs/Impl/JobService.cs
@@ -101,11 +101,6 @@ namespace HordeServer.Services
 		/// </summary>
 		IUserCollection UserCollection;
 
-		/// <summary>
-		/// Collection of counter documents
-		/// </summary>
-		ICounterCollection CounterCollection;
-
 		/// <summary>
 		/// Collection of notification triggers
 		/// </summary>
@@ -175,7 +170,6 @@ namespace HordeServer.Services
 		/// <param name="JobStepRefs">The jobsteprefs collection</param>
 		/// <param name="JobTimings">The job timing document collection</param>
 		/// <param name="UserCollection">User profiles</param>
-		/// <param name="CounterCollection">Counter collection</param>
 		/// <param name="TriggerCollection">Trigger collection</param>
 		/// <param name="JobTaskSource">The queue service</param>
 		/// <param name="StreamService">The stream service</param>
@@ -184,7 +178,7 @@ namespace HordeServer.Services
 		/// <param name="PerforceService">The perforce service</param>
 		/// <param name="Settings">Settings instance</param>
 		/// <param name="Logger">Log output</param>
-		public JobService(IJobCollection Jobs, IGraphCollection Graphs, IAgentCollection Agents, IJobStepRefCollection JobStepRefs, IJobTimingCollection JobTimings, IUserCollection UserCollection, ICounterCollection CounterCollection, INotificationTriggerCollection TriggerCollection, JobTaskSource JobTaskSource, StreamService StreamService, TemplateService TemplateService, IIssueService IssueService, IPerforceService PerforceService, IOptionsMonitor<ServerSettings> Settings, ILogger<JobService> Logger)
+		public JobService(IJobCollection Jobs, IGraphCollection Graphs, IAgentCollection Agents, IJobStepRefCollection JobStepRefs, IJobTimingCollection JobTimings, IUserCollection UserCollection, INotificationTriggerCollection TriggerCollection, JobTaskSource JobTaskSource, StreamService StreamService, TemplateService TemplateService, IIssueService IssueService, IPerforceService PerforceService, IOptionsMonitor<ServerSettings> Settings, ILogger<JobService> Logger)
 		{
 			this.Jobs = Jobs;
 			this.Graphs = Graphs;
@@ -192,7 +186,6 @@ namespace HordeServer.Services
 			this.JobStepRefs = JobStepRefs;
 			this.JobTimings = JobTimings;
 			this.UserCollection = UserCollection;
-			this.CounterCollection = CounterCollection;
 			this.TriggerCollection = TriggerCollection;
 			this.JobTaskSource = JobTaskSource;
 			this.StreamService = StreamService;
@@ -234,10 +227,9 @@ namespace HordeServer.Services
 		/// <param name="NotificationChannel">Notification Channel for this job</param>
 		/// <param name="NotificationChannelFilter">Notification Channel filter for this job</param>
 		/// <param name="HelixSwarmCallbackUrl">Helix Swarm callback URL for review, if any</param>
-		/// <param name="Counters">Counters for the job</param>
 		/// <param name="Arguments">Arguments for the job</param>
 		/// <returns>Unique id representing the job</returns>
-		public async Task<IJob> CreateJobAsync(ObjectId? JobId, IStream Stream, TemplateRefId TemplateRefId, ContentHash TemplateHash, IGraph Graph, string Name, int Change, int CodeChange, int? PreflightChange, int? ClonedPreflightChange, ObjectId? StartedByUserId, string? StartedByUserName, Priority? Priority, bool? AutoSubmit, bool? UpdateIssues, List<ChainedJobTemplate>? JobTriggers, bool ShowUgsBadges, bool ShowUgsAlerts, string? NotificationChannel, string? NotificationChannelFilter, string? HelixSwarmCallbackUrl, IReadOnlyList<ITemplateCounter> Counters, IReadOnlyList<string> Arguments)
+		public async Task<IJob> CreateJobAsync(ObjectId? JobId, IStream Stream, TemplateRefId TemplateRefId, ContentHash TemplateHash, IGraph Graph, string Name, int Change, int CodeChange, int? PreflightChange, int? ClonedPreflightChange, ObjectId? StartedByUserId, string? StartedByUserName, Priority? Priority, bool? AutoSubmit, bool? UpdateIssues, List<ChainedJobTemplate>? JobTriggers, bool ShowUgsBadges, bool ShowUgsAlerts, string? NotificationChannel, string? NotificationChannelFilter, string? HelixSwarmCallbackUrl, IReadOnlyList<string> Arguments)
 		{
 			ObjectId JobIdValue = JobId ?? ObjectId.GenerateNewId();
 			using IDisposable Scope = Logger.BeginScope("CreateJobAsync({JobId})", JobIdValue);
@@ -258,14 +250,6 @@ namespace HordeServer.Services
 			Properties["TemplateId"] = TemplateRefId.ToString();
 			Properties["JobId"] = JobIdValue.ToString();
 
-			foreach (ITemplateCounter Counter in Counters)
-			{
-				string CounterName = StringUtils.ExpandProperties(Counter.CounterName, Properties);
-				int Value = await CounterCollection.IncrementAsync(CounterName);
-				Properties[Counter.PropertyName] = Value.ToString(CultureInfo.InvariantCulture);
-				Properties[$"{Counter.PropertyName}.Suffix"] = (Value == 1) ? String.Empty : $" ({Value})";
-			}
-
 			List<string> ExpandedArguments = new List<string>();
 			if (Arguments != null)
 			{
@@ -1137,7 +1121,7 @@ namespace HordeServer.Services
 					IGraph TriggerGraph = await Graphs.AddAsync(Template);
 					Logger.LogInformation("Creating downstream job {ChainedJobId} from job {JobId}", ChainedJobId, Job.Id);
 
-					await CreateJobAsync(ChainedJobId, Stream, JobTrigger.TemplateRefId, TemplateRef.Hash, TriggerGraph, TemplateRef.Name, Job.Change, Job.CodeChange, Job.PreflightChange, Job.ClonedPreflightChange, Job.StartedByUserId, Job.StartedByUser, Template.Priority, null, Job.UpdateIssues, TemplateRef.ChainedJobs, false, false, TemplateRef.NotificationChannel, TemplateRef.NotificationChannelFilter, null, Template.Counters, Template.Arguments);
+					await CreateJobAsync(ChainedJobId, Stream, JobTrigger.TemplateRefId, TemplateRef.Hash, TriggerGraph, TemplateRef.Name, Job.Change, Job.CodeChange, Job.PreflightChange, Job.ClonedPreflightChange, Job.StartedByUserId, Job.StartedByUser, Template.Priority, null, Job.UpdateIssues, TemplateRef.ChainedJobs, false, false, TemplateRef.NotificationChannel, TemplateRef.NotificationChannelFilter, null, Template.Arguments);
 					break;
 				}
 
diff --git a/Engine/Source/Programs/Horde/HordeServer/Jobs/Impl/TemplateCollection.cs b/Engine/Source/Programs/Horde/HordeServer/Jobs/Impl/TemplateCollection.cs
--- a/Engine/Source/Programs/Horde/HordeServer/Jobs/Impl/TemplateCollection.cs
+++ b/Engine/Source/Programs/Horde/HordeServer/Jobs/Impl/TemplateCollection.cs
@@ -23,27 +23,6 @@ namespace HordeServer.Collections.Impl
 	/// </summary>
 	public sealed class TemplateCollection : ITemplateCollection, IDisposable
 	{
-		/// <summary>
-		/// Counter instance
-		/// </summary>
-		class TemplateCounterDocument : ITemplateCounter
-		{
-			public string PropertyName { get; set; }
-			public string CounterName { get; set; }
-
-			private TemplateCounterDocument()
-			{
-				PropertyName = null!;
-				CounterName = null!;
-			}
-
-			public TemplateCounterDocument(string PropertyName, string CounterName)
-			{
-				this.PropertyName = PropertyName;
-				this.CounterName = CounterName;
-			}
-		}
-
 		/// <summary>
 		/// Document describing a job template. These objects are considered immutable once created and uniquely referenced by hash, in order to de-duplicate across all job runs.
 		/// </summary>
@@ -62,11 +41,9 @@ namespace HordeServer.Collections.Impl
 			[BsonIgnoreIfNull]
 			public string? SubmitNewChange { get; set; }
 
-			public List<TemplateCounterDocument> Counters { get; set; } = new List<TemplateCounterDocument>();
 			public List<string> Arguments { get; private set; } = new List<string>();
 			public List<Parameter> Parameters { get; private set; } = new List<Parameter>();
 
-			IReadOnlyList<ITemplateCounter> ITemplate.Counters => Counters;
 			IReadOnlyList<string> ITemplate.Arguments => Arguments;
 			IReadOnlyList<Parameter> ITemplate.Parameters => Parameters;
 
@@ -76,14 +53,13 @@ namespace HordeServer.Collections.Impl
 				Name = null!;
 			}
 
-			public TemplateDocument(string Name, Priority? Priority, bool bAllowPreflights, string? InitialAgentType, string? SubmitNewChange, List<TemplateCounter>? Counters, List<string>? Arguments, List<Parameter>? Parameters)
+			public TemplateDocument(string Name, Priority? Priority, bool bAllowPreflights, string? InitialAgentType, string? SubmitNewChange, List<string>? Arguments, List<Parameter>? Parameters)
 			{
 				this.Name = Name;
 				this.Priority = Priority;
 				this.AllowPreflights = bAllowPreflights;
 				this.InitialAgentType = InitialAgentType;
 				this.SubmitNewChange = SubmitNewChange;
-				this.Counters = (Counters ?? new List<TemplateCounter>()).ConvertAll(x => new TemplateCounterDocument(x.PropertyName, x.CounterName));
 				this.Arguments = Arguments ?? new List<string>();
 				this.Parameters = Parameters ?? new List<Parameter>();
 
@@ -127,9 +103,9 @@ namespace HordeServer.Collections.Impl
 		}
 
 		/// <inheritdoc/>
-		public async Task<ITemplate> AddAsync(string Name, Priority? Priority, bool bAllowPreflights, string? InitialAgentType, string? SubmitNewChange, List<TemplateCounter>? Counters, List<string>? Arguments, List<Parameter>? Parameters)
+		public async Task<ITemplate> AddAsync(string Name, Priority? Priority, bool bAllowPreflights, string? InitialAgentType, string? SubmitNewChange, List<string>? Arguments, List<Parameter>? Parameters)
 		{
-			TemplateDocument Template = new TemplateDocument(Name, Priority, bAllowPreflights, InitialAgentType, SubmitNewChange, Counters, Arguments, Parameters);
+			TemplateDocument Template = new TemplateDocument(Name, Priority, bAllowPreflights, InitialAgentType, SubmitNewChange, Arguments, Parameters);
 			if (await GetAsync(Template.Id) == null)
 			{
 				await Templates.ReplaceOneAsync(x => x.Id == Template.Id, Template, new ReplaceOptions { IsUpsert = true });
diff --git a/Engine/Source/Programs/Horde/HordeServer/Jobs/TemplateService.cs b/Engine/Source/Programs/Horde/HordeServer/Jobs/TemplateService.cs
--- a/Engine/Source/Programs/Horde/HordeServer/Jobs/TemplateService.cs
+++ b/Engine/Source/Programs/Horde/HordeServer/Jobs/TemplateService.cs
@@ -41,13 +41,12 @@ namespace HordeServer.Services
 		/// <param name="AllowPreflights">Whether to allow preflights of this template</param>
 		/// <param name="InitialAgentType">The initial agent type to parse the BuildGraph script</param>
 		/// <param name="SubmitNewChange">Path to a file within the stream to submit to generate a new changelist for jobs</param>
-		/// <param name="Counters">List of counters for the template</param>
 		/// <param name="Arguments">Common arguments for jobs started from this template</param>
 		/// <param name="Parameters">List of parameters for this template</param>
 		/// <returns>The new template document</returns>
-		public Task<ITemplate> CreateTemplateAsync(string Name, Priority? Priority = null, bool AllowPreflights = true, string? InitialAgentType = null, string? SubmitNewChange = null, List<TemplateCounter>? Counters = null, List<string>? Arguments = null, List<Parameter>? Parameters = null)
+		public Task<ITemplate> CreateTemplateAsync(string Name, Priority? Priority = null, bool AllowPreflights = true, string? InitialAgentType = null, string? SubmitNewChange = null, List<string>? Arguments = null, List<Parameter>? Parameters = null)
 		{
-			return Templates.AddAsync(Name, Priority, AllowPreflights, InitialAgentType, SubmitNewChange, Counters, Arguments, Parameters);
+			return Templates.AddAsync(Name, Priority, AllowPreflights, InitialAgentType, SubmitNewChange, Arguments, Parameters);
 		}
 
 		/// <summary>
diff --git a/Engine/Source/Programs/Horde/HordeServer/Projects/Impl/StreamCollection.cs b/Engine/Source/Programs/Horde/HordeServer/Projects/Impl/StreamCollection.cs
--- a/Engine/Source/Programs/Horde/HordeServer/Projects/Impl/StreamCollection.cs
+++ b/Engine/Source/Programs/Horde/HordeServer/Projects/Impl/StreamCollection.cs
@@ -171,7 +171,7 @@ namespace HordeServer.Collections.Impl
 			foreach (CreateTemplateRefRequest Request in Requests)
 			{
 				// Create the template
-				ITemplate NewTemplate = await TemplateCollection.AddAsync(Request.Name, Request.Priority, Request.AllowPreflights, Request.InitialAgentType, Request.SubmitNewChange, Request.Counters, Request.Arguments, Request.Parameters.ConvertAll(x => x.ToModel()));
+				ITemplate NewTemplate = await TemplateCollection.AddAsync(Request.Name, Request.Priority, Request.AllowPreflights, Request.InitialAgentType, Request.SubmitNewChange, Request.Arguments, Request.Parameters.ConvertAll(x => x.ToModel()));
 
 				// Get an identifier for the new template ref
 				TemplateRefId NewTemplateRefId;
diff --git a/Engine/Source/Programs/Horde/HordeServer/Projects/ScheduleService.cs b/Engine/Source/Programs/Horde/HordeServer/Projects/ScheduleService.cs
--- a/Engine/Source/Programs/Horde/HordeServer/Projects/ScheduleService.cs
+++ b/Engine/Source/Programs/Horde/HordeServer/Projects/ScheduleService.cs
@@ -361,7 +361,7 @@ namespace HordeServer.Services
 			foreach ((int Change, int CodeChange) in TriggerChanges.OrderBy(x => x.Change))
 			{
 				List<string> DefaultArguments = Template.GetDefaultArguments();
-				IJob NewJob = await JobService.CreateJobAsync(null, Stream, TemplateRefId, Template.Id, Graph, Template.Name, Change, CodeChange, null, null, null, null, Template.Priority, null, null, TemplateRef.ChainedJobs, TemplateRef.ShowUgsBadges, TemplateRef.ShowUgsAlerts, TemplateRef.NotificationChannel, TemplateRef.NotificationChannelFilter, null, Template.Counters, DefaultArguments);
+				IJob NewJob = await JobService.CreateJobAsync(null, Stream, TemplateRefId, Template.Id, Graph, Template.Name, Change, CodeChange, null, null, null, null, Template.Priority, null, null, TemplateRef.ChainedJobs, TemplateRef.ShowUgsBadges, TemplateRef.ShowUgsAlerts, TemplateRef.NotificationChannel, TemplateRef.NotificationChannelFilter, null, DefaultArguments);
 				Logger.LogInformation("Started new job for {StreamName} template {TemplateName} at CL {Change} (Code CL {CodeChange}): {JobId}", Stream.Id, TemplateRef.Name, Change, CodeChange, NewJob.Id);
 				await StreamService.UpdateScheduleTriggerAsync(Stream, TemplateRefId, Now, Change, new List<ObjectId> { NewJob.Id }, new List<ObjectId>());
 			}
diff --git a/Engine/Source/Programs/Horde/HordeServer/Startup.cs b/Engine/Source/Programs/Horde/HordeServer/Startup.cs
--- a/Engine/Source/Programs/Horde/HordeServer/Startup.cs
+++ b/Engine/Source/Programs/Horde/HordeServer/Startup.cs
@@ -318,7 +318,6 @@ namespace HordeServer
 			Services.AddSingleton<IAgentSoftwareCollection, AgentSoftwareCollection>();
 			Services.AddSingleton<IArtifactCollection, ArtifactCollection>();
 			Services.AddSingleton<ICommitCollection, CommitCollection>();
-			Services.AddSingleton<ICounterCollection, CounterCollection>();
 			Services.AddSingleton<IGraphCollection, GraphCollection>();
 			Services.AddSingleton<IIssueCollection, IssueCollection>();
 			Services.AddSingleton<IJobCollection, JobCollection>();
