Horde: Rename ArtifactService to ArtifactCollection.

diff --git a/Engine/Source/Programs/Horde/HordeServer/Jobs/Impl/ArtifactCollection.cs b/Engine/Source/Programs/Horde/HordeServer/Jobs/Impl/ArtifactCollection.cs
new file mode 100644
--- /dev/null
+++ b/Engine/Source/Programs/Horde/HordeServer/Jobs/Impl/ArtifactCollection.cs
@@ -0,0 +1,291 @@
+// Copyright Epic Games, Inc. All Rights Reserved.
+
+using Amazon.Extensions.NETCore.Setup;
+using Amazon.Runtime;
+using Amazon.S3;
+using Amazon.S3.Model;
+using HordeServer.Api;
+using HordeServer.Jobs;
+using HordeServer.Models;
+using HordeServer.Storage;
+using HordeServer.Utilities;
+using Microsoft.Extensions.Configuration;
+using Microsoft.Extensions.Hosting;
+using Microsoft.Extensions.Logging;
+using Microsoft.Extensions.Options;
+using MongoDB.Bson;
+using MongoDB.Bson.Serialization.Attributes;
+using MongoDB.Driver;
+using System;
+using System.Collections.Generic;
+using System.IO;
+using System.Linq;
+using System.Security.Claims;
+using System.Text;
+using System.Threading;
+using System.Threading.Tasks;
+
+namespace HordeServer.Services
+{
+	/// <summary>
+	/// Wraps functionality for manipulating artifacts
+	/// </summary>
+	public class ArtifactCollection : IArtifactCollection
+	{
+		class Artifact : IArtifact
+		{
+			[BsonRequired, BsonId]
+			public ObjectId Id { get; set; }
+
+			[BsonRequired]
+			public ObjectId JobId { get; set; }
+
+			[BsonRequired]
+			public string Name { get; set; }
+
+			public SubResourceId? StepId { get; set; }
+
+			[BsonRequired]
+			public long Length { get; set; }
+
+			public string MimeType { get; set; }
+
+			[BsonRequired]
+			public int UpdateIndex { get; set; }
+
+			[BsonConstructor]
+			private Artifact()
+			{
+				this.Name = null!;
+				this.MimeType = null!;
+			}
+
+			public Artifact(ObjectId JobId, SubResourceId? StepId, string Name, long Length, string MimeType)
+			{
+				this.Id = ObjectId.GenerateNewId();
+				this.JobId = JobId;
+				this.StepId = StepId;
+				this.Name = Name;
+				this.Length = Length;
+				this.MimeType = MimeType;
+			}
+		}
+
+		private readonly IStorageBackend StorageBackend;
+		private readonly ILogger<ArtifactCollection> Logger;
+		private readonly IMongoCollection<Artifact> Artifacts;
+
+		/// <summary>
+		/// Constructor
+		/// </summary>
+		/// <param name="DatabaseService">The database service</param>
+		/// <param name="StorageBackend">The storage backend</param>
+		/// <param name="Logger">Log interface</param>
+		public ArtifactCollection(DatabaseService DatabaseService, IStorageBackend<ArtifactCollection> StorageBackend, ILogger<ArtifactCollection> Logger)
+		{
+			this.StorageBackend = StorageBackend;
+			this.Logger = Logger;
+
+			// Initialize Artifacts table
+			Artifacts = DatabaseService.GetCollection<Artifact>("Artifacts");
+			if (!DatabaseService.ReadOnlyMode)
+			{
+				Artifacts.Indexes.CreateOne(new CreateIndexModel<Artifact>(Builders<Artifact>.IndexKeys.Ascending(x => x.JobId)));
+			}
+		}
+
+		/// <summary>
+		/// Creates a new artifact
+		/// </summary>
+		/// <param name="JobId">Unique id of the job that owns this artifact</param>
+		/// <param name="StepId">Optional Step id</param>
+		/// <param name="Name">Name of artifact</param>
+		/// <param name="MimeType">Type of artifact</param>
+		/// <param name="Data">The data to write</param>
+		/// <returns>The new log file document</returns>
+		public async Task<IArtifact> CreateArtifactAsync(ObjectId JobId, SubResourceId? StepId, string Name, string MimeType, System.IO.Stream Data)
+		{
+			// upload first
+			string ArtifactName = ValidateName(Name);
+			await StorageBackend.WriteAsync(GetPath(JobId, StepId, ArtifactName), Data);
+
+			// then create entry
+			Artifact NewArtifact = new Artifact(JobId, StepId, ArtifactName, Data.Length, MimeType);
+			await Artifacts.InsertOneAsync(NewArtifact);
+			return NewArtifact;
+		}
+
+		/// <summary>
+		/// Gets all the available artifacts for a job
+		/// </summary>
+		/// <param name="JobId">Unique id of the job to query</param>
+		/// <param name="StepId">Unique id of the Step to query</param>
+		/// <param name="Name">Name of the artifact</param>
+		/// <returns>List of artifact documents</returns>
+		public async Task<List<IArtifact>> GetArtifactsAsync(ObjectId? JobId, SubResourceId? StepId, string? Name)
+		{
+			FilterDefinitionBuilder<Artifact> Builder = Builders<Artifact>.Filter;
+
+			FilterDefinition<Artifact> Filter = FilterDefinition<Artifact>.Empty;
+			if (JobId != null)
+			{
+				Filter &= Builder.Eq(x => x.JobId, JobId.Value);
+			}
+			if (StepId != null)
+			{
+				Filter &= Builder.Eq(x => x.StepId, StepId.Value);
+			}
+			if (Name != null)
+			{
+				Filter &= Builder.Eq(x => x.Name, Name);
+			}
+
+			return await Artifacts.Find(Filter).ToListAsync<Artifact, IArtifact>();
+		}
+
+		/// <summary>
+		/// Gets a specific list of artifacts based on id
+		/// </summary>
+		/// <param name="ArtifactIds">The list of artifact Ids</param>
+		/// <returns>List of artifact documents</returns>
+		public async Task<List<IArtifact>> GetArtifactsAsync(IEnumerable<ObjectId> ArtifactIds)
+		{
+			FilterDefinitionBuilder<Artifact> Builder = Builders<Artifact>.Filter;
+
+			FilterDefinition<Artifact> Filter = FilterDefinition<Artifact>.Empty;
+			Filter &= Builder.In(x => x.Id, ArtifactIds);
+			
+			return await Artifacts.Find(Filter).ToListAsync<Artifact, IArtifact>();
+		}
+
+		/// <summary>
+		/// Gets an artifact by ID
+		/// </summary>
+		/// <param name="ArtifactId">Unique id of the artifact</param>
+		/// <returns>The artifact document</returns>
+		public async Task<IArtifact?> GetArtifactAsync(ObjectId ArtifactId)
+		{
+			return await Artifacts.Find<Artifact>(x => x.Id == ArtifactId).FirstOrDefaultAsync();
+		}
+
+		/// <summary>
+		/// Attempts to update a single artifact document, if the update counter is valid
+		/// </summary>
+		/// <param name="Current">The current artifact document</param>
+		/// <param name="Update">The update definition</param>
+		/// <returns>True if the document was updated, false if another writer updated the document first</returns>
+		private async Task<bool> TryUpdateArtifactAsync(Artifact Current, UpdateDefinition<Artifact> Update)
+		{
+			UpdateResult Result = await Artifacts.UpdateOneAsync<Artifact>(x => x.Id == Current.Id && x.UpdateIndex == Current.UpdateIndex, Update.Set(x => x.UpdateIndex, Current.UpdateIndex + 1));
+			return Result.ModifiedCount == 1;
+		}
+
+		/// <summary>
+		/// Updates an artifact
+		/// </summary>
+		/// <param name="Artifact">The artifact</param>
+		/// <param name="NewMimeType">New mime type</param>
+		/// <param name="NewData">New data</param>
+		/// <returns>Async task</returns>
+		public async Task<bool> UpdateArtifactAsync(IArtifact? Artifact, string NewMimeType, System.IO.Stream NewData)
+		{
+			while (Artifact != null)
+			{
+				UpdateDefinitionBuilder<Artifact> UpdateBuilder = Builders<Artifact>.Update;
+
+				List<UpdateDefinition<Artifact>> Updates = new List<UpdateDefinition<Artifact>>();
+				Updates.Add(UpdateBuilder.Set(x => x.MimeType, NewMimeType));
+				Updates.Add(UpdateBuilder.Set(x => x.Length, NewData.Length));
+
+				// re-upload the data to external
+				string ArtifactName = ValidateName(Artifact.Name);
+				await StorageBackend.WriteAsync(GetPath(Artifact.JobId, Artifact.StepId, ArtifactName), NewData);
+
+				if (await TryUpdateArtifactAsync((Artifact)Artifact, UpdateBuilder.Combine(Updates)))
+				{
+					return true;
+				}
+
+				Artifact = await GetArtifactAsync(Artifact.Id);
+			}
+			return false;
+		}
+
+		/// <summary>
+		/// gets artifact data
+		/// </summary>
+		/// <param name="Artifact">The artifact</param>
+		/// <returns>The chunk data</returns>
+		public async Task<System.IO.Stream> OpenArtifactReadStreamAsync(IArtifact Artifact)
+		{
+			System.IO.Stream? Stream = await StorageBackend.ReadAsync(GetPath(Artifact.JobId, Artifact.StepId, Artifact.Name));
+			if (Stream == null)
+			{
+				throw new Exception($"Unable to get artifact {Artifact.Id}");
+			}
+			return Stream;
+		}
+
+		/// <summary>
+		/// Get the path for an artifact
+		/// </summary>
+		/// <param name="JobId"></param>
+		/// <param name="StepId"></param>
+		/// <param name="Name"></param>
+		/// <returns></returns>
+		private static string GetPath(ObjectId JobId, SubResourceId? StepId, string Name)
+		{
+			if (StepId == null)
+			{
+				return $"{JobId}/{Name}";
+			}
+			else
+			{
+				return $"{JobId}/{StepId.Value}/{Name}";
+			}
+		}
+
+		/// <summary>
+		/// Checks that the name given is valid for an artifact
+		/// </summary>
+		/// <param name="Name"></param>
+		private static string ValidateName(string Name)
+		{
+			string NewName = Name.Replace('\\', '/');
+			if (NewName.Length == 0)
+			{
+				throw new ArgumentException("Artifact name is empty");
+			}
+			else if (NewName.StartsWith('/'))
+			{
+				throw new ArgumentException($"Artifact has an absolute path ({NewName})");
+			}
+			else if (NewName.EndsWith('/'))
+			{
+				throw new ArgumentException($"Artifact does not have a file name ({NewName})");
+			}
+			else if (NewName.Contains("//", StringComparison.Ordinal))
+			{
+				throw new ArgumentException($"Artifact name contains an invalid directory ({NewName})");
+			}
+
+			string InvalidChars = ":<>|\"";
+			for (int Idx = 0; Idx < NewName.Length; Idx++)
+			{
+				int MinIdx = Idx;
+				for (; Idx < NewName.Length && NewName[Idx] != '/'; Idx++)
+				{
+					if (InvalidChars.Contains(NewName[Idx], StringComparison.Ordinal))
+					{
+						throw new ArgumentException($"Invalid character in artifact name ({NewName})");
+					}
+				}
+				if ((Idx == MinIdx + 1 && NewName[MinIdx] == '.') || (Idx == MinIdx + 2 && NewName[MinIdx] == '.' && NewName[MinIdx + 1] == '.'))
+				{
+					throw new ArgumentException($"Artifact may not contain symbolic directory names ({NewName})");
+				}
+			}
+			return NewName;
+		}
+	}
+}
\ No newline at end of file
diff --git a/Engine/Source/Programs/Horde/HordeServer/Jobs/Impl/ArtifactService.cs b/Engine/Source/Programs/Horde/HordeServer/Jobs/Impl/ArtifactService.cs
deleted file mode 100644
