Horde: Fix step history being supressed by retried steps
#rnx

diff --git a/Engine/Source/Programs/Horde/HordeDashboard/src/backend/JobDetails.ts b/Engine/Source/Programs/Horde/HordeDashboard/src/backend/JobDetails.ts
--- a/Engine/Source/Programs/Horde/HordeDashboard/src/backend/JobDetails.ts
+++ b/Engine/Source/Programs/Horde/HordeDashboard/src/backend/JobDetails.ts
@@ -69,7 +69,7 @@ export class JobDetails {
 
                 if (stepId) {
                     historyIdx = requests.length;
-                    requests.push(backend.getJobStepHistory(this.jobdata!.streamId, this.getStepName(stepId), 1024, this.jobdata!.templateId!));
+                    requests.push(backend.getJobStepHistory(this.jobdata!.streamId, this.getStepName(stepId, false), 1024, this.jobdata!.templateId!));
                 }
             }
 
@@ -330,7 +330,7 @@ export class JobDetails {
 
     }
 
-    getStepName(stepId: string | undefined): string {
+    getStepName(stepId: string | undefined, includeRetry:boolean = true): string {
 
         if (!stepId) {
             return "";
@@ -342,7 +342,7 @@ export class JobDetails {
         }
 
         const idx = this.getStepRetryNumber(stepId);
-        if (!idx) {
+        if (!idx || !includeRetry) {
             return node.name;
         }
 
