Horde: Replace TemplateService with ITemplateCollection.

diff --git a/Engine/Source/Programs/Horde/HordeServer.Compute.Tests/ControllerIntegrationTest.cs b/Engine/Source/Programs/Horde/HordeServer.Compute.Tests/ControllerIntegrationTest.cs
--- a/Engine/Source/Programs/Horde/HordeServer.Compute.Tests/ControllerIntegrationTest.cs
+++ b/Engine/Source/Programs/Horde/HordeServer.Compute.Tests/ControllerIntegrationTest.cs
@@ -7,6 +7,7 @@ using System.Net.Http.Headers;
 using System.Runtime.CompilerServices;
 using System.Threading.Tasks;
 using HordeServer;
+using HordeServer.Collections;
 using HordeServer.Collections.Impl;
 using HordeServer.Jobs;
 using HordeServer.Services;
@@ -56,7 +57,7 @@ namespace HordeServerTests
 
             IServiceProvider Services = GetFactory().Services;
             DatabaseService DatabaseService = Services.GetRequiredService<DatabaseService>();
-            TemplateService TemplateService = Services.GetRequiredService<TemplateService>();
+            ITemplateCollection TemplateService = Services.GetRequiredService<ITemplateCollection>();
             JobService JobService = Services.GetRequiredService<JobService>();
             IArtifactCollection ArtifactCollection = Services.GetRequiredService<IArtifactCollection>();
             StreamService StreamService = Services.GetRequiredService<StreamService>();
diff --git a/Engine/Source/Programs/Horde/HordeServer.Compute.Tests/Fixture.cs b/Engine/Source/Programs/Horde/HordeServer.Compute.Tests/Fixture.cs
--- a/Engine/Source/Programs/Horde/HordeServer.Compute.Tests/Fixture.cs
+++ b/Engine/Source/Programs/Horde/HordeServer.Compute.Tests/Fixture.cs
@@ -33,10 +33,10 @@ namespace HordeServerTests
 		public IAgent Agent1 { get; private set; } = null!;
 		public string Agent1Name { get; private set; } = null!;
 
-		public static async Task<Fixture> Create(bool ForceNewFixture, IGraphCollection GraphCollection, TemplateService TemplateService, JobService JobService, IArtifactCollection ArtifactCollection, StreamService StreamService, AgentService AgentService, IPerforceService PerforceService)
+		public static async Task<Fixture> Create(bool ForceNewFixture, IGraphCollection GraphCollection, ITemplateCollection TemplateCollection, JobService JobService, IArtifactCollection ArtifactCollection, StreamService StreamService, AgentService AgentService, IPerforceService PerforceService)
 		{
 			Fixture _fixture = new Fixture();
-			await _fixture.Populate(GraphCollection, TemplateService, JobService, ArtifactCollection, StreamService, AgentService, PerforceService);
+			await _fixture.Populate(GraphCollection, TemplateCollection, JobService, ArtifactCollection, StreamService, AgentService, PerforceService);
 			
 			(PerforceService as PerforceServiceStub)?.AddChange("//UE5/Main", 112233, "leet.coder", "Did stuff", new []{"file.cpp"});
 			(PerforceService as PerforceServiceStub)?.AddChange("//UE5/Main", 1111, "swarm", "A shelved CL here", new []{"renderer.cpp"});
@@ -44,7 +44,7 @@ namespace HordeServerTests
 			return _fixture;
 		}
 
-		private async Task Populate(IGraphCollection GraphCollection, TemplateService TemplateService, JobService JobService, IArtifactCollection ArtifactCollection, StreamService StreamService, AgentService AgentService, IPerforceService PerforceService)
+		private async Task Populate(IGraphCollection GraphCollection, ITemplateCollection TemplateCollection, JobService JobService, IArtifactCollection ArtifactCollection, StreamService StreamService, AgentService AgentService, IPerforceService PerforceService)
 		{
 			var Fg = new FixtureGraph();
 			Fg.Id = ContentHash.Empty;
@@ -53,7 +53,7 @@ namespace HordeServerTests
 			Fg.Aggregates = new List<IAggregate>();
 			Fg.Labels = new List<ILabel>();
 
-			Template = await TemplateService.CreateTemplateAsync("Test template", null, false, null, null,
+			Template = await TemplateCollection.AddAsync("Test template", null, false, null, null,
 				new List<string>(), new List<Parameter>());
 			Graph = await GraphCollection.AddAsync(Template);
 
diff --git a/Engine/Source/Programs/Horde/HordeServer.Compute.Tests/JobServiceTests.cs b/Engine/Source/Programs/Horde/HordeServer.Compute.Tests/JobServiceTests.cs
--- a/Engine/Source/Programs/Horde/HordeServer.Compute.Tests/JobServiceTests.cs
+++ b/Engine/Source/Programs/Horde/HordeServer.Compute.Tests/JobServiceTests.cs
@@ -28,7 +28,7 @@ namespace HordeServerTests
 			IProject? Project = await TestSetup.ProjectService.Collection.AddOrUpdateAsync(ProjectId, "", "", 0, new ProjectConfig { Name = "UE5" });
 			Assert.IsNotNull(Project);
 
-			ITemplate Template = await TestSetup.TemplateService.CreateTemplateAsync("Test template", null, false, null, null, new List<string>(), new List<Parameter>());
+			ITemplate Template = await TestSetup.TemplateCollection.AddAsync("Test template", null, false, null, null, new List<string>(), new List<Parameter>());
 			IGraph Graph = await TestSetup.GraphCollection.AddAsync(Template);
 
 			TemplateRefId TemplateRefId1 = new TemplateRefId("template1");
@@ -154,7 +154,7 @@ namespace HordeServerTests
 			IStream? Stream = await TestSetup.StreamCollection.GetAsync(StreamId);
 			Stream = await TestSetup.StreamCollection.TryCreateOrReplaceAsync(StreamId, Stream, "", "", Project!.Id, new StreamConfig { Name = "//UE5/Main" });
 
-			ITemplate Template = await TestSetup.TemplateService.CreateTemplateAsync("Test template", null, false, null, null, new List<string>(), new List<Parameter>());
+			ITemplate Template = await TestSetup.TemplateCollection.AddAsync("Test template", null, false, null, null, new List<string>(), new List<Parameter>());
 			IGraph Graph = await TestSetup.GraphCollection.AddAsync(Template);
 
 			NewGroup GroupA = new NewGroup("win", new List<NewNode>());
diff --git a/Engine/Source/Programs/Horde/HordeServer.Compute.Tests/SchedulerTests.cs b/Engine/Source/Programs/Horde/HordeServer.Compute.Tests/SchedulerTests.cs
--- a/Engine/Source/Programs/Horde/HordeServer.Compute.Tests/SchedulerTests.cs
+++ b/Engine/Source/Programs/Horde/HordeServer.Compute.Tests/SchedulerTests.cs
@@ -42,7 +42,7 @@ namespace HordeServerTests
 			IProject ? Project = TestSetup.ProjectService.Collection.AddOrUpdateAsync(ProjectId, "", "", 0, new ProjectConfig { Name = "UE5" }).Result;
 			Assert.IsNotNull(Project);
 
-			Template = TestSetup.TemplateService.CreateTemplateAsync("Test template", null, false, null, null, new List<string>(), new List<Parameter>()).Result;
+			Template = TestSetup.TemplateCollection.AddAsync("Test template", null, false, null, null, new List<string>(), new List<Parameter>()).Result;
 
 			InitialJobIds = new HashSet<ObjectId>(TestSetup.JobCollection.FindAsync().Result.Select(x => x.Id));
 
@@ -355,11 +355,11 @@ namespace HordeServerTests
 			TestSetup.Clock.UtcNow = StartTime;
 
 			// Create two templates, the second dependent on the first
-			ITemplate? NewTemplate1 = await TestSetup.TemplateService.CreateTemplateAsync("Test template 1");
+			ITemplate? NewTemplate1 = await TestSetup.TemplateCollection.AddAsync("Test template 1");
 			TemplateRef NewTemplateRef1 = new TemplateRef(NewTemplate1);
 			TemplateRefId NewTemplateRefId1 = new TemplateRefId("new-template-1");
 
-			ITemplate? NewTemplate2 = await TestSetup.TemplateService.CreateTemplateAsync("Test template 2");
+			ITemplate? NewTemplate2 = await TestSetup.TemplateCollection.AddAsync("Test template 2");
 			TemplateRef NewTemplateRef2 = new TemplateRef(NewTemplate2);
 			NewTemplateRef2.Schedule = new Schedule();
 			NewTemplateRef2.Schedule.Gate = new ScheduleGate(NewTemplateRefId1, "TriggerNext");
diff --git a/Engine/Source/Programs/Horde/HordeServer.Compute.Tests/TestSetup.cs b/Engine/Source/Programs/Horde/HordeServer.Compute.Tests/TestSetup.cs
--- a/Engine/Source/Programs/Horde/HordeServer.Compute.Tests/TestSetup.cs
+++ b/Engine/Source/Programs/Horde/HordeServer.Compute.Tests/TestSetup.cs
@@ -75,7 +75,7 @@ namespace HordeServerTests
 		public AgentSoftwareService AgentSoftwareService => ServiceProvider.GetRequiredService<AgentSoftwareService>();
 		public AgentService AgentService => ServiceProvider.GetRequiredService<AgentService>();
 		public DatabaseService DatabaseService => ServiceProvider.GetRequiredService<DatabaseService>();
-		public TemplateService TemplateService => ServiceProvider.GetRequiredService<TemplateService>();
+		public ITemplateCollection TemplateCollection => ServiceProvider.GetRequiredService<ITemplateCollection>();
 		internal PerforceServiceStub PerforceService => (PerforceServiceStub)ServiceProvider.GetRequiredService<IPerforceService>();
 		public ILogFileService LogFileService => ServiceProvider.GetRequiredService<ILogFileService>();
 		public ProjectService ProjectService => ServiceProvider.GetRequiredService<ProjectService>();
@@ -162,6 +162,7 @@ namespace HordeServerTests
 			Services.AddSingleton<ITemplateCollection, TemplateCollection>();
 			Services.AddSingleton<ITestDataCollection, TestDataCollection>();
 			Services.AddSingleton<ITelemetryCollection, TelemetryCollection>();
+			Services.AddSingleton<ITemplateCollection, TemplateCollection>();
 			Services.AddSingleton<IUgsMetadataCollection, UgsMetadataCollection>();
 			Services.AddSingleton<IUserCollection, UserCollectionV1>();
 
@@ -195,7 +196,6 @@ namespace HordeServerTests
 			Services.AddSingleton<RpcService>();
 			Services.AddSingleton<ScheduleService>();
 			Services.AddSingleton<StreamService>();
-			Services.AddSingleton<TemplateService>();
 			Services.AddSingleton<UpgradeService>();
 
 			Services.AddSingleton<ConformTaskSource>();
@@ -221,7 +221,7 @@ namespace HordeServerTests
 		{
 			if (!DatabaseService.ReadOnlyMode)
 			{
-				Fixture = await Fixture.Create(ForceNewFixture, GraphCollection, TemplateService, JobService, ArtifactCollection, StreamService, AgentService, PerforceService);
+				Fixture = await Fixture.Create(ForceNewFixture, GraphCollection, TemplateCollection, JobService, ArtifactCollection, StreamService, AgentService, PerforceService);
 			}
 		}
 
@@ -229,7 +229,7 @@ namespace HordeServerTests
         {
 			ILogger<JobsController> Logger = ServiceProvider.GetRequiredService<ILogger<JobsController>>();
 			JobsController JobsCtrl = new JobsController(AclService, GraphCollection, PerforceService, StreamService, JobService,
-		        TemplateService, ArtifactCollection, NotificationService, Logger);
+		        TemplateCollection, ArtifactCollection, NotificationService, Logger);
 	        JobsCtrl.ControllerContext = GetControllerContext();
 	        return JobsCtrl;
         }
diff --git a/Engine/Source/Programs/Horde/HordeServer/Controllers/JobsController.cs b/Engine/Source/Programs/Horde/HordeServer/Controllers/JobsController.cs
--- a/Engine/Source/Programs/Horde/HordeServer/Controllers/JobsController.cs
+++ b/Engine/Source/Programs/Horde/HordeServer/Controllers/JobsController.cs
@@ -66,7 +66,7 @@ namespace HordeServer.Controllers
 		/// <summary>
 		/// Instance of the TemplateService singleton
 		/// </summary>
-		private readonly TemplateService TemplateService;
+		private readonly ITemplateCollection TemplateCollection;
 
 		/// <summary>
 		/// 
@@ -86,14 +86,14 @@ namespace HordeServer.Controllers
 		/// <summary>
 		/// Constructor
 		/// </summary>
-		public JobsController(AclService AclService, IGraphCollection Graphs, IPerforceService Perforce, StreamService StreamService, JobService JobService, TemplateService TemplateService, IArtifactCollection ArtifactCollection, INotificationService NotificationService, ILogger<JobsController> Logger)
+		public JobsController(AclService AclService, IGraphCollection Graphs, IPerforceService Perforce, StreamService StreamService, JobService JobService, ITemplateCollection TemplateCollection, IArtifactCollection ArtifactCollection, INotificationService NotificationService, ILogger<JobsController> Logger)
 		{
 			this.AclService = AclService;
 			this.Graphs = Graphs;
 			this.Perforce = Perforce;
 			this.StreamService = StreamService;
 			this.JobService = JobService;
-			this.TemplateService = TemplateService;
+			this.TemplateCollection = TemplateCollection;
 			this.ArtifactCollection = ArtifactCollection;
 			this.NotificationService = NotificationService;
 			this.Logger = Logger;
@@ -132,7 +132,7 @@ namespace HordeServer.Controllers
 				return Forbid();
 			}
 
-			ITemplate? Template = await TemplateService.GetTemplateAsync(TemplateRef.Hash);
+			ITemplate? Template = await TemplateCollection.GetAsync(TemplateRef.Hash);
 			if (Template == null)
 			{
 				return BadRequest($"Missing template referenced by {Create.TemplateId}");
@@ -474,7 +474,7 @@ namespace HordeServer.Controllers
 				return Forbid();
 			}
 
-			ITemplate? Template = await TemplateService.GetTemplateAsync(Job.TemplateHash);
+			ITemplate? Template = await TemplateCollection.GetAsync(Job.TemplateHash);
 			if(Template == null)
 			{
 				return NotFound();
diff --git a/Engine/Source/Programs/Horde/HordeServer/Controllers/StreamsController.cs b/Engine/Source/Programs/Horde/HordeServer/Controllers/StreamsController.cs
--- a/Engine/Source/Programs/Horde/HordeServer/Controllers/StreamsController.cs
+++ b/Engine/Source/Programs/Horde/HordeServer/Controllers/StreamsController.cs
@@ -46,7 +46,7 @@ namespace HordeServer.Controllers
 		/// <summary>
 		/// Singleton instance of the template service
 		/// </summary>
-		private readonly TemplateService TemplateService;
+		private readonly ITemplateCollection TemplateCollection;
 
 		/// <summary>
 		/// Singleton instance of the job service
@@ -69,16 +69,16 @@ namespace HordeServer.Controllers
 		/// <param name="AclService">The ACL service</param>
 		/// <param name="ProjectService">The project service</param>
 		/// <param name="StreamService">The stream service</param>
-		/// <param name="TemplateService">The template service</param>
+		/// <param name="TemplateCollection">The template service</param>
 		/// <param name="JobService">The job service</param>
 		/// <param name="JobStepRefCollection">The jobstep ref collection</param>
 		/// <param name="PerforceService">The perforce service</param>
-		public StreamsController(AclService AclService, ProjectService ProjectService, StreamService StreamService, TemplateService TemplateService, JobService JobService, IJobStepRefCollection JobStepRefCollection, IPerforceService PerforceService)
+		public StreamsController(AclService AclService, ProjectService ProjectService, StreamService StreamService, ITemplateCollection TemplateCollection, JobService JobService, IJobStepRefCollection JobStepRefCollection, IPerforceService PerforceService)
 		{
 			this.AclService = AclService;
 			this.ProjectService = ProjectService;
 			this.StreamService = StreamService;
-			this.TemplateService = TemplateService;
+			this.TemplateCollection = TemplateCollection;
 			this.JobService = JobService;
 			this.JobStepRefCollection = JobStepRefCollection;
 			this.PerforceService = PerforceService;
@@ -156,7 +156,7 @@ namespace HordeServer.Controllers
 			{
 				if (await StreamService.AuthorizeAsync(Stream, Pair.Value, AclAction.ViewTemplate, User, Cache))
 				{
-					ITemplate? Template = await TemplateService.GetTemplateAsync(Pair.Value.Hash);
+					ITemplate? Template = await TemplateCollection.GetAsync(Pair.Value.Hash);
 					if (Template != null)
 					{
 						bool bIncludeTemplateAcl = Pair.Value.Acl != null && await StreamService.AuthorizeAsync(Stream, Pair.Value, AclAction.ViewPermissions, User, Cache);
@@ -336,15 +336,15 @@ namespace HordeServer.Controllers
 		/// </summary>
 		/// <param name="Requests">Request objects</param>
 		/// <param name="Stream">The current stream state</param>
-		/// <param name="TemplateService">The template service</param>
+		/// <param name="TemplateCollection">The template service</param>
 		/// <returns>List of new template references</returns>
-		static async Task<Dictionary<TemplateRefId, TemplateRef>> CreateTemplateRefs(List<CreateTemplateRefRequest> Requests, IStream? Stream, TemplateService TemplateService)
+		static async Task<Dictionary<TemplateRefId, TemplateRef>> CreateTemplateRefs(List<CreateTemplateRefRequest> Requests, IStream? Stream, ITemplateCollection TemplateCollection)
 		{
 			Dictionary<TemplateRefId, TemplateRef> NewTemplateRefs = new Dictionary<TemplateRefId, TemplateRef>();
 			foreach (CreateTemplateRefRequest Request in Requests)
 			{
 				// Create the template
-				ITemplate NewTemplate = await TemplateService.CreateTemplateAsync(Request.Name, Request.Priority, Request.AllowPreflights, Request.InitialAgentType, Request.SubmitNewChange, Request.Arguments, Request.Parameters.ConvertAll(x => x.ToModel()));
+				ITemplate NewTemplate = await TemplateCollection.AddAsync(Request.Name, Request.Priority, Request.AllowPreflights, Request.InitialAgentType, Request.SubmitNewChange, Request.Arguments, Request.Parameters.ConvertAll(x => x.ToModel()));
 
 				// Get an identifier for the new template ref
 				TemplateRefId NewTemplateRefId;
@@ -415,7 +415,7 @@ namespace HordeServer.Controllers
 				return Forbid();
 			}
 
-			ITemplate? Template = await TemplateService.GetTemplateAsync(TemplateRef.Hash);
+			ITemplate? Template = await TemplateCollection.GetAsync(TemplateRef.Hash);
 			if(Template == null)
 			{
 				return NotFound();
diff --git a/Engine/Source/Programs/Horde/HordeServer/Jobs/Impl/JobService.cs b/Engine/Source/Programs/Horde/HordeServer/Jobs/Impl/JobService.cs
--- a/Engine/Source/Programs/Horde/HordeServer/Jobs/Impl/JobService.cs
+++ b/Engine/Source/Programs/Horde/HordeServer/Jobs/Impl/JobService.cs
@@ -119,7 +119,7 @@ namespace HordeServer.Services
 		/// <summary>
 		/// Singleton instance of the template service
 		/// </summary>
-		TemplateService TemplateService;
+		ITemplateCollection TemplateCollection;
 
 		/// <summary>
 		/// Singleton instance of the issue service
@@ -173,12 +173,12 @@ namespace HordeServer.Services
 		/// <param name="TriggerCollection">Trigger collection</param>
 		/// <param name="JobTaskSource">The queue service</param>
 		/// <param name="StreamService">The stream service</param>
-		/// <param name="TemplateService">The template service</param>
+		/// <param name="TemplateCollection">The template service</param>
 		/// <param name="IssueService">The issue service</param>
 		/// <param name="PerforceService">The perforce service</param>
 		/// <param name="Settings">Settings instance</param>
 		/// <param name="Logger">Log output</param>
-		public JobService(IJobCollection Jobs, IGraphCollection Graphs, IAgentCollection Agents, IJobStepRefCollection JobStepRefs, IJobTimingCollection JobTimings, IUserCollection UserCollection, INotificationTriggerCollection TriggerCollection, JobTaskSource JobTaskSource, StreamService StreamService, TemplateService TemplateService, IIssueService IssueService, IPerforceService PerforceService, IOptionsMonitor<ServerSettings> Settings, ILogger<JobService> Logger)
+		public JobService(IJobCollection Jobs, IGraphCollection Graphs, IAgentCollection Agents, IJobStepRefCollection JobStepRefs, IJobTimingCollection JobTimings, IUserCollection UserCollection, INotificationTriggerCollection TriggerCollection, JobTaskSource JobTaskSource, StreamService StreamService, ITemplateCollection TemplateCollection, IIssueService IssueService, IPerforceService PerforceService, IOptionsMonitor<ServerSettings> Settings, ILogger<JobService> Logger)
 		{
 			this.Jobs = Jobs;
 			this.Graphs = Graphs;
@@ -189,7 +189,7 @@ namespace HordeServer.Services
 			this.TriggerCollection = TriggerCollection;
 			this.JobTaskSource = JobTaskSource;
 			this.StreamService = StreamService;
-			this.TemplateService = TemplateService;
+			this.TemplateCollection = TemplateCollection;
 			this.IssueService = IssueService;
 			this.PerforceService = PerforceService;
 			this.Settings = Settings;
@@ -1111,7 +1111,7 @@ namespace HordeServer.Services
 						break;
 					}
 
-					ITemplate? Template = await TemplateService.GetTemplateAsync(TemplateRef.Hash);
+					ITemplate? Template = await TemplateCollection.GetAsync(TemplateRef.Hash);
 					if (Template == null)
 					{
 						Logger.LogWarning("Cannot find template {TemplateHash}", TemplateRef.Hash);
diff --git a/Engine/Source/Programs/Horde/HordeServer/Jobs/TemplateService.cs b/Engine/Source/Programs/Horde/HordeServer/Jobs/TemplateService.cs
deleted file mode 100644
diff --git a/Engine/Source/Programs/Horde/HordeServer/Projects/ScheduleService.cs b/Engine/Source/Programs/Horde/HordeServer/Projects/ScheduleService.cs
--- a/Engine/Source/Programs/Horde/HordeServer/Projects/ScheduleService.cs
+++ b/Engine/Source/Programs/Horde/HordeServer/Projects/ScheduleService.cs
@@ -68,7 +68,7 @@ namespace HordeServer.Services
 		/// <summary>
 		/// The template service
 		/// </summary>
-		TemplateService TemplateService;
+		ITemplateCollection TemplateCollection;
 
 		/// <summary>
 		/// The timezone to use for schedules
@@ -95,11 +95,11 @@ namespace HordeServer.Services
 		/// <param name="JobCollection">Job collection singleton</param>
 		/// <param name="JobService">The job service instance</param>
 		/// <param name="StreamService">The stream service instance</param>
-		/// <param name="TemplateService">The template service instance</param>
+		/// <param name="TemplateCollection">The template service instance</param>
 		/// <param name="Clock"></param>
 		/// <param name="Settings">Settings for the server</param>
 		/// <param name="Logger">Logging instance</param>
-		public ScheduleService(DatabaseService DatabaseService, IGraphCollection Graphs, IDowntimeService DowntimeService, IPerforceService Perforce, IJobCollection JobCollection, JobService JobService, StreamService StreamService, TemplateService TemplateService, IClock Clock, IOptionsMonitor<ServerSettings> Settings, ILogger<ScheduleService> Logger)
+		public ScheduleService(DatabaseService DatabaseService, IGraphCollection Graphs, IDowntimeService DowntimeService, IPerforceService Perforce, IJobCollection JobCollection, JobService JobService, StreamService StreamService, ITemplateCollection TemplateCollection, IClock Clock, IOptionsMonitor<ServerSettings> Settings, ILogger<ScheduleService> Logger)
 			: base(DatabaseService, new ObjectId("6035593e6d721d80fb9efa5c"), Logger)
 		{
 			this.Graphs = Graphs;
@@ -108,7 +108,7 @@ namespace HordeServer.Services
 			this.JobCollection = JobCollection;
 			this.JobService = JobService;
 			this.StreamService = StreamService;
-			this.TemplateService = TemplateService;
+			this.TemplateCollection = TemplateCollection;
 			this.Clock = Clock;
 
 			string? ScheduleTimeZone = Settings.CurrentValue.ScheduleTimeZone;
@@ -339,7 +339,7 @@ namespace HordeServer.Services
 			}
 
 			// Get the matching template
-			ITemplate? Template = await TemplateService.GetTemplateAsync(TemplateRef.Hash);
+			ITemplate? Template = await TemplateCollection.GetAsync(TemplateRef.Hash);
 			if (Template == null)
 			{
 				Logger.LogWarning("Unable to find template '{TemplateHash}' for '{TemplateRefId}'", TemplateRef.Hash, TemplateRefId);
diff --git a/Engine/Source/Programs/Horde/HordeServer/Startup.cs b/Engine/Source/Programs/Horde/HordeServer/Startup.cs
--- a/Engine/Source/Programs/Horde/HordeServer/Startup.cs
+++ b/Engine/Source/Programs/Horde/HordeServer/Startup.cs
@@ -336,6 +336,7 @@ namespace HordeServer
 			Services.AddSingleton<ITemplateCollection, TemplateCollection>();
 			Services.AddSingleton<ITestDataCollection, TestDataCollection>();
 			Services.AddSingleton<ITelemetryCollection, TelemetryCollection>();
+			Services.AddSingleton<ITemplateCollection, TemplateCollection>();
 			Services.AddSingleton<IUgsMetadataCollection, UgsMetadataCollection>();
 			Services.AddSingleton<IUserCollection, UserCollectionV2>();
 			Services.AddSingleton<IDeviceCollection, DeviceCollection>();
@@ -408,7 +409,6 @@ namespace HordeServer
 			Services.AddSingleton<IAvatarService, SlackNotificationSink>(SP => SP.GetRequiredService<SlackNotificationSink>());
 			Services.AddSingleton<INotificationSink, SlackNotificationSink>(SP => SP.GetRequiredService<SlackNotificationSink>());
 			Services.AddSingleton<StreamService>();
-			Services.AddSingleton<TemplateService>();
 			Services.AddSingleton<UpgradeService>();
 
 			Services.AddSingleton<DeviceService>();
