Horde: Add strongly typed JobId, LeaseId, and LogId structs (via ObjectId<>, which operates similarly to StringId<>), and register formatters to add a type attribute to any structured log data.

Messages in audit logs will now include such properties in the format {"$type":"JobId", "$text":"123123123123123123"}.

diff --git a/Engine/Source/Programs/Horde/HordeServer.Compute.Tests/DebugTest.cs b/Engine/Source/Programs/Horde/HordeServer.Compute.Tests/DebugTest.cs
--- a/Engine/Source/Programs/Horde/HordeServer.Compute.Tests/DebugTest.cs
+++ b/Engine/Source/Programs/Horde/HordeServer.Compute.Tests/DebugTest.cs
@@ -20,6 +20,8 @@ using Microsoft.Extensions.Options;
 
 namespace HordeServerTests
 {
+	using JobId = ObjectId<IJob>;
+
     /// <summary>
     /// Tests aimed to be run from an IDE while debugging against real databases
     /// Uncomment the [Ignore] tag to run
@@ -63,7 +65,7 @@ namespace HordeServerTests
         {
 	        TestSetup TestSetup = await GetTestSetupExternal();
 	        
-	        HordeServer.Models.IJob? Job = await TestSetup.JobService.GetJobAsync(new ObjectId("5f98da1d98cf40000141344c"));
+	        HordeServer.Models.IJob? Job = await TestSetup.JobService.GetJobAsync(new JobId("5f98da1d98cf40000141344c"));
 	        Console.WriteLine("Job is " + Job);
         }
         
diff --git a/Engine/Source/Programs/Horde/HordeServer.Compute.Tests/Fixture.cs b/Engine/Source/Programs/Horde/HordeServer.Compute.Tests/Fixture.cs
--- a/Engine/Source/Programs/Horde/HordeServer.Compute.Tests/Fixture.cs
+++ b/Engine/Source/Programs/Horde/HordeServer.Compute.Tests/Fixture.cs
@@ -16,9 +16,12 @@ using TemplateRefId = HordeServer.Utilities.StringId<HordeServer.Models.Template
 using ProjectId = HordeServer.Utilities.StringId<HordeServer.Models.IProject>;
 using HordeServer.Api;
 using HordeServer.Jobs;
+using HordeServer.Utilities;
 
 namespace HordeServerTests
 {
+	using JobId = ObjectId<IJob>;
+
 	public class Fixture
 	{
 		public IJob Job1 { get; private set; } = null!;
@@ -78,7 +81,7 @@ namespace HordeServerTests
 			);
 			
 			Job1 = await JobService.CreateJobAsync(
-				JobId: new ObjectId("5f283932841e7fdbcafb6ab5"),
+				JobId: new JobId("5f283932841e7fdbcafb6ab5"),
 				Stream: Stream!,
 				TemplateRefId: TemplateRefId1,
 				TemplateHash: Template.Id,
@@ -104,7 +107,7 @@ namespace HordeServerTests
 			Job1 = (await JobService.GetJobAsync(Job1.Id))!;
 
 			Job2 = await JobService.CreateJobAsync(
-				JobId: new ObjectId("5f69ea1b68423e921b035106"),
+				JobId: new JobId("5f69ea1b68423e921b035106"),
 				Stream: Stream!,
 				TemplateRefId: new TemplateRefId("template-id-1"),
 				TemplateHash: ContentHash.MD5("made-up-template-hash"),
diff --git a/Engine/Source/Programs/Horde/HordeServer.Compute.Tests/IssueServiceTests.cs b/Engine/Source/Programs/Horde/HordeServer.Compute.Tests/IssueServiceTests.cs
--- a/Engine/Source/Programs/Horde/HordeServer.Compute.Tests/IssueServiceTests.cs
+++ b/Engine/Source/Programs/Horde/HordeServer.Compute.Tests/IssueServiceTests.cs
@@ -41,16 +41,19 @@ using TemplateRefId = HordeServer.Utilities.StringId<HordeServer.Models.Template
 
 namespace HordeServerTests
 {
+	using JobId = ObjectId<IJob>;
+	using LogId = ObjectId<ILogFile>;
+
 	[TestClass]
 	public class IssueServiceTests : DatabaseIntegrationTest
 	{
 		class TestJsonLogger : JsonLogger, IAsyncDisposable
 		{
 			ILogFileService LogFileService;
-			ObjectId LogId;
+			LogId LogId;
 			List<(LogLevel, byte[][])> Events = new List<(LogLevel, byte[][])>();
 
-			public TestJsonLogger(ILogFileService LogFileService, ObjectId LogId)
+			public TestJsonLogger(ILogFileService LogFileService, LogId LogId)
 				: base(null, NullLogger.Instance)
 			{
 				this.LogFileService = LogFileService;
@@ -166,7 +169,7 @@ namespace HordeServerTests
 
 		public IJob CreateJob(StreamId StreamId, int Change, string Name, IGraph Graph, TimeSpan Time = default)
 		{
-			ObjectId JobId = ObjectId.GenerateNewId();
+			JobId JobId = JobId.GenerateNewId();
 
 			List<IJobStepBatch> Batches = new List<IJobStepBatch>();
 			for (int GroupIdx = 0; GroupIdx < Graph.Groups.Count; GroupIdx++)
@@ -224,7 +227,7 @@ namespace HordeServerTests
 
 		async Task AddEvent(IJob Job, int BatchIdx, int StepIdx, object Data, EventSeverity Severity = EventSeverity.Error)
 		{
-			ObjectId LogId = Job.Batches[BatchIdx].Steps[StepIdx].LogId!.Value;
+			LogId LogId = Job.Batches[BatchIdx].Steps[StepIdx].LogId!.Value;
 
 			List<byte> Bytes = new List<byte>();
 			Bytes.AddRange(JsonSerializer.SerializeToUtf8Bytes(Data));
@@ -239,7 +242,7 @@ namespace HordeServerTests
 
 		private async Task ParseEventsAsync(IJob Job, int BatchIdx, int StepIdx, string[] Lines)
 		{
-			ObjectId LogId = Job.Batches[BatchIdx].Steps[StepIdx].LogId!.Value;
+			LogId LogId = Job.Batches[BatchIdx].Steps[StepIdx].LogId!.Value;
 
 			LogParserContext Context = new LogParserContext();
 			Context.WorkspaceDir = WorkspaceDir;
@@ -1233,7 +1236,7 @@ namespace HordeServerTests
 			}
 		}
 
-		private async Task ParseAsync(ObjectId LogId, string[] Lines)
+		private async Task ParseAsync(LogId LogId, string[] Lines)
 		{
 			LogParserContext Context = new LogParserContext();
 			Context.WorkspaceDir = DirectoryReference.GetCurrentDirectory();
diff --git a/Engine/Source/Programs/Horde/HordeServer.Compute.Tests/JobCollectionTests.cs b/Engine/Source/Programs/Horde/HordeServer.Compute.Tests/JobCollectionTests.cs
--- a/Engine/Source/Programs/Horde/HordeServer.Compute.Tests/JobCollectionTests.cs
+++ b/Engine/Source/Programs/Horde/HordeServer.Compute.Tests/JobCollectionTests.cs
@@ -22,9 +22,14 @@ using System.Linq;
 using HordeCommon;
 using HordeServer.Collections.Impl;
 using Microsoft.Extensions.Logging.Abstractions;
+using HordeServer.Utilities;
 
 namespace HordeServerTests
 {
+	using JobId = ObjectId<IJob>;
+	using LeaseId = ObjectId<ILease>;
+	using LogId = ObjectId<ILogFile>;
+
 	[TestClass]
 	public class JobCollectionTests : DatabaseIntegrationTest
 	{
@@ -85,7 +90,7 @@ namespace HordeServerTests
 			Arguments.Add("-Target=Publish Client");
 			Arguments.Add("-Target=Post-Publish Client");
 
-			IJob Job = await JobCollection.AddAsync(ObjectId.GenerateNewId(), new StreamId("ue4-main"), new TemplateRefId("test-build"), ContentHash.SHA1("hello"), BaseGraph, "Test job", 123, 123, null, null, null, "Ben", null, null, null, null, false, false, null, null, null, Arguments);
+			IJob Job = await JobCollection.AddAsync(JobId.GenerateNewId(), new StreamId("ue4-main"), new TemplateRefId("test-build"), ContentHash.SHA1("hello"), BaseGraph, "Test job", 123, 123, null, null, null, "Ben", null, null, null, null, false, false, null, null, null, Arguments);
 
 			await StartBatch(Job, BaseGraph, 0);
 			await RunStep(Job, BaseGraph, 0, 0, JobStepOutcome.Success); // Setup Build
@@ -125,11 +130,11 @@ namespace HordeServerTests
 		{
 			TestSetup TestSetup = await GetTestSetup();
 			await TestSetup.JobCollection.TryAssignLeaseAsync(TestSetup.Fixture!.Job1, 0, new PoolId("foo"), TestSetup.Fixture.Agent1.Id,
-				ObjectId.GenerateNewId(), ObjectId.GenerateNewId(), ObjectId.GenerateNewId());
+				ObjectId.GenerateNewId(), LeaseId.GenerateNewId(), LogId.GenerateNewId());
 			
 			IJob Job = (await TestSetup.JobCollection.GetAsync(TestSetup.Fixture!.Job1.Id))!;
 			await TestSetup.JobCollection.TryAssignLeaseAsync(Job, 0, new PoolId("foo"), TestSetup.Fixture.Agent1.Id,
-				ObjectId.GenerateNewId(), ObjectId.GenerateNewId(), ObjectId.GenerateNewId());
+				ObjectId.GenerateNewId(), LeaseId.GenerateNewId(), LogId.GenerateNewId());
 			
 			// Manually verify the log output
 		}
@@ -157,7 +162,7 @@ namespace HordeServerTests
 			Arguments.Add("-Target=Step 1");
 			Arguments.Add("-Target=Step 3");
 
-			IJob Job = await JobCollection.AddAsync(ObjectId.GenerateNewId(), new StreamId("ue4-main"), new TemplateRefId("test-build"), ContentHash.SHA1("hello"), BaseGraph, "Test job", 123, 123, null, null, null, "Ben", null, null, null, null, false, false, null, null, null, Arguments);
+			IJob Job = await JobCollection.AddAsync(JobId.GenerateNewId(), new StreamId("ue4-main"), new TemplateRefId("test-build"), ContentHash.SHA1("hello"), BaseGraph, "Test job", 123, 123, null, null, null, "Ben", null, null, null, null, false, false, null, null, null, Arguments);
 
 			await StartBatch(Job, BaseGraph, 0);
 			await RunStep(Job, BaseGraph, 0, 0, JobStepOutcome.Success); // Setup Build
@@ -213,7 +218,7 @@ namespace HordeServerTests
 			Arguments.Add("-Target=Step 1");
 			Arguments.Add("-Target=Step 3");
 
-			IJob Job = await JobCollection.AddAsync(ObjectId.GenerateNewId(), new StreamId("ue4-main"), new TemplateRefId("test-build"), ContentHash.SHA1("hello"), BaseGraph, "Test job", 123, 123, null, null, null, "Ben", null, null, null, null, false, false, null, null, null, Arguments);
+			IJob Job = await JobCollection.AddAsync(JobId.GenerateNewId(), new StreamId("ue4-main"), new TemplateRefId("test-build"), ContentHash.SHA1("hello"), BaseGraph, "Test job", 123, 123, null, null, null, "Ben", null, null, null, null, false, false, null, null, null, Arguments);
 			Assert.AreEqual(1, Job.Batches.Count);
 
 			await StartBatch(Job, BaseGraph, 0);
diff --git a/Engine/Source/Programs/Horde/HordeServer.Compute.Tests/JobServiceTests.cs b/Engine/Source/Programs/Horde/HordeServer.Compute.Tests/JobServiceTests.cs
--- a/Engine/Source/Programs/Horde/HordeServer.Compute.Tests/JobServiceTests.cs
+++ b/Engine/Source/Programs/Horde/HordeServer.Compute.Tests/JobServiceTests.cs
@@ -13,9 +13,13 @@ using EpicGames.Core;
 using ProjectId = HordeServer.Utilities.StringId<HordeServer.Models.IProject>;
 using StreamId = HordeServer.Utilities.StringId<HordeServer.Models.IStream>;
 using TemplateRefId = HordeServer.Utilities.StringId<HordeServer.Models.TemplateRef>;
+using HordeServer.Utilities;
 
 namespace HordeServerTests
 {
+	using JobId = ObjectId<IJob>;
+	using LogId = ObjectId<ILogFile>;
+
 	[TestClass]
 	public class JobServiceTests : DatabaseIntegrationTest
 	{
@@ -46,7 +50,7 @@ namespace HordeServerTests
 			IJob Job = await TestSetup.JobService.CreateJobAsync(null, Stream!, TemplateRefId1, Template.Id, Graph, "Hello", 1234, 1233, 999, null, null, "joe", null, null, null, Stream!.Templates[TemplateRefId1].ChainedJobs, true, true, null, null, null, new List<string>());
 			Assert.AreEqual(1, Job.ChainedJobs.Count);
 
-			Assert.IsTrue(await TestSetup.JobService.UpdateBatchAsync(Job, Job.Batches[0].Id, ObjectId.GenerateNewId(), JobStepBatchState.Running));
+			Assert.IsTrue(await TestSetup.JobService.UpdateBatchAsync(Job, Job.Batches[0].Id, LogId.GenerateNewId(), JobStepBatchState.Running));
 			Assert.IsNotNull(await TestSetup.JobService.UpdateStepAsync(Job, Job.Batches[0].Id, Job.Batches[0].Steps[0].Id, JobStepState.Running));
 			Assert.IsNotNull(await TestSetup.JobService.UpdateStepAsync(Job, Job.Batches[0].Id, Job.Batches[0].Steps[0].Id, JobStepState.Completed, JobStepOutcome.Success));
 
@@ -93,7 +97,7 @@ namespace HordeServerTests
 		private Task<IJob> CreatePreflightJob(TestSetup TestSetup, string TemplateRefId, string TemplateHash, string StartedByUserName, int PreflightChange, string[] Arguments)
 		{
 			return TestSetup.JobService.CreateJobAsync(
-				JobId: ObjectId.GenerateNewId(),
+				JobId: JobId.GenerateNewId(),
 				Stream: TestSetup.Fixture!.Stream!,
 				TemplateRefId: new TemplateRefId(TemplateRefId),
 				TemplateHash: new ContentHash(Encoding.ASCII.GetBytes(TemplateHash)),
@@ -169,14 +173,14 @@ namespace HordeServerTests
 
 			IJob Job = await TestSetup.JobService.CreateJobAsync(null, Stream!, new TemplateRefId("temp"), Template.Id, Graph, "Hello", 1234, 1233, 999, null, null, "joe", null, null, null, null, true, true, null, null, null, new List<string> { "-Target=Pak" });
 
-			Assert.IsTrue(await TestSetup.JobService.UpdateBatchAsync(Job, Job.Batches[0].Id, ObjectId.GenerateNewId(), JobStepBatchState.Running));
+			Assert.IsTrue(await TestSetup.JobService.UpdateBatchAsync(Job, Job.Batches[0].Id, LogId.GenerateNewId(), JobStepBatchState.Running));
 			Assert.IsNotNull(await TestSetup.JobService.UpdateStepAsync(Job, Job.Batches[0].Id, Job.Batches[0].Steps[0].Id, JobStepState.Running));
 			Assert.IsNotNull(await TestSetup.JobService.UpdateStepAsync(Job, Job.Batches[0].Id, Job.Batches[0].Steps[0].Id, JobStepState.Completed, JobStepOutcome.Success));
 
-			Assert.IsTrue(await TestSetup.JobService.UpdateBatchAsync(Job, Job.Batches[1].Id, ObjectId.GenerateNewId(), JobStepBatchState.Running));
+			Assert.IsTrue(await TestSetup.JobService.UpdateBatchAsync(Job, Job.Batches[1].Id, LogId.GenerateNewId(), JobStepBatchState.Running));
 			Assert.IsNotNull(await TestSetup.JobService.UpdateStepAsync(Job, Job.Batches[1].Id, Job.Batches[1].Steps[0].Id, JobStepState.Running));
 
-			Assert.IsTrue(await TestSetup.JobService.UpdateBatchAsync(Job, Job.Batches[2].Id, ObjectId.GenerateNewId(), JobStepBatchState.Running));
+			Assert.IsTrue(await TestSetup.JobService.UpdateBatchAsync(Job, Job.Batches[2].Id, LogId.GenerateNewId(), JobStepBatchState.Running));
 			Assert.IsNotNull(await TestSetup.JobService.UpdateStepAsync(Job, Job.Batches[2].Id, Job.Batches[2].Steps[0].Id, JobStepState.Running));
 			Assert.IsNotNull(await TestSetup.JobService.UpdateStepAsync(Job, Job.Batches[2].Id, Job.Batches[2].Steps[0].Id, JobStepState.Completed, JobStepOutcome.Success));
 
diff --git a/Engine/Source/Programs/Horde/HordeServer.Compute.Tests/JobsControllerDbTest.cs b/Engine/Source/Programs/Horde/HordeServer.Compute.Tests/JobsControllerDbTest.cs
--- a/Engine/Source/Programs/Horde/HordeServer.Compute.Tests/JobsControllerDbTest.cs
+++ b/Engine/Source/Programs/Horde/HordeServer.Compute.Tests/JobsControllerDbTest.cs
@@ -46,16 +46,16 @@ namespace HordeServerTests
 	        SubResourceId BatchId = Job.Batches[0].Id;
 	        SubResourceId StepId = Job.Batches[0].Steps[0].Id;
 
-	        object Obj = (await JobsController.GetStepAsync(Job.Id.ToString(), BatchId.ToString(), StepId.ToString())).Value;
+	        object Obj = (await JobsController.GetStepAsync(Job.Id, BatchId.ToString(), StepId.ToString())).Value;
 	        GetStepResponse StepRes = (Obj as GetStepResponse)!;
 	        Assert.IsFalse(StepRes.AbortRequested);
 	        
 	        UpdateStepRequest UpdateReq = new UpdateStepRequest();
 	        UpdateReq.AbortRequested = true;
-	        Obj = (await JobsController.UpdateStepAsync(Job.Id.ToString(), BatchId.ToString(), StepId.ToString(), UpdateReq)).Value;
+	        Obj = (await JobsController.UpdateStepAsync(Job.Id, BatchId.ToString(), StepId.ToString(), UpdateReq)).Value;
 	        UpdateStepResponse UpdateRes = (Obj as UpdateStepResponse)!;
 	        
-	        Obj = (await JobsController.GetStepAsync(Job.Id.ToString(), BatchId.ToString(), StepId.ToString())).Value;
+	        Obj = (await JobsController.GetStepAsync(Job.Id, BatchId.ToString(), StepId.ToString())).Value;
 	        StepRes = (Obj as GetStepResponse)!;
 	        Assert.IsTrue(StepRes.AbortRequested);
 	        Assert.AreEqual("Anonymous", StepRes.AbortByUser);
diff --git a/Engine/Source/Programs/Horde/HordeServer.Compute.Tests/LogBuilderTests.cs b/Engine/Source/Programs/Horde/HordeServer.Compute.Tests/LogBuilderTests.cs
--- a/Engine/Source/Programs/Horde/HordeServer.Compute.Tests/LogBuilderTests.cs
+++ b/Engine/Source/Programs/Horde/HordeServer.Compute.Tests/LogBuilderTests.cs
@@ -3,6 +3,8 @@
 using HordeServer.Api;
 using HordeServer.Logs;
 using HordeServer.Logs.Builder;
+using HordeServer.Models;
+using HordeServer.Utilities;
 using Microsoft.Extensions.Logging.Abstractions;
 using Microsoft.VisualStudio.TestTools.UnitTesting;
 using MongoDB.Bson;
@@ -13,6 +15,8 @@ using System.Threading.Tasks;
 
 namespace HordeServerTests
 {
+	using LogId = ObjectId<ILogFile>;
+
 	[TestClass]
 	public class LogBuilderTests : DatabaseIntegrationTest
 	{
@@ -32,7 +36,7 @@ namespace HordeServerTests
 
 		public async Task TestBuilder(ILogBuilder Builder)
 		{
-			ObjectId LogId = new ObjectId();
+			LogId LogId = LogId.GenerateNewId();
 
 			const long Offset = 100;
 			Assert.IsTrue(await Builder.AppendAsync(LogId, Offset, Offset, 0, 1, Encoding.UTF8.GetBytes("hello\n"), LogType.Text));
@@ -76,7 +80,7 @@ namespace HordeServerTests
 			Assert.AreEqual(16, Chunk3!.GetLineOffsetWithinChunk(3));
 			Assert.AreEqual(20, Chunk3!.GetLineOffsetWithinChunk(4));
 
-			List<(ObjectId, long)> Chunks1 = await Builder.TouchChunksAsync(TimeSpan.Zero);
+			List<(LogId, long)> Chunks1 = await Builder.TouchChunksAsync(TimeSpan.Zero);
 			Assert.AreEqual(1, Chunks1.Count);
 			Assert.AreEqual((LogId, Offset), Chunks1[0]);
 
@@ -85,12 +89,12 @@ namespace HordeServerTests
 			LogChunkData? Chunk4 = await Builder.GetChunkAsync(LogId, Offset, 5);
 			Assert.AreEqual(24, Chunk4!.Length);
 
-			List<(ObjectId, long)> Chunks2 = await Builder.TouchChunksAsync(TimeSpan.Zero);
+			List<(LogId, long)> Chunks2 = await Builder.TouchChunksAsync(TimeSpan.Zero);
 			Assert.AreEqual(1, Chunks2.Count);
 
 			await Builder.RemoveChunkAsync(LogId, Offset);
 
-			List<(ObjectId, long)> Chunks3 = await Builder.TouchChunksAsync(TimeSpan.Zero);
+			List<(LogId, long)> Chunks3 = await Builder.TouchChunksAsync(TimeSpan.Zero);
 			Assert.AreEqual(0, Chunks3.Count);
 		}
 	}
diff --git a/Engine/Source/Programs/Horde/HordeServer.Compute.Tests/LogFileServiceTests.cs b/Engine/Source/Programs/Horde/HordeServer.Compute.Tests/LogFileServiceTests.cs
--- a/Engine/Source/Programs/Horde/HordeServer.Compute.Tests/LogFileServiceTests.cs
+++ b/Engine/Source/Programs/Horde/HordeServer.Compute.Tests/LogFileServiceTests.cs
@@ -32,7 +32,10 @@ using PoolId = HordeServer.Utilities.StringId<HordeServer.Models.IPool>;
 
 namespace HordeServerTests
 {
-    [TestClass]
+	using JobId = ObjectId<IJob>;
+	using LogId = ObjectId<ILogFile>;
+
+	[TestClass]
     public class LogFileServiceTest : DatabaseIntegrationTest
     {
         private readonly LogFileService LogFileService;
@@ -55,7 +58,7 @@ namespace HordeServerTests
 		[TestMethod]
         public async Task WriteLogLifecycleOldTest()
         {
-            ObjectId JobId = ObjectId.GenerateNewId();
+			JobId JobId = JobId.GenerateNewId();
             ILogFile LogFile = await LogFileService.CreateLogFileAsync(JobId, null, LogType.Text);
 
             LogFile = (await ((ILogFileService)LogFileService).WriteLogDataAsync(LogFile, 0, 0, Encoding.ASCII.GetBytes("hello\n"), true))!;
@@ -89,7 +92,7 @@ namespace HordeServerTests
 	        Assert.AreEqual(ExpectedMaxLineIndex, Metadata.MaxLineIndex);
         }
 
-        private static async Task AssertChunk(ILogFileService LogFileService, ObjectId LogFileId, long NumChunks, int ChunkId, long Offset, long Length,
+        private static async Task AssertChunk(ILogFileService LogFileService, LogId LogFileId, long NumChunks, int ChunkId, long Offset, long Length,
 	        long LineIndex)
         {
 	        ILogFile? LogFile = await LogFileService.GetLogFileAsync(LogFileId);
@@ -101,7 +104,7 @@ namespace HordeServerTests
 	        Assert.AreEqual(LineIndex, Chunk.LineIndex);
         }
 
-        private static async Task AssertLineOffset(ILogFileService LogFileService, ObjectId LogFileId, int LineIndex, int ClampedLineIndex, long Offset)
+        private static async Task AssertLineOffset(ILogFileService LogFileService, LogId LogFileId, int LineIndex, int ClampedLineIndex, long Offset)
         {
 	        ILogFile? LogFile = await LogFileService.GetLogFileAsync(LogFileId);
 	        Assert.AreEqual((ClampedLineIndex, Offset), await LogFileService.GetLineOffsetAsync(LogFile!, LineIndex));
@@ -115,7 +118,7 @@ namespace HordeServerTests
         
         public static async Task WriteLogLifecycle(ILogFileService Lfs, int MaxChunkLength)
         {
-	        ObjectId JobId = ObjectId.GenerateNewId();
+			JobId JobId = JobId.GenerateNewId();
             ILogFile LogFile = await Lfs.CreateLogFileAsync(JobId, null, LogType.Text);
 
             string Str1 = "hello\n";
@@ -225,8 +228,8 @@ namespace HordeServerTests
             await GetDatabaseService().Database.DropCollectionAsync("LogFiles");
             Assert.AreEqual(0, (await LogFileService.GetLogFilesAsync()).Count);
 
-            // Will implicitly test GetLogFileAsync(), AddCachedLogFile()
-            ObjectId JobId = ObjectId.GenerateNewId();
+			// Will implicitly test GetLogFileAsync(), AddCachedLogFile()
+			JobId JobId = JobId.GenerateNewId();
             ObjectId SessionId = ObjectId.GenerateNewId();
             ILogFile A = await LogFileService.CreateLogFileAsync(JobId, SessionId, LogType.Text);
             ILogFile B = (await LogFileService.GetCachedLogFileAsync(A.Id))!;
@@ -234,17 +237,17 @@ namespace HordeServerTests
             Assert.AreEqual(A.SessionId, B.SessionId);
             Assert.AreEqual(A.Type, B.Type);
 
-            ILogFile? NotFound = await LogFileService.GetCachedLogFileAsync(ObjectId.GenerateNewId());
+            ILogFile? NotFound = await LogFileService.GetCachedLogFileAsync(LogId.GenerateNewId());
             Assert.IsNull(NotFound);
 
-            await LogFileService.CreateLogFileAsync(ObjectId.GenerateNewId(), ObjectId.GenerateNewId(), LogType.Text);
+            await LogFileService.CreateLogFileAsync(JobId.GenerateNewId(), ObjectId.GenerateNewId(), LogType.Text);
             Assert.AreEqual(2, (await LogFileService.GetLogFilesAsync()).Count);
         }
 
         [TestMethod]
         public async Task AuthorizeForSession()
         {
-            ObjectId JobId = ObjectId.GenerateNewId();
+			JobId JobId = JobId.GenerateNewId();
             ObjectId SessionId = ObjectId.GenerateNewId();
             ILogFile LogFile = await LogFileService.CreateLogFileAsync(JobId, SessionId, LogType.Text);
             ILogFile LogFileNoSession = await LogFileService.CreateLogFileAsync(JobId, null, LogType.Text);
@@ -266,7 +269,7 @@ namespace HordeServerTests
 		[TestMethod]
 		public async Task ChunkSplitting()
 		{
-			ObjectId JobId = ObjectId.GenerateNewId();
+			JobId JobId = JobId.GenerateNewId();
 			ILogFile LogFile = await LogFileService.CreateLogFileAsync(JobId, null, LogType.Text);
 
 			long Offset = 0;
diff --git a/Engine/Source/Programs/Horde/HordeServer.Compute.Tests/LogIndexingTests.cs b/Engine/Source/Programs/Horde/HordeServer.Compute.Tests/LogIndexingTests.cs
--- a/Engine/Source/Programs/Horde/HordeServer.Compute.Tests/LogIndexingTests.cs
+++ b/Engine/Source/Programs/Horde/HordeServer.Compute.Tests/LogIndexingTests.cs
@@ -25,6 +25,8 @@ using HordeServer.Collections.Impl;
 
 namespace HordeServerTests
 {
+	using JobId = ObjectId<IJob>;
+
 	[TestClass]
 	public class LogIndexingTests : DatabaseIntegrationTest
 	{
@@ -53,7 +55,7 @@ namespace HordeServerTests
 		[TestMethod]
 		public async Task IndexTests()
 		{
-			ObjectId JobId = ObjectId.GenerateNewId();
+			JobId JobId = JobId.GenerateNewId();
 			ILogFile LogFile = await LogFileService.CreateLogFileAsync(JobId, null, LogType.Text);
 
 			// Write the test data to the log file in blocks
@@ -124,7 +126,7 @@ namespace HordeServerTests
 		[TestMethod]
 		public async Task PartialTokenTests()
 		{
-			ObjectId JobId = ObjectId.GenerateNewId();
+			JobId JobId = JobId.GenerateNewId();
 			ILogFile LogFile = await LogFileService.CreateLogFileAsync(JobId, null, LogType.Text);
 
 			string[] Lines =
@@ -166,7 +168,7 @@ namespace HordeServerTests
 		[TestMethod]
 		public async Task AppendIndexTests()
 		{
-			ObjectId JobId = ObjectId.GenerateNewId();
+			JobId JobId = JobId.GenerateNewId();
 			ILogFile LogFile = await LogFileService.CreateLogFileAsync(JobId, null, LogType.Text);
 
 			LogFile = await WriteLogDataAsync(LogFile, 0, 0, Encoding.UTF8.GetBytes("abc\n"), true);
diff --git a/Engine/Source/Programs/Horde/HordeServer.Compute.Tests/Models/Job/GetLabelStatesTests.cs b/Engine/Source/Programs/Horde/HordeServer.Compute.Tests/Models/Job/GetLabelStatesTests.cs
--- a/Engine/Source/Programs/Horde/HordeServer.Compute.Tests/Models/Job/GetLabelStatesTests.cs
+++ b/Engine/Source/Programs/Horde/HordeServer.Compute.Tests/Models/Job/GetLabelStatesTests.cs
@@ -7,9 +7,12 @@ using HordeServer.Models;
 using Microsoft.VisualStudio.TestTools.UnitTesting;
 using MongoDB.Bson;
 using Moq;
+using HordeServer.Utilities;
 
 namespace HordeServerTests
 {
+	using JobId = ObjectId<IJob>;
+
 	[TestClass]
 	public class GetLabelStatesTests
 	{
@@ -62,7 +65,7 @@ namespace HordeServerTests
 
 		public static IJob CreateGetLabelStatesJob(string Name, IGraph Graph, List<(JobStepState, JobStepOutcome)> JobStepData)
 		{
-			ObjectId JobId = new ObjectId("5ec16da1774cb4000107c2c1");
+			JobId JobId = JobId.Parse("5ec16da1774cb4000107c2c1");
 			List<IJobStepBatch> Batches = new List<IJobStepBatch>();
 			int DataIdx = 0;
 			for (int GroupIdx = 0; GroupIdx < Graph.Groups.Count; GroupIdx++)
diff --git a/Engine/Source/Programs/Horde/HordeServer.Compute.Tests/NotificationServiceTests.cs b/Engine/Source/Programs/Horde/HordeServer.Compute.Tests/NotificationServiceTests.cs
--- a/Engine/Source/Programs/Horde/HordeServer.Compute.Tests/NotificationServiceTests.cs
+++ b/Engine/Source/Programs/Horde/HordeServer.Compute.Tests/NotificationServiceTests.cs
@@ -4,6 +4,7 @@ using System.Collections.Generic;
 using System.Threading.Tasks;
 using HordeServer.Api;
 using HordeServer.Models;
+using HordeServer.Utilities;
 using Microsoft.VisualStudio.TestTools.UnitTesting;
 using MongoDB.Bson;
 using Moq;
@@ -12,12 +13,14 @@ using StreamId = HordeServer.Utilities.StringId<HordeServer.Models.IStream>;
 
 namespace HordeServerTests
 {
+	using JobId = ObjectId<IJob>;
+
 	[TestClass]
 	public class NotificationServiceTests
 	{
 		public IJob CreateJob(StreamId StreamId, int Change, string Name, IGraph Graph)
 		{
-			ObjectId JobId = new ObjectId("5ec16da1774cb4000107c2c1");
+			JobId JobId = new JobId("5ec16da1774cb4000107c2c1");
 
 			List<IJobStepBatch> Batches = new List<IJobStepBatch>();
 			for (int GroupIdx = 0; GroupIdx < Graph.Groups.Count; GroupIdx++)
diff --git a/Engine/Source/Programs/Horde/HordeServer.Compute.Tests/RpcServiceTest.cs b/Engine/Source/Programs/Horde/HordeServer.Compute.Tests/RpcServiceTest.cs
--- a/Engine/Source/Programs/Horde/HordeServer.Compute.Tests/RpcServiceTest.cs
+++ b/Engine/Source/Programs/Horde/HordeServer.Compute.Tests/RpcServiceTest.cs
@@ -33,6 +33,9 @@ using PoolId = HordeServer.Utilities.StringId<HordeServer.Models.IPool>;
 
 namespace HordeServerTests
 {
+	using LeaseId = ObjectId<ILease>;
+	using LogId = ObjectId<ILogFile>;
+
 	public class AppLifetimeStub : IHostApplicationLifetime
 	{
 		public CancellationToken ApplicationStarted { get; }
@@ -409,7 +412,7 @@ namespace HordeServerTests
 			// Set the session ID on the job batch to pass auth later
 			Assert.IsTrue(await TestSetup.JobCollection.TryAssignLeaseAsync(TestSetup.Fixture.Job1, 0, new PoolId("foo"),
 				new AgentId("test"), SessionId,
-				ObjectId.GenerateNewId(), ObjectId.GenerateNewId()));
+				LeaseId.GenerateNewId(), LogId.GenerateNewId()));
 /*
 			TestAsyncStreamReader<UploadArtifactRequest> RequestStream = new TestAsyncStreamReader<UploadArtifactRequest>(Context);
 			Task<UploadArtifactResponse> Call = TestSetup.RpcService.UploadArtifact(RequestStream,  Context);
diff --git a/Engine/Source/Programs/Horde/HordeServer.Compute.Tests/SchedulerTests.cs b/Engine/Source/Programs/Horde/HordeServer.Compute.Tests/SchedulerTests.cs
--- a/Engine/Source/Programs/Horde/HordeServer.Compute.Tests/SchedulerTests.cs
+++ b/Engine/Source/Programs/Horde/HordeServer.Compute.Tests/SchedulerTests.cs
@@ -18,10 +18,14 @@ using MongoDB.Bson;
 using System.Linq;
 using HordeServerTests.Stubs.Services;
 using HordeServer.Collections;
+using HordeServer.Utilities;
 
 namespace HordeServerTests
 {
-    [TestClass]
+	using JobId = ObjectId<IJob>;
+	using LogId = ObjectId<ILogFile>;
+
+	[TestClass]
     public class SchedulerTests : DatabaseIntegrationTest
 	{
 		ProjectId ProjectId { get; } = new ProjectId("ue5");
@@ -31,7 +35,7 @@ namespace HordeServerTests
 		IServiceProvider ServiceProvider;
 		TestSetup TestSetup;
 		ITemplate Template;
-		HashSet<ObjectId> InitialJobIds;
+		HashSet<JobId> InitialJobIds;
 		ScheduleService ScheduleService;
 		PerforceServiceStub PerforceService;
 
@@ -44,7 +48,7 @@ namespace HordeServerTests
 
 			Template = TestSetup.TemplateCollection.AddAsync("Test template", null, false, null, null, new List<string>(), new List<Parameter>()).Result;
 
-			InitialJobIds = new HashSet<ObjectId>(TestSetup.JobCollection.FindAsync().Result.Select(x => x.Id));
+			InitialJobIds = new HashSet<JobId>(TestSetup.JobCollection.FindAsync().Result.Select(x => x.Id));
 
 			ServiceProvider = TestSetup.ServiceProvider;
 			ScheduleService = ServiceProvider.GetRequiredService<ScheduleService>();
@@ -388,9 +392,9 @@ namespace HordeServerTests
 
 			// Create a job and fail it
 			IJob Job1 = await TestSetup.JobService.CreateJobAsync(null, Stream, NewTemplateRefId1, Template.Id, GraphA, "Hello", 1234, 1233, 999, null, null, "joe", null, null, null, null, true, true, null, null, null, new List<string> { "-Target=TriggerNext" });
-			Assert.IsTrue(await TestSetup.JobService.UpdateBatchAsync(Job1, Job1.Batches[0].Id, ObjectId.GenerateNewId(), JobStepBatchState.Running));
+			Assert.IsTrue(await TestSetup.JobService.UpdateBatchAsync(Job1, Job1.Batches[0].Id, LogId.GenerateNewId(), JobStepBatchState.Running));
 			Assert.IsNotNull(await TestSetup.JobService.UpdateStepAsync(Job1, Job1.Batches[0].Id, Job1.Batches[0].Steps[0].Id, JobStepState.Completed, JobStepOutcome.Failure));
-			Assert.IsTrue(await TestSetup.JobService.UpdateBatchAsync(Job1, Job1.Batches[0].Id, ObjectId.GenerateNewId(), JobStepBatchState.Complete));
+			Assert.IsTrue(await TestSetup.JobService.UpdateBatchAsync(Job1, Job1.Batches[0].Id, LogId.GenerateNewId(), JobStepBatchState.Complete));
 			await GetNewJobs();
 
 			// Tick the schedule and make sure it doesn't trigger
@@ -401,7 +405,7 @@ namespace HordeServerTests
 
 			// Create a job and make it succeed
 			IJob Job2 = await TestSetup.JobService.CreateJobAsync(null, Stream, NewTemplateRefId1, Template.Id, GraphA, "Hello", 1234, 1233, 999, null, null, "joe", null, null, null, null, true, true, null, null, null, new List<string> { "-Target=TriggerNext" });
-			Assert.IsTrue(await TestSetup.JobService.UpdateBatchAsync(Job2, Job2.Batches[0].Id, ObjectId.GenerateNewId(), JobStepBatchState.Running));
+			Assert.IsTrue(await TestSetup.JobService.UpdateBatchAsync(Job2, Job2.Batches[0].Id, LogId.GenerateNewId(), JobStepBatchState.Running));
 			Assert.IsNotNull(await TestSetup.JobService.UpdateStepAsync(Job2, Job2.Batches[0].Id, Job2.Batches[0].Steps[0].Id, JobStepState.Completed, JobStepOutcome.Success));
 
 			// Tick the schedule and make sure it does trigger
diff --git a/Engine/Source/Programs/Horde/HordeServer.Compute.Tests/Stubs/Collections/JobCollectionStub.cs b/Engine/Source/Programs/Horde/HordeServer.Compute.Tests/Stubs/Collections/JobCollectionStub.cs
--- a/Engine/Source/Programs/Horde/HordeServer.Compute.Tests/Stubs/Collections/JobCollectionStub.cs
+++ b/Engine/Source/Programs/Horde/HordeServer.Compute.Tests/Stubs/Collections/JobCollectionStub.cs
@@ -12,6 +12,8 @@ using System.Threading.Tasks;
 
 namespace HordeServerTests.Stubs.Collections
 {
+	using JobId = ObjectId<IJob>;
+
 	class JobCollectionStub : IJobCollection
 	{
 		public Task AddAsync(Job NewJob)
@@ -19,7 +21,7 @@ namespace HordeServerTests.Stubs.Collections
 			throw new NotImplementedException();
 		}
 
-		public Task AddIssueToJobAsync(ObjectId JobId, ObjectId IssueId)
+		public Task AddIssueToJobAsync(JobId JobId, ObjectId IssueId)
 		{
 			throw new NotImplementedException();
 		}
diff --git a/Engine/Source/Programs/Horde/HordeServer.Compute.Tests/Stubs/Collections/JobStepRefCollectionStub.cs b/Engine/Source/Programs/Horde/HordeServer.Compute.Tests/Stubs/Collections/JobStepRefCollectionStub.cs
--- a/Engine/Source/Programs/Horde/HordeServer.Compute.Tests/Stubs/Collections/JobStepRefCollectionStub.cs
+++ b/Engine/Source/Programs/Horde/HordeServer.Compute.Tests/Stubs/Collections/JobStepRefCollectionStub.cs
@@ -17,6 +17,9 @@ using PoolId = HordeServer.Utilities.StringId<HordeServer.Models.IPool>;
 
 namespace HordeServerTests.Stubs.Collections
 {
+	using JobId = ObjectId<IJob>;
+	using LogId = ObjectId<ILogFile>;
+
 	class JobStepRefStub : IJobStepRef
 	{
 		public JobStepRefId Id { get; set; }
@@ -25,7 +28,7 @@ namespace HordeServerTests.Stubs.Collections
 		public StreamId StreamId { get; set; }
 		public TemplateRefId TemplateId { get; set; }
 		public int Change { get; set; }
-		public ObjectId? LogId { get; set; }
+		public LogId? LogId { get; set; }
 		public PoolId? PoolId { get; set; }
 		public AgentId? AgentId { get; set; }
 		public JobStepOutcome? Outcome { get; set; }
@@ -35,7 +38,7 @@ namespace HordeServerTests.Stubs.Collections
 		public virtual DateTime StartTimeUtc => throw new NotImplementedException();
 		public virtual DateTime? FinishTimeUtc => throw new NotImplementedException();
 
-		public JobStepRefStub(ObjectId JobId, SubResourceId BatchId, SubResourceId StepId, string JobName, string NodeName, StreamId StreamId, TemplateRefId TemplateId, int Change, JobStepOutcome? Outcome)
+		public JobStepRefStub(JobId JobId, SubResourceId BatchId, SubResourceId StepId, string JobName, string NodeName, StreamId StreamId, TemplateRefId TemplateId, int Change, JobStepOutcome? Outcome)
 		{
 			this.Id = new JobStepRefId(JobId, BatchId, StepId);
 			this.JobName = JobName;
@@ -93,7 +96,7 @@ namespace HordeServerTests.Stubs.Collections
 			throw new NotImplementedException();
 		}
 
-		Task<IJobStepRef> IJobStepRefCollection.InsertOrReplaceAsync(JobStepRefId Id, string JobName, string NodeName, StreamId StreamId, TemplateRefId TemplateId, int Change, ObjectId? LogId, PoolId? PoolId, AgentId? AgentId, JobStepOutcome? Outcome, float WaitTime, float InitTime, DateTime StartTime, DateTime? FinishTime)
+		Task<IJobStepRef> IJobStepRefCollection.InsertOrReplaceAsync(JobStepRefId Id, string JobName, string NodeName, StreamId StreamId, TemplateRefId TemplateId, int Change, LogId? LogId, PoolId? PoolId, AgentId? AgentId, JobStepOutcome? Outcome, float WaitTime, float InitTime, DateTime StartTime, DateTime? FinishTime)
 		{
 			throw new NotImplementedException();
 		}
diff --git a/Engine/Source/Programs/Horde/HordeServer/Agents/IAgent.cs b/Engine/Source/Programs/Horde/HordeServer/Agents/IAgent.cs
--- a/Engine/Source/Programs/Horde/HordeServer/Agents/IAgent.cs
+++ b/Engine/Source/Programs/Horde/HordeServer/Agents/IAgent.cs
@@ -25,6 +25,8 @@ using EpicGames.Horde.Common;
 
 namespace HordeServer.Models
 {
+	using LeaseId = ObjectId<ILease>;
+	using LogId = ObjectId<ILogFile>;
 	using PoolId = StringId<IPool>;
 	using StreamId = StringId<IStream>;
 	using AgentSoftwareVersion = StringId<IAgentSoftwareCollection>;
@@ -185,13 +187,13 @@ namespace HordeServer.Models
 		/// <summary>
 		/// Lease returned when a task source wants to run, but needs to wait for other tasks to finish first.
 		/// </summary>
-		public static AgentLease Drain { get; } = new AgentLease(ObjectId.Empty, String.Empty, null, null, null, LeaseState.Pending, null, true, null);
+		public static AgentLease Drain { get; } = new AgentLease(LeaseId.Empty, String.Empty, null, null, null, LeaseState.Pending, null, true, null);
 
 		/// <summary>
 		/// Name of this lease
 		/// </summary>
 		[BsonRequired]
-		public ObjectId Id { get; set; }
+		public LeaseId Id { get; set; }
 
 		/// <summary>
 		/// Name of this lease
@@ -216,7 +218,7 @@ namespace HordeServer.Models
 		/// <summary>
 		/// Optional log for this lease
 		/// </summary>
-		public ObjectId? LogId { get; set; }
+		public LogId? LogId { get; set; }
 
 		/// <summary>
 		/// Time at which the lease started
@@ -270,7 +272,7 @@ namespace HordeServer.Models
 		/// <param name="Resources">Resources required for this lease</param>
 		/// <param name="Exclusive">Whether to reserve the entire device</param>
 		/// <param name="Payload">Encoded "any" protobuf describing the contents of the payload</param>
-		public AgentLease(ObjectId Id, string Name, StreamId? StreamId, PoolId? PoolId, ObjectId? LogId, LeaseState State, IReadOnlyDictionary<string, int>? Resources, bool Exclusive, byte[]? Payload)
+		public AgentLease(LeaseId Id, string Name, StreamId? StreamId, PoolId? PoolId, LogId? LogId, LeaseState State, IReadOnlyDictionary<string, int>? Resources, bool Exclusive, byte[]? Payload)
 		{
 			this.Id = Id;
 			this.Name = Name;
diff --git a/Engine/Source/Programs/Horde/HordeServer/Agents/ILease.cs b/Engine/Source/Programs/Horde/HordeServer/Agents/ILease.cs
--- a/Engine/Source/Programs/Horde/HordeServer/Agents/ILease.cs
+++ b/Engine/Source/Programs/Horde/HordeServer/Agents/ILease.cs
@@ -15,6 +15,8 @@ using Google.Protobuf;
 
 namespace HordeServer.Models
 {
+	using LeaseId = ObjectId<ILease>;
+	using LogId = ObjectId<ILogFile>;
 	using PoolId = StringId<IPool>;
 	using StreamId = StringId<IStream>;
 
@@ -26,7 +28,7 @@ namespace HordeServer.Models
 		/// <summary>
 		/// The unique id of this lease
 		/// </summary>
-		public ObjectId Id { get; }
+		public LeaseId Id { get; }
 
 		/// <summary>
 		/// Name of this lease
@@ -56,7 +58,7 @@ namespace HordeServer.Models
 		/// <summary>
 		/// The log for this lease, if applicable
 		/// </summary>
-		public ObjectId? LogId { get; }
+		public LogId? LogId { get; }
 
 		/// <summary>
 		/// Time at which this lease started
diff --git a/Engine/Source/Programs/Horde/HordeServer/Agents/ILeaseCollection.cs b/Engine/Source/Programs/Horde/HordeServer/Agents/ILeaseCollection.cs
--- a/Engine/Source/Programs/Horde/HordeServer/Agents/ILeaseCollection.cs
+++ b/Engine/Source/Programs/Horde/HordeServer/Agents/ILeaseCollection.cs
@@ -14,6 +14,8 @@ using System.Threading.Tasks;
 
 namespace HordeServer.Collections
 {
+	using LeaseId = ObjectId<ILease>;
+	using LogId = ObjectId<ILogFile>;
 	using PoolId = StringId<IPool>;
 	using StreamId = StringId<IStream>;
 
@@ -35,21 +37,21 @@ namespace HordeServer.Collections
 		/// <param name="StartTime">Start time of the lease</param>
 		/// <param name="Payload">Payload for the lease</param>
 		/// <returns>Async task</returns>
-		Task<ILease> AddAsync(ObjectId Id, string Name, AgentId AgentId, ObjectId SessionId, StreamId? StreamId, PoolId? PoolId, ObjectId? LogId, DateTime StartTime, byte[] Payload);
+		Task<ILease> AddAsync(LeaseId Id, string Name, AgentId AgentId, ObjectId SessionId, StreamId? StreamId, PoolId? PoolId, LogId? LogId, DateTime StartTime, byte[] Payload);
 
 		/// <summary>
 		/// Deletes a lease from the collection
 		/// </summary>
 		/// <param name="LeaseId">Unique id of the lease</param>
 		/// <returns>Async task</returns>
-		Task DeleteAsync(ObjectId LeaseId);
+		Task DeleteAsync(LeaseId LeaseId);
 
 		/// <summary>
 		/// Gets a specific lease
 		/// </summary>
 		/// <param name="LeaseId">Unique id of the lease</param>
 		/// <returns>The lease that was found, or null if it does not exist</returns>
-		Task<ILease?> GetAsync(ObjectId LeaseId);
+		Task<ILease?> GetAsync(LeaseId LeaseId);
 
 		/// <summary>
 		/// Finds all leases matching a set of criteria
@@ -79,6 +81,6 @@ namespace HordeServer.Collections
 		/// <param name="Outcome">Outcome of the lease</param>
 		/// <param name="Output">Output data from the task</param>
 		/// <returns>True if the lease was updated, false otherwise</returns>
-		Task<bool> TrySetOutcomeAsync(ObjectId LeaseId, DateTime FinishTime, LeaseOutcome Outcome, byte[]? Output);
+		Task<bool> TrySetOutcomeAsync(LeaseId LeaseId, DateTime FinishTime, LeaseOutcome Outcome, byte[]? Output);
 	}
 }
diff --git a/Engine/Source/Programs/Horde/HordeServer/Agents/ITaskSource.cs b/Engine/Source/Programs/Horde/HordeServer/Agents/ITaskSource.cs
--- a/Engine/Source/Programs/Horde/HordeServer/Agents/ITaskSource.cs
+++ b/Engine/Source/Programs/Horde/HordeServer/Agents/ITaskSource.cs
@@ -6,6 +6,7 @@ using Google.Protobuf.WellKnownTypes;
 using HordeCommon;
 using HordeServer.Models;
 using HordeServer.Services;
+using HordeServer.Utilities;
 using Microsoft.Extensions.Logging;
 using MongoDB.Bson;
 using System;
@@ -18,6 +19,8 @@ using System.Threading.Tasks;
 
 namespace HordeServer.Tasks
 {
+	using LeaseId = ObjectId<ILease>;
+
 	/// <summary>
 	/// Handler for a certain lease type
 	/// </summary>
@@ -48,7 +51,7 @@ namespace HordeServer.Tasks
 		/// <param name="LeaseId">The lease id</param>
 		/// <param name="Payload">Payload for the lease</param>
 		/// <returns></returns>
-		Task CancelLeaseAsync(IAgent Agent, ObjectId LeaseId, Any Payload);
+		Task CancelLeaseAsync(IAgent Agent, LeaseId LeaseId, Any Payload);
 
 		/// <summary>
 		/// Notification that a lease has been started
@@ -57,7 +60,7 @@ namespace HordeServer.Tasks
 		/// <param name="LeaseId">The lease id</param>
 		/// <param name="Payload">Payload for the lease</param>
 		/// <param name="Logger">Logger for the agent</param>
-		Task OnLeaseStartedAsync(IAgent Agent, ObjectId LeaseId, Any Payload, ILogger Logger);
+		Task OnLeaseStartedAsync(IAgent Agent, LeaseId LeaseId, Any Payload, ILogger Logger);
 
 		/// <summary>
 		/// Notification that a task has completed
@@ -68,7 +71,7 @@ namespace HordeServer.Tasks
 		/// <param name="Outcome">Outcome of the lease</param>
 		/// <param name="Output">Output from the task</param>
 		/// <param name="Logger">Logger for the agent</param>
-		Task OnLeaseFinishedAsync(IAgent Agent, ObjectId LeaseId, Any Payload, LeaseOutcome Outcome, ReadOnlyMemory<byte> Output, ILogger Logger);
+		Task OnLeaseFinishedAsync(IAgent Agent, LeaseId LeaseId, Any Payload, LeaseOutcome Outcome, ReadOnlyMemory<byte> Output, ILogger Logger);
 	}
 	
 	/// <summary>
@@ -89,26 +92,26 @@ namespace HordeServer.Tasks
 		public abstract Task<AgentLease?> AssignLeaseAsync(IAgent Agent, CancellationToken CancellationToken);
 
 		/// <inheritdoc/>
-		public Task CancelLeaseAsync(IAgent Agent, ObjectId LeaseId, Any Payload) => CancelLeaseAsync(Agent, LeaseId, Payload.Unpack<TMessage>());
+		public Task CancelLeaseAsync(IAgent Agent, LeaseId LeaseId, Any Payload) => CancelLeaseAsync(Agent, LeaseId, Payload.Unpack<TMessage>());
 
 		/// <inheritdoc/>
-		public Task OnLeaseStartedAsync(IAgent Agent, ObjectId LeaseId, Any Payload, ILogger Logger) => OnLeaseStartedAsync(Agent, LeaseId, Payload.Unpack<TMessage>(), Logger);
+		public Task OnLeaseStartedAsync(IAgent Agent, LeaseId LeaseId, Any Payload, ILogger Logger) => OnLeaseStartedAsync(Agent, LeaseId, Payload.Unpack<TMessage>(), Logger);
 
 		/// <inheritdoc/>
-		public Task OnLeaseFinishedAsync(IAgent Agent, ObjectId LeaseId, Any Payload, LeaseOutcome Outcome, ReadOnlyMemory<byte> Output, ILogger Logger) => OnLeaseFinishedAsync(Agent, LeaseId, Payload.Unpack<TMessage>(), Outcome, Output, Logger);
+		public Task OnLeaseFinishedAsync(IAgent Agent, LeaseId LeaseId, Any Payload, LeaseOutcome Outcome, ReadOnlyMemory<byte> Output, ILogger Logger) => OnLeaseFinishedAsync(Agent, LeaseId, Payload.Unpack<TMessage>(), Outcome, Output, Logger);
 
-		/// <inheritdoc cref="ITaskSource.CancelLeaseAsync(IAgent, ObjectId, Any)"/>
-		public virtual Task CancelLeaseAsync(IAgent Agent, ObjectId LeaseId, TMessage Payload) => Task.CompletedTask;
+		/// <inheritdoc cref="ITaskSource.CancelLeaseAsync(IAgent, LeaseId, Any)"/>
+		public virtual Task CancelLeaseAsync(IAgent Agent, LeaseId LeaseId, TMessage Payload) => Task.CompletedTask;
 
-		/// <inheritdoc cref="ITaskSource.OnLeaseStartedAsync(IAgent, ObjectId, Any, ILogger)"/>
-		public virtual Task OnLeaseStartedAsync(IAgent Agent, ObjectId LeaseId, TMessage Payload, ILogger Logger)
+		/// <inheritdoc cref="ITaskSource.OnLeaseStartedAsync(IAgent, LeaseId, Any, ILogger)"/>
+		public virtual Task OnLeaseStartedAsync(IAgent Agent, LeaseId LeaseId, TMessage Payload, ILogger Logger)
 		{
 			LogLeaseStartedInfo<TMessage>(LeaseId, Payload, Logger);
 			return Task.CompletedTask;
 		}
 
-		/// <inheritdoc cref="ITaskSource.OnLeaseFinishedAsync(IAgent, ObjectId, Any, LeaseOutcome, ReadOnlyMemory{byte}, ILogger)"/>
-		public virtual Task OnLeaseFinishedAsync(IAgent Agent, ObjectId LeaseId, TMessage Payload, LeaseOutcome Outcome, ReadOnlyMemory<byte> Output, ILogger Logger)
+		/// <inheritdoc cref="ITaskSource.OnLeaseFinishedAsync(IAgent, LeaseId, Any, LeaseOutcome, ReadOnlyMemory{byte}, ILogger)"/>
+		public virtual Task OnLeaseFinishedAsync(IAgent Agent, LeaseId LeaseId, TMessage Payload, LeaseOutcome Outcome, ReadOnlyMemory<byte> Output, ILogger Logger)
 		{
 			Logger.LogInformation("Lease {LeaseId} complete", LeaseId);
 			return Task.CompletedTask;
@@ -138,7 +141,7 @@ namespace HordeServer.Tasks
 		/// <param name="LeaseId"></param>
 		/// <param name="Payload"></param>
 		/// <param name="Logger"></param>
-		protected void LogLeaseStartedInfo<T>(ObjectId LeaseId, T Payload, ILogger Logger)
+		protected void LogLeaseStartedInfo<T>(LeaseId LeaseId, T Payload, ILogger Logger)
 		{
 			PropertyInfo[] Properties = TypeInfo<T>.Properties;
 
diff --git a/Engine/Source/Programs/Horde/HordeServer/Agents/Impl/AgentService.cs b/Engine/Source/Programs/Horde/HordeServer/Agents/Impl/AgentService.cs
--- a/Engine/Source/Programs/Horde/HordeServer/Agents/Impl/AgentService.cs
+++ b/Engine/Source/Programs/Horde/HordeServer/Agents/Impl/AgentService.cs
@@ -24,6 +24,7 @@ using System.Threading.Tasks;
 namespace HordeServer.Services
 {
 	using AgentSoftwareChannelName = StringId<AgentSoftwareChannels>;
+	using LeaseId = ObjectId<ILease>;
 	using PoolId = StringId<IPool>;
 
 	/// <summary>
@@ -464,7 +465,7 @@ namespace HordeServer.Services
 		/// <param name="Agent">Agent to cancel the lease on</param>
 		/// <param name="LeaseId">The lease id to cancel</param>
 		/// <returns></returns>
-		public async Task<bool> CancelLeaseAsync(IAgent Agent, ObjectId LeaseId)
+		public async Task<bool> CancelLeaseAsync(IAgent Agent, LeaseId LeaseId)
 		{		
 			int Index = 0;
 			while (Index < Agent.Leases.Count && Agent.Leases[Index].Id != LeaseId)
@@ -571,7 +572,7 @@ namespace HordeServer.Services
 				List<AgentLease> Leases = new List<AgentLease>(Agent.Leases);
 
 				// Remove any completed leases from the agent
-				Dictionary<ObjectId, HordeCommon.Rpc.Messages.Lease> LeaseIdToNewState = NewLeases.ToDictionary(x => x.Id.ToObjectId(), x => x);
+				Dictionary<LeaseId, HordeCommon.Rpc.Messages.Lease> LeaseIdToNewState = NewLeases.ToDictionary(x => new LeaseId(x.Id), x => x);
 				for (int Idx = 0; Idx < Leases.Count; Idx++)
 				{
 					AgentLease Lease = Leases[Idx];
@@ -718,7 +719,7 @@ namespace HordeServer.Services
 		/// </summary>
 		/// <param name="LeaseId">Unique id of the lease</param>
 		/// <returns>The lease that was found, or null if it does not exist</returns>
-		public Task<ILease?> GetLeaseAsync(ObjectId LeaseId)
+		public Task<ILease?> GetLeaseAsync(LeaseId LeaseId)
 		{
 			return Leases.GetAsync(LeaseId);
 		}
diff --git a/Engine/Source/Programs/Horde/HordeServer/Agents/Impl/ConformList.cs b/Engine/Source/Programs/Horde/HordeServer/Agents/Impl/ConformList.cs
--- a/Engine/Source/Programs/Horde/HordeServer/Agents/Impl/ConformList.cs
+++ b/Engine/Source/Programs/Horde/HordeServer/Agents/Impl/ConformList.cs
@@ -10,6 +10,8 @@ using System.Threading.Tasks;
 
 namespace HordeServer.Models
 {
+	using LeaseId = ObjectId<ILease>;
+
 	/// <summary>
 	/// Entry to the conform list
 	/// </summary>
@@ -25,7 +27,7 @@ namespace HordeServer.Models
 		/// The lease id
 		/// </summary>
 		[BsonElement("l")]
-		public ObjectId LeaseId { get; set; }
+		public LeaseId LeaseId { get; set; }
 
 		/// <summary>
 		/// Last timestamp that the lease was checked for validity 
diff --git a/Engine/Source/Programs/Horde/HordeServer/Agents/Impl/LeaseCollection.cs b/Engine/Source/Programs/Horde/HordeServer/Agents/Impl/LeaseCollection.cs
--- a/Engine/Source/Programs/Horde/HordeServer/Agents/Impl/LeaseCollection.cs
+++ b/Engine/Source/Programs/Horde/HordeServer/Agents/Impl/LeaseCollection.cs
@@ -14,6 +14,8 @@ using System.Threading.Tasks;
 
 namespace HordeServer.Collections.Impl
 {
+	using LeaseId = ObjectId<ILease>;
+	using LogId = ObjectId<ILogFile>;
 	using PoolId = StringId<IPool>;
 	using StreamId = StringId<IStream>;
 
@@ -27,13 +29,13 @@ namespace HordeServer.Collections.Impl
 		/// </summary>
 		class LeaseDocument : ILease
 		{
-			public ObjectId Id { get; set; }
+			public LeaseId Id { get; set; }
 			public string Name { get; set; }
 			public AgentId AgentId { get; set; }
 			public ObjectId SessionId { get; set; }
 			public StreamId? StreamId { get; set; }
 			public PoolId? PoolId { get; set; }
-			public ObjectId? LogId { get; set; }
+			public LogId? LogId { get; set; }
 			public DateTime StartTime { get; set; }
 			public DateTime? FinishTime { get; set; }
 			public byte[] Payload { get; set; }
@@ -52,7 +54,7 @@ namespace HordeServer.Collections.Impl
 				Payload = null!;
 			}
 
-			public LeaseDocument(ObjectId Id, string Name, AgentId AgentId, ObjectId SessionId, StreamId? StreamId, PoolId? PoolId, ObjectId? LogId, DateTime StartTime, byte[] Payload)
+			public LeaseDocument(LeaseId Id, string Name, AgentId AgentId, ObjectId SessionId, StreamId? StreamId, PoolId? PoolId, LogId? LogId, DateTime StartTime, byte[] Payload)
 			{
 				this.Id = Id;
 				this.Name = Name;
@@ -89,7 +91,7 @@ namespace HordeServer.Collections.Impl
 		}
 
 		/// <inheritdoc/>
-		public async Task<ILease> AddAsync(ObjectId Id, string Name, AgentId AgentId, ObjectId SessionId, StreamId? StreamId, PoolId? PoolId, ObjectId? LogId, DateTime StartTime, byte[] Payload)
+		public async Task<ILease> AddAsync(LeaseId Id, string Name, AgentId AgentId, ObjectId SessionId, StreamId? StreamId, PoolId? PoolId, LogId? LogId, DateTime StartTime, byte[] Payload)
 		{
 			LeaseDocument Lease = new LeaseDocument(Id, Name, AgentId, SessionId, StreamId, PoolId, LogId, StartTime, Payload);
 			await Leases.ReplaceOneAsync(x => x.Id == Id, Lease, new ReplaceOptions { IsUpsert = true });
@@ -97,13 +99,13 @@ namespace HordeServer.Collections.Impl
 		}
 
 		/// <inheritdoc/>
-		public async Task DeleteAsync(ObjectId LeaseId)
+		public async Task DeleteAsync(LeaseId LeaseId)
 		{
 			await Leases.DeleteOneAsync(x => x.Id == LeaseId);
 		}
 
 		/// <inheritdoc/>
-		public async Task<ILease?> GetAsync(ObjectId LeaseId)
+		public async Task<ILease?> GetAsync(LeaseId LeaseId)
 		{
 			return await Leases.Find(x => x.Id == LeaseId).FirstOrDefaultAsync();
 		}
@@ -143,7 +145,7 @@ namespace HordeServer.Collections.Impl
 		}
 
 		/// <inheritdoc/>
-		public async Task<bool> TrySetOutcomeAsync(ObjectId LeaseId, DateTime FinishTime, LeaseOutcome Outcome, byte[]? Output)
+		public async Task<bool> TrySetOutcomeAsync(LeaseId LeaseId, DateTime FinishTime, LeaseOutcome Outcome, byte[]? Output)
 		{
 			FilterDefinitionBuilder<LeaseDocument> FilterBuilder = Builders<LeaseDocument>.Filter;
 			FilterDefinition<LeaseDocument> Filter = FilterBuilder.Eq(x => x.Id, LeaseId) & FilterBuilder.Eq(x => x.FinishTime, null);
diff --git a/Engine/Source/Programs/Horde/HordeServer/Agents/Impl/Tasks/ConformTaskSource.cs b/Engine/Source/Programs/Horde/HordeServer/Agents/Impl/Tasks/ConformTaskSource.cs
--- a/Engine/Source/Programs/Horde/HordeServer/Agents/Impl/Tasks/ConformTaskSource.cs
+++ b/Engine/Source/Programs/Horde/HordeServer/Agents/Impl/Tasks/ConformTaskSource.cs
@@ -21,6 +21,9 @@ using System.Threading.Tasks;
 
 namespace HordeServer.Tasks.Impl
 {
+	using JobId = ObjectId<IJob>;
+	using LeaseId = ObjectId<ILease>;
+
 	/// <summary>
 	/// Generates tasks telling agents to sync their workspaces
 	/// </summary>
@@ -92,7 +95,7 @@ namespace HordeServer.Tasks.Impl
 			ConformList List = await ConformList.GetAsync();
 
 			// Update any leases that are older than LastCheckTimeUtc
-			Dictionary<ObjectId, bool> RemoveLeases = new Dictionary<ObjectId, bool>();
+			Dictionary<LeaseId, bool> RemoveLeases = new Dictionary<LeaseId, bool>();
 			foreach (ConformListEntry Entry in Enumerable.Concat(List.Entries, List.Servers.SelectMany(x => x.Entries)))
 			{
 				if (Entry.LastCheckTimeUtc < LastCheckTimeUtc)
@@ -123,7 +126,7 @@ namespace HordeServer.Tasks.Impl
 		/// <param name="List">The list to update</param>
 		/// <param name="UtcNow">Current time</param>
 		/// <param name="RemoveLeases">List of leases to update. Entries with values set to true will be removed, entries with values set to false will have their timestamp updated.</param>
-		static void UpdateConformList(ConformList List, DateTime UtcNow, Dictionary<ObjectId, bool> RemoveLeases)
+		static void UpdateConformList(ConformList List, DateTime UtcNow, Dictionary<LeaseId, bool> RemoveLeases)
 		{
 			UpdateConformList(List.Entries, UtcNow, RemoveLeases);
 			foreach (ConformListServer Server in List.Servers)
@@ -138,7 +141,7 @@ namespace HordeServer.Tasks.Impl
 		/// <param name="Entries">The list to update</param>
 		/// <param name="UtcNow">Current time</param>
 		/// <param name="RemoveLeases">List of leases to update. Entries with values set to true will be removed, entries with values set to false will have their timestamp updated.</param>
-		static void UpdateConformList(List<ConformListEntry> Entries, DateTime UtcNow, Dictionary<ObjectId, bool> RemoveLeases)
+		static void UpdateConformList(List<ConformListEntry> Entries, DateTime UtcNow, Dictionary<LeaseId, bool> RemoveLeases)
 		{
 			for (int Idx = 0; Idx < Entries.Count; Idx++)
 			{
@@ -171,13 +174,13 @@ namespace HordeServer.Tasks.Impl
 				ConformTask Task = new ConformTask();
 				if (await GetWorkspacesAsync(Agent, Task.Workspaces))
 				{
-					ObjectId LeaseId = ObjectId.GenerateNewId();
+					LeaseId LeaseId = LeaseId.GenerateNewId();
 					if (!await AllocateConformLeaseAsync(Agent.Id, Task.Workspaces, LeaseId))
 					{
 						return null;
 					}
 
-					ILogFile Log = await LogService.CreateLogFileAsync(ObjectId.Empty, Agent.SessionId, LogType.Json);
+					ILogFile Log = await LogService.CreateLogFileAsync(JobId.Empty, Agent.SessionId, LogType.Json);
 					Task.LogId = Log.Id.ToString();
 
 					byte[] Payload = Any.Pack(Task).ToByteArray();
@@ -213,7 +216,7 @@ namespace HordeServer.Tasks.Impl
 		}
 
 		/// <inheritdoc/>
-		public override Task OnLeaseFinishedAsync(IAgent Agent, ObjectId LeaseId, ConformTask Payload, LeaseOutcome Outcome, ReadOnlyMemory<byte> Output, ILogger Logger)
+		public override Task OnLeaseFinishedAsync(IAgent Agent, LeaseId LeaseId, ConformTask Payload, LeaseOutcome Outcome, ReadOnlyMemory<byte> Output, ILogger Logger)
 		{
 			return ReleaseConformLeaseAsync(LeaseId);
 		}
@@ -225,7 +228,7 @@ namespace HordeServer.Tasks.Impl
 		/// <param name="Workspaces">List of workspaces that are required</param>
 		/// <param name="LeaseId">The lease id</param>
 		/// <returns>True if the resource was allocated, false otherwise</returns>
-		private async Task<bool> AllocateConformLeaseAsync(AgentId AgentId, IEnumerable<HordeCommon.Rpc.Messages.AgentWorkspace> Workspaces, ObjectId LeaseId)
+		private async Task<bool> AllocateConformLeaseAsync(AgentId AgentId, IEnumerable<HordeCommon.Rpc.Messages.AgentWorkspace> Workspaces, LeaseId LeaseId)
 		{
 			Globals Globals = await DatabaseService.GetGlobalsAsync();
 			for (; ; )
@@ -288,7 +291,7 @@ namespace HordeServer.Tasks.Impl
 		/// </summary>
 		/// <param name="LeaseId">The lease id</param>
 		/// <returns>Async task</returns>
-		public async Task ReleaseConformLeaseAsync(ObjectId LeaseId)
+		public async Task ReleaseConformLeaseAsync(LeaseId LeaseId)
 		{
 			for (; ; )
 			{
diff --git a/Engine/Source/Programs/Horde/HordeServer/Agents/Impl/Tasks/RestartTaskSource.cs b/Engine/Source/Programs/Horde/HordeServer/Agents/Impl/Tasks/RestartTaskSource.cs
--- a/Engine/Source/Programs/Horde/HordeServer/Agents/Impl/Tasks/RestartTaskSource.cs
+++ b/Engine/Source/Programs/Horde/HordeServer/Agents/Impl/Tasks/RestartTaskSource.cs
@@ -8,6 +8,7 @@ using HordeCommon.Rpc.Tasks;
 using HordeServer.Api;
 using HordeServer.Models;
 using HordeServer.Services;
+using HordeServer.Utilities;
 using MongoDB.Bson;
 using System;
 using System.Collections.Generic;
@@ -17,6 +18,9 @@ using System.Threading.Tasks;
 
 namespace HordeServer.Tasks.Impl
 {
+	using JobId = ObjectId<IJob>;
+	using LeaseId = ObjectId<ILease>;
+
 	class RestartTaskSource : TaskSourceBase<RestartTask>
 	{
 		public override string Type => "Restart";
@@ -39,14 +43,14 @@ namespace HordeServer.Tasks.Impl
 				return AgentLease.Drain;
 			}
 
-			ILogFile Log = await LogService.CreateLogFileAsync(ObjectId.Empty, Agent.SessionId, LogType.Json);
+			ILogFile Log = await LogService.CreateLogFileAsync(JobId.Empty, Agent.SessionId, LogType.Json);
 
 			RestartTask Task = new RestartTask();
 			Task.LogId = Log.Id.ToString();
 
 			byte[] Payload = Any.Pack(Task).ToByteArray();
 
-			return new AgentLease(ObjectId.GenerateNewId(), "Restart", null, null, Log.Id, LeaseState.Pending, null, true, Payload);
+			return new AgentLease(LeaseId.GenerateNewId(), "Restart", null, null, Log.Id, LeaseState.Pending, null, true, Payload);
 		}
 	}
 }
diff --git a/Engine/Source/Programs/Horde/HordeServer/Agents/Impl/Tasks/ShutdownTaskSource.cs b/Engine/Source/Programs/Horde/HordeServer/Agents/Impl/Tasks/ShutdownTaskSource.cs
--- a/Engine/Source/Programs/Horde/HordeServer/Agents/Impl/Tasks/ShutdownTaskSource.cs
+++ b/Engine/Source/Programs/Horde/HordeServer/Agents/Impl/Tasks/ShutdownTaskSource.cs
@@ -9,6 +9,7 @@ using HordeCommon.Rpc.Tasks;
 using HordeServer.Api;
 using HordeServer.Models;
 using HordeServer.Services;
+using HordeServer.Utilities;
 using MongoDB.Bson;
 using System;
 using System.Collections.Generic;
@@ -18,6 +19,9 @@ using System.Threading.Tasks;
 
 namespace HordeServer.Tasks.Impl
 {
+	using JobId = ObjectId<IJob>;
+	using LeaseId = ObjectId<ILease>;
+
 	class ShutdownTaskSource : TaskSourceBase<ShutdownTask>
 	{
 		public override string Type => "Shutdown";
@@ -40,14 +44,14 @@ namespace HordeServer.Tasks.Impl
 				return AgentLease.Drain;
 			}
 
-			ILogFile Log = await LogService.CreateLogFileAsync(ObjectId.Empty, Agent.SessionId, LogType.Json);
+			ILogFile Log = await LogService.CreateLogFileAsync(JobId.Empty, Agent.SessionId, LogType.Json);
 
 			ShutdownTask Task = new ShutdownTask();
 			Task.LogId = Log.Id.ToString();
 
 			byte[] Payload = Any.Pack(Task).ToByteArray();
 
-			return new AgentLease(ObjectId.GenerateNewId(), "Shutdown", null, null, Log.Id, LeaseState.Pending, null, true, Payload);
+			return new AgentLease(LeaseId.GenerateNewId(), "Shutdown", null, null, Log.Id, LeaseState.Pending, null, true, Payload);
 		}
 	}
 }
diff --git a/Engine/Source/Programs/Horde/HordeServer/Agents/Impl/Tasks/UpgradeTaskSource.cs b/Engine/Source/Programs/Horde/HordeServer/Agents/Impl/Tasks/UpgradeTaskSource.cs
--- a/Engine/Source/Programs/Horde/HordeServer/Agents/Impl/Tasks/UpgradeTaskSource.cs
+++ b/Engine/Source/Programs/Horde/HordeServer/Agents/Impl/Tasks/UpgradeTaskSource.cs
@@ -19,6 +19,8 @@ using System.Threading.Tasks;
 namespace HordeServer.Tasks.Impl
 {
 	using AgentSoftwareChannelName = StringId<AgentSoftwareChannels>;
+	using JobId = ObjectId<IJob>;
+	using LeaseId = ObjectId<ILease>;
 
 	class UpgradeTaskSource : TaskSourceBase<UpgradeTask>
 	{
@@ -43,7 +45,7 @@ namespace HordeServer.Tasks.Impl
 				AgentLease? Lease = null;
 				if (Agent.Leases.Count == 0 && (Agent.LastUpgradeTime == null || Agent.LastUpgradeTime.Value + TimeSpan.FromMinutes(5.0) < Clock.UtcNow || Agent.LastUpgradeVersion != RequiredVersion.ToString()))
 				{
-					ILogFile LogFile = await LogService.CreateLogFileAsync(ObjectId.Empty, Agent.SessionId, LogType.Json);
+					ILogFile LogFile = await LogService.CreateLogFileAsync(JobId.Empty, Agent.SessionId, LogType.Json);
 
 					UpgradeTask Task = new UpgradeTask();
 					Task.SoftwareId = RequiredVersion.ToString();
@@ -62,7 +64,7 @@ namespace HordeServer.Tasks.Impl
 						Payload = Any.Pack(Task).ToByteArray();
 					}
 
-					Lease = new AgentLease(ObjectId.GenerateNewId(), $"Upgrade to {RequiredVersion}", null, null, LogFile.Id, LeaseState.Pending, null, true, Payload);
+					Lease = new AgentLease(LeaseId.GenerateNewId(), $"Upgrade to {RequiredVersion}", null, null, LogFile.Id, LeaseState.Pending, null, true, Payload);
 				}
 				return Lease;
 			}
diff --git a/Engine/Source/Programs/Horde/HordeServer/Compute/IComputeService.cs b/Engine/Source/Programs/Horde/HordeServer/Compute/IComputeService.cs
--- a/Engine/Source/Programs/Horde/HordeServer/Compute/IComputeService.cs
+++ b/Engine/Source/Programs/Horde/HordeServer/Compute/IComputeService.cs
@@ -19,6 +19,7 @@ using System.Threading.Tasks;
 namespace HordeServer.Compute
 {
 	using ChannelId = StringId<IComputeChannel>;
+	using LeaseId = ObjectId<ILease>;
 	using NamespaceId = StringId<INamespace>;
 
 	/// <summary>
@@ -62,7 +63,7 @@ namespace HordeServer.Compute
 		/// <summary>
 		/// The lease id
 		/// </summary>
-		public ObjectId? LeaseId { get; }
+		public LeaseId? LeaseId { get; }
 
 		/// <summary>
 		/// Hash of the result, if complete
diff --git a/Engine/Source/Programs/Horde/HordeServer/Compute/Impl/ComputeService.cs b/Engine/Source/Programs/Horde/HordeServer/Compute/Impl/ComputeService.cs
--- a/Engine/Source/Programs/Horde/HordeServer/Compute/Impl/ComputeService.cs
+++ b/Engine/Source/Programs/Horde/HordeServer/Compute/Impl/ComputeService.cs
@@ -38,6 +38,7 @@ using HordeServer.Collections;
 namespace HordeServer.Compute.Impl
 {
 	using ChannelId = StringId<IComputeChannel>;
+	using LeaseId = ObjectId<ILease>;
 	using NamespaceId = StringId<INamespace>;
 
 	/// <summary>
@@ -105,23 +106,23 @@ namespace HordeServer.Compute.Impl
 		AgentId? IComputeTaskStatus.AgentId => AgentId.IsEmpty ? (AgentId?)null : new AgentId(AgentId.ToString());
 
 		/// <inheritdoc/>
-		public ObjectId? LeaseId
+		public LeaseId? LeaseId
 		{
-			get => (LeaseIdBytes == null || LeaseIdBytes.Length == 0)? (ObjectId?)null : new ObjectId(LeaseIdBytes);
-			set => LeaseIdBytes = value?.ToByteArray();
+			get => (LeaseIdBytes == null || LeaseIdBytes.Length == 0)? (LeaseId?)null : new LeaseId(LeaseIdBytes);
+			set => LeaseIdBytes = (value != null) ? value.Value.Value.ToByteArray() : null;
 		}
 
 		private ComputeTaskStatus()
 		{
 		}
 
-		public ComputeTaskStatus(CbObjectAttachment Task, ComputeTaskState State, AgentId? AgentId, ObjectId? LeaseId)
+		public ComputeTaskStatus(CbObjectAttachment Task, ComputeTaskState State, AgentId? AgentId, LeaseId? LeaseId)
 		{
 			this.Task = Task;
 			this.Time = DateTime.UtcNow;
 			this.State = State;
 			this.AgentId = (AgentId == null)? Utf8String.Empty : AgentId.Value.ToString();
-			this.LeaseIdBytes = LeaseId?.ToByteArray();
+			this.LeaseIdBytes = (LeaseId != null)? LeaseId.Value.Value.ToByteArray() : null;
 		}
 	}
 
@@ -238,7 +239,7 @@ namespace HordeServer.Compute.Impl
 					string LeaseName = $"Remote action ({TaskInfo.TaskHash})";
 					byte[] Payload = Any.Pack(ComputeTask).ToByteArray();
 
-					AgentLease Lease = new AgentLease(ObjectId.GenerateNewId(), LeaseName, null, null, null, LeaseState.Pending, Requirements.Resources, Requirements.Exclusive, Payload);
+					AgentLease Lease = new AgentLease(LeaseId.GenerateNewId(), LeaseName, null, null, null, LeaseState.Pending, Requirements.Resources, Requirements.Exclusive, Payload);
 					Logger.LogDebug("Created lease {LeaseId} for channel {ChannelId} task {TaskHash} req {RequirementsHash}", Lease.Id, ComputeTask.ChannelId, ComputeTask.Task.Hash, ComputeTask.Requirements.Hash);
 					return Lease;
 				}
@@ -247,7 +248,7 @@ namespace HordeServer.Compute.Impl
 		}
 
 		/// <inheritdoc/>
-		public override Task CancelLeaseAsync(IAgent Agent, ObjectId LeaseId, ComputeTaskMessage Message)
+		public override Task CancelLeaseAsync(IAgent Agent, LeaseId LeaseId, ComputeTaskMessage Message)
 		{
 			ComputeTaskInfo TaskInfo = new ComputeTaskInfo(Message.Task, new ChannelId(Message.ChannelId));
 			return TaskScheduler.EnqueueAsync(Message.Requirements.Hash, TaskInfo, true);
@@ -266,7 +267,7 @@ namespace HordeServer.Compute.Impl
 			return Messages.ConvertAll<IComputeTaskStatus>(x => x);
 		}
 
-		public override Task OnLeaseStartedAsync(IAgent Agent, ObjectId LeaseId, ComputeTaskMessage ComputeTask, ILogger Logger)
+		public override Task OnLeaseStartedAsync(IAgent Agent, LeaseId LeaseId, ComputeTaskMessage ComputeTask, ILogger Logger)
 		{
 			Logger.LogInformation("Compute lease started (lease: {LeaseId}, task: {TaskHash}, agent: {AgentId}, channel: {ChannelId})", LeaseId, ComputeTask.Task.Hash, Agent.Id, ComputeTask.ChannelId);
 
@@ -274,7 +275,7 @@ namespace HordeServer.Compute.Impl
 			return MessageQueue.PostAsync(ComputeTask.ChannelId, Status);
 		}
 
-		public override Task OnLeaseFinishedAsync(IAgent Agent, ObjectId LeaseId, ComputeTaskMessage ComputeTask, LeaseOutcome Outcome, ReadOnlyMemory<byte> Output, ILogger Logger)
+		public override Task OnLeaseFinishedAsync(IAgent Agent, LeaseId LeaseId, ComputeTaskMessage ComputeTask, LeaseOutcome Outcome, ReadOnlyMemory<byte> Output, ILogger Logger)
 		{
 			ComputeTaskResultMessage Message = ComputeTaskResultMessage.Parser.ParseFrom(Output.ToArray());
 
diff --git a/Engine/Source/Programs/Horde/HordeServer/Controllers/AgentsController.cs b/Engine/Source/Programs/Horde/HordeServer/Controllers/AgentsController.cs
--- a/Engine/Source/Programs/Horde/HordeServer/Controllers/AgentsController.cs
+++ b/Engine/Source/Programs/Horde/HordeServer/Controllers/AgentsController.cs
@@ -18,6 +18,7 @@ using System.Threading.Tasks;
 namespace HordeServer.Controllers
 {
 	using AgentSoftwareChannelName = StringId<AgentSoftwareChannels>;
+	using LeaseId = ObjectId<ILease>;
 	using PoolId = StringId<IPool>;
 
 	/// <summary>
@@ -352,11 +353,9 @@ namespace HordeServer.Controllers
 		/// <returns>Lease matching the given id</returns>
 		[HttpGet]
 		[Route("/api/v1/agents/{AgentId}/leases/{LeaseId}")]
-		public async Task<ActionResult<GetAgentLeaseResponse>> GetLeaseAsync(string AgentId, string LeaseId)
+		public async Task<ActionResult<GetAgentLeaseResponse>> GetLeaseAsync(AgentId AgentId, LeaseId LeaseId)
 		{
-			AgentId AgentIdValue = new AgentId(AgentId);
-
-			IAgent? Agent = await AgentService.GetAgentAsync(AgentIdValue);
+			IAgent? Agent = await AgentService.GetAgentAsync(AgentId);
 			if(Agent == null)
 			{
 				return NotFound();
@@ -366,10 +365,8 @@ namespace HordeServer.Controllers
 				return Forbid();
 			}
 
-			ObjectId LeaseIdValue = LeaseId.ToObjectId();
-
-			ILease? Lease = await AgentService.GetLeaseAsync(LeaseIdValue);
-			if (Lease == null || Lease.AgentId != AgentIdValue)
+			ILease? Lease = await AgentService.GetLeaseAsync(LeaseId);
+			if (Lease == null || Lease.AgentId != AgentId)
 			{
 				return NotFound();
 			}
diff --git a/Engine/Source/Programs/Horde/HordeServer/Controllers/ArtifactsController.cs b/Engine/Source/Programs/Horde/HordeServer/Controllers/ArtifactsController.cs
--- a/Engine/Source/Programs/Horde/HordeServer/Controllers/ArtifactsController.cs
+++ b/Engine/Source/Programs/Horde/HordeServer/Controllers/ArtifactsController.cs
@@ -27,6 +27,8 @@ using System.Threading.Tasks;
 
 namespace HordeServer.Controllers
 {
+	using JobId = ObjectId<IJob>;
+
 	/// <summary>
 	/// Controller for the /api/artifacts endpoint
 	/// </summary>
@@ -75,9 +77,9 @@ namespace HordeServer.Controllers
 		[HttpPost]
 		[Authorize]
 		[Route("/api/v1/artifacts")]
-		public async Task<ActionResult<CreateArtifactResponse>> CreateArtifact([FromQuery]string JobId, [FromQuery]string? StepId, IFormFile File)
+		public async Task<ActionResult<CreateArtifactResponse>> CreateArtifact([FromQuery] JobId JobId, [FromQuery]string? StepId, IFormFile File)
 		{
-			IJob? Job = await JobService.GetJobAsync(JobId.ToObjectId());
+			IJob? Job = await JobService.GetJobAsync(JobId);
 			if(Job == null)
 			{
 				return NotFound();
@@ -146,17 +148,16 @@ namespace HordeServer.Controllers
 		[Authorize]
 		[Route("/api/v1/artifacts")]
 		[ProducesResponseType(typeof(List<GetArtifactResponse>), 200)]
-		public async Task<ActionResult<List<object>>> GetArtifacts([FromQuery] string JobId, [FromQuery] string? StepId = null, [FromQuery] bool Code = false, [FromQuery] PropertyFilter? Filter = null)
+		public async Task<ActionResult<List<object>>> GetArtifacts([FromQuery] JobId JobId, [FromQuery] string? StepId = null, [FromQuery] bool Code = false, [FromQuery] PropertyFilter? Filter = null)
 		{
-			ObjectId JobIdValue = JobId.ToObjectId();
-			if (!await JobService.AuthorizeAsync(JobIdValue, AclAction.DownloadArtifact, User, null))
+			if (!await JobService.AuthorizeAsync(JobId, AclAction.DownloadArtifact, User, null))
 			{
 				return Forbid();
 			}
 
-			string? DownloadCode = Code ? (string?)GetDirectDownloadCodeForJob(JobIdValue) : null;
+			string? DownloadCode = Code ? (string?)GetDirectDownloadCodeForJob(JobId) : null;
 
-			List<IArtifact> Artifacts = await ArtifactCollection.GetArtifactsAsync(JobIdValue, StepId?.ToSubResourceId(), null);
+			List<IArtifact> Artifacts = await ArtifactCollection.GetArtifactsAsync(JobId, StepId?.ToSubResourceId(), null);
 			return Artifacts.ConvertAll(x => new GetArtifactResponse(x, DownloadCode).ApplyFilter(Filter));
 		}
 
@@ -165,7 +166,7 @@ namespace HordeServer.Controllers
 		/// </summary>
 		/// <param name="JobId">The job id</param>
 		/// <returns>The required claim</returns>
-		static Claim GetDirectDownloadClaim(ObjectId JobId)
+		static Claim GetDirectDownloadClaim(JobId JobId)
 		{
 			return new Claim(HordeClaimTypes.JobArtifacts, JobId.ToString());
 		}
@@ -175,7 +176,7 @@ namespace HordeServer.Controllers
 		/// </summary>
 		/// <param name="JobId">The job id</param>
 		/// <returns>The download code</returns>
-		string GetDirectDownloadCodeForJob(ObjectId JobId)
+		string GetDirectDownloadCodeForJob(JobId JobId)
 		{
 			Claim DownloadClaim = GetDirectDownloadClaim(JobId);
 			return AclService.IssueBearerToken(new[] { DownloadClaim }, TimeSpan.FromHours(4.0));
@@ -241,17 +242,16 @@ namespace HordeServer.Controllers
 		/// <returns>Raw artifact data</returns>
 		[HttpGet]
 		[Route("/api/v1/jobs/{JobId}/steps/{StepId}/artifacts/{FileName}/data")]
-		public async Task<ActionResult<object>> GetArtifactDataByFilename(string JobId, string StepId, string Filename)
+		public async Task<ActionResult<object>> GetArtifactDataByFilename(JobId JobId, string StepId, string Filename)
 		{
-			ObjectId JobIdValue = JobId.ToObjectId();
 			SubResourceId StepIdValue = StepId.ToSubResourceId();
 
-			if (!await JobService.AuthorizeAsync(JobIdValue, AclAction.DownloadArtifact, User, null))
+			if (!await JobService.AuthorizeAsync(JobId, AclAction.DownloadArtifact, User, null))
 			{
 				return Forbid();
 			}
 			
-			List<IArtifact> Artifacts = await ArtifactCollection.GetArtifactsAsync(JobIdValue, StepIdValue, Filename);
+			List<IArtifact> Artifacts = await ArtifactCollection.GetArtifactsAsync(JobId, StepIdValue, Filename);
 			if (Artifacts.Count == 0)
 			{
 				return NotFound();
@@ -349,7 +349,7 @@ namespace HordeServer.Controllers
 				return BadRequest("Must specify a JobId");
 			}
 
-			IJob? Job = await JobService.GetJobAsync(ArtifactZipRequest.JobId!.ToObjectId());
+			IJob? Job = await JobService.GetJobAsync(new JobId(ArtifactZipRequest.JobId!));
 			if (Job == null)
 			{
 				return NotFound();
diff --git a/Engine/Source/Programs/Horde/HordeServer/Controllers/DebugController.cs b/Engine/Source/Programs/Horde/HordeServer/Controllers/DebugController.cs
--- a/Engine/Source/Programs/Horde/HordeServer/Controllers/DebugController.cs
+++ b/Engine/Source/Programs/Horde/HordeServer/Controllers/DebugController.cs
@@ -40,6 +40,7 @@ using JetBrains.Profiler.SelfApi;
 
 namespace HordeServer.Controllers
 {
+	using LogId = ObjectId<ILogFile>;
 	using PoolId = StringId<IPool>;
 	using ProjectId = StringId<IProject>;
 	using TemplateRefId = StringId<TemplateRef>;
@@ -530,14 +531,14 @@ namespace HordeServer.Controllers
 		/// <returns>Information about the requested project</returns>
 		[HttpGet]
 		[Route("/api/v1/debug/logs/{LogFileId}")]
-		public async Task<ActionResult<object>> GetLogAsync(string LogFileId, [FromQuery] PropertyFilter? Filter = null)
+		public async Task<ActionResult<object>> GetLogAsync(LogId LogFileId, [FromQuery] PropertyFilter? Filter = null)
 		{
 			if (!await AclService.AuthorizeAsync(AclAction.AdminRead, User))
 			{
 				return Forbid();
 			}
 
-			ILogFile? LogFile = await LogFileCollection.GetLogFileAsync(LogFileId.ToObjectId());
+			ILogFile? LogFile = await LogFileCollection.GetLogFileAsync(LogFileId);
 			if (LogFile == null)
 			{
 				return NotFound();
diff --git a/Engine/Source/Programs/Horde/HordeServer/Controllers/IssuesController.cs b/Engine/Source/Programs/Horde/HordeServer/Controllers/IssuesController.cs
--- a/Engine/Source/Programs/Horde/HordeServer/Controllers/IssuesController.cs
+++ b/Engine/Source/Programs/Horde/HordeServer/Controllers/IssuesController.cs
@@ -21,6 +21,8 @@ using System.Threading.Tasks;
 
 namespace HordeServer.Controllers
 {
+	using JobId = ObjectId<IJob>;
+	using LogId = ObjectId<ILogFile>;
 	using StreamId = StringId<IStream>;
 
 	/// <summary>
@@ -101,7 +103,7 @@ namespace HordeServer.Controllers
 		[HttpGet]
 		[Route("/api/v1/issues")]
 		[ProducesResponseType(typeof(List<GetIssueResponse>), 200)]
-		public async Task<ActionResult<object>> FindIssuesAsync([FromQuery(Name = "Id")] int[]? Ids = null, [FromQuery] string? StreamId = null, [FromQuery] int? Change = null, [FromQuery] int? MinChange = null, [FromQuery] int? MaxChange = null, [FromQuery] string? JobId = null, [FromQuery] string? BatchId = null, [FromQuery] string? StepId = null, [FromQuery(Name = "label")] int? LabelIdx = null, [FromQuery] string? UserId = null, [FromQuery] bool? Resolved = null, [FromQuery] int Index = 0, [FromQuery] int Count = 10, [FromQuery] PropertyFilter? Filter = null)
+		public async Task<ActionResult<object>> FindIssuesAsync([FromQuery(Name = "Id")] int[]? Ids = null, [FromQuery] string? StreamId = null, [FromQuery] int? Change = null, [FromQuery] int? MinChange = null, [FromQuery] int? MaxChange = null, [FromQuery] JobId? JobId = null, [FromQuery] string? BatchId = null, [FromQuery] string? StepId = null, [FromQuery(Name = "label")] int? LabelIdx = null, [FromQuery] string? UserId = null, [FromQuery] bool? Resolved = null, [FromQuery] int Index = 0, [FromQuery] int Count = 10, [FromQuery] PropertyFilter? Filter = null)
 		{
 			if(Ids != null && Ids.Length == 0)
 			{
@@ -127,9 +129,7 @@ namespace HordeServer.Controllers
 			}
 			else
 			{
-				ObjectId JobIdValue = JobId.ToObjectId();
-
-				IJob? Job = await JobService.GetJobAsync(JobIdValue);
+				IJob? Job = await JobService.GetJobAsync(JobId.Value);
 				if (Job == null)
 				{
 					return NotFound();
@@ -289,12 +289,12 @@ namespace HordeServer.Controllers
 		[HttpGet]
 		[Route("/api/v1/issues/{IssueId}/events")]
 		[ProducesResponseType(typeof(List<GetLogEventResponse>), 200)]
-		public async Task<ActionResult<List<object>>> GetIssueEventsAsync(int IssueId, [FromQuery] string? JobId = null, [FromQuery] string? BatchId = null, [FromQuery] string? StepId = null, [FromQuery(Name = "label")] int? LabelIdx = null, [FromQuery] string[]? LogIds = null, [FromQuery] int Index = 0, [FromQuery] int Count = 10, [FromQuery] PropertyFilter? Filter = null)
+		public async Task<ActionResult<List<object>>> GetIssueEventsAsync(int IssueId, [FromQuery] JobId? JobId = null, [FromQuery] string? BatchId = null, [FromQuery] string? StepId = null, [FromQuery(Name = "label")] int? LabelIdx = null, [FromQuery] string[]? LogIds = null, [FromQuery] int Index = 0, [FromQuery] int Count = 10, [FromQuery] PropertyFilter? Filter = null)
 		{
-			HashSet<ObjectId> LogIdValues = new HashSet<ObjectId>();
+			HashSet<LogId> LogIdValues = new HashSet<LogId>();
 			if(JobId != null)
 			{
-				IJob? Job = await JobService.GetJobAsync(JobId.ToObjectId());
+				IJob? Job = await JobService.GetJobAsync(JobId.Value);
 				if(Job == null)
 				{
 					return NotFound();
@@ -341,13 +341,13 @@ namespace HordeServer.Controllers
 			}
 			if(LogIds != null)
 			{
-				LogIdValues.UnionWith(LogIds.Select(x => x.ToObjectId()));
+				LogIdValues.UnionWith(LogIds.Select(x => new LogId(x)));
 			}
 
 			List<ILogEvent> Events = await IssueService.FindEventsForIssueAsync(IssueId, LogIdValues.ToArray(), Index, Count);
 
 			JobPermissionsCache PermissionsCache = new JobPermissionsCache();
-			Dictionary<ObjectId, ILogFile?> LogFiles = new Dictionary<ObjectId, ILogFile?>();
+			Dictionary<LogId, ILogFile?> LogFiles = new Dictionary<LogId, ILogFile?>();
 
 			List<object> Responses = new List<object>();
 			foreach (ILogEvent Event in Events)
diff --git a/Engine/Source/Programs/Horde/HordeServer/Controllers/JobsController.cs b/Engine/Source/Programs/Horde/HordeServer/Controllers/JobsController.cs
--- a/Engine/Source/Programs/Horde/HordeServer/Controllers/JobsController.cs
+++ b/Engine/Source/Programs/Horde/HordeServer/Controllers/JobsController.cs
@@ -27,6 +27,7 @@ using HordeServer.Jobs;
 
 namespace HordeServer.Controllers
 {
+	using JobId = ObjectId<IJob>;
 	using StreamId = StringId<IStream>;
 	using TemplateRefId = StringId<TemplateRef>;
 
@@ -190,7 +191,7 @@ namespace HordeServer.Controllers
 
 			// Create the job
 			IJob Job = await JobService.CreateJobAsync(null, Stream, TemplateRefId, Template.Id, Graph, Name, Change, CodeChange, Create.PreflightChange, null, User.GetUserId(), User.GetUserName(), Priority, Create.AutoSubmit, Create.UpdateIssues, TemplateRef.ChainedJobs, TemplateRef.ShowUgsBadges, TemplateRef.ShowUgsAlerts, TemplateRef.NotificationChannel, TemplateRef.NotificationChannelFilter, null, Arguments);
-			await UpdateNotificationsAsync(Job.Id.ToString(), new UpdateNotificationsRequest { Slack = true });
+			await UpdateNotificationsAsync(Job.Id, new UpdateNotificationsRequest { Slack = true });
 			return new CreateJobResponse(Job.Id.ToString());
 		}
 
@@ -224,9 +225,9 @@ namespace HordeServer.Controllers
 		/// <returns>Async task</returns>
 		[HttpDelete]
 		[Route("/api/v1/jobs/{JobId}")]
-		public async Task<ActionResult> DeleteJobAsync(string JobId)
+		public async Task<ActionResult> DeleteJobAsync(JobId JobId)
 		{
-			IJob? Job = await JobService.GetJobAsync(JobId.ToObjectId());
+			IJob? Job = await JobService.GetJobAsync(JobId);
 			if (Job == null)
 			{
 				return NotFound();
@@ -250,9 +251,9 @@ namespace HordeServer.Controllers
 		/// <returns>Async task</returns>
 		[HttpPut]
 		[Route("/api/v1/jobs/{JobId}")]
-		public async Task<ActionResult> UpdateJobAsync(string JobId, [FromBody] UpdateJobRequest Request)
+		public async Task<ActionResult> UpdateJobAsync(JobId JobId, [FromBody] UpdateJobRequest Request)
 		{
-			IJob? Job = await JobService.GetJobAsync(JobId.ToObjectId());
+			IJob? Job = await JobService.GetJobAsync(JobId);
 			if (Job == null)
 			{
 				return NotFound();
@@ -296,9 +297,9 @@ namespace HordeServer.Controllers
 		/// <returns>Information about the requested job</returns>
 		[HttpPut]
 		[Route("/api/v1/jobs/{JobId}/notifications")]
-		public async Task<ActionResult> UpdateNotificationsAsync(string JobId, [FromBody] UpdateNotificationsRequest Request)
+		public async Task<ActionResult> UpdateNotificationsAsync(JobId JobId, [FromBody] UpdateNotificationsRequest Request)
 		{
-			IJob? Job = await JobService.GetJobAsync(JobId.ToObjectId());
+			IJob? Job = await JobService.GetJobAsync(JobId);
 			if (Job == null)
 			{
 				return NotFound();
@@ -325,9 +326,9 @@ namespace HordeServer.Controllers
 		/// <returns>Information about the requested job</returns>
 		[HttpGet]
 		[Route("/api/v1/jobs/{JobId}/notifications")]
-		public async Task<ActionResult<GetNotificationResponse>> GetNotificationsAsync(string JobId)
+		public async Task<ActionResult<GetNotificationResponse>> GetNotificationsAsync(JobId JobId)
 		{
-			IJob? Job = await JobService.GetJobAsync(JobId.ToObjectId());
+			IJob? Job = await JobService.GetJobAsync(JobId);
 			if (Job == null)
 			{
 				return NotFound();
@@ -359,9 +360,9 @@ namespace HordeServer.Controllers
 		[HttpGet]
 		[Route("/api/v1/jobs/{JobId}")]
 		[ProducesResponseType(typeof(GetJobResponse), 200)]
-		public async Task<ActionResult<object>> GetJobAsync(string JobId, [FromQuery] DateTimeOffset? ModifiedAfter = null, [FromQuery] PropertyFilter? Filter = null)
+		public async Task<ActionResult<object>> GetJobAsync(JobId JobId, [FromQuery] DateTimeOffset? ModifiedAfter = null, [FromQuery] PropertyFilter? Filter = null)
 		{
-			IJob? Job = await JobService.GetJobAsync(JobId.ToObjectId());
+			IJob? Job = await JobService.GetJobAsync(JobId);
 			if (Job == null)
 			{
 				return NotFound();
@@ -391,9 +392,9 @@ namespace HordeServer.Controllers
 		[HttpGet]
 		[Route("/api/v1/jobs/{JobId}/graph")]
 		[ProducesResponseType(typeof(GetGraphResponse), 200)]
-		public async Task<ActionResult<object>> GetJobGraphAsync(string JobId, [FromQuery] PropertyFilter? Filter = null)
+		public async Task<ActionResult<object>> GetJobGraphAsync(JobId JobId, [FromQuery] PropertyFilter? Filter = null)
 		{
-			IJob? Job = await JobService.GetJobAsync(JobId.ToObjectId());
+			IJob? Job = await JobService.GetJobAsync(JobId);
 			if (Job == null)
 			{
 				return NotFound();
@@ -416,9 +417,9 @@ namespace HordeServer.Controllers
 		[HttpGet]
 		[Route("/api/v1/jobs/{JobId}/timing")]
 		[ProducesResponseType(typeof(GetJobTimingResponse), 200)]
-		public async Task<ActionResult<object>> GetJobTimingAsync(string JobId, [FromQuery] PropertyFilter? Filter = null)
+		public async Task<ActionResult<object>> GetJobTimingAsync(JobId JobId, [FromQuery] PropertyFilter? Filter = null)
 		{
-			IJob? Job = await JobService.GetJobAsync(JobId.ToObjectId());
+			IJob? Job = await JobService.GetJobAsync(JobId);
 			if (Job == null)
 			{
 				return NotFound();
@@ -462,9 +463,9 @@ namespace HordeServer.Controllers
 		[HttpGet]
 		[Route("/api/v1/jobs/{JobId}/template")]
 		[ProducesResponseType(typeof(GetTemplateResponse), 200)]
-		public async Task<ActionResult<object>> GetJobTemplateAsync(string JobId, [FromQuery] PropertyFilter? Filter = null)
+		public async Task<ActionResult<object>> GetJobTemplateAsync(JobId JobId, [FromQuery] PropertyFilter? Filter = null)
 		{
-			IJob? Job = await JobService.GetJobAsync(JobId.ToObjectId());
+			IJob? Job = await JobService.GetJobAsync(JobId);
 			if (Job == null || Job.TemplateHash == null)
 			{
 				return NotFound();
@@ -532,7 +533,7 @@ namespace HordeServer.Controllers
 			[FromQuery] int Index = 0,
 			[FromQuery] int Count = 100)
 		{
-			ObjectId[]? JobIdValues = (Ids == null) ? (ObjectId[]?)null : Array.ConvertAll(Ids, x => x.ToObjectId());
+			JobId[]? JobIdValues = (Ids == null) ? (JobId[]?)null : Array.ConvertAll(Ids, x => new JobId(x));
 			StreamId? StreamIdValue = (StreamId == null)? (StreamId?)null : new StreamId(StreamId);
 			
 			TemplateRefId[]? TemplateRefIds = (Templates != null && Templates.Length > 0) ? Templates.Select(x => new TemplateRefId(x)).ToArray() : null;
@@ -609,14 +610,13 @@ namespace HordeServer.Controllers
 		/// <returns>Id of the new job</returns>
 		[HttpPost]
 		[Route("/api/v1/jobs/{JobId}/groups")]
-		public async Task<ActionResult> CreateGroupsAsync(string JobId, [FromBody] List<NewGroup> Requests)
+		public async Task<ActionResult> CreateGroupsAsync(JobId JobId, [FromBody] List<NewGroup> Requests)
 		{
 			Dictionary<string, int> ExpectedDurationCache = new Dictionary<string, int>();
 
-			ObjectId JobIdValue = JobId.ToObjectId();
 			for (; ; )
 			{
-				IJob? Job = await JobService.GetJobAsync(JobIdValue);
+				IJob? Job = await JobService.GetJobAsync(JobId);
 				if (Job == null)
 				{
 					return NotFound();
@@ -645,9 +645,9 @@ namespace HordeServer.Controllers
 		[HttpGet]
 		[Route("/api/v1/jobs/{JobId}/groups")]
 		[ProducesResponseType(typeof(List<GetGroupResponse>), 200)]
-		public async Task<ActionResult<List<object>>> GetGroupsAsync(string JobId, [FromQuery] PropertyFilter? Filter = null)
+		public async Task<ActionResult<List<object>>> GetGroupsAsync(JobId JobId, [FromQuery] PropertyFilter? Filter = null)
 		{
-			IJob? Job = await JobService.GetJobAsync(JobId.ToObjectId());
+			IJob? Job = await JobService.GetJobAsync(JobId);
 			if (Job == null)
 			{
 				return NotFound();
@@ -671,9 +671,9 @@ namespace HordeServer.Controllers
 		[HttpGet]
 		[Route("/api/v1/jobs/{JobId}/groups/{GroupIdx}")]
 		[ProducesResponseType(typeof(GetGroupResponse), 200)]
-		public async Task<ActionResult<object>> GetGroupAsync(string JobId, int GroupIdx, [FromQuery] PropertyFilter? Filter = null)
+		public async Task<ActionResult<object>> GetGroupAsync(JobId JobId, int GroupIdx, [FromQuery] PropertyFilter? Filter = null)
 		{
-			IJob? Job = await JobService.GetJobAsync(JobId.ToObjectId());
+			IJob? Job = await JobService.GetJobAsync(JobId);
 			if (Job == null)
 			{
 				return NotFound();
@@ -701,9 +701,9 @@ namespace HordeServer.Controllers
 		[HttpGet]
 		[Route("/api/v1/jobs/{JobId}/groups/{GroupIdx}/nodes")]
 		[ProducesResponseType(typeof(List<GetNodeResponse>), 200)]
-		public async Task<ActionResult<List<object>>> GetNodesAsync(string JobId, int GroupIdx, [FromQuery] PropertyFilter? Filter = null)
+		public async Task<ActionResult<List<object>>> GetNodesAsync(JobId JobId, int GroupIdx, [FromQuery] PropertyFilter? Filter = null)
 		{
-			IJob? Job = await JobService.GetJobAsync(JobId.ToObjectId());
+			IJob? Job = await JobService.GetJobAsync(JobId);
 			if (Job == null)
 			{
 				return NotFound();
@@ -732,9 +732,9 @@ namespace HordeServer.Controllers
 		[HttpGet]
 		[Route("/api/v1/jobs/{JobId}/groups/{GroupIdx}/nodes/{NodeIdx}")]
 		[ProducesResponseType(typeof(GetNodeResponse), 200)]
-		public async Task<ActionResult<object>> GetNodeAsync(string JobId, int GroupIdx, int NodeIdx, [FromQuery] PropertyFilter? Filter = null)
+		public async Task<ActionResult<object>> GetNodeAsync(JobId JobId, int GroupIdx, int NodeIdx, [FromQuery] PropertyFilter? Filter = null)
 		{
-			IJob? Job = await JobService.GetJobAsync(JobId.ToObjectId());
+			IJob? Job = await JobService.GetJobAsync(JobId);
 			if (Job == null)
 			{
 				return NotFound();
@@ -762,11 +762,9 @@ namespace HordeServer.Controllers
 		[HttpGet]
 		[Route("/api/v1/jobs/{JobId}/batches")]
 		[ProducesResponseType(typeof(List<GetBatchResponse>), 200)]
-		public async Task<ActionResult<List<object>>> GetBatchesAsync(string JobId, [FromQuery] PropertyFilter? Filter = null)
+		public async Task<ActionResult<List<object>>> GetBatchesAsync(JobId JobId, [FromQuery] PropertyFilter? Filter = null)
 		{
-			ObjectId JobIdValue = JobId.ToObjectId();
-
-			IJob? Job = await JobService.GetJobAsync(JobIdValue);
+			IJob? Job = await JobService.GetJobAsync(JobId);
 			if (Job == null)
 			{
 				return NotFound();
@@ -794,12 +792,11 @@ namespace HordeServer.Controllers
 		/// <param name="Request">Updates to apply to the node</param>
 		[HttpPut]
 		[Route("/api/v1/jobs/{JobId}/batches/{BatchId}")]
-		public async Task<ActionResult> UpdateBatchAsync(string JobId, string BatchId, [FromBody] UpdateBatchRequest Request)
+		public async Task<ActionResult> UpdateBatchAsync(JobId JobId, string BatchId, [FromBody] UpdateBatchRequest Request)
 		{
-			ObjectId JobIdValue = JobId.ToObjectId();
 			SubResourceId BatchIdValue = BatchId.ToSubResourceId();
 
-			IJob? Job = await JobService.GetJobAsync(JobIdValue);
+			IJob? Job = await JobService.GetJobAsync(JobId);
 			if (Job == null)
 			{
 				return NotFound();
@@ -815,7 +812,7 @@ namespace HordeServer.Controllers
 				return Forbid();
 			}
 
-			if (!await JobService.UpdateBatchAsync(Job, BatchIdValue, Request.LogId?.ToObjectId(), Request.State))
+			if (!await JobService.UpdateBatchAsync(Job, BatchIdValue, Request.LogId?.ToObjectId<ILogFile>(), Request.State))
 			{
 				return NotFound();
 			}
@@ -832,12 +829,11 @@ namespace HordeServer.Controllers
 		[HttpGet]
 		[Route("/api/v1/jobs/{JobId}/batches/{BatchId}")]
 		[ProducesResponseType(typeof(GetBatchResponse), 200)]
-		public async Task<ActionResult<object>> GetBatchAsync(string JobId, string BatchId, [FromQuery] PropertyFilter? Filter = null)
+		public async Task<ActionResult<object>> GetBatchAsync(JobId JobId, string BatchId, [FromQuery] PropertyFilter? Filter = null)
 		{
-			ObjectId JobIdValue = JobId.ToObjectId();
 			SubResourceId BatchIdValue = BatchId.ToSubResourceId();
 
-			IJob? Job = await JobService.GetJobAsync(JobIdValue);
+			IJob? Job = await JobService.GetJobAsync(JobId);
 			if (Job == null)
 			{
 				return NotFound();
@@ -869,12 +865,11 @@ namespace HordeServer.Controllers
 		[HttpGet]
 		[Route("/api/v1/jobs/{JobId}/batches/{BatchId}/steps")]
 		[ProducesResponseType(typeof(List<GetStepResponse>), 200)]
-		public async Task<ActionResult<List<object>>> GetStepsAsync(string JobId, string BatchId, [FromQuery] PropertyFilter? Filter = null)
+		public async Task<ActionResult<List<object>>> GetStepsAsync(JobId JobId, string BatchId, [FromQuery] PropertyFilter? Filter = null)
 		{
-			ObjectId JobIdValue = JobId.ToObjectId();
 			SubResourceId BatchIdValue = BatchId.ToSubResourceId();
 
-			IJob? Job = await JobService.GetJobAsync(JobIdValue);
+			IJob? Job = await JobService.GetJobAsync(JobId);
 			if (Job == null)
 			{
 				return NotFound();
@@ -912,13 +907,12 @@ namespace HordeServer.Controllers
 		/// <param name="Request">Updates to apply to the node</param>
 		[HttpPut]
 		[Route("/api/v1/jobs/{JobId}/batches/{BatchId}/steps/{StepId}")]
-		public async Task<ActionResult<UpdateStepResponse>> UpdateStepAsync(string JobId, string BatchId, string StepId, [FromBody] UpdateStepRequest Request)
+		public async Task<ActionResult<UpdateStepResponse>> UpdateStepAsync(JobId JobId, string BatchId, string StepId, [FromBody] UpdateStepRequest Request)
 		{
-			ObjectId JobIdValue = JobId.ToObjectId();
 			SubResourceId BatchIdValue = BatchId.ToSubResourceId();
 			SubResourceId StepIdValue = StepId.ToSubResourceId();
 
-			IJob? Job = await JobService.GetJobAsync(JobIdValue);
+			IJob? Job = await JobService.GetJobAsync(JobId);
 			if (Job == null)
 			{
 				return NotFound();
@@ -957,7 +951,7 @@ namespace HordeServer.Controllers
 
 			try
 			{
-				IJob? NewJob = await JobService.UpdateStepAsync(Job, BatchIdValue, StepIdValue, Request.State, Request.Outcome, Request.AbortRequested, AbortByUser, Request.LogId?.ToObjectId(), null, RetryByUser, Request.Priority, null, Request.Properties);
+				IJob? NewJob = await JobService.UpdateStepAsync(Job, BatchIdValue, StepIdValue, Request.State, Request.Outcome, Request.AbortRequested, AbortByUser, Request.LogId?.ToObjectId<ILogFile>(), null, RetryByUser, Request.Priority, null, Request.Properties);
 				if (NewJob == null)
 				{
 					return NotFound();
@@ -1022,13 +1016,12 @@ namespace HordeServer.Controllers
 		[HttpGet]
 		[Route("/api/v1/jobs/{JobId}/batches/{BatchId}/steps/{StepId}")]
 		[ProducesResponseType(typeof(GetStepResponse), 200)]
-		public async Task<ActionResult<object>> GetStepAsync(string JobId, string BatchId, string StepId, [FromQuery] PropertyFilter? Filter = null)
+		public async Task<ActionResult<object>> GetStepAsync(JobId JobId, string BatchId, string StepId, [FromQuery] PropertyFilter? Filter = null)
 		{
-			ObjectId JobIdValue = JobId.ToObjectId();
 			SubResourceId BatchIdValue = BatchId.ToSubResourceId();
 			SubResourceId StepIdValue = StepId.ToSubResourceId();
 
-			IJob? Job = await JobService.GetJobAsync(JobIdValue);
+			IJob? Job = await JobService.GetJobAsync(JobId);
 			if (Job == null)
 			{
 				return NotFound();
@@ -1067,13 +1060,12 @@ namespace HordeServer.Controllers
 		/// <returns>Information about the requested job</returns>
 		[HttpPut]
 		[Route("/api/v1/jobs/{JobId}/batches/{BatchId}/steps/{StepId}/notifications")]
-		public async Task<ActionResult> UpdateStepNotificationsAsync(string JobId, string BatchId, string StepId, [FromBody] UpdateNotificationsRequest Request)
+		public async Task<ActionResult> UpdateStepNotificationsAsync(JobId JobId, string BatchId, string StepId, [FromBody] UpdateNotificationsRequest Request)
 		{
-			ObjectId JobIdValue = JobId.ToObjectId();
 			SubResourceId BatchIdValue = BatchId.ToSubResourceId();
 			SubResourceId StepIdValue = StepId.ToSubResourceId();
 
-			IJob? Job = await JobService.GetJobAsync(JobIdValue);
+			IJob? Job = await JobService.GetJobAsync(JobId);
 			if (Job == null)
 			{
 				return NotFound();
@@ -1116,9 +1108,9 @@ namespace HordeServer.Controllers
 		/// <returns>Information about the requested job</returns>
 		[HttpGet]
 		[Route("/api/v1/jobs/{JobId}/batches/{BatchId}/steps/{StepId}/notifications")]
-		public async Task<ActionResult<GetNotificationResponse>> GetStepNotificationsAsync(string JobId, string BatchId, string StepId)
+		public async Task<ActionResult<GetNotificationResponse>> GetStepNotificationsAsync(JobId JobId, string BatchId, string StepId)
 		{
-			IJob? Job = await JobService.GetJobAsync(JobId.ToObjectId());
+			IJob? Job = await JobService.GetJobAsync(JobId);
 			if (Job == null)
 			{
 				return NotFound();
@@ -1156,13 +1148,12 @@ namespace HordeServer.Controllers
 		/// <returns>List of nodes to be executed</returns>
 		[HttpGet]
 		[Route("/api/v1/jobs/{JobId}/batches/{BatchId}/steps/{StepId}/artifacts/{*Name}")]
-		public async Task<ActionResult> GetArtifactAsync(string JobId, string BatchId, string StepId, string Name)
+		public async Task<ActionResult> GetArtifactAsync(JobId JobId, string BatchId, string StepId, string Name)
 		{
-			ObjectId JobIdValue = JobId.ToObjectId();
 			SubResourceId BatchIdValue = BatchId.ToSubResourceId();
 			SubResourceId StepIdValue = StepId.ToSubResourceId();
 
-			IJob? Job = await JobService.GetJobAsync(JobIdValue);
+			IJob? Job = await JobService.GetJobAsync(JobId);
 			if (Job == null)
 			{
 				return NotFound();
@@ -1176,7 +1167,7 @@ namespace HordeServer.Controllers
 				return NotFound();
 			}
 
-			List<IArtifact> Artifacts = await ArtifactCollection.GetArtifactsAsync(JobIdValue, StepIdValue, Name);
+			List<IArtifact> Artifacts = await ArtifactCollection.GetArtifactsAsync(JobId, StepIdValue, Name);
 			if (Artifacts.Count == 0)
 			{
 				return NotFound();
@@ -1195,13 +1186,12 @@ namespace HordeServer.Controllers
 		/// <returns>List of nodes to be executed</returns>
 		[HttpGet]
 		[Route("/api/v1/jobs/{JobId}/batches/{BatchId}/steps/{StepId}/trace")]
-		public async Task<ActionResult> GetStepTraceAsync(string JobId, string BatchId, string StepId)
+		public async Task<ActionResult> GetStepTraceAsync(JobId JobId, string BatchId, string StepId)
 		{
-			ObjectId JobIdValue = JobId.ToObjectId();
 			SubResourceId BatchIdValue = BatchId.ToSubResourceId();
 			SubResourceId StepIdValue = StepId.ToSubResourceId();
 
-			IJob? Job = await JobService.GetJobAsync(JobIdValue);
+			IJob? Job = await JobService.GetJobAsync(JobId);
 			if (Job == null)
 			{
 				return NotFound();
@@ -1215,7 +1205,7 @@ namespace HordeServer.Controllers
 				return NotFound();
 			}
 
-			List<IArtifact> Artifacts = await ArtifactCollection.GetArtifactsAsync(JobIdValue, StepIdValue, null);
+			List<IArtifact> Artifacts = await ArtifactCollection.GetArtifactsAsync(JobId, StepIdValue, null);
 			foreach (IArtifact Artifact in Artifacts)
 			{
 				if (Artifact.Name.Equals("trace.json", StringComparison.OrdinalIgnoreCase))
@@ -1234,14 +1224,12 @@ namespace HordeServer.Controllers
 		/// <param name="Request">The notification request</param>
 		[HttpPut]
 		[Route("/api/v1/jobs/{JobId}/labels/{LabelIndex}/notifications")]
-		public async Task<ActionResult> UpdateLabelNotificationsAsync(string JobId, int LabelIndex, [FromBody] UpdateNotificationsRequest Request)
+		public async Task<ActionResult> UpdateLabelNotificationsAsync(JobId JobId, int LabelIndex, [FromBody] UpdateNotificationsRequest Request)
 		{
-			ObjectId JobIdValue = JobId.ToObjectId();
-
 			ObjectId TriggerId;
 			for (; ; )
 			{
-				IJob? Job = await JobService.GetJobAsync(JobIdValue);
+				IJob? Job = await JobService.GetJobAsync(JobId);
 				if (Job == null)
 				{
 					return NotFound();
@@ -1278,9 +1266,9 @@ namespace HordeServer.Controllers
 		/// <returns>Notification info for the requested label in the job</returns>
 		[HttpGet]
 		[Route("/api/v1/jobs/{JobId}/labels/{LabelIndex}/notifications")]
-		public async Task<ActionResult<GetNotificationResponse>> GetLabelNotificationsAsync(string JobId, int LabelIndex)
+		public async Task<ActionResult<GetNotificationResponse>> GetLabelNotificationsAsync(JobId JobId, int LabelIndex)
 		{
-			IJob? Job = await JobService.GetJobAsync(JobId.ToObjectId());
+			IJob? Job = await JobService.GetJobAsync(JobId);
 			if (Job == null)
 			{
 				return NotFound();
diff --git a/Engine/Source/Programs/Horde/HordeServer/Controllers/LeasesController.cs b/Engine/Source/Programs/Horde/HordeServer/Controllers/LeasesController.cs
--- a/Engine/Source/Programs/Horde/HordeServer/Controllers/LeasesController.cs
+++ b/Engine/Source/Programs/Horde/HordeServer/Controllers/LeasesController.cs
@@ -24,6 +24,7 @@ using Microsoft.Extensions.Logging;
 
 namespace HordeServer.Controllers
 {
+	using LeaseId = ObjectId<ILease>;
 	using PoolId = StringId<IPool>;
 	using AgentSoftwareChannelName = StringId<AgentSoftwareChannels>;
 
@@ -111,11 +112,9 @@ namespace HordeServer.Controllers
 		/// <returns>Lease matching the given id</returns>
 		[HttpGet]
 		[Route("/api/v1/leases/{LeaseId}")]
-		public async Task<ActionResult<GetAgentLeaseResponse>> GetLeaseAsync(string LeaseId)
+		public async Task<ActionResult<GetAgentLeaseResponse>> GetLeaseAsync(LeaseId LeaseId)
 		{
-			ObjectId LeaseIdValue = LeaseId.ToObjectId();
-
-			ILease? Lease = await AgentService.GetLeaseAsync(LeaseIdValue);
+			ILease? Lease = await AgentService.GetLeaseAsync(LeaseId);
 			if (Lease == null)
 			{
 				return NotFound();
@@ -142,17 +141,15 @@ namespace HordeServer.Controllers
 		/// <returns>Lease matching the given id</returns>
 		[HttpPut]
 		[Route("/api/v1/leases/{LeaseId}")]
-		public async Task<ActionResult> UpdateLeaseAsync(string LeaseId, [FromBody] UpdateLeaseRequest Request)
+		public async Task<ActionResult> UpdateLeaseAsync(LeaseId LeaseId, [FromBody] UpdateLeaseRequest Request)
 		{
-			ObjectId LeaseIdValue = LeaseId.ToObjectId();
-
 			// only update supported right now is abort
 			if (!Request.Aborted.HasValue || !Request.Aborted.Value)
 			{
 				return Ok();
 			}
 
-			ILease? Lease = await AgentService.GetLeaseAsync(LeaseIdValue);
+			ILease? Lease = await AgentService.GetLeaseAsync(LeaseId);
 			if (Lease == null)
 			{
 				return NotFound();
@@ -169,7 +166,7 @@ namespace HordeServer.Controllers
 				return Forbid();
 			}
 
-			AgentLease? AgentLease = Agent.Leases.FirstOrDefault(x => x.Id == LeaseIdValue);
+			AgentLease? AgentLease = Agent.Leases.FirstOrDefault(x => x.Id == LeaseId);
 
 			if (AgentLease == null)
 			{
@@ -182,7 +179,7 @@ namespace HordeServer.Controllers
 				return Ok();
 			}
 
-			await AgentService.CancelLeaseAsync(Agent, LeaseIdValue);
+			await AgentService.CancelLeaseAsync(Agent, LeaseId);
 
 			return Ok();
 		}
diff --git a/Engine/Source/Programs/Horde/HordeServer/Controllers/LogsController.cs b/Engine/Source/Programs/Horde/HordeServer/Controllers/LogsController.cs
--- a/Engine/Source/Programs/Horde/HordeServer/Controllers/LogsController.cs
+++ b/Engine/Source/Programs/Horde/HordeServer/Controllers/LogsController.cs
@@ -24,6 +24,9 @@ using MongoDB.Bson;
 
 namespace HordeServer.Controllers
 {
+	using JobId = ObjectId<IJob>;
+	using LogId = ObjectId<ILogFile>;
+
 	/// <summary>
 	/// Format for the returned data
 	/// </summary>
@@ -92,9 +95,9 @@ namespace HordeServer.Controllers
 		[HttpGet]
 		[Route("/api/v1/logs/{LogFileId}")]
 		[ProducesResponseType(typeof(GetLogFileResponse), 200)]
-		public async Task<ActionResult<object>> GetLog(string LogFileId, [FromQuery] PropertyFilter? Filter = null)
+		public async Task<ActionResult<object>> GetLog(LogId LogFileId, [FromQuery] PropertyFilter? Filter = null)
 		{
-			ILogFile? LogFile = await LogFileService.GetLogFileAsync(LogFileId.ToObjectId());
+			ILogFile? LogFile = await LogFileService.GetLogFileAsync(LogFileId);
 			if (LogFile == null)
 			{
 				return NotFound();
@@ -120,9 +123,9 @@ namespace HordeServer.Controllers
 		/// <returns>Raw log data for the requested range</returns>
 		[HttpGet]
 		[Route("/api/v1/logs/{LogFileId}/data")]
-		public async Task<ActionResult> GetLogData(string LogFileId, [FromQuery] LogOutputFormat Format = LogOutputFormat.Raw, [FromQuery] long Offset = 0, [FromQuery] long Length = long.MaxValue, [FromQuery] string? FileName = null, [FromQuery] bool Download = false)
+		public async Task<ActionResult> GetLogData(LogId LogFileId, [FromQuery] LogOutputFormat Format = LogOutputFormat.Raw, [FromQuery] long Offset = 0, [FromQuery] long Length = long.MaxValue, [FromQuery] string? FileName = null, [FromQuery] bool Download = false)
 		{
-			ILogFile? LogFile = await LogFileService.GetLogFileAsync(LogFileId.ToObjectId());
+			ILogFile? LogFile = await LogFileService.GetLogFileAsync(LogFileId);
 			if (LogFile == null)
 			{
 				return NotFound();
@@ -154,9 +157,9 @@ namespace HordeServer.Controllers
 		/// <returns>Information about the requested project</returns>
 		[HttpGet]
 		[Route("/api/v1/logs/{LogFileId}/lines")]
-		public async Task<ActionResult> GetLogLines(string LogFileId, [FromQuery] int Index = 0, [FromQuery] int Count = int.MaxValue)
+		public async Task<ActionResult> GetLogLines(LogId LogFileId, [FromQuery] int Index = 0, [FromQuery] int Count = int.MaxValue)
 		{
-			ILogFile? LogFile = await LogFileService.GetLogFileAsync(LogFileId.ToObjectId());
+			ILogFile? LogFile = await LogFileService.GetLogFileAsync(LogFileId);
 			if (LogFile == null)
 			{
 				return NotFound();
@@ -284,9 +287,9 @@ namespace HordeServer.Controllers
 		/// <returns>Raw log data for the requested range</returns>
 		[HttpGet]
 		[Route("/api/v1/logs/{LogFileId}/search")]
-		public async Task<ActionResult<SearchLogFileResponse>> SearchLogFileAsync(string LogFileId, [FromQuery] string Text, [FromQuery] int FirstLine = 0, [FromQuery] int Count = 5)
+		public async Task<ActionResult<SearchLogFileResponse>> SearchLogFileAsync(LogId LogFileId, [FromQuery] string Text, [FromQuery] int FirstLine = 0, [FromQuery] int Count = 5)
 		{
-			ILogFile? LogFile = await LogFileService.GetLogFileAsync(LogFileId.ToObjectId());
+			ILogFile? LogFile = await LogFileService.GetLogFileAsync(LogFileId);
 			if (LogFile == null)
 			{
 				return NotFound();
@@ -312,9 +315,9 @@ namespace HordeServer.Controllers
 		[HttpGet]
 		[Route("/api/v1/logs/{LogFileId}/events")]
 		[ProducesResponseType(typeof(List<GetLogEventResponse>), 200)]
-		public async Task<ActionResult<List<GetLogEventResponse>>> GetEventsAsync(string LogFileId, [FromQuery] int? Index = null, [FromQuery] int? Count = null)
+		public async Task<ActionResult<List<GetLogEventResponse>>> GetEventsAsync(LogId LogFileId, [FromQuery] int? Index = null, [FromQuery] int? Count = null)
 		{
-			ILogFile? LogFile = await LogFileService.GetLogFileAsync(LogFileId.ToObjectId());
+			ILogFile? LogFile = await LogFileService.GetLogFileAsync(LogFileId);
 			if (LogFile == null)
 			{
 				return NotFound();
@@ -355,9 +358,9 @@ namespace HordeServer.Controllers
 		/// <returns>Http result code</returns>
 		[HttpPost]
 		[Route("/api/v1/logs/{LogFileId}")]
-		public async Task<ActionResult> WriteData(string LogFileId, [FromQuery] long Offset, [FromQuery] int LineIndex)
+		public async Task<ActionResult> WriteData(LogId LogFileId, [FromQuery] long Offset, [FromQuery] int LineIndex)
 		{
-			ILogFile? LogFile = await LogFileService.GetLogFileAsync(LogFileId.ToObjectId());
+			ILogFile? LogFile = await LogFileService.GetLogFileAsync(LogFileId);
 			if (LogFile == null)
 			{
 				return NotFound();
@@ -385,7 +388,7 @@ namespace HordeServer.Controllers
 		/// <returns>True if the action is authorized</returns>
 		async Task<bool> AuthorizeAsync(ILogFile LogFile, AclAction Action, ClaimsPrincipal User, JobPermissionsCache? PermissionsCache)
 		{
-			if (LogFile.JobId != ObjectId.Empty && await JobService.AuthorizeAsync(LogFile.JobId, Action, User, PermissionsCache))
+			if (LogFile.JobId != JobId.Empty && await JobService.AuthorizeAsync(LogFile.JobId, Action, User, PermissionsCache))
 			{
 				return true;
 			}
diff --git a/Engine/Source/Programs/Horde/HordeServer/Controllers/TestDataController.cs b/Engine/Source/Programs/Horde/HordeServer/Controllers/TestDataController.cs
--- a/Engine/Source/Programs/Horde/HordeServer/Controllers/TestDataController.cs
+++ b/Engine/Source/Programs/Horde/HordeServer/Controllers/TestDataController.cs
@@ -20,6 +20,7 @@ using System.Threading.Tasks;
 
 namespace HordeServer.Controllers
 {
+	using JobId = ObjectId<IJob>;
 	using StreamId = StringId<IStream>;
 
 	/// <summary>
@@ -59,7 +60,7 @@ namespace HordeServer.Controllers
 		[Route("/api/v1/testdata")]
 		public async Task<ActionResult<CreateTestDataResponse>> CreateAsync(CreateTestDataRequest Request)
 		{
-			IJob? Job = await JobService.GetJobAsync(Request.JobId.ToObjectId());
+			IJob? Job = await JobService.GetJobAsync(new JobId(Request.JobId));
 			if (Job == null)
 			{
 				return NotFound();
@@ -107,7 +108,7 @@ namespace HordeServer.Controllers
 
 			List<object> Results = new List<object>();
 
-			List<ITestData> Documents = await TestDataCollection.FindAsync(StreamIdValue, MinChange, MaxChange, JobId?.ToObjectId(), JobStepId?.ToSubResourceId(), Key, Index, Count);
+			List<ITestData> Documents = await TestDataCollection.FindAsync(StreamIdValue, MinChange, MaxChange, JobId?.ToObjectId<IJob>(), JobStepId?.ToSubResourceId(), Key, Index, Count);
 			foreach (ITestData Document in Documents)
 			{
 				if (await JobService.AuthorizeAsync(Document.JobId, AclAction.ViewJob, User, Cache))
diff --git a/Engine/Source/Programs/Horde/HordeServer/Controllers/UgsController.cs b/Engine/Source/Programs/Horde/HordeServer/Controllers/UgsController.cs
--- a/Engine/Source/Programs/Horde/HordeServer/Controllers/UgsController.cs
+++ b/Engine/Source/Programs/Horde/HordeServer/Controllers/UgsController.cs
@@ -23,6 +23,7 @@ using Microsoft.Extensions.Logging;
 
 namespace HordeServer.Controllers
 {
+	using LogId = ObjectId<ILogFile>;
 	using StreamId = StringId<IStream>;
 
 	/// <summary>
@@ -291,7 +292,7 @@ namespace HordeServer.Controllers
 		{
 			List<GetUgsIssueDiagnosticResponse> Diagnostics = new List<GetUgsIssueDiagnosticResponse>();
 
-			Dictionary<ObjectId, ILogFile?> LogFiles = new Dictionary<ObjectId, ILogFile?>();
+			Dictionary<LogId, ILogFile?> LogFiles = new Dictionary<LogId, ILogFile?>();
 
 			JobPermissionsCache PermissionsCache = new JobPermissionsCache();
 
diff --git a/Engine/Source/Programs/Horde/HordeServer/Controllers/UserController.cs b/Engine/Source/Programs/Horde/HordeServer/Controllers/UserController.cs
--- a/Engine/Source/Programs/Horde/HordeServer/Controllers/UserController.cs
+++ b/Engine/Source/Programs/Horde/HordeServer/Controllers/UserController.cs
@@ -15,6 +15,8 @@ using System.Threading.Tasks;
 
 namespace HordeServer.Controllers
 {
+	using JobId = ObjectId<IJob>;
+
 	/// <summary>
 	/// Controller for the /api/v1/user endpoint
 	/// </summary>
@@ -78,7 +80,7 @@ namespace HordeServer.Controllers
 			{
 				return BadRequest("Current user does not have a registered profile");
 			}
-			await UserCollection.UpdateSettingsAsync(UserId.Value, Request.EnableExperimentalFeatures, Request.EnableIssueNotifications, Request.DashboardSettings?.ToBsonValue(), Request.AddPinnedJobIds?.Select(x => new ObjectId(x)), Request.RemovePinnedJobIds?.Select(x => new ObjectId(x)));
+			await UserCollection.UpdateSettingsAsync(UserId.Value, Request.EnableExperimentalFeatures, Request.EnableIssueNotifications, Request.DashboardSettings?.ToBsonValue(), Request.AddPinnedJobIds?.Select(x => new JobId(x)), Request.RemovePinnedJobIds?.Select(x => new JobId(x)));
 			return Ok();
 		}
 	}
diff --git a/Engine/Source/Programs/Horde/HordeServer/Devices/Impl/DeviceService.cs b/Engine/Source/Programs/Horde/HordeServer/Devices/Impl/DeviceService.cs
--- a/Engine/Source/Programs/Horde/HordeServer/Devices/Impl/DeviceService.cs
+++ b/Engine/Source/Programs/Horde/HordeServer/Devices/Impl/DeviceService.cs
@@ -23,6 +23,7 @@ namespace HordeServer.Services
 	using DeviceId = StringId<IDevice>;
 	using DevicePlatformId = StringId<IDevicePlatform>;
 	using DevicePoolId = StringId<IDevicePool>;
+	using JobId = ObjectId<IJob>;
 
 	/// <summary>
 	/// Platform map required by V1 API
@@ -314,7 +315,7 @@ namespace HordeServer.Services
 
 				if (JobId != null)
 				{
-					Job = await JobService.GetJobAsync(new ObjectId(JobId));
+					Job = await JobService.GetJobAsync(new JobId(JobId));
 
 					if (Job != null)
 					{
diff --git a/Engine/Source/Programs/Horde/HordeServer/Issues/IIssueCollection.cs b/Engine/Source/Programs/Horde/HordeServer/Issues/IIssueCollection.cs
--- a/Engine/Source/Programs/Horde/HordeServer/Issues/IIssueCollection.cs
+++ b/Engine/Source/Programs/Horde/HordeServer/Issues/IIssueCollection.cs
@@ -23,6 +23,8 @@ using System.Transactions;
 
 namespace HordeServer.Collections
 {
+	using JobId = ObjectId<IJob>;
+	using LogId = ObjectId<ILogFile>;
 	using StreamId = StringId<IStream>;
 	using TemplateRefId = StringId<TemplateRef>;
 
@@ -194,7 +196,7 @@ namespace HordeServer.Collections
 		/// <summary>
 		/// The unique job id
 		/// </summary>
-		public ObjectId JobId { get; set; }
+		public JobId JobId { get; set; }
 
 		/// <summary>
 		/// Batch id for the step
@@ -214,7 +216,7 @@ namespace HordeServer.Collections
 		/// <summary>
 		/// The log id for this step
 		/// </summary>
-		public ObjectId? LogId { get; set; }
+		public LogId? LogId { get; set; }
 
 		/// <summary>
 		/// Whether to notify suspects for this step
@@ -233,7 +235,7 @@ namespace HordeServer.Collections
 		/// <param name="StepTime">Time that the step started</param>
 		/// <param name="LogId">Unique id of the log file for this step</param>
 		/// <param name="NotifySuspects">Whether to notify suspects for this step</param>
-		public NewIssueStepData(int Change, IssueSeverity Severity, string JobName, ObjectId JobId, SubResourceId BatchId, SubResourceId StepId, DateTime StepTime, ObjectId? LogId, bool NotifySuspects)
+		public NewIssueStepData(int Change, IssueSeverity Severity, string JobName, JobId JobId, SubResourceId BatchId, SubResourceId StepId, DateTime StepTime, LogId? LogId, bool NotifySuspects)
 		{
 			this.Change = Change;
 			this.Severity = Severity;
@@ -537,7 +539,7 @@ namespace HordeServer.Collections
 		/// <param name="BatchId">The batch id</param>
 		/// <param name="StepId">The step id</param>
 		/// <returns>List of steps</returns>
-		Task<List<IIssueStep>> FindStepsAsync(ObjectId JobId, SubResourceId? BatchId, SubResourceId? StepId);
+		Task<List<IIssueStep>> FindStepsAsync(JobId JobId, SubResourceId? BatchId, SubResourceId? StepId);
 
 		#endregion
 	}
diff --git a/Engine/Source/Programs/Horde/HordeServer/Issues/IIssueService.cs b/Engine/Source/Programs/Horde/HordeServer/Issues/IIssueService.cs
--- a/Engine/Source/Programs/Horde/HordeServer/Issues/IIssueService.cs
+++ b/Engine/Source/Programs/Horde/HordeServer/Issues/IIssueService.cs
@@ -11,6 +11,7 @@ using System.Threading.Tasks;
 
 namespace HordeServer.Services
 {
+	using LogId = ObjectId<ILogFile>;
 	using StreamId = StringId<IStream>;
 
 	/// <summary>
@@ -122,7 +123,7 @@ namespace HordeServer.Services
 		/// <param name="Index">Index of the first event to return</param>
 		/// <param name="Count">Number of issues to return</param>
 		/// <returns>List of log events</returns>
-		Task<List<ILogEvent>> FindEventsForIssueAsync(int IssueId, ObjectId[]? LogIds = null, int Index = 0, int Count = 10);
+		Task<List<ILogEvent>> FindEventsForIssueAsync(int IssueId, LogId[]? LogIds = null, int Index = 0, int Count = 10);
 
 		/// <summary>
 		/// Updates the state of an issue
diff --git a/Engine/Source/Programs/Horde/HordeServer/Issues/IIssueStep.cs b/Engine/Source/Programs/Horde/HordeServer/Issues/IIssueStep.cs
--- a/Engine/Source/Programs/Horde/HordeServer/Issues/IIssueStep.cs
+++ b/Engine/Source/Programs/Horde/HordeServer/Issues/IIssueStep.cs
@@ -1,5 +1,6 @@
 // Copyright Epic Games, Inc. All Rights Reserved.
 
+using HordeServer.Utilities;
 using MongoDB.Bson;
 using System;
 using System.Collections.Generic;
@@ -10,6 +11,9 @@ using SubResourceId = HordeServer.Models.SubResourceId;
 
 namespace HordeServer.Models
 {
+	using JobId = ObjectId<IJob>;
+	using LogId = ObjectId<ILogFile>;
+
 	/// <summary>
 	/// Identifies a particular changelist and job that contributes to a span
 	/// </summary>
@@ -38,7 +42,7 @@ namespace HordeServer.Models
 		/// <summary>
 		/// The unique job id
 		/// </summary>
-		public ObjectId JobId { get; }
+		public JobId JobId { get; }
 
 		/// <summary>
 		/// Unique id of the batch within the job
@@ -63,6 +67,6 @@ namespace HordeServer.Models
 		/// <summary>
 		/// The log id for this step
 		/// </summary>
-		public ObjectId? LogId { get; }
+		public LogId? LogId { get; }
 	}
 }
diff --git a/Engine/Source/Programs/Horde/HordeServer/Issues/Impl/IssueCollection.cs b/Engine/Source/Programs/Horde/HordeServer/Issues/Impl/IssueCollection.cs
--- a/Engine/Source/Programs/Horde/HordeServer/Issues/Impl/IssueCollection.cs
+++ b/Engine/Source/Programs/Horde/HordeServer/Issues/Impl/IssueCollection.cs
@@ -20,6 +20,8 @@ using System.Threading.Tasks;
 
 namespace HordeServer.Collections.Impl
 {
+	using JobId = ObjectId<IJob>;
+	using LogId = ObjectId<ILogFile>;
 	using StreamId = StringId<IStream>;
 	using TemplateRefId = StringId<TemplateRef>;
 
@@ -393,7 +395,7 @@ namespace HordeServer.Collections.Impl
 			public string JobName { get; set; }
 
 			[BsonRequired]
-			public ObjectId JobId { get; set; }
+			public JobId JobId { get; set; }
 
 			[BsonRequired]
 			public SubResourceId BatchId { get; set; }
@@ -403,7 +405,7 @@ namespace HordeServer.Collections.Impl
 
 			public DateTime StepTime { get; set; }
 
-			public ObjectId? LogId { get; set; }
+			public LogId? LogId { get; set; }
 
 			public bool NotifySuspects { get; set; }
 
@@ -1264,7 +1266,7 @@ namespace HordeServer.Collections.Impl
 		}
 
 		/// <inheritdoc/>
-		public Task<List<IIssueStep>> FindStepsAsync(ObjectId JobId, SubResourceId? BatchId, SubResourceId? StepId)
+		public Task<List<IIssueStep>> FindStepsAsync(JobId JobId, SubResourceId? BatchId, SubResourceId? StepId)
 		{
 			FilterDefinition<IssueStep> Filter = Builders<IssueStep>.Filter.Eq(x => x.JobId, JobId);
 			if (BatchId != null)
diff --git a/Engine/Source/Programs/Horde/HordeServer/Issues/Impl/IssueService.cs b/Engine/Source/Programs/Horde/HordeServer/Issues/Impl/IssueService.cs
--- a/Engine/Source/Programs/Horde/HordeServer/Issues/Impl/IssueService.cs
+++ b/Engine/Source/Programs/Horde/HordeServer/Issues/Impl/IssueService.cs
@@ -32,6 +32,7 @@ using OpenTracing.Util;
 
 namespace HordeServer.Services.Impl
 {
+	using LogId = ObjectId<ILogFile>;
 	using StreamId = StringId<IStream>;
 	using TemplateRefId = StringId<TemplateRef>;
 
@@ -1239,7 +1240,7 @@ namespace HordeServer.Services.Impl
 		}
 
 		/// <inheritdoc/>
-		public async Task<List<ILogEvent>> FindEventsForIssueAsync(int IssueId, ObjectId[]? LogIds, int Index, int Count)
+		public async Task<List<ILogEvent>> FindEventsForIssueAsync(int IssueId, LogId[]? LogIds, int Index, int Count)
 		{
 			List<IIssueSpan> Spans = await IssueCollection.FindSpansAsync(IssueId);
 			return await LogFileService.FindEventsForSpansAsync(Spans.Select(x => x.Id), LogIds, Index, Count);
diff --git a/Engine/Source/Programs/Horde/HordeServer/Jobs/IArtifact.cs b/Engine/Source/Programs/Horde/HordeServer/Jobs/IArtifact.cs
--- a/Engine/Source/Programs/Horde/HordeServer/Jobs/IArtifact.cs
+++ b/Engine/Source/Programs/Horde/HordeServer/Jobs/IArtifact.cs
@@ -7,6 +7,8 @@ using MongoDB.Bson.Serialization.Attributes;
 
 namespace HordeServer.Models
 {
+	using JobId = ObjectId<IJob>;
+
 	/// <summary>
 	/// Information about an artifact
 	/// </summary>
@@ -20,7 +22,7 @@ namespace HordeServer.Models
 		/// <summary>
 		/// Unique id of the job containing this artifact
 		/// </summary>
-		public ObjectId JobId { get; }
+		public JobId JobId { get; }
 
 		/// <summary>
 		/// Unique id of the job containing this artifact
diff --git a/Engine/Source/Programs/Horde/HordeServer/Jobs/IArtifactCollection.cs b/Engine/Source/Programs/Horde/HordeServer/Jobs/IArtifactCollection.cs
--- a/Engine/Source/Programs/Horde/HordeServer/Jobs/IArtifactCollection.cs
+++ b/Engine/Source/Programs/Horde/HordeServer/Jobs/IArtifactCollection.cs
@@ -1,12 +1,15 @@
 // Copyright Epic Games, Inc. All Rights Reserved.
 
 using HordeServer.Models;
+using HordeServer.Utilities;
 using MongoDB.Bson;
 using System.Collections.Generic;
 using System.Threading.Tasks;
 
 namespace HordeServer.Jobs
 {
+	using JobId = ObjectId<IJob>;
+
 	/// <summary>
 	/// Interface for a collection of artifacts
 	/// </summary>
@@ -21,7 +24,7 @@ namespace HordeServer.Jobs
 		/// <param name="MimeType">Type of artifact</param>
 		/// <param name="Data">The data to write</param>
 		/// <returns>The new log file document</returns>
-		Task<IArtifact> CreateArtifactAsync(ObjectId JobId, SubResourceId? StepId, string Name, string MimeType, System.IO.Stream Data);
+		Task<IArtifact> CreateArtifactAsync(JobId JobId, SubResourceId? StepId, string Name, string MimeType, System.IO.Stream Data);
 
 		/// <summary>
 		/// Gets all the available artifacts for a job
@@ -30,7 +33,7 @@ namespace HordeServer.Jobs
 		/// <param name="StepId">Unique id of the Step to query</param>
 		/// <param name="Name">Name of the artifact</param>
 		/// <returns>List of artifact documents</returns>
-		Task<List<IArtifact>> GetArtifactsAsync(ObjectId? JobId, SubResourceId? StepId, string? Name);
+		Task<List<IArtifact>> GetArtifactsAsync(JobId? JobId, SubResourceId? StepId, string? Name);
 
 		/// <summary>
 		/// Gets a specific list of artifacts based on id
diff --git a/Engine/Source/Programs/Horde/HordeServer/Jobs/IJob.cs b/Engine/Source/Programs/Horde/HordeServer/Jobs/IJob.cs
--- a/Engine/Source/Programs/Horde/HordeServer/Jobs/IJob.cs
+++ b/Engine/Source/Programs/Horde/HordeServer/Jobs/IJob.cs
@@ -20,9 +20,12 @@ using System.Threading.Tasks;
 namespace HordeServer.Models
 {
 	using ReportPlacement = HordeCommon.Rpc.ReportPlacement;
-	using StreamId = StringId<HordeServer.Models.IStream>;
-	using TemplateRefId = StringId<HordeServer.Models.TemplateRef>;
-	using PoolId = StringId<HordeServer.Models.IPool>;
+	using JobId = ObjectId<IJob>;
+	using LeaseId = ObjectId<ILease>;
+	using LogId = ObjectId<ILogFile>;
+	using StreamId = StringId<IStream>;
+	using TemplateRefId = StringId<TemplateRef>;
+	using PoolId = StringId<IPool>;
 
 	/// <summary>
 	/// Report for a job or jobstep
@@ -88,7 +91,7 @@ namespace HordeServer.Models
 		/// <summary>
 		/// The log id for this step
 		/// </summary>
-		public ObjectId? LogId { get; }
+		public LogId? LogId { get; }
 
 		/// <summary>
 		/// Unique id for notifications
@@ -178,7 +181,7 @@ namespace HordeServer.Models
 		/// <summary>
 		/// The log file id for this batch
 		/// </summary>
-		public ObjectId? LogId { get; }
+		public LogId? LogId { get; }
 
 		/// <summary>
 		/// Index of the group being executed
@@ -218,7 +221,7 @@ namespace HordeServer.Models
 		/// <summary>
 		/// The lease that's executing this group
 		/// </summary>
-		public ObjectId? LeaseId { get; }
+		public LeaseId? LeaseId { get; }
 
 		/// <summary>
 		/// The weighted priority of this batch for the scheduler
@@ -474,7 +477,7 @@ namespace HordeServer.Models
 		/// <summary>
 		/// The triggered job id
 		/// </summary>
-		public ObjectId? JobId { get; }
+		public JobId? JobId { get; }
 	}
 
 	/// <summary>
@@ -496,7 +499,7 @@ namespace HordeServer.Models
 		/// <summary>
 		/// Identifier for the job. Randomly generated.
 		/// </summary>
-		public ObjectId Id { get; }
+		public JobId Id { get; }
 
 		/// <summary>
 		/// The stream that this job belongs to
diff --git a/Engine/Source/Programs/Horde/HordeServer/Jobs/IJobCollection.cs b/Engine/Source/Programs/Horde/HordeServer/Jobs/IJobCollection.cs
--- a/Engine/Source/Programs/Horde/HordeServer/Jobs/IJobCollection.cs
+++ b/Engine/Source/Programs/Horde/HordeServer/Jobs/IJobCollection.cs
@@ -19,6 +19,9 @@ using System.Threading.Tasks;
 
 namespace HordeServer.Collections
 {
+	using JobId = ObjectId<IJob>;
+	using LeaseId = ObjectId<ILease>;
+	using LogId = ObjectId<ILogFile>;
 	using PoolId = StringId<IPool>;
 	using StreamId = StringId<IStream>;
 	using TemplateRefId = StringId<TemplateRef>;
@@ -55,14 +58,14 @@ namespace HordeServer.Collections
 		/// <param name="HelixSwarmCallbackUrl">Callback URL to a Helix Server review, if any</param>
 		/// <param name="Arguments">Arguments for the job</param>
 		/// <returns>The new job document</returns>
-		Task<IJob> AddAsync(ObjectId JobId, StreamId StreamId, TemplateRefId TemplateRefId, ContentHash TemplateHash, IGraph Graph, string Name, int Change, int CodeChange, int? PreflightChange, int? ClonedPreflightChange, ObjectId? StartedByUserId, string? StartedByUserName, Priority? Priority, bool? AutoSubmit, bool? UpdateIssues, List<ChainedJobTemplate>? JobTriggers, bool ShowUgsBadges, bool ShowUgsAlerts, string? NotificationChannel, string? NotificationChannelFilter, string? HelixSwarmCallbackUrl, List<string>? Arguments);
+		Task<IJob> AddAsync(JobId JobId, StreamId StreamId, TemplateRefId TemplateRefId, ContentHash TemplateHash, IGraph Graph, string Name, int Change, int CodeChange, int? PreflightChange, int? ClonedPreflightChange, ObjectId? StartedByUserId, string? StartedByUserName, Priority? Priority, bool? AutoSubmit, bool? UpdateIssues, List<ChainedJobTemplate>? JobTriggers, bool ShowUgsBadges, bool ShowUgsAlerts, string? NotificationChannel, string? NotificationChannelFilter, string? HelixSwarmCallbackUrl, List<string>? Arguments);
 
 		/// <summary>
 		/// Gets a job with the given unique id
 		/// </summary>
 		/// <param name="JobId">Job id to search for</param>
 		/// <returns>Information about the given job</returns>
-		Task<IJob?> GetAsync(ObjectId JobId);
+		Task<IJob?> GetAsync(JobId JobId);
 
 		/// <summary>
 		/// Deletes a job
@@ -82,7 +85,7 @@ namespace HordeServer.Collections
 		/// </summary>
 		/// <param name="JobId">Unique id of the job</param>
 		/// <returns>The job document</returns>
-		Task<IJobPermissions?> GetPermissionsAsync(ObjectId JobId);
+		Task<IJobPermissions?> GetPermissionsAsync(JobId JobId);
 
 		/// <summary>
 		/// Searches for jobs matching the given criteria
@@ -103,7 +106,7 @@ namespace HordeServer.Collections
 		/// <param name="Index">Index of the first result to return</param>
 		/// <param name="Count">Number of results to return</param>
 		/// <returns>List of jobs matching the given criteria</returns>
-		Task<List<IJob>> FindAsync(ObjectId[]? JobIds = null, StreamId? StreamId = null, string? Name = null, TemplateRefId[]? Templates = null, int? MinChange = null, int? MaxChange = null, int? PreflightChange = null, ObjectId? PreflightStartedByUser = null, ObjectId? StartedByUser = null, DateTimeOffset ? MinCreateTime = null, DateTimeOffset? MaxCreateTime = null, DateTimeOffset? ModifiedBefore = null, DateTimeOffset? ModifiedAfter = null, int? Index = null, int? Count = null);
+		Task<List<IJob>> FindAsync(JobId[]? JobIds = null, StreamId? StreamId = null, string? Name = null, TemplateRefId[]? Templates = null, int? MinChange = null, int? MaxChange = null, int? PreflightChange = null, ObjectId? PreflightStartedByUser = null, ObjectId? StartedByUser = null, DateTimeOffset ? MinCreateTime = null, DateTimeOffset? MaxCreateTime = null, DateTimeOffset? ModifiedBefore = null, DateTimeOffset? ModifiedAfter = null, int? Index = null, int? Count = null);
 
 		/// <summary>
 		/// Updates a new job
@@ -121,7 +124,7 @@ namespace HordeServer.Collections
 		/// <param name="Arguments">New arguments for the job</param>
 		/// <param name="LabelIdxToTriggerId">New trigger ID for a label in the job</param>
 		/// <param name="JobTrigger">New downstream job id</param>
-		Task<bool> TryUpdateJobAsync(IJob Job, IGraph Graph, string? Name = null, Priority? Priority = null, bool? AutoSubmit = null, int? AutoSubmitChange = null, string? AutoSubmitMessage = null, string? AbortedByUser = null, ObjectId? NotificationTriggerId = null, List<Report>? Reports = null, List<string>? Arguments = null, KeyValuePair<int, ObjectId>? LabelIdxToTriggerId = null, KeyValuePair<TemplateRefId, ObjectId>? JobTrigger = null);
+		Task<bool> TryUpdateJobAsync(IJob Job, IGraph Graph, string? Name = null, Priority? Priority = null, bool? AutoSubmit = null, int? AutoSubmitChange = null, string? AutoSubmitMessage = null, string? AbortedByUser = null, ObjectId? NotificationTriggerId = null, List<Report>? Reports = null, List<string>? Arguments = null, KeyValuePair<int, ObjectId>? LabelIdxToTriggerId = null, KeyValuePair<TemplateRefId, JobId>? JobTrigger = null);
 
 		/// <summary>
 		/// Updates the state of a batch
@@ -133,7 +136,7 @@ namespace HordeServer.Collections
 		/// <param name="NewState">New state of the jobstep</param>
 		/// <param name="NewError">Error code for the batch</param>
 		/// <returns>True if the job was updated, false if it was deleted</returns>
-		Task<bool> TryUpdateBatchAsync(IJob Job, IGraph Graph, SubResourceId BatchId, ObjectId? NewLogId, JobStepBatchState? NewState, JobStepBatchError? NewError);
+		Task<bool> TryUpdateBatchAsync(IJob Job, IGraph Graph, SubResourceId BatchId, LogId? NewLogId, JobStepBatchState? NewState, JobStepBatchError? NewError);
 
 		/// <summary>
 		/// Update a jobstep state
@@ -153,7 +156,7 @@ namespace HordeServer.Collections
 		/// <param name="NewReports">New report documents</param>
 		/// <param name="NewProperties">Property changes. Any properties with a null value will be removed.</param>
 		/// <returns>True if the job was updated, false if it was deleted in the meantime</returns>
-		Task<bool> TryUpdateStepAsync(IJob Job, IGraph Graph, SubResourceId BatchId, SubResourceId StepId, JobStepState NewState = default, JobStepOutcome NewOutcome = default, bool? NewAbortRequested = null, string? NewAbortByUser = null, ObjectId? NewLogId = null, ObjectId? NewNotificationTriggerId = null, string? NewRetryByUser = null, Priority? NewPriority = null, List<Report>? NewReports = null, Dictionary<string, string?>? NewProperties = null);
+		Task<bool> TryUpdateStepAsync(IJob Job, IGraph Graph, SubResourceId BatchId, SubResourceId StepId, JobStepState NewState = default, JobStepOutcome NewOutcome = default, bool? NewAbortRequested = null, string? NewAbortByUser = null, LogId? NewLogId = null, ObjectId? NewNotificationTriggerId = null, string? NewRetryByUser = null, Priority? NewPriority = null, List<Report>? NewReports = null, Dictionary<string, string?>? NewProperties = null);
 
 		/// <summary>
 		/// Attempts to update the node groups to be executed for a job. Fails if another write happens in the meantime.
@@ -176,7 +179,7 @@ namespace HordeServer.Collections
 		/// <param name="JobId">The job id</param>
 		/// <param name="IssueId">The issue to add</param>
 		/// <returns>Async task</returns>
-		Task AddIssueToJobAsync(ObjectId JobId, int IssueId);
+		Task AddIssueToJobAsync(JobId JobId, int IssueId);
 
 		/// <summary>
 		/// Gets a queue of jobs to consider for execution
@@ -223,7 +226,7 @@ namespace HordeServer.Collections
 		/// <param name="LeaseId">The lease unique id</param>
 		/// <param name="LogId">Unique id of the log for the batch</param>
 		/// <returns>True if the batch is updated</returns>
-		Task<bool> TryAssignLeaseAsync(IJob Job, int BatchIdx, PoolId PoolId, AgentId AgentId, ObjectId SessionId, ObjectId LeaseId, ObjectId LogId);
+		Task<bool> TryAssignLeaseAsync(IJob Job, int BatchIdx, PoolId PoolId, AgentId AgentId, ObjectId SessionId, LeaseId LeaseId, LogId LogId);
 
 		/// <summary>
 		/// Cancel a lease reservation on a batch (before it has started)
diff --git a/Engine/Source/Programs/Horde/HordeServer/Jobs/IJobStepRef.cs b/Engine/Source/Programs/Horde/HordeServer/Jobs/IJobStepRef.cs
--- a/Engine/Source/Programs/Horde/HordeServer/Jobs/IJobStepRef.cs
+++ b/Engine/Source/Programs/Horde/HordeServer/Jobs/IJobStepRef.cs
@@ -12,6 +12,8 @@ using System.Threading.Tasks;
 
 namespace HordeServer.Models
 {
+	using JobId = ObjectId<IJob>;
+	using LogId = ObjectId<ILogFile>;
 	using StreamId = StringId<IStream>;
 	using TemplateRefId = StringId<TemplateRef>;
 	using PoolId = StringId<IPool>;
@@ -25,7 +27,7 @@ namespace HordeServer.Models
 		/// <summary>
 		/// The job id
 		/// </summary>
-		public ObjectId JobId { get; set; }
+		public JobId JobId { get; set; }
 
 		/// <summary>
 		/// The batch id within the job
@@ -43,7 +45,7 @@ namespace HordeServer.Models
 		/// <param name="JobId">The job id</param>
 		/// <param name="BatchId">The batch id within the job</param>
 		/// <param name="StepId">The step id</param>
-		public JobStepRefId(ObjectId JobId, SubResourceId BatchId, SubResourceId StepId)
+		public JobStepRefId(JobId JobId, SubResourceId BatchId, SubResourceId StepId)
 		{
 			this.JobId = JobId;
 			this.BatchId = BatchId;
@@ -58,7 +60,7 @@ namespace HordeServer.Models
 		public static JobStepRefId Parse(string Text)
 		{
 			string[] Components = Text.Split(':');
-			return new JobStepRefId(ObjectId.Parse(Components[0]), SubResourceId.Parse(Components[1]), SubResourceId.Parse(Components[2]));
+			return new JobStepRefId(JobId.Parse(Components[0]), SubResourceId.Parse(Components[1]), SubResourceId.Parse(Components[2]));
 		}
 
 		/// <summary>
@@ -145,7 +147,7 @@ namespace HordeServer.Models
 		/// <summary>
 		/// Log for this step
 		/// </summary>
-		public ObjectId? LogId { get; }
+		public LogId? LogId { get; }
 
 		/// <summary>
 		/// The agent type
diff --git a/Engine/Source/Programs/Horde/HordeServer/Jobs/IJobStepRefCollection.cs b/Engine/Source/Programs/Horde/HordeServer/Jobs/IJobStepRefCollection.cs
--- a/Engine/Source/Programs/Horde/HordeServer/Jobs/IJobStepRefCollection.cs
+++ b/Engine/Source/Programs/Horde/HordeServer/Jobs/IJobStepRefCollection.cs
@@ -14,6 +14,7 @@ using HordeServer.Utilities;
 
 namespace HordeServer.Collections
 {
+	using LogId = ObjectId<ILogFile>;
 	using PoolId = StringId<IPool>;
 	using StreamId = StringId<IStream>;
 	using TemplateRefId = StringId<TemplateRef>;
@@ -40,7 +41,7 @@ namespace HordeServer.Collections
 		/// <param name="InitTime">Time taken for the batch containing this step to initializer</param>
 		/// <param name="StartTimeUtc">Start time</param>
 		/// <param name="FinishTimeUtc">Finish time for the step, if known</param>
-		Task<IJobStepRef> InsertOrReplaceAsync(JobStepRefId Id, string JobName, string StepName, StreamId StreamId, TemplateRefId TemplateId, int Change, ObjectId? LogId, PoolId? PoolId, AgentId? AgentId, JobStepOutcome? Outcome, float WaitTime, float InitTime, DateTime StartTimeUtc, DateTime? FinishTimeUtc);
+		Task<IJobStepRef> InsertOrReplaceAsync(JobStepRefId Id, string JobName, string StepName, StreamId StreamId, TemplateRefId TemplateId, int Change, LogId? LogId, PoolId? PoolId, AgentId? AgentId, JobStepOutcome? Outcome, float WaitTime, float InitTime, DateTime StartTimeUtc, DateTime? FinishTimeUtc);
 
 		/// <summary>
 		/// Gets the history of a given node
diff --git a/Engine/Source/Programs/Horde/HordeServer/Jobs/IJobTimingCollection.cs b/Engine/Source/Programs/Horde/HordeServer/Jobs/IJobTimingCollection.cs
--- a/Engine/Source/Programs/Horde/HordeServer/Jobs/IJobTimingCollection.cs
+++ b/Engine/Source/Programs/Horde/HordeServer/Jobs/IJobTimingCollection.cs
@@ -16,6 +16,8 @@ using System.Threading.Tasks;
 
 namespace HordeServer.Collections
 {
+	using JobId = ObjectId<IJob>;
+
 	/// <summary>
 	/// Average timing information for a node
 	/// </summary>
@@ -80,14 +82,14 @@ namespace HordeServer.Collections
 		/// <param name="JobId">The job id</param>
 		/// <param name="Steps">List of timing info for each step</param>
 		/// <returns>New timing document</returns>
-		Task<IJobTiming?> TryAddAsync(ObjectId JobId, List<JobStepTimingData> Steps);
+		Task<IJobTiming?> TryAddAsync(JobId JobId, List<JobStepTimingData> Steps);
 
 		/// <summary>
 		/// Attempts to get the timing information for a particular job
 		/// </summary>
 		/// <param name="JobId">The unique job id</param>
 		/// <returns>Timing info for the requested jbo</returns>
-		Task<IJobTiming?> TryGetAsync(ObjectId JobId);
+		Task<IJobTiming?> TryGetAsync(JobId JobId);
 
 		/// <summary>
 		/// Adds timing information for the particular job
diff --git a/Engine/Source/Programs/Horde/HordeServer/Jobs/ITestData.cs b/Engine/Source/Programs/Horde/HordeServer/Jobs/ITestData.cs
--- a/Engine/Source/Programs/Horde/HordeServer/Jobs/ITestData.cs
+++ b/Engine/Source/Programs/Horde/HordeServer/Jobs/ITestData.cs
@@ -9,6 +9,7 @@ using System.Threading.Tasks;
 
 namespace HordeServer.Models
 {
+	using JobId = ObjectId<IJob>;
 	using StreamId = StringId<IStream>;
 	using TemplateRefId = StringId<TemplateRef>;
 
@@ -35,7 +36,7 @@ namespace HordeServer.Models
 		/// <summary>
 		/// The job which produced the data
 		/// </summary>
-		ObjectId JobId { get; }
+		JobId JobId { get; }
 
 		/// <summary>
 		/// The step that ran
diff --git a/Engine/Source/Programs/Horde/HordeServer/Jobs/ITestDataCollection.cs b/Engine/Source/Programs/Horde/HordeServer/Jobs/ITestDataCollection.cs
--- a/Engine/Source/Programs/Horde/HordeServer/Jobs/ITestDataCollection.cs
+++ b/Engine/Source/Programs/Horde/HordeServer/Jobs/ITestDataCollection.cs
@@ -14,6 +14,7 @@ using System.Threading.Tasks;
 
 namespace HordeServer.Collections
 {
+	using JobId = ObjectId<IJob>;
 	using ProjectId = StringId<IProject>;
 	using StreamId = StringId<IStream>;
 	using TemplateRefId = StringId<TemplateRef>;
@@ -52,7 +53,7 @@ namespace HordeServer.Collections
 		/// <param name="Index">Offset within the results to return</param>
 		/// <param name="Count">Number of results to return</param>
 		/// <returns>The stream document</returns>
-		Task<List<ITestData>> FindAsync(StreamId? StreamId, int? MinChange, int? MaxChange, ObjectId? JobId, SubResourceId? StepId, string? Key = null, int Index = 0, int Count = 10);
+		Task<List<ITestData>> FindAsync(StreamId? StreamId, int? MinChange, int? MaxChange, JobId? JobId, SubResourceId? StepId, string? Key = null, int Index = 0, int Count = 10);
 
 		/// <summary>
 		/// Delete the test data
diff --git a/Engine/Source/Programs/Horde/HordeServer/Jobs/Impl/ArtifactCollection.cs b/Engine/Source/Programs/Horde/HordeServer/Jobs/Impl/ArtifactCollection.cs
--- a/Engine/Source/Programs/Horde/HordeServer/Jobs/Impl/ArtifactCollection.cs
+++ b/Engine/Source/Programs/Horde/HordeServer/Jobs/Impl/ArtifactCollection.cs
@@ -27,6 +27,8 @@ using System.Threading.Tasks;
 
 namespace HordeServer.Services
 {
+	using JobId = ObjectId<IJob>;
+
 	/// <summary>
 	/// Wraps functionality for manipulating artifacts
 	/// </summary>
@@ -38,7 +40,7 @@ namespace HordeServer.Services
 			public ObjectId Id { get; set; }
 
 			[BsonRequired]
-			public ObjectId JobId { get; set; }
+			public JobId JobId { get; set; }
 
 			[BsonRequired]
 			public string Name { get; set; }
@@ -60,7 +62,7 @@ namespace HordeServer.Services
 				this.MimeType = null!;
 			}
 
-			public Artifact(ObjectId JobId, SubResourceId? StepId, string Name, long Length, string MimeType)
+			public Artifact(JobId JobId, SubResourceId? StepId, string Name, long Length, string MimeType)
 			{
 				this.Id = ObjectId.GenerateNewId();
 				this.JobId = JobId;
@@ -103,7 +105,7 @@ namespace HordeServer.Services
 		/// <param name="MimeType">Type of artifact</param>
 		/// <param name="Data">The data to write</param>
 		/// <returns>The new log file document</returns>
-		public async Task<IArtifact> CreateArtifactAsync(ObjectId JobId, SubResourceId? StepId, string Name, string MimeType, System.IO.Stream Data)
+		public async Task<IArtifact> CreateArtifactAsync(JobId JobId, SubResourceId? StepId, string Name, string MimeType, System.IO.Stream Data)
 		{
 			// upload first
 			string ArtifactName = ValidateName(Name);
@@ -122,7 +124,7 @@ namespace HordeServer.Services
 		/// <param name="StepId">Unique id of the Step to query</param>
 		/// <param name="Name">Name of the artifact</param>
 		/// <returns>List of artifact documents</returns>
-		public async Task<List<IArtifact>> GetArtifactsAsync(ObjectId? JobId, SubResourceId? StepId, string? Name)
+		public async Task<List<IArtifact>> GetArtifactsAsync(JobId? JobId, SubResourceId? StepId, string? Name)
 		{
 			FilterDefinitionBuilder<Artifact> Builder = Builders<Artifact>.Filter;
 
@@ -233,7 +235,7 @@ namespace HordeServer.Services
 		/// <param name="StepId"></param>
 		/// <param name="Name"></param>
 		/// <returns></returns>
-		private static string GetPath(ObjectId JobId, SubResourceId? StepId, string Name)
+		private static string GetPath(JobId JobId, SubResourceId? StepId, string Name)
 		{
 			if (StepId == null)
 			{
diff --git a/Engine/Source/Programs/Horde/HordeServer/Jobs/Impl/JobCollection.cs b/Engine/Source/Programs/Horde/HordeServer/Jobs/Impl/JobCollection.cs
--- a/Engine/Source/Programs/Horde/HordeServer/Jobs/Impl/JobCollection.cs
+++ b/Engine/Source/Programs/Horde/HordeServer/Jobs/Impl/JobCollection.cs
@@ -21,6 +21,9 @@ using OpenTracing;
 
 namespace HordeServer.Collections.Impl
 {
+	using JobId = ObjectId<IJob>;
+	using LeaseId = ObjectId<ILease>;
+	using LogId = ObjectId<ILogFile>;
 	using StreamId = StringId<IStream>;
 	using TemplateRefId = StringId<TemplateRef>;
 	using PoolId = StringId<IPool>;
@@ -48,7 +51,7 @@ namespace HordeServer.Collections.Impl
 			public JobStepOutcome Outcome { get; set; } = JobStepOutcome.Success;
 
 			[BsonIgnoreIfNull]
-			public ObjectId? LogId { get; set; }
+			public LogId? LogId { get; set; }
 
 			[BsonIgnoreIfNull]
 			public ObjectId? NotificationTriggerId { get; set; }
@@ -103,7 +106,7 @@ namespace HordeServer.Collections.Impl
 			[BsonRequired]
 			public SubResourceId Id { get; set; }
 
-			public ObjectId? LogId { get; set; }
+			public LogId? LogId { get; set; }
 
 			[BsonRequired]
 			public int GroupIdx { get; set; }
@@ -126,7 +129,7 @@ namespace HordeServer.Collections.Impl
 			public ObjectId? SessionId { get; set; }
 
 			[BsonIgnoreIfNull]
-			public ObjectId? LeaseId { get; set; }
+			public LeaseId? LeaseId { get; set; }
 
 			public int SchedulePriority { get; set; }
 
@@ -160,7 +163,7 @@ namespace HordeServer.Collections.Impl
 		{
 			public string Target { get; set; }
 			public TemplateRefId TemplateRefId { get; set; }
-			public ObjectId? JobId { get; set; }
+			public JobId? JobId { get; set; }
 
 			[BsonConstructor]
 			private ChainedJobDocument()
@@ -184,7 +187,7 @@ namespace HordeServer.Collections.Impl
 		class JobDocument : IJob
 		{
 			[BsonRequired, BsonId]
-			public ObjectId Id { get; set; }
+			public JobId Id { get; set; }
 
 			public StreamId StreamId { get; set; }
 			public TemplateRefId TemplateId { get; set; }
@@ -271,7 +274,7 @@ namespace HordeServer.Collections.Impl
 				GraphHash = null!;
 			}
 
-			public JobDocument(ObjectId Id, StreamId StreamId, TemplateRefId TemplateId, ContentHash TemplateHash, ContentHash GraphHash, string Name, int Change, int CodeChange, int PreflightChange, int ClonedPreflightChange, ObjectId? StartedByUserId, string? StartedByUserName, Priority? Priority, bool? AutoSubmit, bool? UpdateIssues, DateTime CreateTimeUtc, List<ChainedJobDocument> ChainedJobs, bool ShowUgsBadges, bool ShowUgsAlerts, string? NotificationChannel, string? NotificationChannelFilter, string? HelixSwarmCallbackUrl, List<string>? Arguments)
+			public JobDocument(JobId Id, StreamId StreamId, TemplateRefId TemplateId, ContentHash TemplateHash, ContentHash GraphHash, string Name, int Change, int CodeChange, int PreflightChange, int ClonedPreflightChange, ObjectId? StartedByUserId, string? StartedByUserName, Priority? Priority, bool? AutoSubmit, bool? UpdateIssues, DateTime CreateTimeUtc, List<ChainedJobDocument> ChainedJobs, bool ShowUgsBadges, bool ShowUgsAlerts, string? NotificationChannel, string? NotificationChannelFilter, string? HelixSwarmCallbackUrl, List<string>? Arguments)
 			{
 				this.Id = Id;
 				this.StreamId = StreamId;
@@ -363,7 +366,7 @@ namespace HordeServer.Collections.Impl
 
 		/// <inheritdoc/>
 		[SuppressMessage("Compiler", "CA1054:URI parameters should not be strings")]
-		public async Task<IJob> AddAsync(ObjectId JobId, StreamId StreamId, TemplateRefId TemplateRefId, ContentHash TemplateHash, IGraph Graph, string Name, int Change, int CodeChange, int? PreflightChange, int? ClonedPreflightChange, ObjectId? StartedByUserId, string? StartedByUserName, Priority? Priority, bool? AutoSubmit, bool? UpdateIssues, List<ChainedJobTemplate>? ChainedJobs, bool ShowUgsBadges, bool ShowUgsAlerts, string? NotificationChannel, string? NotificationChannelFilter, string? HelixSwarmCallbackUrl, List<string>? Arguments)
+		public async Task<IJob> AddAsync(JobId JobId, StreamId StreamId, TemplateRefId TemplateRefId, ContentHash TemplateHash, IGraph Graph, string Name, int Change, int CodeChange, int? PreflightChange, int? ClonedPreflightChange, ObjectId? StartedByUserId, string? StartedByUserName, Priority? Priority, bool? AutoSubmit, bool? UpdateIssues, List<ChainedJobTemplate>? ChainedJobs, bool ShowUgsBadges, bool ShowUgsAlerts, string? NotificationChannel, string? NotificationChannelFilter, string? HelixSwarmCallbackUrl, List<string>? Arguments)
 		{
 			List<ChainedJobDocument> JobTriggers = new List<ChainedJobDocument>();
 			if (ChainedJobs == null)
@@ -384,7 +387,7 @@ namespace HordeServer.Collections.Impl
 		}
 
 		/// <inheritdoc/>
-		public async Task<IJob?> GetAsync(ObjectId JobId)
+		public async Task<IJob?> GetAsync(JobId JobId)
 		{
 			JobDocument? Job = await Jobs.Find<JobDocument>(x => x.Id == JobId).FirstOrDefaultAsync();
 			if (Job != null)
@@ -408,13 +411,13 @@ namespace HordeServer.Collections.Impl
 		}
 
 		/// <inheritdoc/>
-		public async Task<IJobPermissions?> GetPermissionsAsync(ObjectId JobId)
+		public async Task<IJobPermissions?> GetPermissionsAsync(JobId JobId)
 		{
 			return await Jobs.Find<JobDocument>(x => x.Id == JobId).Project<JobPermissions>(JobPermissions.Projection).FirstOrDefaultAsync();
 		}
 
 		/// <inheritdoc/>
-		public async Task<List<IJob>> FindAsync(ObjectId[]? JobIds, StreamId? StreamId, string? Name, TemplateRefId[]? Templates, int? MinChange, int? MaxChange, int? PreflightChange, ObjectId? PreflightStartedByUser, ObjectId? StartedByUser, DateTimeOffset? MinCreateTime, DateTimeOffset? MaxCreateTime, DateTimeOffset? ModifiedBefore, DateTimeOffset? ModifiedAfter, int? Index, int? Count)
+		public async Task<List<IJob>> FindAsync(JobId[]? JobIds, StreamId? StreamId, string? Name, TemplateRefId[]? Templates, int? MinChange, int? MaxChange, int? PreflightChange, ObjectId? PreflightStartedByUser, ObjectId? StartedByUser, DateTimeOffset? MinCreateTime, DateTimeOffset? MaxCreateTime, DateTimeOffset? ModifiedBefore, DateTimeOffset? ModifiedAfter, int? Index, int? Count)
 		{
 			FilterDefinitionBuilder<JobDocument> FilterBuilder = Builders<JobDocument>.Filter;
 
@@ -494,7 +497,7 @@ namespace HordeServer.Collections.Impl
 		}
 
 		/// <inheritdoc/>
-		public async Task<bool> TryUpdateJobAsync(IJob Job, IGraph Graph, string? Name, Priority? Priority, bool? AutoSubmit, int? AutoSubmitChange, string? AutoSubmitMessage, string? AbortedByUser, ObjectId? NotificationTriggerId, List<Report>? Reports, List<string>? Arguments, KeyValuePair<int, ObjectId>? LabelIdxToTriggerId, KeyValuePair<TemplateRefId, ObjectId>? JobTrigger)
+		public async Task<bool> TryUpdateJobAsync(IJob Job, IGraph Graph, string? Name, Priority? Priority, bool? AutoSubmit, int? AutoSubmitChange, string? AutoSubmitMessage, string? AbortedByUser, ObjectId? NotificationTriggerId, List<Report>? Reports, List<string>? Arguments, KeyValuePair<int, ObjectId>? LabelIdxToTriggerId, KeyValuePair<TemplateRefId, JobId>? JobTrigger)
 		{
 			// Create the update 
 			UpdateDefinitionBuilder<JobDocument> UpdateBuilder = Builders<JobDocument>.Update;
@@ -598,7 +601,7 @@ namespace HordeServer.Collections.Impl
 		}
 
 		/// <inheritdoc/>
-		public async Task<bool> TryUpdateBatchAsync(IJob Job, IGraph Graph, SubResourceId BatchId, ObjectId? NewLogId, JobStepBatchState? NewState, JobStepBatchError? NewError)
+		public async Task<bool> TryUpdateBatchAsync(IJob Job, IGraph Graph, SubResourceId BatchId, LogId? NewLogId, JobStepBatchState? NewState, JobStepBatchError? NewError)
 		{
 			JobDocument JobDocument = (JobDocument)Job;
 
@@ -684,7 +687,7 @@ namespace HordeServer.Collections.Impl
 		}
 
 		/// <inheritdoc/>
-		public async Task<bool> TryUpdateStepAsync(IJob Job, IGraph Graph, SubResourceId BatchId, SubResourceId StepId, JobStepState NewState, JobStepOutcome NewOutcome, bool? NewAbortRequested, string? NewAbortByUser, ObjectId? NewLogId, ObjectId? NewNotificationTriggerId, string? NewRetryByUser, Priority? NewPriority, List<Report>? NewReports, Dictionary<string, string?>? NewProperties)
+		public async Task<bool> TryUpdateStepAsync(IJob Job, IGraph Graph, SubResourceId BatchId, SubResourceId StepId, JobStepState NewState, JobStepOutcome NewOutcome, bool? NewAbortRequested, string? NewAbortByUser, LogId? NewLogId, ObjectId? NewNotificationTriggerId, string? NewRetryByUser, Priority? NewPriority, List<Report>? NewReports, Dictionary<string, string?>? NewProperties)
 		{
 			JobDocument JobDocument = (JobDocument)Job;
 
@@ -916,7 +919,7 @@ namespace HordeServer.Collections.Impl
 		}
 
 		/// <inheritdoc/>
-		public async Task AddIssueToJobAsync(ObjectId JobId, int IssueId)
+		public async Task AddIssueToJobAsync(JobId JobId, int IssueId)
 		{
 			FilterDefinition<JobDocument> JobFilter = Builders<JobDocument>.Filter.Eq(x => x.Id, JobId);
 			UpdateDefinition<JobDocument> JobUpdate = Builders<JobDocument>.Update.AddToSet(x => x.ReferencedByIssues, IssueId).Inc(x => x.UpdateIndex, 1).Max(x => x.UpdateTimeUtc, DateTime.UtcNow);
@@ -1016,7 +1019,7 @@ namespace HordeServer.Collections.Impl
 		}
 
 		/// <inheritdoc/>
-		public async Task<bool> TryAssignLeaseAsync(IJob Job, int BatchIdx, PoolId PoolId, AgentId AgentId, ObjectId SessionId, ObjectId LeaseId, ObjectId LogId)
+		public async Task<bool> TryAssignLeaseAsync(IJob Job, int BatchIdx, PoolId PoolId, AgentId AgentId, ObjectId SessionId, LeaseId LeaseId, LogId LogId)
 		{
 			// Try to update the job with this agent id
 			UpdateDefinitionBuilder<JobDocument> UpdateBuilder = Builders<JobDocument>.Update;
diff --git a/Engine/Source/Programs/Horde/HordeServer/Jobs/Impl/JobService.cs b/Engine/Source/Programs/Horde/HordeServer/Jobs/Impl/JobService.cs
--- a/Engine/Source/Programs/Horde/HordeServer/Jobs/Impl/JobService.cs
+++ b/Engine/Source/Programs/Horde/HordeServer/Jobs/Impl/JobService.cs
@@ -30,6 +30,9 @@ using OpenTracing;
 
 namespace HordeServer.Services
 {
+	using JobId = ObjectId<IJob>;
+	using LeaseId = ObjectId<ILease>;
+	using LogId = ObjectId<ILogFile>;
 	using StreamId = StringId<IStream>;
 	using TemplateRefId = StringId<TemplateRef>;
 
@@ -62,7 +65,7 @@ namespace HordeServer.Services
 		/// <summary>
 		/// Map of job id to permissions for that job
 		/// </summary>
-		public Dictionary<ObjectId, IJobPermissions?> Jobs { get; set; } = new Dictionary<ObjectId, IJobPermissions?>();
+		public Dictionary<JobId, IJobPermissions?> Jobs { get; set; } = new Dictionary<JobId, IJobPermissions?>();
 	}
 
 	/// <summary>
@@ -229,9 +232,9 @@ namespace HordeServer.Services
 		/// <param name="HelixSwarmCallbackUrl">Helix Swarm callback URL for review, if any</param>
 		/// <param name="Arguments">Arguments for the job</param>
 		/// <returns>Unique id representing the job</returns>
-		public async Task<IJob> CreateJobAsync(ObjectId? JobId, IStream Stream, TemplateRefId TemplateRefId, ContentHash TemplateHash, IGraph Graph, string Name, int Change, int CodeChange, int? PreflightChange, int? ClonedPreflightChange, ObjectId? StartedByUserId, string? StartedByUserName, Priority? Priority, bool? AutoSubmit, bool? UpdateIssues, List<ChainedJobTemplate>? JobTriggers, bool ShowUgsBadges, bool ShowUgsAlerts, string? NotificationChannel, string? NotificationChannelFilter, string? HelixSwarmCallbackUrl, IReadOnlyList<string> Arguments)
+		public async Task<IJob> CreateJobAsync(JobId? JobId, IStream Stream, TemplateRefId TemplateRefId, ContentHash TemplateHash, IGraph Graph, string Name, int Change, int CodeChange, int? PreflightChange, int? ClonedPreflightChange, ObjectId? StartedByUserId, string? StartedByUserName, Priority? Priority, bool? AutoSubmit, bool? UpdateIssues, List<ChainedJobTemplate>? JobTriggers, bool ShowUgsBadges, bool ShowUgsAlerts, string? NotificationChannel, string? NotificationChannelFilter, string? HelixSwarmCallbackUrl, IReadOnlyList<string> Arguments)
 		{
-			ObjectId JobIdValue = JobId ?? ObjectId.GenerateNewId();
+			JobId JobIdValue = JobId ?? HordeServer.Utilities.ObjectId<IJob>.GenerateNewId();
 			using IDisposable Scope = Logger.BeginScope("CreateJobAsync({JobId})", JobIdValue);
 
 			if (PreflightChange != null && ShouldClonePreflightChange(Stream.Id))
@@ -404,7 +407,7 @@ namespace HordeServer.Services
 		/// <param name="AgentId">The agent to retreive</param>
 		/// <param name="LeaseId">The lease id to update</param>
 		/// <returns></returns>
-		async Task CancelLeaseAsync(AgentId AgentId, ObjectId LeaseId)
+		async Task CancelLeaseAsync(AgentId AgentId, LeaseId LeaseId)
 		{
 			for (; ; )
 			{
@@ -444,7 +447,7 @@ namespace HordeServer.Services
 		/// </summary>
 		/// <param name="JobId">Job id to search for</param>
 		/// <returns>Information about the given job</returns>
-		public Task<IJob?> GetJobAsync(ObjectId JobId)
+		public Task<IJob?> GetJobAsync(JobId JobId)
 		{
 			return Jobs.GetAsync(JobId);
 		}
@@ -464,7 +467,7 @@ namespace HordeServer.Services
 		/// </summary>
 		/// <param name="JobId">Unique id of the job</param>
 		/// <returns>The job document</returns>
-		public Task<IJobPermissions?> GetJobPermissionsAsync(ObjectId JobId)
+		public Task<IJobPermissions?> GetJobPermissionsAsync(JobId JobId)
 		{
 			return Jobs.GetPermissionsAsync(JobId);
 		}
@@ -491,7 +494,7 @@ namespace HordeServer.Services
 		/// <param name="Index">Index of the first result to return</param>
 		/// <param name="Count">Number of results to return</param>
 		/// <returns>List of jobs matching the given criteria</returns>
-		public async Task<List<IJob>> FindJobsAsync(ObjectId[]? JobIds = null, StreamId? StreamId = null, string? Name = null, TemplateRefId[]? Templates = null, int? MinChange = null, int? MaxChange = null, int? PreflightChange = null, ObjectId ? PreflightStartedByUser = null, ObjectId? StartedByUser = null, DateTimeOffset ? MinCreateTime = null, DateTimeOffset? MaxCreateTime = null, string? Target = null, JobStepState[]? State = null, JobStepOutcome[]? Outcome = null, DateTimeOffset? ModifiedBefore = null, DateTimeOffset? ModifiedAfter = null, int? Index = null, int? Count = null)
+		public async Task<List<IJob>> FindJobsAsync(JobId[]? JobIds = null, StreamId? StreamId = null, string? Name = null, TemplateRefId[]? Templates = null, int? MinChange = null, int? MaxChange = null, int? PreflightChange = null, ObjectId ? PreflightStartedByUser = null, ObjectId? StartedByUser = null, DateTimeOffset ? MinCreateTime = null, DateTimeOffset? MaxCreateTime = null, string? Target = null, JobStepState[]? State = null, JobStepOutcome[]? Outcome = null, DateTimeOffset? ModifiedBefore = null, DateTimeOffset? ModifiedAfter = null, int? Index = null, int? Count = null)
 		{
 			if (Target == null && (State == null || State.Length == 0) && (Outcome == null || Outcome.Length == 0))
 			{
@@ -681,7 +684,7 @@ namespace HordeServer.Services
 		/// <param name="NewLogId">The new log file id</param>
 		/// <param name="NewState">New state of the jobstep</param>
 		/// <returns>True if the job was updated, false if it was deleted</returns>
-		public async Task<bool> UpdateBatchAsync(IJob Job, SubResourceId BatchId, ObjectId? NewLogId = null, JobStepBatchState? NewState = null)
+		public async Task<bool> UpdateBatchAsync(IJob Job, SubResourceId BatchId, LogId? NewLogId = null, JobStepBatchState? NewState = null)
 		{
 			using IDisposable Scope = Logger.BeginScope("UpdateBatchAsync({JobId})", Job.Id);
 
@@ -759,7 +762,7 @@ namespace HordeServer.Services
 		/// <param name="NewState">New state of the jobstep</param>
 		/// <param name="NewError">New error state</param>
 		/// <returns>The updated job, otherwise null</returns>
-		public async Task<bool> TryUpdateBatchAsync(IJob Job, SubResourceId BatchId, ObjectId? NewLogId = null, JobStepBatchState? NewState = null, JobStepBatchError? NewError = null)
+		public async Task<bool> TryUpdateBatchAsync(IJob Job, SubResourceId BatchId, LogId? NewLogId = null, JobStepBatchState? NewState = null, JobStepBatchError? NewError = null)
 		{
 			IGraph Graph = await GetGraphAsync(Job);
 			if (!await Jobs.TryUpdateBatchAsync(Job, Graph, BatchId, NewLogId, NewState, NewError))
@@ -786,7 +789,7 @@ namespace HordeServer.Services
 		/// <param name="NewReports">New list of reports</param>
 		/// <param name="NewProperties">Property changes. Any properties with a null value will be removed.</param>
 		/// <returns>True if the job was updated, false if it was deleted in the meantime</returns>
-		public async Task<IJob?> UpdateStepAsync(IJob Job, SubResourceId BatchId, SubResourceId StepId, JobStepState NewState = JobStepState.Unspecified, JobStepOutcome NewOutcome = JobStepOutcome.Unspecified, bool? NewAbortRequested = null, string? NewAbortByUser = null, ObjectId? NewLogId = null, ObjectId? NewNotificationTriggerId = null, string? NewRetryByUser = null, Priority? NewPriority = null, List<Report>? NewReports = null, Dictionary<string, string?>? NewProperties = null)
+		public async Task<IJob?> UpdateStepAsync(IJob Job, SubResourceId BatchId, SubResourceId StepId, JobStepState NewState = JobStepState.Unspecified, JobStepOutcome NewOutcome = JobStepOutcome.Unspecified, bool? NewAbortRequested = null, string? NewAbortByUser = null, LogId? NewLogId = null, ObjectId? NewNotificationTriggerId = null, string? NewRetryByUser = null, Priority? NewPriority = null, List<Report>? NewReports = null, Dictionary<string, string?>? NewProperties = null)
 		{
 			using IDisposable Scope = Logger.BeginScope("UpdateStepAsync({JobId})", Job.Id);
 			for (; ;)
@@ -823,7 +826,7 @@ namespace HordeServer.Services
 		/// <param name="NewReports">New reports</param>
 		/// <param name="NewProperties">Property changes. Any properties with a null value will be removed.</param>
 		/// <returns>True if the job was updated, false if it was deleted in the meantime</returns>
-		public async Task<bool> TryUpdateStepAsync(IJob Job, SubResourceId BatchId, SubResourceId StepId, JobStepState NewState = JobStepState.Unspecified, JobStepOutcome NewOutcome = JobStepOutcome.Unspecified, bool? NewAbortRequested = null, string? NewAbortByUser = null, ObjectId? NewLogId = null, ObjectId? NewTriggerId = null, string? NewRetryByUser = null, Priority? NewPriority = null, List<Report>? NewReports = null, Dictionary<string, string?>? NewProperties = null)
+		public async Task<bool> TryUpdateStepAsync(IJob Job, SubResourceId BatchId, SubResourceId StepId, JobStepState NewState = JobStepState.Unspecified, JobStepOutcome NewOutcome = JobStepOutcome.Unspecified, bool? NewAbortRequested = null, string? NewAbortByUser = null, LogId? NewLogId = null, ObjectId? NewTriggerId = null, string? NewRetryByUser = null, Priority? NewPriority = null, List<Report>? NewReports = null, Dictionary<string, string?>? NewProperties = null)
 		{
 			using IDisposable Scope = Logger.BeginScope("TryUpdateStepAsync({JobId})", Job.Id);
 
@@ -1094,8 +1097,8 @@ namespace HordeServer.Services
 			for (; ; )
 			{
 				// Update the job
-				ObjectId ChainedJobId = ObjectId.GenerateNewId();
-				if (await Jobs.TryUpdateJobAsync(Job, Graph, JobTrigger: new KeyValuePair<TemplateRefId, ObjectId>(JobTrigger.TemplateRefId, ChainedJobId)))
+				JobId ChainedJobId = JobId.GenerateNewId();
+				if (await Jobs.TryUpdateJobAsync(Job, Graph, JobTrigger: new KeyValuePair<TemplateRefId, JobId>(JobTrigger.TemplateRefId, ChainedJobId)))
 				{
 					IStream? Stream = await StreamService.GetStreamAsync(Job.StreamId);
 					if(Stream == null)
@@ -1180,7 +1183,7 @@ namespace HordeServer.Services
 		/// <param name="User">The principal to authorize</param>
 		/// <param name="Cache">Cache of job permissions</param>
 		/// <returns>True if the action is authorized</returns>
-		public async Task<bool> AuthorizeAsync(ObjectId JobId, AclAction Action, ClaimsPrincipal User, JobPermissionsCache? Cache)
+		public async Task<bool> AuthorizeAsync(JobId JobId, AclAction Action, ClaimsPrincipal User, JobPermissionsCache? Cache)
 		{
 			IJobPermissions? Permissions;
 			if (Cache == null)
diff --git a/Engine/Source/Programs/Horde/HordeServer/Jobs/Impl/JobStepRefCollection.cs b/Engine/Source/Programs/Horde/HordeServer/Jobs/Impl/JobStepRefCollection.cs
--- a/Engine/Source/Programs/Horde/HordeServer/Jobs/Impl/JobStepRefCollection.cs
+++ b/Engine/Source/Programs/Horde/HordeServer/Jobs/Impl/JobStepRefCollection.cs
@@ -14,9 +14,10 @@ using HordeServer.Utilities;
 
 namespace HordeServer.Collections.Impl
 {
+	using LogId = ObjectId<ILogFile>;
+	using PoolId = StringId<IPool>;
 	using StreamId = StringId<IStream>;
 	using TemplateRefId = StringId<TemplateRef>;
-	using PoolId = StringId<IPool>;
 
 	/// <summary>
 	/// Collection of JobStepRef documents
@@ -35,7 +36,7 @@ namespace HordeServer.Collections.Impl
 			public StreamId StreamId { get; set; }
 			public TemplateRefId TemplateId { get; set; }
 			public int Change { get; set; }
-			public ObjectId? LogId { get; set; }
+			public LogId? LogId { get; set; }
 			public PoolId? PoolId { get; set; }
 			public AgentId? AgentId { get; set; }
 			public JobStepOutcome? Outcome { get; set; }
@@ -58,7 +59,7 @@ namespace HordeServer.Collections.Impl
 			DateTime? IJobStepRef.FinishTimeUtc => FinishTimeUtc ?? FinishTime?.UtcDateTime;
 			string IJobStepRef.NodeName => Name;
 
-			public JobStepRef(JobStepRefId Id, string JobName, string NodeName, StreamId StreamId, TemplateRefId TemplateId, int Change, ObjectId? LogId, PoolId? PoolId, AgentId? AgentId, JobStepOutcome? Outcome, float BatchWaitTime, float BatchInitTime, DateTime StartTimeUtc, DateTime? FinishTimeUtc)
+			public JobStepRef(JobStepRefId Id, string JobName, string NodeName, StreamId StreamId, TemplateRefId TemplateId, int Change, LogId? LogId, PoolId? PoolId, AgentId? AgentId, JobStepOutcome? Outcome, float BatchWaitTime, float BatchInitTime, DateTime StartTimeUtc, DateTime? FinishTimeUtc)
 			{
 				this.Id = Id;
 				this.JobName = JobName;
@@ -96,7 +97,7 @@ namespace HordeServer.Collections.Impl
 		}
 
 		/// <inheritdoc/>
-		public async Task<IJobStepRef> InsertOrReplaceAsync(JobStepRefId Id, string JobName, string StepName, StreamId StreamId, TemplateRefId TemplateId, int Change, ObjectId? LogId, PoolId? PoolId, AgentId? AgentId, JobStepOutcome? Outcome, float WaitTime, float InitTime, DateTime StartTimeUtc, DateTime? FinishTimeUtc)
+		public async Task<IJobStepRef> InsertOrReplaceAsync(JobStepRefId Id, string JobName, string StepName, StreamId StreamId, TemplateRefId TemplateId, int Change, LogId? LogId, PoolId? PoolId, AgentId? AgentId, JobStepOutcome? Outcome, float WaitTime, float InitTime, DateTime StartTimeUtc, DateTime? FinishTimeUtc)
 		{
 			JobStepRef NewJobStepRef = new JobStepRef(Id, JobName, StepName, StreamId, TemplateId, Change, LogId, PoolId, AgentId, Outcome, WaitTime, InitTime, StartTimeUtc, FinishTimeUtc);
 			await JobStepRefs.ReplaceOneAsync(Builders<JobStepRef>.Filter.Eq(x => x.Id, NewJobStepRef.Id), NewJobStepRef, new ReplaceOptions { IsUpsert = true });
diff --git a/Engine/Source/Programs/Horde/HordeServer/Jobs/Impl/JobTaskSource.cs b/Engine/Source/Programs/Horde/HordeServer/Jobs/Impl/JobTaskSource.cs
--- a/Engine/Source/Programs/Horde/HordeServer/Jobs/Impl/JobTaskSource.cs
+++ b/Engine/Source/Programs/Horde/HordeServer/Jobs/Impl/JobTaskSource.cs
@@ -32,6 +32,9 @@ using EpicGames.Horde.Compute;
 
 namespace HordeServer.Tasks.Impl
 {
+	using JobId = ObjectId<IJob>;
+	using LeaseId = ObjectId<ILease>;
+	using LogId = ObjectId<ILogFile>;
 	using PoolId = StringId<IPool>;
 	using StreamId = StringId<IStream>;
 
@@ -89,7 +92,7 @@ namespace HordeServer.Tasks.Impl
 			/// <summary>
 			/// Returns an identifier describing this unique batch
 			/// </summary>
-			public (ObjectId, SubResourceId) Id
+			public (JobId, SubResourceId) Id
 			{
 				get { return (Job.Id, Batch.Id); }
 			}
@@ -255,7 +258,7 @@ namespace HordeServer.Tasks.Impl
 		/// <summary>
 		/// Map from batch id to the corresponding queue item
 		/// </summary>
-		Dictionary<(ObjectId, SubResourceId), QueueItem> BatchIdToQueueItem = new Dictionary<(ObjectId, SubResourceId), QueueItem>();
+		Dictionary<(JobId, SubResourceId), QueueItem> BatchIdToQueueItem = new Dictionary<(JobId, SubResourceId), QueueItem>();
 
 		/// <summary>
 		/// Set of long-poll tasks waiting to be satisfied 
@@ -421,7 +424,7 @@ namespace HordeServer.Tasks.Impl
 
 			// New list of queue items
 			SortedSet<QueueItem> NewQueue = new SortedSet<QueueItem>(Queue.Comparer);
-			Dictionary<(ObjectId, SubResourceId), QueueItem> NewBatchIdToQueueItem = new Dictionary<(ObjectId, SubResourceId), QueueItem>();
+			Dictionary<(JobId, SubResourceId), QueueItem> NewBatchIdToQueueItem = new Dictionary<(JobId, SubResourceId), QueueItem>();
 
 			// Query for a new list of jobs for the queue
 			List<IJob> NewJobs = await Jobs.GetDispatchQueueAsync();
@@ -723,9 +726,9 @@ namespace HordeServer.Tasks.Impl
 		}
 
 		/// <inheritdoc/>
-		public override Task CancelLeaseAsync(IAgent Agent, ObjectId LeaseId, ExecuteJobTask Task)
+		public override Task CancelLeaseAsync(IAgent Agent, LeaseId LeaseId, ExecuteJobTask Task)
 		{
-			return CancelLeaseAsync(Agent, Task.JobId.ToObjectId(), Task.BatchId.ToSubResourceId());
+			return CancelLeaseAsync(Agent, new JobId(Task.JobId), Task.BatchId.ToSubResourceId());
 		}
 
 		/// <summary>
@@ -742,13 +745,13 @@ namespace HordeServer.Tasks.Impl
 			Logger.LogDebug("Assigning job {JobId}, batch {BatchId} to waiter (agent {AgentID})", Job.Id, Batch.Id, Agent.Id);
 
 			// Generate a new unique id for the lease
-			ObjectId LeaseId = ObjectId.GenerateNewId();
+			LeaseId LeaseId = LeaseId.GenerateNewId();
 
 			// The next time to try assigning to another agent
 			DateTime BackOffTime = DateTime.UtcNow + TimeSpan.FromMinutes(1.0);
 
 			// Try to update the job with this agent id
-			ObjectId LogId = (await LogFileService.CreateLogFileAsync(Job.Id, Agent.SessionId, LogType.Json)).Id;
+			LogId LogId = (await LogFileService.CreateLogFileAsync(Job.Id, Agent.SessionId, LogType.Json)).Id;
 			if (await Jobs.TryAssignLeaseAsync(Item.Job, Item.BatchIdx, Item.PoolId, Agent.Id, Agent.SessionId!.Value, LeaseId, LogId))
 			{
 				// Get the lease name
@@ -816,7 +819,7 @@ namespace HordeServer.Tasks.Impl
 			return null;
 		}
 
-		async Task<ExecuteJobTask?> CreateExecuteJobTaskAsync(IStream Stream, IJob Job, IJobStepBatch Batch, IAgent Agent, AgentWorkspace Workspace, ObjectId LogId)
+		async Task<ExecuteJobTask?> CreateExecuteJobTaskAsync(IStream Stream, IJob Job, IJobStepBatch Batch, IAgent Agent, AgentWorkspace Workspace, LogId LogId)
 		{
 			// Get the lease name
 			StringBuilder LeaseName = new StringBuilder($"{Stream.Name} - ");
@@ -974,12 +977,12 @@ namespace HordeServer.Tasks.Impl
 		}
 
 		/// <inheritdoc/>
-		public override async Task OnLeaseFinishedAsync(IAgent Agent, ObjectId LeaseId, ExecuteJobTask Task, LeaseOutcome Outcome, ReadOnlyMemory<byte> Output, ILogger Logger)
+		public override async Task OnLeaseFinishedAsync(IAgent Agent, LeaseId LeaseId, ExecuteJobTask Task, LeaseOutcome Outcome, ReadOnlyMemory<byte> Output, ILogger Logger)
 		{
 			if (Outcome != LeaseOutcome.Success)
 			{
 				AgentId AgentId = Agent.Id;
-				ObjectId JobId = Task.JobId.ToObjectId();
+				JobId JobId = new JobId(Task.JobId);
 				SubResourceId BatchId = Task.BatchId.ToSubResourceId();
 
 				// Update the batch
@@ -1029,7 +1032,7 @@ namespace HordeServer.Tasks.Impl
 		/// <param name="JobId"></param>
 		/// <param name="BatchId"></param>
 		/// <returns></returns>
-		async Task CancelLeaseAsync(IAgent Agent, ObjectId JobId, SubResourceId BatchId)
+		async Task CancelLeaseAsync(IAgent Agent, JobId JobId, SubResourceId BatchId)
 		{
 			Logger.LogDebug("Cancelling lease for job {JobId}, batch {BatchId}", JobId, BatchId);
 
diff --git a/Engine/Source/Programs/Horde/HordeServer/Jobs/Impl/JobTimingCollection.cs b/Engine/Source/Programs/Horde/HordeServer/Jobs/Impl/JobTimingCollection.cs
--- a/Engine/Source/Programs/Horde/HordeServer/Jobs/Impl/JobTimingCollection.cs
+++ b/Engine/Source/Programs/Horde/HordeServer/Jobs/Impl/JobTimingCollection.cs
@@ -16,6 +16,8 @@ using System.Threading.Tasks;
 
 namespace HordeServer.Collections.Impl
 {
+	using JobId = ObjectId<IJob>;
+
 	/// <summary>
 	/// Concrete implementation of IJobTimingCollection
 	/// </summary>
@@ -51,7 +53,7 @@ namespace HordeServer.Collections.Impl
 		/// </summary>
 		class JobTimingDocument : IJobTiming
 		{
-			public ObjectId Id { get; set; }
+			public JobId Id { get; set; }
 			public List<JobStepTimingDocument> Steps { get; set; } = new List<JobStepTimingDocument>();
 			public int UpdateIndex { get; set; }
 
@@ -91,7 +93,7 @@ namespace HordeServer.Collections.Impl
 		}
 
 		/// <inheritdoc/>
-		public async Task<IJobTiming?> TryAddAsync(ObjectId JobId, List<JobStepTimingData> Steps)
+		public async Task<IJobTiming?> TryAddAsync(JobId JobId, List<JobStepTimingData> Steps)
 		{
 			JobTimingDocument JobTiming = new JobTimingDocument();
 			JobTiming.Id = JobId;
@@ -117,7 +119,7 @@ namespace HordeServer.Collections.Impl
 		}
 
 		/// <inheritdoc/>
-		public async Task<IJobTiming?> TryGetAsync(ObjectId JobId)
+		public async Task<IJobTiming?> TryGetAsync(JobId JobId)
 		{
 			return await Collection.Find(x => x.Id == JobId).FirstOrDefaultAsync();
 		}
diff --git a/Engine/Source/Programs/Horde/HordeServer/Jobs/Impl/TestDataCollection.cs b/Engine/Source/Programs/Horde/HordeServer/Jobs/Impl/TestDataCollection.cs
--- a/Engine/Source/Programs/Horde/HordeServer/Jobs/Impl/TestDataCollection.cs
+++ b/Engine/Source/Programs/Horde/HordeServer/Jobs/Impl/TestDataCollection.cs
@@ -14,6 +14,7 @@ using System.Threading.Tasks;
 
 namespace HordeServer.Collections.Impl
 {
+	using JobId = ObjectId<IJob>;
 	using ProjectId = StringId<IProject>;
 	using StreamId = StringId<IStream>;
 	using TemplateRefId = StringId<TemplateRef>;
@@ -31,7 +32,7 @@ namespace HordeServer.Collections.Impl
 			public ObjectId Id { get; set; }
 			public StreamId StreamId { get; set; }
 			public TemplateRefId TemplateRefId { get; set; }
-			public ObjectId JobId { get; set; }
+			public JobId JobId { get; set; }
 			public SubResourceId StepId { get; set; }
 			public int Change { get; set; }
 			public string Key { get; set; }
@@ -91,7 +92,7 @@ namespace HordeServer.Collections.Impl
 		}
 
 		/// <inheritdoc/>
-		public async Task<List<ITestData>> FindAsync(StreamId? StreamId, int? MinChange, int? MaxChange, ObjectId? JobId, SubResourceId? StepId, string? Key = null, int Index = 0, int Count = 10)
+		public async Task<List<ITestData>> FindAsync(StreamId? StreamId, int? MinChange, int? MaxChange, JobId? JobId, SubResourceId? StepId, string? Key = null, int Index = 0, int Count = 10)
 		{
 			FilterDefinition<TestDataDocument> Filter = FilterDefinition<TestDataDocument>.Empty;
 			if (StreamId != null)
diff --git a/Engine/Source/Programs/Horde/HordeServer/Logs/Builder/ILogBuilder.cs b/Engine/Source/Programs/Horde/HordeServer/Logs/Builder/ILogBuilder.cs
--- a/Engine/Source/Programs/Horde/HordeServer/Logs/Builder/ILogBuilder.cs
+++ b/Engine/Source/Programs/Horde/HordeServer/Logs/Builder/ILogBuilder.cs
@@ -2,6 +2,7 @@
 
 using HordeServer.Api;
 using HordeServer.Models;
+using HordeServer.Utilities;
 using MongoDB.Bson;
 using System;
 using System.Collections.Generic;
@@ -10,6 +11,8 @@ using System.Threading.Tasks;
 
 namespace HordeServer.Logs
 {
+	using LogId = ObjectId<ILogFile>;
+
 	/// <summary>
 	/// Interface for the log file write cache
 	/// </summary>
@@ -31,7 +34,7 @@ namespace HordeServer.Logs
 		/// <param name="Data">Data to be appended</param>
 		/// <param name="Type">Type of data stored in this log file</param>
 		/// <returns>True if the data was appended to the given chunk. False if the chunk has been completed.</returns>
-		Task<bool> AppendAsync(ObjectId LogId, long ChunkOffset, long WriteOffset, int WriteLineIndex, int WriteLineCount, ReadOnlyMemory<byte> Data, LogType Type);
+		Task<bool> AppendAsync(LogId LogId, long ChunkOffset, long WriteOffset, int WriteLineIndex, int WriteLineCount, ReadOnlyMemory<byte> Data, LogType Type);
 
 		/// <summary>
 		/// Finish the current sub chunk
@@ -39,7 +42,7 @@ namespace HordeServer.Logs
 		/// <param name="LogId">The log file id</param>
 		/// <param name="Offset">Offset of the chunk within the log file</param>
 		/// <returns>Async task</returns>
-		Task CompleteSubChunkAsync(ObjectId LogId, long Offset);
+		Task CompleteSubChunkAsync(LogId LogId, long Offset);
 
 		/// <summary>
 		/// Finish the current chunk
@@ -47,7 +50,7 @@ namespace HordeServer.Logs
 		/// <param name="LogId">The log file id</param>
 		/// <param name="Offset">Offset of the chunk within the log file</param>
 		/// <returns>Async task</returns>
-		Task CompleteChunkAsync(ObjectId LogId, long Offset);
+		Task CompleteChunkAsync(LogId LogId, long Offset);
 
 		/// <summary>
 		/// Remove a complete chunk from the builder
@@ -55,7 +58,7 @@ namespace HordeServer.Logs
 		/// <param name="LogId">The log file id</param>
 		/// <param name="Offset">Offset of the chunk within the log file</param>
 		/// <returns>Async task</returns>
-		Task RemoveChunkAsync(ObjectId LogId, long Offset);
+		Task RemoveChunkAsync(LogId LogId, long Offset);
 
 		/// <summary>
 		/// Gets the current chunk for the given log file
@@ -64,13 +67,13 @@ namespace HordeServer.Logs
 		/// <param name="Offset">Offset of the chunk within the log file</param>
 		/// <param name="LineIndex">Line index of the chunk within the log file</param>
 		/// <returns></returns>
-		Task<LogChunkData?> GetChunkAsync(ObjectId LogId, long Offset, int LineIndex);
+		Task<LogChunkData?> GetChunkAsync(LogId LogId, long Offset, int LineIndex);
 
 		/// <summary>
 		/// Touches the timestamps of all the chunks after the given age, and returns them. Used for flushing the builder.
 		/// </summary>
 		/// <param name="MinAge">Minimum age of the chunks to enumerate. If specified, only chunks last modified longer than this period will be returned.</param>
 		/// <returns>List of chunks, identified by log id and chunk index</returns>
-		Task<List<(ObjectId, long)>> TouchChunksAsync(TimeSpan MinAge);
+		Task<List<(LogId, long)>> TouchChunksAsync(TimeSpan MinAge);
 	}
 }
diff --git a/Engine/Source/Programs/Horde/HordeServer/Logs/Builder/Impl/LocalLogBuilder.cs b/Engine/Source/Programs/Horde/HordeServer/Logs/Builder/Impl/LocalLogBuilder.cs
--- a/Engine/Source/Programs/Horde/HordeServer/Logs/Builder/Impl/LocalLogBuilder.cs
+++ b/Engine/Source/Programs/Horde/HordeServer/Logs/Builder/Impl/LocalLogBuilder.cs
@@ -13,6 +13,8 @@ using System.Threading.Tasks;
 
 namespace HordeServer.Logs.Builder
 {
+	using LogId = ObjectId<ILogFile>;
+
 	/// <summary>
 	/// In-memory implementation of a log write buffer
 	/// </summary>
@@ -109,13 +111,13 @@ namespace HordeServer.Logs.Builder
 		/// <summary>
 		/// Current log chunk state
 		/// </summary>
-		ConcurrentDictionary<(ObjectId, long), PendingChunk> PendingChunks = new ConcurrentDictionary<(ObjectId, long), PendingChunk>();
+		ConcurrentDictionary<(LogId, long), PendingChunk> PendingChunks = new ConcurrentDictionary<(LogId, long), PendingChunk>();
 
 		/// <inheritdoc/>
 		public bool FlushOnShutdown => true;
 
 		/// <inheritdoc/>
-		public Task<bool> AppendAsync(ObjectId LogId, long ChunkOffset, long WriteOffset, int WriteLineIndex, int WriteLineCount, ReadOnlyMemory<byte> Data, LogType Type)
+		public Task<bool> AppendAsync(LogId LogId, long ChunkOffset, long WriteOffset, int WriteLineIndex, int WriteLineCount, ReadOnlyMemory<byte> Data, LogType Type)
 		{
 			PendingChunk? PendingChunk;
 			while (!PendingChunks.TryGetValue((LogId, ChunkOffset), out PendingChunk))
@@ -146,7 +148,7 @@ namespace HordeServer.Logs.Builder
 		}
 
 		/// <inheritdoc/>
-		public Task CompleteSubChunkAsync(ObjectId LogId, long Offset)
+		public Task CompleteSubChunkAsync(LogId LogId, long Offset)
 		{
 			PendingChunk? PendingChunk;
 			if (PendingChunks.TryGetValue((LogId, Offset), out PendingChunk))
@@ -160,7 +162,7 @@ namespace HordeServer.Logs.Builder
 		}
 
 		/// <inheritdoc/>
-		public Task CompleteChunkAsync(ObjectId LogId, long Offset)
+		public Task CompleteChunkAsync(LogId LogId, long Offset)
 		{
 			PendingChunk? PendingChunk;
 			if (PendingChunks.TryGetValue((LogId, Offset), out PendingChunk))
@@ -178,14 +180,14 @@ namespace HordeServer.Logs.Builder
 		}
 
 		/// <inheritdoc/>
-		public Task RemoveChunkAsync(ObjectId LogId, long Offset)
+		public Task RemoveChunkAsync(LogId LogId, long Offset)
 		{
 			PendingChunks.TryRemove((LogId, Offset), out _);
 			return Task.CompletedTask;
 		}
 
 		/// <inheritdoc/>
-		public Task<LogChunkData?> GetChunkAsync(ObjectId LogId, long Offset, int LineIndex)
+		public Task<LogChunkData?> GetChunkAsync(LogId LogId, long Offset, int LineIndex)
 		{
 			PendingChunk? PendingChunk;
 			if (PendingChunks.TryGetValue((LogId, Offset), out PendingChunk))
@@ -216,12 +218,12 @@ namespace HordeServer.Logs.Builder
 		}
 
 		/// <inheritdoc/>
-		public Task<List<(ObjectId, long)>> TouchChunksAsync(TimeSpan MinAge)
+		public Task<List<(LogId, long)>> TouchChunksAsync(TimeSpan MinAge)
 		{
 			DateTime UtcNow = DateTime.UtcNow;
 
-			List<(ObjectId, long)> Chunks = new List<(ObjectId, long)>();
-			foreach (KeyValuePair<(ObjectId, long), PendingChunk> PendingChunk in PendingChunks.ToArray())
+			List<(LogId, long)> Chunks = new List<(LogId, long)>();
+			foreach (KeyValuePair<(LogId, long), PendingChunk> PendingChunk in PendingChunks.ToArray())
 			{
 				lock (PendingChunk.Value)
 				{
diff --git a/Engine/Source/Programs/Horde/HordeServer/Logs/Builder/Impl/RedisLogBuilder.cs b/Engine/Source/Programs/Horde/HordeServer/Logs/Builder/Impl/RedisLogBuilder.cs
--- a/Engine/Source/Programs/Horde/HordeServer/Logs/Builder/Impl/RedisLogBuilder.cs
+++ b/Engine/Source/Programs/Horde/HordeServer/Logs/Builder/Impl/RedisLogBuilder.cs
@@ -2,6 +2,7 @@
 
 using EpicGames.Core;
 using HordeServer.Api;
+using HordeServer.Utilities;
 using Microsoft.Extensions.Logging;
 using Microsoft.Extensions.Primitives;
 using MongoDB.Bson;
@@ -12,12 +13,13 @@ using System.Globalization;
 using System.Linq;
 using System.Text.RegularExpressions;
 using System.Threading.Tasks;
-using HordeServer.Utilities;
 using OpenTracing;
 using OpenTracing.Util;
+using HordeServer.Models;
 
 namespace HordeServer.Logs.Builder
 {
+	using LogId = ObjectId<ILogFile>;
 	using Condition = StackExchange.Redis.Condition;
 
 	/// <summary>
@@ -37,17 +39,17 @@ namespace HordeServer.Logs.Builder
 			public string SubChunkData => $"{Prefix}-subchunk";
 			public string Complete => $"{Prefix}-complete";
 
-			public ChunkKeys(ObjectId LogId, long Offset)
+			public ChunkKeys(LogId LogId, long Offset)
 			{
 				Prefix = $"log-{LogId}-chunk-{Offset}-builder";
 			}
 
-			public static bool TryParse(string Prefix, out ObjectId LogId, out long Offset)
+			public static bool TryParse(string Prefix, out LogId LogId, out long Offset)
 			{
 				Match Match = Regex.Match(Prefix, "log-([^-]+)-chunk-([^-]+)-builder");
-				if (!Match.Success || !ObjectId.TryParse(Match.Groups[1].Value, out LogId) || !long.TryParse(Match.Groups[2].Value, out Offset))
+				if (!Match.Success || !LogId.TryParse(Match.Groups[1].Value, out LogId) || !long.TryParse(Match.Groups[2].Value, out Offset))
 				{
-					LogId = new ObjectId();
+					LogId = LogId.Empty;
 					Offset = 0;
 					return false;
 				}
@@ -80,7 +82,7 @@ namespace HordeServer.Logs.Builder
 		}
 
 		/// <inheritdoc/>
-		public async Task<bool> AppendAsync(ObjectId LogId, long ChunkOffset, long WriteOffset, int WriteLineIndex, int WriteLineCount, ReadOnlyMemory<byte> Data, LogType Type)
+		public async Task<bool> AppendAsync(LogId LogId, long ChunkOffset, long WriteOffset, int WriteLineIndex, int WriteLineCount, ReadOnlyMemory<byte> Data, LogType Type)
 		{
 			IDatabase RedisDb = RedisConnectionPool.GetDatabase();
 			ChunkKeys Keys = new ChunkKeys(LogId, ChunkOffset);
@@ -131,7 +133,7 @@ namespace HordeServer.Logs.Builder
 		}
 
 		/// <inheritdoc/>
-		public async Task CompleteSubChunkAsync(ObjectId LogId, long Offset)
+		public async Task CompleteSubChunkAsync(LogId LogId, long Offset)
 		{
 			IDatabase RedisDb = RedisConnectionPool.GetDatabase();
 			ChunkKeys Keys = new ChunkKeys(LogId, Offset);
@@ -174,7 +176,7 @@ namespace HordeServer.Logs.Builder
 		}
 
 		/// <inheritdoc/>
-		public async Task CompleteChunkAsync(ObjectId LogId, long Offset)
+		public async Task CompleteChunkAsync(LogId LogId, long Offset)
 		{
 			IDatabase RedisDb = RedisConnectionPool.GetDatabase();
 			using IScope Scope = GlobalTracer.Instance.BuildSpan("Redis.CompleteChunk").StartActive();
@@ -196,7 +198,7 @@ namespace HordeServer.Logs.Builder
 		}
 
 		/// <inheritdoc/>
-		public async Task RemoveChunkAsync(ObjectId LogId, long Offset)
+		public async Task RemoveChunkAsync(LogId LogId, long Offset)
 		{
 			IDatabase RedisDb = RedisConnectionPool.GetDatabase();
 			using IScope Scope = GlobalTracer.Instance.BuildSpan("Redis.RemoveChunk").StartActive();
@@ -218,7 +220,7 @@ namespace HordeServer.Logs.Builder
 		}
 
 		/// <inheritdoc/>
-		public async Task<LogChunkData?> GetChunkAsync(ObjectId LogId, long Offset, int LineIndex)
+		public async Task<LogChunkData?> GetChunkAsync(LogId LogId, long Offset, int LineIndex)
 		{
 			IDatabase RedisDb = RedisConnectionPool.GetDatabase();
 			ChunkKeys Keys = new ChunkKeys(LogId, Offset);
@@ -274,7 +276,7 @@ namespace HordeServer.Logs.Builder
 		}
 
 		/// <inheritdoc/>
-		public async Task<List<(ObjectId, long)>> TouchChunksAsync(TimeSpan MinAge)
+		public async Task<List<(LogId, long)>> TouchChunksAsync(TimeSpan MinAge)
 		{
 			IDatabase RedisDb = RedisConnectionPool.GetDatabase();
 			
@@ -283,10 +285,10 @@ namespace HordeServer.Logs.Builder
 			SortedSetEntry[] Entries = await RedisDb.SortedSetRangeByScoreWithScoresAsync(ItemsKey, stop: (UtcNow - MinAge).Ticks);
 
 			// Update the score for each element in a transaction. If it succeeds, we can write the chunk. Otherwise another pod has beat us to it.
-			List<(ObjectId, long)> Results = new List<(ObjectId, long)>();
+			List<(LogId, long)> Results = new List<(LogId, long)>();
 			foreach (SortedSetEntry Entry in Entries)
 			{
-				if (ChunkKeys.TryParse(Entry.Element.ToString(), out ObjectId LogId, out long Offset))
+				if (ChunkKeys.TryParse(Entry.Element.ToString(), out LogId LogId, out long Offset))
 				{
 					ITransaction Transaction = RedisDb.CreateTransaction();
 					Transaction.AddCondition(Condition.SortedSetEqual(ItemsKey, Entry.Element, Entry.Score));
diff --git a/Engine/Source/Programs/Horde/HordeServer/Logs/ILogEvent.cs b/Engine/Source/Programs/Horde/HordeServer/Logs/ILogEvent.cs
--- a/Engine/Source/Programs/Horde/HordeServer/Logs/ILogEvent.cs
+++ b/Engine/Source/Programs/Horde/HordeServer/Logs/ILogEvent.cs
@@ -2,6 +2,7 @@
 
 using HordeCommon;
 using HordeCommon.Rpc;
+using HordeServer.Utilities;
 using Microsoft.Extensions.Logging;
 using MongoDB.Bson;
 using System;
@@ -13,6 +14,8 @@ using System.Threading.Tasks;
 
 namespace HordeServer.Models
 {
+	using LogId = ObjectId<ILogFile>;
+
 	/// <summary>
 	/// Interface for a log event line
 	/// </summary>
@@ -73,7 +76,7 @@ namespace HordeServer.Models
 		/// <summary>
 		/// Unique id of the log containing this event
 		/// </summary>
-		public ObjectId LogId { get; }
+		public LogId LogId { get; }
 
 		/// <summary>
 		/// Severity of the event
@@ -104,7 +107,7 @@ namespace HordeServer.Models
 		/// <summary>
 		/// Unique id of the log containing this event
 		/// </summary>
-		public ObjectId LogId { get; set; }
+		public LogId LogId { get; set; }
 
 		/// <summary>
 		/// Severity of the event
diff --git a/Engine/Source/Programs/Horde/HordeServer/Logs/ILogEventCollection.cs b/Engine/Source/Programs/Horde/HordeServer/Logs/ILogEventCollection.cs
--- a/Engine/Source/Programs/Horde/HordeServer/Logs/ILogEventCollection.cs
+++ b/Engine/Source/Programs/Horde/HordeServer/Logs/ILogEventCollection.cs
@@ -17,6 +17,8 @@ using System.Threading.Tasks;
 
 namespace HordeServer.Collections
 {
+	using LogId = ObjectId<ILogFile>;
+
 	/// <summary>
 	/// Interface for a collection of event documents
 	/// </summary>
@@ -41,7 +43,7 @@ namespace HordeServer.Collections
 		/// <param name="Index">Start index within the matching results</param>
 		/// <param name="Count">Maximum number of results to return</param>
 		/// <returns>List of events matching the query</returns>
-		Task<List<ILogEvent>> FindAsync(ObjectId LogId, int? Index = null, int? Count = null);
+		Task<List<ILogEvent>> FindAsync(LogId LogId, int? Index = null, int? Count = null);
 
 		/// <summary>
 		/// Finds a list of events for a set of spans
@@ -51,14 +53,14 @@ namespace HordeServer.Collections
 		/// <param name="Index">Index of the first result to return</param>
 		/// <param name="Count">Number of results to return</param>
 		/// <returns>List of events for this issue</returns>
-		Task<List<ILogEvent>> FindEventsForSpansAsync(IEnumerable<ObjectId> SpanIds, ObjectId[]? LogIds = null, int Index = 0, int Count = 10);
+		Task<List<ILogEvent>> FindEventsForSpansAsync(IEnumerable<ObjectId> SpanIds, LogId[]? LogIds = null, int Index = 0, int Count = 10);
 
 		/// <summary>
 		/// Delete all the events for a log file
 		/// </summary>
 		/// <param name="LogId">Unique id of the log</param>
 		/// <returns>Async task</returns>
-		Task DeleteLogAsync(ObjectId LogId);
+		Task DeleteLogAsync(LogId LogId);
 
 		/// <summary>
 		/// Update the span for an event
diff --git a/Engine/Source/Programs/Horde/HordeServer/Logs/ILogFile.cs b/Engine/Source/Programs/Horde/HordeServer/Logs/ILogFile.cs
--- a/Engine/Source/Programs/Horde/HordeServer/Logs/ILogFile.cs
+++ b/Engine/Source/Programs/Horde/HordeServer/Logs/ILogFile.cs
@@ -14,6 +14,9 @@ using System.Threading.Tasks;
 
 namespace HordeServer.Models
 {
+	using JobId = ObjectId<IJob>;
+	using LogId = ObjectId<ILogFile>;
+
 	/// <summary>
 	/// Information about a log file chunk
 	/// </summary>
@@ -48,12 +51,12 @@ namespace HordeServer.Models
 		/// <summary>
 		/// Identifier for the LogFile. Randomly generated.
 		/// </summary>
-		public ObjectId Id { get; }
+		public LogId Id { get; }
 
 		/// <summary>
 		/// Unique id of the job containing this log
 		/// </summary>
-		public ObjectId JobId { get; }
+		public JobId JobId { get; }
 
 		/// <summary>
 		/// The session allowed to write to this log
diff --git a/Engine/Source/Programs/Horde/HordeServer/Logs/ILogFileCollection.cs b/Engine/Source/Programs/Horde/HordeServer/Logs/ILogFileCollection.cs
--- a/Engine/Source/Programs/Horde/HordeServer/Logs/ILogFileCollection.cs
+++ b/Engine/Source/Programs/Horde/HordeServer/Logs/ILogFileCollection.cs
@@ -3,6 +3,7 @@
 using HordeServer.Api;
 using HordeServer.Models;
 using HordeServer.Services;
+using HordeServer.Utilities;
 using MongoDB.Bson;
 using MongoDB.Bson.Serialization.Attributes;
 using MongoDB.Driver;
@@ -14,6 +15,9 @@ using System.Threading.Tasks;
 
 namespace HordeServer.Collections
 {
+	using JobId = ObjectId<IJob>;
+	using LogId = ObjectId<ILogFile>;
+
 	/// <summary>
 	/// Updates a log file chunk
 	/// </summary>
@@ -60,7 +64,7 @@ namespace HordeServer.Collections
 		/// <param name="SessionId">Agent session allowed to update the log</param>
 		/// <param name="Type">Type of events to be stored in the log</param>
 		/// <returns>The new log file document</returns>
-		Task<ILogFile> CreateLogFileAsync(ObjectId JobId, ObjectId? SessionId, LogType Type);
+		Task<ILogFile> CreateLogFileAsync(JobId JobId, ObjectId? SessionId, LogType Type);
 
 		/// <summary>
 		/// Adds a new chunk
@@ -92,7 +96,7 @@ namespace HordeServer.Collections
 		/// </summary>
 		/// <param name="LogFileId">Unique id of the log file</param>
 		/// <returns>The logfile document</returns>
-		Task<ILogFile?> GetLogFileAsync(ObjectId LogFileId);
+		Task<ILogFile?> GetLogFileAsync(LogId LogFileId);
 
 		/// <summary>
 		/// Gets all the log files
diff --git a/Engine/Source/Programs/Horde/HordeServer/Logs/Impl/LogEventCollection.cs b/Engine/Source/Programs/Horde/HordeServer/Logs/Impl/LogEventCollection.cs
--- a/Engine/Source/Programs/Horde/HordeServer/Logs/Impl/LogEventCollection.cs
+++ b/Engine/Source/Programs/Horde/HordeServer/Logs/Impl/LogEventCollection.cs
@@ -17,6 +17,8 @@ using System.Threading.Tasks;
 
 namespace HordeServer.Collections.Impl
 {
+	using LogId = ObjectId<ILogFile>;
+
 	/// <summary>
 	/// Collection of event documents
 	/// </summary>
@@ -25,7 +27,7 @@ namespace HordeServer.Collections.Impl
 		class LogEventId
 		{
 			[BsonElement("l")]
-			public ObjectId LogId { get; set; }
+			public LogId LogId { get; set; }
 
 			[BsonElement("n")]
 			public int LineIndex { get; set; }
@@ -45,7 +47,7 @@ namespace HordeServer.Collections.Impl
 			[BsonElement("s")]
 			public ObjectId? SpanId { get; set; }
 
-			ObjectId ILogEvent.LogId => Id.LogId;
+			LogId ILogEvent.LogId => Id.LogId;
 			int ILogEvent.LineIndex => Id.LineIndex;
 			int ILogEvent.LineCount => LineCount ?? 1;
 			EventSeverity ILogEvent.Severity => IsWarning ? EventSeverity.Warning : EventSeverity.Error;
@@ -55,7 +57,7 @@ namespace HordeServer.Collections.Impl
 				this.Id = new LogEventId();
 			}
 
-			public LogEventDocument(ObjectId LogId, EventSeverity Severity, int LineIndex, int LineCount, ObjectId? SpanId)
+			public LogEventDocument(LogId LogId, EventSeverity Severity, int LineIndex, int LineCount, ObjectId? SpanId)
 			{
 				this.Id = new LogEventId { LogId = LogId, LineIndex = LineIndex };
 				this.IsWarning = Severity == EventSeverity.Warning;
@@ -74,7 +76,7 @@ namespace HordeServer.Collections.Impl
 			public ObjectId Id { get; set; }
 			public DateTime Time { get; set; }
 			public EventSeverity Severity { get; set; }
-			public ObjectId LogId { get; set; }
+			public LogId LogId { get; set; }
 			public int LineIndex { get; set; }
 			public int LineCount { get; set; }
 
@@ -136,9 +138,9 @@ namespace HordeServer.Collections.Impl
 		}
 
 		/// <inheritdoc/>
-		public async Task<List<ILogEvent>> FindAsync(ObjectId LogId, int? Index = null, int? Count = null)
+		public async Task<List<ILogEvent>> FindAsync(LogId LogId, int? Index = null, int? Count = null)
 		{
-			Logger.LogInformation("Querying for log events for log {LogId} creation time {CreateTime}", LogId, LogId.CreationTime);
+			Logger.LogInformation("Querying for log events for log {LogId} creation time {CreateTime}", LogId, LogId.Value.CreationTime);
 
 			FilterDefinitionBuilder<LogEventDocument> Builder = Builders<LogEventDocument>.Filter;
 
@@ -159,7 +161,7 @@ namespace HordeServer.Collections.Impl
 		}
 
 		/// <inheritdoc/>
-		public async Task<List<ILogEvent>> FindEventsForSpanAsync(ObjectId SpanId, ObjectId[]? LogIds, int Index, int Count)
+		public async Task<List<ILogEvent>> FindEventsForSpanAsync(ObjectId SpanId, LogId[]? LogIds, int Index, int Count)
 		{
 			FilterDefinition<LogEventDocument> Filter = Builders<LogEventDocument>.Filter.Eq(x => x.SpanId, SpanId);
 			if (LogIds != null && LogIds.Length > 0)
@@ -172,7 +174,7 @@ namespace HordeServer.Collections.Impl
 		}
 
 		/// <inheritdoc/>
-		public async Task<List<ILogEvent>> FindEventsForSpansAsync(IEnumerable<ObjectId> SpanIds, ObjectId[]? LogIds, int Index, int Count)
+		public async Task<List<ILogEvent>> FindEventsForSpansAsync(IEnumerable<ObjectId> SpanIds, LogId[]? LogIds, int Index, int Count)
 		{
 			FilterDefinition<LogEventDocument> Filter = Builders<LogEventDocument>.Filter.In(x => x.SpanId, SpanIds.Select<ObjectId, ObjectId?>(x => x));
 			if (LogIds != null && LogIds.Length > 0)
@@ -185,7 +187,7 @@ namespace HordeServer.Collections.Impl
 		}
 
 		/// <inheritdoc/>
-		public async Task DeleteLogAsync(ObjectId LogId)
+		public async Task DeleteLogAsync(LogId LogId)
 		{
 			await LogEvents.DeleteManyAsync(x => x.Id.LogId == LogId);
 			await LegacyEvents.DeleteManyAsync(x => x.LogId == LogId);
diff --git a/Engine/Source/Programs/Horde/HordeServer/Logs/Impl/LogFileCollection.cs b/Engine/Source/Programs/Horde/HordeServer/Logs/Impl/LogFileCollection.cs
--- a/Engine/Source/Programs/Horde/HordeServer/Logs/Impl/LogFileCollection.cs
+++ b/Engine/Source/Programs/Horde/HordeServer/Logs/Impl/LogFileCollection.cs
@@ -3,6 +3,7 @@
 using HordeServer.Api;
 using HordeServer.Models;
 using HordeServer.Services;
+using HordeServer.Utilities;
 using MongoDB.Bson;
 using MongoDB.Bson.Serialization.Attributes;
 using MongoDB.Driver;
@@ -14,6 +15,9 @@ using System.Threading.Tasks;
 
 namespace HordeServer.Collections.Impl
 {
+	using JobId = ObjectId<IJob>;
+	using LogId = ObjectId<ILogFile>;
+
 	/// <summary>
 	/// Wrapper around the jobs collection in a mongo DB
 	/// </summary>
@@ -50,10 +54,10 @@ namespace HordeServer.Collections.Impl
 		class LogFileDocument : ILogFile
 		{
 			[BsonRequired, BsonId]
-			public ObjectId Id { get; set; }
+			public LogId Id { get; set; }
 
 			[BsonRequired]
-			public ObjectId JobId { get; set; }
+			public JobId JobId { get; set; }
 
 			public ObjectId? SessionId { get; set; }
 			public LogType Type { get; set; }
@@ -76,9 +80,9 @@ namespace HordeServer.Collections.Impl
 			{
 			}
 
-			public LogFileDocument(ObjectId JobId, ObjectId? SessionId, LogType Type)
+			public LogFileDocument(JobId JobId, ObjectId? SessionId, LogType Type)
 			{
-				this.Id = ObjectId.GenerateNewId();
+				this.Id = LogId.GenerateNewId();
 				this.JobId = JobId;
 				this.SessionId = SessionId;
 				this.Type = Type;
@@ -114,7 +118,7 @@ namespace HordeServer.Collections.Impl
 		}
 
 		/// <inheritdoc/>
-		public async Task<ILogFile> CreateLogFileAsync(ObjectId JobId, ObjectId? SessionId, LogType Type)
+		public async Task<ILogFile> CreateLogFileAsync(JobId JobId, ObjectId? SessionId, LogType Type)
 		{
 			LogFileDocument NewLogFile = new LogFileDocument(JobId, SessionId, Type);
 			await LogFiles.InsertOneAsync(NewLogFile);
@@ -204,7 +208,7 @@ namespace HordeServer.Collections.Impl
 		}
 
 		/// <inheritdoc/>
-		public async Task<ILogFile?> GetLogFileAsync(ObjectId LogFileId)
+		public async Task<ILogFile?> GetLogFileAsync(LogId LogFileId)
 		{
 			LogFileDocument LogFile = await LogFiles.Find<LogFileDocument>(x => x.Id == LogFileId).FirstOrDefaultAsync();
 			return LogFile;
diff --git a/Engine/Source/Programs/Horde/HordeServer/Logs/Impl/LogFileService.cs b/Engine/Source/Programs/Horde/HordeServer/Logs/Impl/LogFileService.cs
--- a/Engine/Source/Programs/Horde/HordeServer/Logs/Impl/LogFileService.cs
+++ b/Engine/Source/Programs/Horde/HordeServer/Logs/Impl/LogFileService.cs
@@ -37,6 +37,9 @@ using OpenTracing;
 
 namespace HordeServer.Services
 {
+	using JobId = ObjectId<IJob>;
+	using LogId = ObjectId<ILogFile>;
+
 	/// <summary>
 	/// Metadata about a log file
 	/// </summary>
@@ -65,21 +68,21 @@ namespace HordeServer.Services
 		/// <param name="SessionId">Agent session allowed to update the log</param>
 		/// <param name="Type">Type of events to be stored in the log</param>
 		/// <returns>The new log file document</returns>
-		Task<ILogFile> CreateLogFileAsync(ObjectId JobId, ObjectId? SessionId, LogType Type);
+		Task<ILogFile> CreateLogFileAsync(JobId JobId, ObjectId? SessionId, LogType Type);
 
 		/// <summary>
 		/// Gets a logfile by ID
 		/// </summary>
 		/// <param name="LogFileId">Unique id of the log file</param>
 		/// <returns>The logfile document</returns>
-		Task<ILogFile?> GetLogFileAsync(ObjectId LogFileId);
+		Task<ILogFile?> GetLogFileAsync(LogId LogFileId);
 
 		/// <summary>
 		/// Gets a logfile by ID, returning a cached copy if available. This should only be used to retrieve constant properties set at creation, such as the session or job it's associated with.
 		/// </summary>
 		/// <param name="LogFileId">Unique id of the log file</param>
 		/// <returns>The logfile document</returns>
-		Task<ILogFile?> GetCachedLogFileAsync(ObjectId LogFileId);
+		Task<ILogFile?> GetCachedLogFileAsync(LogId LogFileId);
 
 		/// <summary>
 		/// Returns a list of log files
@@ -141,7 +144,7 @@ namespace HordeServer.Services
 		/// <param name="Index">Index within the events for results to return</param>
 		/// <param name="Count">Number of results to return</param>
 		/// <returns>Async task</returns>
-		Task<List<ILogEvent>> FindEventsForSpansAsync(IEnumerable<ObjectId> SpanIds, ObjectId[]? LogIds, int Index, int Count);
+		Task<List<ILogEvent>> FindEventsForSpansAsync(IEnumerable<ObjectId> SpanIds, LogId[]? LogIds, int Index, int Count);
 
 		/// <summary>
 		/// Gets the data for an event
@@ -317,7 +320,7 @@ namespace HordeServer.Services
 		/// <summary>
 		/// Set of chunks which is currently being written to
 		/// </summary>
-		private HashSet<(ObjectId, long)> WriteChunks = new HashSet<(ObjectId, long)>();
+		private HashSet<(LogId, long)> WriteChunks = new HashSet<(LogId, long)>();
 
 		/// <summary>
 		/// Cache of session ids for each log
@@ -512,13 +515,13 @@ namespace HordeServer.Services
 		}
 
 		/// <inheritdoc/>
-		public Task<ILogFile> CreateLogFileAsync(ObjectId JobId, ObjectId? SessionId, LogType Type)
+		public Task<ILogFile> CreateLogFileAsync(JobId JobId, ObjectId? SessionId, LogType Type)
 		{
 			return LogFiles.CreateLogFileAsync(JobId, SessionId, Type);
 		}
 
 		/// <inheritdoc/>
-		public async Task<ILogFile?> GetLogFileAsync(ObjectId LogFileId)
+		public async Task<ILogFile?> GetLogFileAsync(LogId LogFileId)
 		{
 			ILogFile? LogFile = await LogFiles.GetLogFileAsync(LogFileId);
 			if(LogFile != null)
@@ -543,7 +546,7 @@ namespace HordeServer.Services
 		/// </summary>
 		/// <param name="LogFileId">The log file id</param>
 		/// <returns>New log file, or null if not found</returns>
-		public async Task<ILogFile?> GetCachedLogFileAsync(ObjectId LogFileId)
+		public async Task<ILogFile?> GetCachedLogFileAsync(LogId LogFileId)
 		{
 			object? LogFile;
 			if (!LogFileCache.TryGetValue(LogFileId, out LogFile))
@@ -871,7 +874,7 @@ namespace HordeServer.Services
 		}
 
 		/// <inheritdoc/>
-		public Task<List<ILogEvent>> FindEventsForSpansAsync(IEnumerable<ObjectId> SpanIds, ObjectId[]? LogIds, int Index, int Count)
+		public Task<List<ILogEvent>> FindEventsForSpansAsync(IEnumerable<ObjectId> SpanIds, LogId[]? LogIds, int Index, int Count)
 		{
 			return LogEvents.FindEventsForSpansAsync(SpanIds, LogIds, Index, Count);
 		}
@@ -1019,11 +1022,11 @@ namespace HordeServer.Services
 		private async Task IncrementalFlush()
 		{
 			// Get all the chunks older than 20 minutes
-			List<(ObjectId, long)> FlushChunks = await Builder.TouchChunksAsync(TimeSpan.FromMinutes(10.0));
+			List<(LogId, long)> FlushChunks = await Builder.TouchChunksAsync(TimeSpan.FromMinutes(10.0));
 			Logger.LogDebug("Performing incremental flush of log builder ({NumChunks} chunks)", FlushChunks.Count);
 
 			// Mark them all as complete
-			foreach ((ObjectId LogId, long Offset) in FlushChunks)
+			foreach ((LogId LogId, long Offset) in FlushChunks)
 			{
 				await Builder.CompleteChunkAsync(LogId, Offset);
 			}
@@ -1057,7 +1060,7 @@ namespace HordeServer.Services
 			Logger.LogInformation("Forcing flush of pending log chunks...");
 
 			// Mark everything in the cache as complete
-			List<(ObjectId, long)> WriteChunks = await Builder.TouchChunksAsync(TimeSpan.Zero);
+			List<(LogId, long)> WriteChunks = await Builder.TouchChunksAsync(TimeSpan.Zero);
 			WriteCompleteChunks(WriteChunks, true);
 
 			// Wait for everything to flush
@@ -1095,11 +1098,11 @@ namespace HordeServer.Services
 		/// </summary>
 		/// <param name="ChunksToWrite">List of chunks to write</param>
 		/// <param name="bCreateIndex">Create an index for the log</param>
-		private void WriteCompleteChunks(List<(ObjectId, long)> ChunksToWrite, bool bCreateIndex)
+		private void WriteCompleteChunks(List<(LogId, long)> ChunksToWrite, bool bCreateIndex)
 		{
-			foreach (IGrouping<ObjectId, long> Group in ChunksToWrite.GroupBy(x => x.Item1, x => x.Item2))
+			foreach (IGrouping<LogId, long> Group in ChunksToWrite.GroupBy(x => x.Item1, x => x.Item2))
 			{
-				ObjectId LogId = Group.Key;
+				LogId LogId = Group.Key;
 
 				// Find offsets of new chunks to write
 				List<long> Offsets = new List<long>();
@@ -1133,7 +1136,7 @@ namespace HordeServer.Services
 		/// <param name="Offsets">Chunks to write</param>
 		/// <param name="bCreateIndex">Whether to create the index for this log</param>
 		/// <returns>Async task</returns>
-		private async Task<ILogFile?> WriteCompleteChunksForLogAsync(ObjectId LogId, List<long> Offsets, bool bCreateIndex)
+		private async Task<ILogFile?> WriteCompleteChunksForLogAsync(LogId LogId, List<long> Offsets, bool bCreateIndex)
 		{
 			ILogFile? LogFile = await LogFiles.GetLogFileAsync(LogId);
 			if(LogFile != null)
@@ -1427,7 +1430,7 @@ namespace HordeServer.Services
 		/// <param name="LineIndex">First line index of the chunk</param>
 		/// <returns>Chunk daata</returns>
 		[SuppressMessage("Design", "CA1031:Do not catch general exception types", Justification = "<Pending>")]
-		private async Task<LogChunkData?> WriteChunkAsync(ObjectId LogFileId, long Offset, int LineIndex)
+		private async Task<LogChunkData?> WriteChunkAsync(LogId LogFileId, long Offset, int LineIndex)
 		{
 			// Write the chunk to storage
 			LogChunkData? ChunkData = await Builder.GetChunkAsync(LogFileId, Offset, LineIndex);
@@ -1489,7 +1492,7 @@ namespace HordeServer.Services
 		/// <param name="Index">Index to write</param>
 		/// <returns>Async task</returns>
 		[SuppressMessage("Design", "CA1031:Do not catch general exception types", Justification = "<Pending>")]
-		private async Task WriteIndexAsync(ObjectId LogFileId, long Length, LogIndexData Index)
+		private async Task WriteIndexAsync(LogId LogFileId, long Length, LogIndexData Index)
 		{
 			try
 			{
diff --git a/Engine/Source/Programs/Horde/HordeServer/Logs/Storage/ILogStorage.cs b/Engine/Source/Programs/Horde/HordeServer/Logs/Storage/ILogStorage.cs
--- a/Engine/Source/Programs/Horde/HordeServer/Logs/Storage/ILogStorage.cs
+++ b/Engine/Source/Programs/Horde/HordeServer/Logs/Storage/ILogStorage.cs
@@ -13,6 +13,8 @@ using System.Threading.Tasks;
 
 namespace HordeServer.Logs
 {
+	using LogId = ObjectId<ILogFile>;
+
 	/// <summary>
 	/// Caching interface for reading and writing log data.
 	/// </summary>
@@ -24,7 +26,7 @@ namespace HordeServer.Logs
 		/// <param name="LogId">Unique id of the log file</param>
 		/// <param name="Length">Length of the file covered by the index</param>
 		/// <returns>Index for the log file</returns>
-		Task<LogIndexData?> ReadIndexAsync(ObjectId LogId, long Length);
+		Task<LogIndexData?> ReadIndexAsync(LogId LogId, long Length);
 
 		/// <summary>
 		/// Log file to write an index for
@@ -33,7 +35,7 @@ namespace HordeServer.Logs
 		/// <param name="Length">Length of the file covered by the index</param>
 		/// <param name="Index">The log file index</param>
 		/// <returns>Async task</returns>
-		Task WriteIndexAsync(ObjectId LogId, long Length, LogIndexData Index);
+		Task WriteIndexAsync(LogId LogId, long Length, LogIndexData Index);
 
 		/// <summary>
 		/// Retrieves an item from the cache
@@ -42,7 +44,7 @@ namespace HordeServer.Logs
 		/// <param name="Offset">Offset of the chunk to read</param>
 		/// <param name="LineIndex">First line of the chunk</param>
 		/// <returns>Data for the given key, or null if it's not present</returns>
-		Task<LogChunkData?> ReadChunkAsync(ObjectId LogId, long Offset, int LineIndex);
+		Task<LogChunkData?> ReadChunkAsync(LogId LogId, long Offset, int LineIndex);
 
 		/// <summary>
 		/// Writes a chunk to storage
@@ -51,6 +53,6 @@ namespace HordeServer.Logs
 		/// <param name="Offset">Offset of the chunk to write</param>
 		/// <param name="ChunkData">Information about the chunk data</param>
 		/// <returns>Async task</returns>
-		Task WriteChunkAsync(ObjectId LogId, long Offset, LogChunkData ChunkData);
+		Task WriteChunkAsync(LogId LogId, long Offset, LogChunkData ChunkData);
 	}
 }
diff --git a/Engine/Source/Programs/Horde/HordeServer/Logs/Storage/Impl/LocalLogStorage.cs b/Engine/Source/Programs/Horde/HordeServer/Logs/Storage/Impl/LocalLogStorage.cs
--- a/Engine/Source/Programs/Horde/HordeServer/Logs/Storage/Impl/LocalLogStorage.cs
+++ b/Engine/Source/Programs/Horde/HordeServer/Logs/Storage/Impl/LocalLogStorage.cs
@@ -1,6 +1,7 @@
 // Copyright Epic Games, Inc. All Rights Reserved.
 
 using HordeServer.Models;
+using HordeServer.Utilities;
 using Microsoft.Extensions.Caching.Memory;
 using MongoDB.Bson;
 using System;
@@ -10,6 +11,8 @@ using System.Threading.Tasks;
 
 namespace HordeServer.Logs.Storage
 {
+	using LogId = ObjectId<ILogFile>;
+
 	/// <summary>
 	/// In-memory cache for chunk and index data
 	/// </summary>
@@ -52,7 +55,7 @@ namespace HordeServer.Logs.Storage
 		/// <param name="LogId">The log file to retrieve an index for</param>
 		/// <param name="Length">Length of the file covered by the index</param>
 		/// <returns>Cache key for the index</returns>
-		static string IndexKey(ObjectId LogId, long Length)
+		static string IndexKey(LogId LogId, long Length)
 		{
 			return $"{LogId}/index-{Length}";
 		}
@@ -63,7 +66,7 @@ namespace HordeServer.Logs.Storage
 		/// <param name="LogId">The log file to retrieve an index for</param>
 		/// <param name="Offset">The chunk offset</param>
 		/// <returns>Cache key for the chunk</returns>
-		static string ChunkKey(ObjectId LogId, long Offset)
+		static string ChunkKey(LogId LogId, long Offset)
 		{
 			return $"{LogId}/chunk-{Offset}";
 		}
@@ -116,26 +119,26 @@ namespace HordeServer.Logs.Storage
 		}
 
 		/// <inheritdoc/>
-		public Task<LogIndexData?> ReadIndexAsync(ObjectId LogId, long Length)
+		public Task<LogIndexData?> ReadIndexAsync(LogId LogId, long Length)
 		{
 			return ReadValueAsync(IndexKey(LogId, Length), () => Inner.ReadIndexAsync(LogId, Length));
 		}
 
 		/// <inheritdoc/>
-		public Task WriteIndexAsync(ObjectId LogId, long Length, LogIndexData IndexData)
+		public Task WriteIndexAsync(LogId LogId, long Length, LogIndexData IndexData)
 		{
 			AddEntry(IndexKey(LogId, Length), IndexData);
 			return Inner.WriteIndexAsync(LogId, Length, IndexData);
 		}
 
 		/// <inheritdoc/>
-		public Task<LogChunkData?> ReadChunkAsync(ObjectId LogId, long Offset, int LineIndex)
+		public Task<LogChunkData?> ReadChunkAsync(LogId LogId, long Offset, int LineIndex)
 		{
 			return ReadValueAsync(ChunkKey(LogId, Offset), () => Inner.ReadChunkAsync(LogId, Offset, LineIndex));
 		}
 
 		/// <inheritdoc/>
-		public Task WriteChunkAsync(ObjectId LogId, long Offset, LogChunkData ChunkData)
+		public Task WriteChunkAsync(LogId LogId, long Offset, LogChunkData ChunkData)
 		{
 			AddEntry(ChunkKey(LogId, Offset), ChunkData);
 			return Inner.WriteChunkAsync(LogId, Offset, ChunkData);
diff --git a/Engine/Source/Programs/Horde/HordeServer/Logs/Storage/Impl/NullLogStorage.cs b/Engine/Source/Programs/Horde/HordeServer/Logs/Storage/Impl/NullLogStorage.cs
--- a/Engine/Source/Programs/Horde/HordeServer/Logs/Storage/Impl/NullLogStorage.cs
+++ b/Engine/Source/Programs/Horde/HordeServer/Logs/Storage/Impl/NullLogStorage.cs
@@ -1,5 +1,7 @@
 // Copyright Epic Games, Inc. All Rights Reserved.
 
+using HordeServer.Models;
+using HordeServer.Utilities;
 using MongoDB.Bson;
 using System;
 using System.Collections.Generic;
@@ -8,6 +10,8 @@ using System.Threading.Tasks;
 
 namespace HordeServer.Logs.Storage.Impl
 {
+	using LogId = ObjectId<ILogFile>;
+
 	/// <summary>
 	/// Empty implementation of log storage
 	/// </summary>
@@ -19,25 +23,25 @@ namespace HordeServer.Logs.Storage.Impl
 		}
 
 		/// <inheritdoc/>
-		public Task<LogIndexData?> ReadIndexAsync(ObjectId LogId, long Length)
+		public Task<LogIndexData?> ReadIndexAsync(LogId LogId, long Length)
 		{
 			return Task.FromResult<LogIndexData?>(null);
 		}
 
 		/// <inheritdoc/>
-		public Task WriteIndexAsync(ObjectId LogId, long Length, LogIndexData Index)
+		public Task WriteIndexAsync(LogId LogId, long Length, LogIndexData Index)
 		{
 			return Task.CompletedTask;
 		}
 
 		/// <inheritdoc/>
-		public Task<LogChunkData?> ReadChunkAsync(ObjectId LogId, long Offset, int LineIndex)
+		public Task<LogChunkData?> ReadChunkAsync(LogId LogId, long Offset, int LineIndex)
 		{
 			return Task.FromResult<LogChunkData?>(null);
 		}
 
 		/// <inheritdoc/>
-		public Task WriteChunkAsync(ObjectId LogId, long Offset, LogChunkData ChunkData)
+		public Task WriteChunkAsync(LogId LogId, long Offset, LogChunkData ChunkData)
 		{
 			return Task.CompletedTask;
 		}
diff --git a/Engine/Source/Programs/Horde/HordeServer/Logs/Storage/Impl/PersistentLogStorage.cs b/Engine/Source/Programs/Horde/HordeServer/Logs/Storage/Impl/PersistentLogStorage.cs
--- a/Engine/Source/Programs/Horde/HordeServer/Logs/Storage/Impl/PersistentLogStorage.cs
+++ b/Engine/Source/Programs/Horde/HordeServer/Logs/Storage/Impl/PersistentLogStorage.cs
@@ -14,6 +14,8 @@ using System.Threading.Tasks;
 
 namespace HordeServer.Logs.Readers
 {
+	using LogId = ObjectId<ILogFile>;
+
 	/// <summary>
 	/// Bulk storage for log file data
 	/// </summary>
@@ -46,7 +48,7 @@ namespace HordeServer.Logs.Readers
 		}
 
 		/// <inheritdoc/>
-		public async Task<LogIndexData?> ReadIndexAsync(ObjectId LogId, long Length)
+		public async Task<LogIndexData?> ReadIndexAsync(LogId LogId, long Length)
 		{
 			Logger.LogDebug("Reading log {LogId} index length {Length} from persistent storage", LogId, Length);
 
@@ -60,7 +62,7 @@ namespace HordeServer.Logs.Readers
 		}
 
 		/// <inheritdoc/>
-		public Task WriteIndexAsync(ObjectId LogId, long Length, LogIndexData IndexData)
+		public Task WriteIndexAsync(LogId LogId, long Length, LogIndexData IndexData)
 		{
 			Logger.LogDebug("Writing log {LogId} index length {Length} to persistent storage", LogId, Length);
 
@@ -70,7 +72,7 @@ namespace HordeServer.Logs.Readers
 		}
 
 		/// <inheritdoc/>
-		public async Task<LogChunkData?> ReadChunkAsync(ObjectId LogId, long Offset, int LineIndex)
+		public async Task<LogChunkData?> ReadChunkAsync(LogId LogId, long Offset, int LineIndex)
 		{
 			Logger.LogDebug("Reading log {LogId} chunk offset {Offset} from persistent storage", LogId, Offset);
 
@@ -93,7 +95,7 @@ namespace HordeServer.Logs.Readers
 		}
 
 		/// <inheritdoc/>
-		public Task WriteChunkAsync(ObjectId LogId, long Offset, LogChunkData ChunkData)
+		public Task WriteChunkAsync(LogId LogId, long Offset, LogChunkData ChunkData)
 		{
 			Logger.LogDebug("Writing log {LogId} chunk offset {Offset} to persistent storage", LogId, Offset);
 
diff --git a/Engine/Source/Programs/Horde/HordeServer/Logs/Storage/Impl/RedisLogStorage.cs b/Engine/Source/Programs/Horde/HordeServer/Logs/Storage/Impl/RedisLogStorage.cs
--- a/Engine/Source/Programs/Horde/HordeServer/Logs/Storage/Impl/RedisLogStorage.cs
+++ b/Engine/Source/Programs/Horde/HordeServer/Logs/Storage/Impl/RedisLogStorage.cs
@@ -1,5 +1,7 @@
 // Copyright Epic Games, Inc. All Rights Reserved.
 
+using HordeServer.Models;
+using HordeServer.Utilities;
 using Microsoft.Extensions.Logging;
 using MongoDB.Bson;
 using StackExchange.Redis;
@@ -11,6 +13,8 @@ using System.Threading.Tasks;
 
 namespace HordeServer.Logs.Storage.Impl
 {
+	using LogId = ObjectId<ILogFile>;
+
 	/// <summary>
 	/// Redis log file storage
 	/// </summary>
@@ -56,7 +60,7 @@ namespace HordeServer.Logs.Storage.Impl
 		/// <param name="LogId">The log file id</param>
 		/// <param name="Length">Length of the file covered by the index</param>
 		/// <returns>The index key</returns>
-		static string IndexKey(ObjectId LogId, long Length)
+		static string IndexKey(LogId LogId, long Length)
 		{
 			return $"log-{LogId}-index-{Length}";
 		}
@@ -67,7 +71,7 @@ namespace HordeServer.Logs.Storage.Impl
 		/// <param name="LogId">The log file id</param>
 		/// <param name="Offset">Offset of the chunk within the log file</param>
 		/// <returns>The chunk key</returns>
-		static string ChunkKey(ObjectId LogId, long Offset)
+		static string ChunkKey(LogId LogId, long Offset)
 		{
 			return $"log-{LogId}-chunk-{Offset}";
 		}
@@ -123,7 +127,7 @@ namespace HordeServer.Logs.Storage.Impl
 		}
 
 		/// <inheritdoc/>
-		public async Task<LogIndexData?> ReadIndexAsync(ObjectId LogId, long Length)
+		public async Task<LogIndexData?> ReadIndexAsync(LogId LogId, long Length)
 		{
 			string Key = IndexKey(LogId, Length);
 
@@ -140,13 +144,13 @@ namespace HordeServer.Logs.Storage.Impl
 		}
 
 		/// <inheritdoc/>
-		public async Task WriteIndexAsync(ObjectId LogId, long Length, LogIndexData Index)
+		public async Task WriteIndexAsync(LogId LogId, long Length, LogIndexData Index)
 		{
 			await Inner.WriteIndexAsync(LogId, Length, Index);
 		}
 
 		/// <inheritdoc/>
-		public async Task<LogChunkData?> ReadChunkAsync(ObjectId LogId, long Offset, int LineIndex)
+		public async Task<LogChunkData?> ReadChunkAsync(LogId LogId, long Offset, int LineIndex)
 		{
 			string Key = ChunkKey(LogId, Offset);
 
@@ -163,7 +167,7 @@ namespace HordeServer.Logs.Storage.Impl
 		}
 
 		/// <inheritdoc/>
-		public async Task WriteChunkAsync(ObjectId LogId, long Offset, LogChunkData ChunkData)
+		public async Task WriteChunkAsync(LogId LogId, long Offset, LogChunkData ChunkData)
 		{
 			await Inner.WriteChunkAsync(LogId, Offset, ChunkData);
 		}
diff --git a/Engine/Source/Programs/Horde/HordeServer/Logs/Storage/Impl/SequencedLogStorage.cs b/Engine/Source/Programs/Horde/HordeServer/Logs/Storage/Impl/SequencedLogStorage.cs
--- a/Engine/Source/Programs/Horde/HordeServer/Logs/Storage/Impl/SequencedLogStorage.cs
+++ b/Engine/Source/Programs/Horde/HordeServer/Logs/Storage/Impl/SequencedLogStorage.cs
@@ -1,6 +1,7 @@
 // Copyright Epic Games, Inc. All Rights Reserved.
 
 using HordeServer.Models;
+using HordeServer.Utilities;
 using MongoDB.Bson;
 using System;
 using System.Collections.Generic;
@@ -9,6 +10,8 @@ using System.Threading.Tasks;
 
 namespace HordeServer.Logs.Storage.Impl
 {
+	using LogId = ObjectId<ILogFile>;
+
 	/// <summary>
 	/// Storage layer which caches pending read tasks, to avoid fetching the same item more than once
 	/// </summary>
@@ -22,22 +25,22 @@ namespace HordeServer.Logs.Storage.Impl
 		/// <summary>
 		/// Pending index reads
 		/// </summary>
-		Dictionary<(ObjectId, long), Task<LogIndexData?>> IndexReadTasks = new Dictionary<(ObjectId, long), Task<LogIndexData?>>();
+		Dictionary<(LogId, long), Task<LogIndexData?>> IndexReadTasks = new Dictionary<(LogId, long), Task<LogIndexData?>>();
 
 		/// <summary>
 		/// Pending index reads
 		/// </summary>
-		Dictionary<(ObjectId, long), Task> IndexWriteTasks = new Dictionary<(ObjectId, long), Task>();
+		Dictionary<(LogId, long), Task> IndexWriteTasks = new Dictionary<(LogId, long), Task>();
 
 		/// <summary>
 		/// Pending chunk reads
 		/// </summary>
-		Dictionary<(ObjectId, long), Task<LogChunkData?>> ChunkReadTasks = new Dictionary<(ObjectId, long), Task<LogChunkData?>>();
+		Dictionary<(LogId, long), Task<LogChunkData?>> ChunkReadTasks = new Dictionary<(LogId, long), Task<LogChunkData?>>();
 
 		/// <summary>
 		/// Pending chunk reads
 		/// </summary>
-		Dictionary<(ObjectId, long), Task> ChunkWriteTasks = new Dictionary<(ObjectId, long), Task>();
+		Dictionary<(LogId, long), Task> ChunkWriteTasks = new Dictionary<(LogId, long), Task>();
 
 		/// <summary>
 		/// Constructor
@@ -55,7 +58,7 @@ namespace HordeServer.Logs.Storage.Impl
 		}
 
 		/// <inheritdoc/>
-		public Task<LogIndexData?> ReadIndexAsync(ObjectId LogId, long Length)
+		public Task<LogIndexData?> ReadIndexAsync(LogId LogId, long Length)
 		{
 			Task<LogIndexData?>? Task;
 			lock (IndexReadTasks)
@@ -70,7 +73,7 @@ namespace HordeServer.Logs.Storage.Impl
 		}
 
 		/// <inheritdoc/>
-		public Task WriteIndexAsync(ObjectId LogId, long Length, LogIndexData IndexData)
+		public Task WriteIndexAsync(LogId LogId, long Length, LogIndexData IndexData)
 		{
 			Task? Task;
 			lock (IndexWriteTasks)
@@ -85,7 +88,7 @@ namespace HordeServer.Logs.Storage.Impl
 		}
 
 		/// <inheritdoc/>
-		public Task<LogChunkData?> ReadChunkAsync(ObjectId LogId, long Offset, int LineIndex)
+		public Task<LogChunkData?> ReadChunkAsync(LogId LogId, long Offset, int LineIndex)
 		{
 			Task<LogChunkData?>? Task;
 			lock (ChunkReadTasks)
@@ -100,7 +103,7 @@ namespace HordeServer.Logs.Storage.Impl
 		}
 
 		/// <inheritdoc/>
-		public Task WriteChunkAsync(ObjectId LogId, long Offset, LogChunkData ChunkData)
+		public Task WriteChunkAsync(LogId LogId, long Offset, LogChunkData ChunkData)
 		{
 			Task? Task;
 			lock (ChunkWriteTasks)
@@ -120,7 +123,7 @@ namespace HordeServer.Logs.Storage.Impl
 		/// <param name="LogId">The log file to read the index for</param>
 		/// <param name="Length">Length of the file that's indexed</param>
 		/// <returns>The index data</returns>
-		async Task<LogIndexData?> InnerReadIndexAsync(ObjectId LogId, long Length)
+		async Task<LogIndexData?> InnerReadIndexAsync(LogId LogId, long Length)
 		{
 			await Task.Yield();
 
@@ -139,7 +142,7 @@ namespace HordeServer.Logs.Storage.Impl
 		/// <param name="Length">Length of the indexed data</param>
 		/// <param name="IndexData">The index data to write</param>
 		/// <returns>The index data</returns>
-		async Task InnerWriteIndexAsync(ObjectId LogId, long Length, LogIndexData IndexData)
+		async Task InnerWriteIndexAsync(LogId LogId, long Length, LogIndexData IndexData)
 		{
 			await Task.Yield();
 
@@ -158,7 +161,7 @@ namespace HordeServer.Logs.Storage.Impl
 		/// <param name="Offset">Offset of the chunk to read</param>
 		/// <param name="LineIndex">Index of the first line in this chunk</param>
 		/// <returns>The index data</returns>
-		async Task<LogChunkData?> InnerReadChunkAsync(ObjectId LogId, long Offset, int LineIndex)
+		async Task<LogChunkData?> InnerReadChunkAsync(LogId LogId, long Offset, int LineIndex)
 		{
 			await Task.Yield();
 
@@ -177,7 +180,7 @@ namespace HordeServer.Logs.Storage.Impl
 		/// <param name="Offset">Offset of the chunk within the log</param>
 		/// <param name="ChunkData">The chunk data</param>
 		/// <returns>The index data</returns>
-		async Task<LogChunkData> InnerWriteChunkAsync(ObjectId LogId, long Offset, LogChunkData ChunkData)
+		async Task<LogChunkData> InnerWriteChunkAsync(LogId LogId, long Offset, LogChunkData ChunkData)
 		{
 			await Task.Yield();
 
diff --git a/Engine/Source/Programs/Horde/HordeServer/Projects/ISchedule.cs b/Engine/Source/Programs/Horde/HordeServer/Projects/ISchedule.cs
--- a/Engine/Source/Programs/Horde/HordeServer/Projects/ISchedule.cs
+++ b/Engine/Source/Programs/Horde/HordeServer/Projects/ISchedule.cs
@@ -12,6 +12,7 @@ using System.Threading.Tasks;
 
 namespace HordeServer.Models
 {
+	using JobId = ObjectId<IJob>;
 	using TemplateRefId = StringId<TemplateRef>;
 
 	/// <summary>
@@ -209,7 +210,7 @@ namespace HordeServer.Models
 		/// <summary>
 		/// List of jobs that are currently active
 		/// </summary>
-		public List<ObjectId> ActiveJobs { get; set; } = new List<ObjectId>();
+		public List<JobId> ActiveJobs { get; set; } = new List<JobId>();
 
 		/// <summary>
 		/// Private constructor for serializtaion
diff --git a/Engine/Source/Programs/Horde/HordeServer/Projects/IStreamCollection.cs b/Engine/Source/Programs/Horde/HordeServer/Projects/IStreamCollection.cs
--- a/Engine/Source/Programs/Horde/HordeServer/Projects/IStreamCollection.cs
+++ b/Engine/Source/Programs/Horde/HordeServer/Projects/IStreamCollection.cs
@@ -15,6 +15,7 @@ using System.Threading.Tasks;
 
 namespace HordeServer.Collections
 {
+	using JobId = ObjectId<IJob>;
 	using ProjectId = StringId<IProject>;
 	using StreamId = StringId<IStream>;
 	using TemplateRefId = StringId<TemplateRef>;
@@ -81,7 +82,7 @@ namespace HordeServer.Collections
 		/// <param name="LastTriggerChange">New last trigger changelist for the schedule</param>
 		/// <param name="NewActiveJobs">New list of active jobs</param>
 		/// <returns>The updated stream if successful, null otherwise</returns>
-		Task<IStream?> TryUpdateScheduleTriggerAsync(IStream Stream, TemplateRefId TemplateRefId, DateTimeOffset? LastTriggerTime, int? LastTriggerChange, List<ObjectId> NewActiveJobs);
+		Task<IStream?> TryUpdateScheduleTriggerAsync(IStream Stream, TemplateRefId TemplateRefId, DateTimeOffset? LastTriggerTime, int? LastTriggerChange, List<JobId> NewActiveJobs);
 
 		/// <summary>
 		/// Delete a stream
diff --git a/Engine/Source/Programs/Horde/HordeServer/Projects/Impl/StreamCollection.cs b/Engine/Source/Programs/Horde/HordeServer/Projects/Impl/StreamCollection.cs
--- a/Engine/Source/Programs/Horde/HordeServer/Projects/Impl/StreamCollection.cs
+++ b/Engine/Source/Programs/Horde/HordeServer/Projects/Impl/StreamCollection.cs
@@ -17,6 +17,7 @@ using System.Threading.Tasks;
 
 namespace HordeServer.Collections.Impl
 {
+	using JobId = ObjectId<IJob>;
 	using ProjectId = StringId<IProject>;
 	using StreamId = StringId<IStream>;
 	using TemplateRefId = StringId<TemplateRef>;
@@ -329,7 +330,7 @@ namespace HordeServer.Collections.Impl
 		}
 
 		/// <inheritdoc/>
-		public async Task<IStream?> TryUpdateScheduleTriggerAsync(IStream StreamInterface, TemplateRefId TemplateRefId, DateTimeOffset? LastTriggerTime, int? LastTriggerChange, List<ObjectId> NewActiveJobs)
+		public async Task<IStream?> TryUpdateScheduleTriggerAsync(IStream StreamInterface, TemplateRefId TemplateRefId, DateTimeOffset? LastTriggerTime, int? LastTriggerChange, List<JobId> NewActiveJobs)
 		{
 			StreamDocument Stream = (StreamDocument)StreamInterface;
 			Schedule Schedule = Stream.Templates[TemplateRefId].Schedule!;
@@ -350,7 +351,7 @@ namespace HordeServer.Collections.Impl
 			}
 			if (NewActiveJobs != null)
 			{
-				FieldDefinition<StreamDocument, List<ObjectId>> Field = $"{nameof(Stream.Templates)}.{TemplateRefId}.{nameof(Schedule)}.{nameof(Schedule.ActiveJobs)}";
+				FieldDefinition<StreamDocument, List<JobId>> Field = $"{nameof(Stream.Templates)}.{TemplateRefId}.{nameof(Schedule)}.{nameof(Schedule.ActiveJobs)}";
 				Updates.Add(Builders<StreamDocument>.Update.Set(Field, NewActiveJobs));
 				Schedule.ActiveJobs = NewActiveJobs;
 			}
diff --git a/Engine/Source/Programs/Horde/HordeServer/Projects/Impl/StreamService.cs b/Engine/Source/Programs/Horde/HordeServer/Projects/Impl/StreamService.cs
--- a/Engine/Source/Programs/Horde/HordeServer/Projects/Impl/StreamService.cs
+++ b/Engine/Source/Programs/Horde/HordeServer/Projects/Impl/StreamService.cs
@@ -19,6 +19,7 @@ using Microsoft.Extensions.Caching.Memory;
 
 namespace HordeServer.Services
 {
+	using JobId = ObjectId<IJob>;
 	using ProjectId = StringId<IProject>;
 	using StreamId = StringId<IStream>;
 	using TemplateRefId = StringId<TemplateRef>;
@@ -118,7 +119,7 @@ namespace HordeServer.Services
 		/// <param name="AddJobs">Jobs to add</param>
 		/// <param name="RemoveJobs">Jobs to remove</param>
 		/// <returns>True if the stream was updated</returns>
-		public async Task<IStream?> UpdateScheduleTriggerAsync(IStream Stream, TemplateRefId TemplateRefId, DateTimeOffset? LastTriggerTime, int? LastTriggerChange, List<ObjectId> AddJobs, List<ObjectId> RemoveJobs)
+		public async Task<IStream?> UpdateScheduleTriggerAsync(IStream Stream, TemplateRefId TemplateRefId, DateTimeOffset? LastTriggerTime, int? LastTriggerChange, List<JobId> AddJobs, List<JobId> RemoveJobs)
 		{
 			IStream? NewStream = Stream;
 			while (NewStream != null)
@@ -133,7 +134,7 @@ namespace HordeServer.Services
 					break;
 				}
 
-				List<ObjectId> NewActiveJobs = TemplateRef.Schedule.ActiveJobs.Except(RemoveJobs).Union(AddJobs).ToList();
+				List<JobId> NewActiveJobs = TemplateRef.Schedule.ActiveJobs.Except(RemoveJobs).Union(AddJobs).ToList();
 
 				NewStream = await Streams.TryUpdateScheduleTriggerAsync(NewStream, TemplateRefId, LastTriggerTime, LastTriggerChange, NewActiveJobs);
 
diff --git a/Engine/Source/Programs/Horde/HordeServer/Projects/ScheduleService.cs b/Engine/Source/Programs/Horde/HordeServer/Projects/ScheduleService.cs
--- a/Engine/Source/Programs/Horde/HordeServer/Projects/ScheduleService.cs
+++ b/Engine/Source/Programs/Horde/HordeServer/Projects/ScheduleService.cs
@@ -27,6 +27,7 @@ using TimeZoneConverter;
 
 namespace HordeServer.Services
 {
+	using JobId = ObjectId<IJob>;
 	using StreamId = StringId<IStream>;
 	using TemplateRefId = StringId<TemplateRef>;
 
@@ -198,8 +199,8 @@ namespace HordeServer.Services
 			Logger.LogInformation("Schedule {StreamId}/{TemplateRefId} is ready to trigger (last: {Change} at {Time})", Stream.Id, TemplateRefId, Schedule.LastTriggerChange, Schedule.LastTriggerTime);
 
 			// Get a list of jobs that we need to remove
-			List<ObjectId> RemoveJobIds = new List<ObjectId>();
-			foreach (ObjectId ActiveJobId in Schedule.ActiveJobs)
+			List<JobId> RemoveJobIds = new List<JobId>();
+			foreach (JobId ActiveJobId in Schedule.ActiveJobs)
 			{
 				IJob? Job = await JobService.GetJobAsync(ActiveJobId);
 				if (Job == null || Job.Batches.All(x => x.State == JobStepBatchState.Complete))
@@ -210,7 +211,7 @@ namespace HordeServer.Services
 			}
 			if (RemoveJobIds.Count > 0)
 			{
-				await StreamService.UpdateScheduleTriggerAsync(Stream, TemplateRefId, Now, null, new List<ObjectId>(), RemoveJobIds);
+				await StreamService.UpdateScheduleTriggerAsync(Stream, TemplateRefId, Now, null, new List<JobId>(), RemoveJobIds);
 			}
 
 			// Trigger this schedule
@@ -243,7 +244,7 @@ namespace HordeServer.Services
 			if (Schedule.MaxActive != 0 && NumActiveJobs >= Schedule.MaxActive)
 			{
 				Logger.LogInformation("Skipping trigger of {StreamId} template {TemplateRefId} - already have maximum number of jobs running ({NumJobs})", Stream.Id, TemplateRefId, Schedule.MaxActive);
-				foreach (ObjectId JobId in Schedule.ActiveJobs)
+				foreach (JobId JobId in Schedule.ActiveJobs)
 				{
 					Logger.LogInformation("Active job for {StreamId} template {TemplateRefId}: {JobId}", Stream.Id, TemplateRefId, JobId);
 				}
@@ -363,7 +364,7 @@ namespace HordeServer.Services
 				List<string> DefaultArguments = Template.GetDefaultArguments();
 				IJob NewJob = await JobService.CreateJobAsync(null, Stream, TemplateRefId, Template.Id, Graph, Template.Name, Change, CodeChange, null, null, null, null, Template.Priority, null, null, TemplateRef.ChainedJobs, TemplateRef.ShowUgsBadges, TemplateRef.ShowUgsAlerts, TemplateRef.NotificationChannel, TemplateRef.NotificationChannelFilter, null, DefaultArguments);
 				Logger.LogInformation("Started new job for {StreamName} template {TemplateName} at CL {Change} (Code CL {CodeChange}): {JobId}", Stream.Id, TemplateRef.Name, Change, CodeChange, NewJob.Id);
-				await StreamService.UpdateScheduleTriggerAsync(Stream, TemplateRefId, Now, Change, new List<ObjectId> { NewJob.Id }, new List<ObjectId>());
+				await StreamService.UpdateScheduleTriggerAsync(Stream, TemplateRefId, Now, Change, new List<JobId> { NewJob.Id }, new List<JobId>());
 			}
 		}
 
diff --git a/Engine/Source/Programs/Horde/HordeServer/Services/RpcService.cs b/Engine/Source/Programs/Horde/HordeServer/Services/RpcService.cs
--- a/Engine/Source/Programs/Horde/HordeServer/Services/RpcService.cs
+++ b/Engine/Source/Programs/Horde/HordeServer/Services/RpcService.cs
@@ -32,6 +32,8 @@ namespace HordeServer.Services
 	using AgentSoftwareChannelName = StringId<AgentSoftwareChannels>;
 	using AgentSoftwareVersion = StringId<IAgentSoftwareCollection>;
 	using IStream = HordeServer.Models.IStream;
+	using JobId = ObjectId<IJob>;
+	using LogId = ObjectId<ILogFile>;
 	using StreamId = StringId<IStream>;
 	using RpcAgentCapabilities = HordeCommon.Rpc.Messages.AgentCapabilities;
 	using RpcDeviceCapabilities = HordeCommon.Rpc.Messages.DeviceCapabilities;
@@ -471,7 +473,7 @@ namespace HordeServer.Services
 		/// <returns>Information about the new agent</returns>
 		public async override Task<GetJobResponse> GetJob(GetJobRequest Request, ServerCallContext Context)
 		{
-			ObjectId JobIdValue = Request.JobId.ToObjectId();
+			JobId JobIdValue = new JobId(Request.JobId.ToObjectId());
 
 			IJob? Job = await JobService.GetJobAsync(JobIdValue);
 			if (Job == null)
@@ -494,7 +496,7 @@ namespace HordeServer.Services
 		/// <returns>Information about the new agent</returns>
 		public async override Task<Empty> UpdateJob(UpdateJobRequest Request, ServerCallContext Context)
 		{
-			ObjectId JobIdValue = Request.JobId.ToObjectId();
+			JobId JobIdValue = new JobId(Request.JobId);
 
 			IJob? Job = await JobService.GetJobAsync(JobIdValue);
 			if (Job == null)
@@ -521,7 +523,7 @@ namespace HordeServer.Services
 			ObjectId JobId = Request.JobId.ToObjectId();
 			SubResourceId BatchId = Request.BatchId.ToSubResourceId();
 
-			IJob? Job = await JobService.GetJobAsync(Request.JobId.ToObjectId());
+			IJob? Job = await JobService.GetJobAsync(new JobId(Request.JobId));
 			if (Job == null)
 			{
 				throw new StructuredRpcException(StatusCode.NotFound, "Job {JobId} not found", Request.JobId);
@@ -549,7 +551,7 @@ namespace HordeServer.Services
 		/// <returns>Information about the new agent</returns>
 		public async override Task<Empty> FinishBatch(FinishBatchRequest Request, ServerCallContext Context)
 		{
-			IJob? Job = await JobService.GetJobAsync(Request.JobId.ToObjectId());
+			IJob? Job = await JobService.GetJobAsync(new JobId(Request.JobId));
 			if (Job == null)
 			{
 				throw new StructuredRpcException(StatusCode.NotFound, "Job {JobId} not found", Request.JobId);
@@ -582,7 +584,7 @@ namespace HordeServer.Services
 		async Task<BeginStepResponse?> TryBeginStep(BeginStepRequest Request, Boxed<ILogFile?> Log, ServerCallContext Context)
 		{
 			// Check the job exists and we can access it
-			IJob? Job = await JobService.GetJobAsync(Request.JobId.ToObjectId());
+			IJob? Job = await JobService.GetJobAsync(new JobId(Request.JobId));
 			if (Job == null)
 			{
 				throw new StructuredRpcException(StatusCode.NotFound, "Job {JobId} not found", Request.JobId);
@@ -728,7 +730,7 @@ namespace HordeServer.Services
 			return true;
 		}
 
-		async Task<IJob> GetJobAsync(ObjectId JobId)
+		async Task<IJob> GetJobAsync(JobId JobId)
 		{
 			IJob? Job = await JobService.GetJobAsync(JobId);
 			if (Job == null)
@@ -767,7 +769,7 @@ namespace HordeServer.Services
 		/// <returns>Information about the new agent</returns>
 		public async override Task<Empty> UpdateStep(UpdateStepRequest Request, ServerCallContext Context)
 		{
-			IJob? Job = await JobService.GetJobAsync(Request.JobId.ToObjectId());
+			IJob? Job = await JobService.GetJobAsync(new JobId(Request.JobId));
 			if (Job == null)
 			{
 				throw new StructuredRpcException(StatusCode.NotFound, "Job {JobId} not found", Request.JobId);
@@ -787,7 +789,7 @@ namespace HordeServer.Services
 		/// <returns>Information about the step</returns>
 		public async override Task<GetStepResponse> GetStep(GetStepRequest Request, ServerCallContext Context)
 		{
-			IJob Job = await GetJobAsync(Request.JobId.ToObjectId());
+			IJob Job = await GetJobAsync(new JobId(Request.JobId));
 			IJobStepBatch Batch = AuthorizeBatch(Job, Request.BatchId.ToSubResourceId(), Context);
 
 			SubResourceId StepId = Request.StepId.ToSubResourceId();
@@ -840,7 +842,7 @@ namespace HordeServer.Services
 				NewLabels.Add(NewLabel);
 			}
 
-			ObjectId JobIdValue = Request.JobId.ToObjectId();
+			JobId JobIdValue = new JobId(Request.JobId);
 			for (; ; )
 			{
 				IJob? Job = await JobService.GetJobAsync(JobIdValue);
@@ -880,7 +882,7 @@ namespace HordeServer.Services
 			foreach (CreateEventRequest Event in Request.Events)
 			{
 				NewLogEventData NewEvent = new NewLogEventData();
-				NewEvent.LogId = Event.LogId.ToObjectId();
+				NewEvent.LogId = new LogId(Event.LogId);
 				NewEvent.Severity = Event.Severity;
 				NewEvent.LineIndex = Event.LineIndex;
 				NewEvent.LineCount = Event.LineCount;
@@ -898,7 +900,7 @@ namespace HordeServer.Services
 		/// <returns>Information about the new agent</returns>
 		public async override Task<Empty> WriteOutput(WriteOutputRequest Request, ServerCallContext Context)
 		{
-			ILogFile? LogFile = await LogFileService.GetCachedLogFileAsync(Request.LogId.ToObjectId());
+			ILogFile? LogFile = await LogFileService.GetCachedLogFileAsync(new LogId(Request.LogId));
 			if (LogFile == null)
 			{
 				throw new StructuredRpcException(StatusCode.NotFound, "Resource not found");
@@ -987,7 +989,7 @@ namespace HordeServer.Services
 			}
 
 			// Get the job and step
-			IJob Job = await GetJobAsync(Metadata.JobId.ToObjectId());
+			IJob Job = await GetJobAsync(new JobId(Metadata.JobId));
 			IJobStepBatch Batch = AuthorizeBatch(Job, Metadata.BatchId.ToSubResourceId(), Context);
 
 			IJobStep? Step;
@@ -1022,7 +1024,7 @@ namespace HordeServer.Services
 			{
 				UploadTestDataRequest Request = Reader.Current;
 
-				ObjectId JobId = Request.JobId.ToObjectId();
+				JobId JobId = new JobId(Request.JobId);
 				if (Job == null || JobId != Job.Id)
 				{
 					Job = await JobService.GetJobAsync(JobId);
@@ -1058,7 +1060,7 @@ namespace HordeServer.Services
 		/// <returns></returns>
 		public override async Task<CreateReportResponse> CreateReport(CreateReportRequest Request, ServerCallContext Context)
 		{
-			IJob Job = await GetJobAsync(Request.JobId.ToObjectId());
+			IJob Job = await GetJobAsync(new JobId(Request.JobId));
 			IJobStepBatch Batch = AuthorizeBatch(Job, Request.BatchId.ToSubResourceId(), Context);
 
 			Report NewReport = new Report { Name = Request.Name, Placement = Request.Placement, ArtifactId = Request.ArtifactId.ToObjectId() };
diff --git a/Engine/Source/Programs/Horde/HordeServer/Startup.cs b/Engine/Source/Programs/Horde/HordeServer/Startup.cs
--- a/Engine/Source/Programs/Horde/HordeServer/Startup.cs
+++ b/Engine/Source/Programs/Horde/HordeServer/Startup.cs
@@ -86,6 +86,12 @@ using HordeServer.Jobs;
 
 namespace HordeServer
 {
+	using MessageTemplate = EpicGames.Core.MessageTemplate;
+	using JobId = ObjectId<IJob>;
+	using LeaseId = ObjectId<ILease>;
+	using LogId = ObjectId<ILogFile>;
+	using ILogger = Microsoft.Extensions.Logging.ILogger;
+
 	class Startup
 	{
 		class GrpcExceptionInterceptor : Interceptor
@@ -596,10 +602,19 @@ namespace HordeServer
 			}
 
 			ConfigureMongoDbClient();
+			ConfigureFormatters();
 
 			OnAddHealthChecks(Services);
 		}
 
+		public static void ConfigureFormatters()
+		{
+			MessageTemplate.RegisterFormatter(typeof(AgentId), new MessageTemplateTypeNameFormatter("AgentId"));
+			MessageTemplate.RegisterFormatter(typeof(JobId), new MessageTemplateTypeNameFormatter("JobId"));
+			MessageTemplate.RegisterFormatter(typeof(LeaseId), new MessageTemplateTypeNameFormatter("LeaseId"));
+			MessageTemplate.RegisterFormatter(typeof(LogId), new MessageTemplateTypeNameFormatter("LogId"));
+		}
+
 		public static void ConfigureJsonSerializer(JsonSerializerOptions Options)
 		{
 			Options.IgnoreNullValues = true;
@@ -740,6 +755,7 @@ namespace HordeServer
 			// Register the custom serializers
 			BsonSerializer.RegisterSerializationProvider(new BsonSerializationProvider());
 			BsonSerializer.RegisterSerializationProvider(new StringIdSerializationProvider());
+			BsonSerializer.RegisterSerializationProvider(new ObjectIdSerializationProvider());
 		}
 
 		private static void OnAddHealthChecks(IServiceCollection Services)
diff --git a/Engine/Source/Programs/Horde/HordeServer/Users/IUser.cs b/Engine/Source/Programs/Horde/HordeServer/Users/IUser.cs
--- a/Engine/Source/Programs/Horde/HordeServer/Users/IUser.cs
+++ b/Engine/Source/Programs/Horde/HordeServer/Users/IUser.cs
@@ -1,5 +1,6 @@
 // Copyright Epic Games, Inc. All Rights Reserved.
 
+using HordeServer.Utilities;
 using MongoDB.Bson;
 using System;
 using System.Collections.Generic;
@@ -9,6 +10,8 @@ using System.Threading.Tasks;
 
 namespace HordeServer.Models
 {
+	using JobId = ObjectId<IJob>;
+
 	/// <summary>
 	/// Document which collates information about a user, and their personal settings
 	/// </summary>
@@ -79,6 +82,6 @@ namespace HordeServer.Models
 		/// <summary>
 		/// List of pinned jobs
 		/// </summary>
-		public IReadOnlyList<ObjectId> PinnedJobIds { get; }
+		public IReadOnlyList<JobId> PinnedJobIds { get; }
 	}
 }
diff --git a/Engine/Source/Programs/Horde/HordeServer/Users/IUserCollection.cs b/Engine/Source/Programs/Horde/HordeServer/Users/IUserCollection.cs
--- a/Engine/Source/Programs/Horde/HordeServer/Users/IUserCollection.cs
+++ b/Engine/Source/Programs/Horde/HordeServer/Users/IUserCollection.cs
@@ -12,6 +12,8 @@ using System.Threading.Tasks;
 
 namespace HordeServer.Collections
 {
+	using JobId = ObjectId<IJob>;
+
 	/// <summary>
 	/// Manages user documents
 	/// </summary>
@@ -89,7 +91,7 @@ namespace HordeServer.Collections
 		/// <param name="AddPinnedJobIds"></param>
 		/// <param name="RemovePinnedJobIds"></param>
 		/// <returns>Updated user object</returns>
-		Task UpdateSettingsAsync(ObjectId UserId, bool? EnableExperimentalFeatures = null, bool? EnableIssueNotifications = null, BsonValue? DashboardSettings = null, IEnumerable<ObjectId>? AddPinnedJobIds = null, IEnumerable<ObjectId>? RemovePinnedJobIds = null);
+		Task UpdateSettingsAsync(ObjectId UserId, bool? EnableExperimentalFeatures = null, bool? EnableIssueNotifications = null, BsonValue? DashboardSettings = null, IEnumerable<JobId>? AddPinnedJobIds = null, IEnumerable<JobId>? RemovePinnedJobIds = null);
 	}
 
 	/// <summary>
diff --git a/Engine/Source/Programs/Horde/HordeServer/Users/Impl/UserCollection.cs b/Engine/Source/Programs/Horde/HordeServer/Users/Impl/UserCollection.cs
--- a/Engine/Source/Programs/Horde/HordeServer/Users/Impl/UserCollection.cs
+++ b/Engine/Source/Programs/Horde/HordeServer/Users/Impl/UserCollection.cs
@@ -15,6 +15,8 @@ using System.Threading.Tasks;
 
 namespace HordeServer.Collections.Impl
 {
+	using JobId = ObjectId<IJob>;
+
 	/// <summary>
 	/// Manages user documents
 	/// </summary>
@@ -34,7 +36,7 @@ namespace HordeServer.Collections.Impl
 			public bool EnableIssueNotifications { get; set; }
 
 			public BsonValue DashboardSettings { get; set; } = BsonNull.Value;
-			public List<ObjectId> PinnedJobIds { get; set; } = new List<ObjectId>();
+			public List<JobId> PinnedJobIds { get; set; } = new List<JobId>();
 
 			string IUser.Name => Claims.FirstOrDefault(x => String.Equals(x.Type, "name", StringComparison.Ordinal))?.Value ?? PrimaryClaim.Value;
 			string IUser.Login => Claims.FirstOrDefault(x => String.Equals(x.Type, ClaimTypes.Name, StringComparison.Ordinal))?.Value ?? PrimaryClaim.Value;
@@ -44,7 +46,7 @@ namespace HordeServer.Collections.Impl
 			IReadOnlyList<IUserClaim> IUserClaims.Claims => Claims;
 
 			ObjectId IUserSettings.UserId => Id;
-			IReadOnlyList<ObjectId> IUserSettings.PinnedJobIds => PinnedJobIds;
+			IReadOnlyList<JobId> IUserSettings.PinnedJobIds => PinnedJobIds;
 		}
 
 		class ClaimDocument : IUserClaim
@@ -150,13 +152,13 @@ namespace HordeServer.Collections.Impl
 		}
 
 		/// <inheritdoc/>
-		public async Task UpdateSettingsAsync(ObjectId UserId, bool? EnableExperimentalFeatures, bool? EnableIssueNotifications, BsonValue? DashboardSettings = null, IEnumerable<ObjectId>? AddPinnedJobIds = null, IEnumerable<ObjectId>? RemovePinnedJobIds = null)
+		public async Task UpdateSettingsAsync(ObjectId UserId, bool? EnableExperimentalFeatures, bool? EnableIssueNotifications, BsonValue? DashboardSettings = null, IEnumerable<JobId>? AddPinnedJobIds = null, IEnumerable<JobId>? RemovePinnedJobIds = null)
 		{
 			if (AddPinnedJobIds != null)
 			{
-				foreach (ObjectId PinnedJobId in AddPinnedJobIds)
+				foreach (JobId PinnedJobId in AddPinnedJobIds)
 				{
-					FilterDefinition<UserDocument> Filter = Builders<UserDocument>.Filter.Eq(x => x.Id, UserId) & Builders<UserDocument>.Filter.AnyNin<ObjectId>(x => x.PinnedJobIds, new[] { PinnedJobId });
+					FilterDefinition<UserDocument> Filter = Builders<UserDocument>.Filter.Eq(x => x.Id, UserId) & Builders<UserDocument>.Filter.AnyNin<JobId>(x => x.PinnedJobIds, new[] { PinnedJobId });
 					UpdateDefinition<UserDocument> Update = Builders<UserDocument>.Update.PushEach(x => x.PinnedJobIds, new[] { PinnedJobId }, -50);
 					await Users.UpdateOneAsync(Filter, Update);
 				}
diff --git a/Engine/Source/Programs/Horde/HordeServer/Users/Impl/UserCollectionV2.cs b/Engine/Source/Programs/Horde/HordeServer/Users/Impl/UserCollectionV2.cs
--- a/Engine/Source/Programs/Horde/HordeServer/Users/Impl/UserCollectionV2.cs
+++ b/Engine/Source/Programs/Horde/HordeServer/Users/Impl/UserCollectionV2.cs
@@ -17,6 +17,8 @@ using System.Threading.Tasks;
 
 namespace HordeServer.Collections.Impl
 {
+	using JobId = ObjectId<IJob>;
+
 	/// <summary>
 	/// Manages user documents
 	/// </summary>
@@ -95,10 +97,10 @@ namespace HordeServer.Collections.Impl
 			public bool EnableIssueNotifications { get; set; } = true;
 
 			public BsonValue DashboardSettings { get; set; } = BsonNull.Value;
-			public List<ObjectId> PinnedJobIds { get; set; } = new List<ObjectId>();
+			public List<JobId> PinnedJobIds { get; set; } = new List<JobId>();
 
 			ObjectId IUserSettings.UserId => Id;
-			IReadOnlyList<ObjectId> IUserSettings.PinnedJobIds => PinnedJobIds;
+			IReadOnlyList<JobId> IUserSettings.PinnedJobIds => PinnedJobIds;
 
 			[BsonConstructor]
 			private UserSettingsDocument()
@@ -116,7 +118,7 @@ namespace HordeServer.Collections.Impl
 				this.EnableExperimentalFeatures = Other.EnableExperimentalFeatures;
 				this.EnableIssueNotifications = Other.EnableIssueNotifications;
 				this.DashboardSettings = Other.DashboardSettings;
-				this.PinnedJobIds = new List<ObjectId>(Other.PinnedJobIds);
+				this.PinnedJobIds = new List<JobId>(Other.PinnedJobIds);
 			}
 		}
 
@@ -300,7 +302,7 @@ namespace HordeServer.Collections.Impl
 		}
 
 		/// <inheritdoc/>
-		public async Task UpdateSettingsAsync(ObjectId UserId, bool? EnableExperimentalFeatures = null, bool? EnableIssueNotifications = null, BsonValue? DashboardSettings = null, IEnumerable<ObjectId>? AddPinnedJobIds = null, IEnumerable<ObjectId>? RemovePinnedJobIds = null)
+		public async Task UpdateSettingsAsync(ObjectId UserId, bool? EnableExperimentalFeatures = null, bool? EnableIssueNotifications = null, BsonValue? DashboardSettings = null, IEnumerable<JobId>? AddPinnedJobIds = null, IEnumerable<JobId>? RemovePinnedJobIds = null)
 		{
 			List<UpdateDefinition<UserSettingsDocument>> Updates = new List<UpdateDefinition<UserSettingsDocument>>();
 			if (EnableExperimentalFeatures != null)
diff --git a/Engine/Source/Programs/Horde/HordeServer/Utilities/ObjectId.cs b/Engine/Source/Programs/Horde/HordeServer/Utilities/ObjectId.cs
new file mode 100644
--- /dev/null
+++ b/Engine/Source/Programs/Horde/HordeServer/Utilities/ObjectId.cs
@@ -0,0 +1,244 @@
+// Copyright Epic Games, Inc. All Rights Reserved.
+
+using EpicGames.Core;
+using EpicGames.Serialization.Converters;
+using MongoDB.Bson;
+using MongoDB.Bson.Serialization;
+using MongoDB.Bson.Serialization.Serializers;
+using System;
+using System.ComponentModel;
+using System.Diagnostics.CodeAnalysis;
+using System.Globalization;
+using System.Text.Json;
+using System.Text.Json.Serialization;
+
+namespace HordeServer.Utilities
+{
+	/// <summary>
+	/// Normalized string identifier for a resource
+	/// </summary>
+	[JsonSchemaString]
+	[TypeConverter(typeof(ObjectIdTypeConverter))]
+	public struct ObjectId<T> : IEquatable<ObjectId<T>>, IComparable<ObjectId<T>>
+	{
+		/// <summary>
+		/// Empty string
+		/// </summary>
+		[SuppressMessage("Design", "CA1000:Do not declare static members on generic types", Justification = "<Pending>")]
+		public static ObjectId<T> Empty { get; } = new ObjectId<T>(ObjectId.Empty);
+
+		/// <summary>
+		/// The text representing this id
+		/// </summary>
+		public ObjectId Value { get; }
+
+		/// <summary>
+		/// Constructor
+		/// </summary>
+		public ObjectId(ObjectId Value) => this.Value = Value;
+
+		/// <summary>
+		/// Constructor
+		/// </summary>
+		public ObjectId(byte[] Bytes) => this.Value = new ObjectId(Bytes);
+
+		/// <summary>
+		/// Constructor
+		/// </summary>
+		public ObjectId(string Value) => this.Value = new ObjectId(Value);
+
+		/// <inheritdoc cref="ObjectId.GenerateNewId()"/>
+		[SuppressMessage("Design", "CA1000:Do not declare static members on generic types", Justification = "<Pending>")]
+		public static ObjectId<T> GenerateNewId() => new ObjectId<T>(ObjectId.GenerateNewId());
+
+		/// <inheritdoc cref="ObjectId.Parse(string)"/>
+		[SuppressMessage("Design", "CA1000:Do not declare static members on generic types", Justification = "<Pending>")]
+		public static ObjectId<T> Parse(string Str) => new ObjectId<T>(ObjectId.Parse(Str));
+
+		/// <inheritdoc cref="ObjectId.TryParse(string, out ObjectId)"/>
+		[SuppressMessage("Design", "CA1000:Do not declare static members on generic types", Justification = "<Pending>")]
+		public static bool TryParse(string Str, out ObjectId<T> Value)
+		{
+			if (ObjectId.TryParse(Str, out ObjectId Id))
+			{
+				Value = new ObjectId<T>(Id);
+				return true;
+			}
+			else
+			{
+				Value = Empty;
+				return false;
+			}
+		}
+
+		/// <inheritdoc/>
+		public override bool Equals(object? Obj) => Obj is ObjectId<T> Other && Equals(Other);
+
+		/// <inheritdoc/>
+		public override int GetHashCode() => Value.GetHashCode();
+
+		/// <inheritdoc/>
+		public bool Equals(ObjectId<T> Other) => Value.Equals(Other.Value);
+
+		/// <inheritdoc/>
+		public override string ToString() => Value.ToString();
+
+		/// <inheritdoc/>
+		public int CompareTo(ObjectId<T> Other) => Value.CompareTo(Other.Value);
+
+		/// <inheritdoc cref="ObjectId.op_LessThan"/>
+		public static bool operator <(ObjectId<T> Left, ObjectId<T> Right) => Left.Value < Right.Value;
+
+		/// <inheritdoc cref="ObjectId.op_GreaterThan"/>
+		public static bool operator >(ObjectId<T> Left, ObjectId<T> Right) => Left.Value > Right.Value;
+
+		/// <inheritdoc cref="ObjectId.op_LessThanOrEqual"/>
+		public static bool operator <=(ObjectId<T> Left, ObjectId<T> Right) => Left.Value <= Right.Value;
+
+		/// <inheritdoc cref="ObjectId.op_GreaterThanOrEqual"/>
+		public static bool operator >=(ObjectId<T> Left, ObjectId<T> Right) => Left.Value >= Right.Value;
+
+		/// <inheritdoc cref="ObjectId.op_Equality"/>
+		public static bool operator ==(ObjectId<T> Left, ObjectId<T> Right) => Left.Value == Right.Value;
+
+		/// <inheritdoc cref="ObjectId.op_Inequality"/>
+		public static bool operator !=(ObjectId<T> Left, ObjectId<T> Right) => Left.Value != Right.Value;
+	}
+
+	/// <summary>
+	/// Converts <see cref="ObjectId{T}"/> values to and from JSON
+	/// </summary>
+	public class ObjectIdJsonConverter<T> : JsonConverter<ObjectId<T>>
+	{
+		/// <inheritdoc/>
+		public override ObjectId<T> Read(ref Utf8JsonReader Reader, Type TypeToConvert, JsonSerializerOptions Options)
+		{
+			return new ObjectId<T>(ObjectId.Parse(Reader.GetString()!));
+		}
+
+		/// <inheritdoc/>
+		public override void Write(Utf8JsonWriter Writer, ObjectId<T> Value, JsonSerializerOptions Options)
+		{
+			Writer.WriteStringValue(Value.ToString());
+		}
+	}
+
+	/// <summary>
+	/// Converts <see cref="ObjectId{T}"/> values to and from JSON
+	/// </summary>
+	public class JsonObjectIdConverterFactory : JsonConverterFactory
+	{
+		/// <inheritdoc/>
+		public override bool CanConvert(Type TypeToConvert)
+		{
+			return TypeToConvert.IsGenericType && TypeToConvert.GetGenericTypeDefinition() == typeof(ObjectId<>);
+		}
+
+		/// <inheritdoc/>
+		public override JsonConverter? CreateConverter(Type Type, JsonSerializerOptions Options)
+		{
+			return (JsonConverter?)Activator.CreateInstance(typeof(ObjectIdJsonConverter<>).MakeGenericType(Type.GetGenericArguments()));
+		}
+	}
+
+	/// <summary>
+	/// Serializer for ObjectId objects
+	/// </summary>
+	public sealed class ObjectIdBsonSerializer<T> : SerializerBase<ObjectId<T>>
+	{
+		/// <inheritdoc/>
+		public override ObjectId<T> Deserialize(BsonDeserializationContext Context, BsonDeserializationArgs Args)
+		{
+			return new ObjectId<T>(Context.Reader.ReadObjectId());
+		}
+
+		/// <inheritdoc/>
+		public override void Serialize(BsonSerializationContext Context, BsonSerializationArgs Args, ObjectId<T> Value)
+		{
+			Context.Writer.WriteObjectId(Value.Value);
+		}
+	}
+
+	/// <summary>
+	/// Serializer for ObjectId objects
+	/// </summary>
+	public sealed class ObjectIdSerializationProvider : BsonSerializationProviderBase
+	{
+		/// <inheritdoc/>
+		public override IBsonSerializer? GetSerializer(Type Type, IBsonSerializerRegistry SerializerRegistry)
+		{
+			if (Type.IsGenericType && Type.GetGenericTypeDefinition() == typeof(ObjectId<>))
+			{
+				return (IBsonSerializer?)Activator.CreateInstance(typeof(ObjectIdBsonSerializer<>).MakeGenericType(Type.GetGenericArguments()));
+			}
+			else
+			{
+				return null;
+			}
+		}
+	}
+
+	/// <summary>
+	/// Type converter from strings to PropertyFilter objects
+	/// </summary>
+	sealed class ObjectIdTypeConverter : TypeConverter
+	{
+		Type Type;
+
+		/// <summary>
+		/// Constructor
+		/// </summary>
+		/// <param name="Type"></param>
+		public ObjectIdTypeConverter(Type Type)
+		{
+			this.Type = Type;
+		}
+
+		/// <inheritdoc/>
+		public override bool CanConvertFrom(ITypeDescriptorContext Context, Type SourceType)
+		{
+			return SourceType == typeof(string) || base.CanConvertFrom(Context, SourceType);
+		}
+
+		/// <inheritdoc/>
+		public override object ConvertFrom(ITypeDescriptorContext Context, CultureInfo Culture, object Value)
+		{
+			return Activator.CreateInstance(Type, Value)!;
+		}
+
+		/// <inheritdoc/>
+		public override bool CanConvertTo(ITypeDescriptorContext Context, Type DestinationType)
+		{
+			if (DestinationType == typeof(string))
+			{
+				return true;
+			}
+			if (DestinationType.IsGenericType)
+			{
+				Type GenericTypeDefinition = DestinationType.GetGenericTypeDefinition();
+				if (GenericTypeDefinition == typeof(ObjectId<>))
+				{
+					return true;
+				}
+				if (GenericTypeDefinition == typeof(Nullable<>))
+				{
+					return CanConvertTo(Context, GenericTypeDefinition.GetGenericArguments()[0]);
+				}
+			}
+			return false;
+		}
+
+		/// <inheritdoc/>
+		public override object? ConvertTo(ITypeDescriptorContext Context, CultureInfo Culture, object Value, Type DestinationType)
+		{
+			if (DestinationType == typeof(string))
+			{
+				return Value.ToString();
+			}
+			else
+			{
+				return Activator.CreateInstance(DestinationType, Value);
+			}
+		}
+	}
+}
diff --git a/Engine/Source/Programs/Horde/HordeServer/Utilities/ObjectIdHelpers.cs b/Engine/Source/Programs/Horde/HordeServer/Utilities/ObjectIdHelpers.cs
--- a/Engine/Source/Programs/Horde/HordeServer/Utilities/ObjectIdHelpers.cs
+++ b/Engine/Source/Programs/Horde/HordeServer/Utilities/ObjectIdHelpers.cs
@@ -21,5 +21,10 @@ namespace HordeServer.Utilities
 				return ObjectId.Parse(Text);
 			}
 		}
+
+		public static ObjectId<T> ToObjectId<T>(this string Text)
+		{
+			return new ObjectId<T>(Text.ToObjectId());
+		}
 	}
 }
diff --git a/Engine/Source/Programs/Shared/EpicGames.Core/MessageTemplate.cs b/Engine/Source/Programs/Shared/EpicGames.Core/MessageTemplate.cs
--- a/Engine/Source/Programs/Shared/EpicGames.Core/MessageTemplate.cs
+++ b/Engine/Source/Programs/Shared/EpicGames.Core/MessageTemplate.cs
@@ -18,23 +18,49 @@ namespace EpicGames.Core
 	/// </summary>
 	public interface IMessageTemplateFormatter
 	{
+		/// <summary>
+		/// The type name to serialize
+		/// </summary>
+		public string Name { get; }
+
 		/// <summary>
 		/// Serializes the given value. The state of the supplied JSON writer will be within a JSON object; named properties should typically be serialized directly here.
 		/// </summary>
 		public void Write(Utf8JsonWriter Writer, object Value);
 	}
 
+	/// <summary>
+	/// Implementation of <see cref="IMessageTemplateFormatter"/> that just serializes a type name
+	/// </summary>
+	public class MessageTemplateTypeNameFormatter : IMessageTemplateFormatter
+	{
+		/// <inheritdoc/>
+		public string Name { get; }
+
+		/// <summary>
+		/// Constructor
+		/// </summary>
+		/// <param name="Name">The type name</param>
+		public MessageTemplateTypeNameFormatter(string Name)
+		{
+			this.Name = Name;
+		}
+
+		/// <inheritdoc/>
+		public void Write(Utf8JsonWriter Writer, object Value)
+		{
+		}
+	}
+
 	/// <summary>
 	/// Attribute indicating the formatter to use for a message template argument
 	/// </summary>
 	public class MessageTemplateFormatterAttribute : Attribute
 	{
-		public string Name { get; }
 		public Type Type { get; }
 
-		public MessageTemplateFormatterAttribute(string Name, Type Type)
+		public MessageTemplateFormatterAttribute(Type Type)
 		{
-			this.Name = Name;
 			this.Type = Type;
 		}
 	}
@@ -52,7 +78,7 @@ namespace EpicGames.Core
 		/// <summary>
 		/// Static cache of formatters for property values
 		/// </summary>
-		static ConcurrentDictionary<Type, IMessageTemplateFormatter> TypeToWriter = new ConcurrentDictionary<Type, IMessageTemplateFormatter>();
+		static ConcurrentDictionary<Type, IMessageTemplateFormatter?> TypeToFormatter = new ConcurrentDictionary<Type, IMessageTemplateFormatter?>();
 
 		/// <summary>
 		/// Register a formatter for the given type
@@ -61,7 +87,7 @@ namespace EpicGames.Core
 		/// <param name="Formatter"></param>
 		public static void RegisterFormatter(Type Type, IMessageTemplateFormatter Formatter)
 		{
-			TypeToWriter.TryAdd(Type, Formatter);
+			TypeToFormatter.TryAdd(Type, Formatter);
 		}
 
 		/// <summary>
@@ -318,18 +344,21 @@ namespace EpicGames.Core
 		/// <param name="JsonWriter">The writer to output to</param>
 		static void WriteComplexPropertyValue(object Value, Type ValueType, Utf8JsonWriter JsonWriter)
 		{
-			MessageTemplateFormatterAttribute? Attribute = ValueType.GetCustomAttribute<MessageTemplateFormatterAttribute>();
-			if (Attribute != null)
+			IMessageTemplateFormatter? Formatter;
+			if (!TypeToFormatter.TryGetValue(ValueType, out Formatter))
 			{
-				IMessageTemplateFormatter? Formatter;
-				if (!TypeToWriter.TryGetValue(Attribute.Type, out Formatter))
+				MessageTemplateFormatterAttribute? Attribute = ValueType.GetCustomAttribute<MessageTemplateFormatterAttribute>();
+				if (Attribute != null)
 				{
 					Formatter = (IMessageTemplateFormatter)Activator.CreateInstance(Attribute.Type)!;
-					TypeToWriter.TryAdd(Attribute.Type, Formatter);
 				}
+				TypeToFormatter.TryAdd(ValueType, Formatter);
+			}
 
+			if (Formatter != null)
+			{
 				JsonWriter.WriteStartObject();
-				JsonWriter.WriteString("$type", Attribute.Name);
+				JsonWriter.WriteString("$type", Formatter.Name);
 				JsonWriter.WriteString("$text", Value.ToString());
 				Formatter.Write(JsonWriter, Value);
 				JsonWriter.WriteEndObject();
