Horde: Use the default Kestrel ports unless overridden by the deployment settings. The Horde agent defaults to port 5001 by default.

diff --git a/Engine/Source/Programs/Horde/HordeServer/Program.cs b/Engine/Source/Programs/Horde/HordeServer/Program.cs
--- a/Engine/Source/Programs/Horde/HordeServer/Program.cs
+++ b/Engine/Source/Programs/Horde/HordeServer/Program.cs
@@ -240,15 +240,21 @@ namespace HordeServer
 				{
 					WebBuilder.ConfigureKestrel(Options =>
 					{
-						Options.ListenAnyIP(ServerSettings.HttpPort, Configure => { Configure.Protocols = HttpProtocols.Http1AndHttp2; });
-							
+						if (ServerSettings.HttpPort != 0)
+						{
+							Options.ListenAnyIP(ServerSettings.HttpPort, Configure => { Configure.Protocols = HttpProtocols.Http1AndHttp2; });
+						}
+
+						if (ServerSettings.HttpsPort != 0)
+						{
+							Options.ListenAnyIP(ServerSettings.HttpsPort, Configure => { if (SslCert != null) { Configure.UseHttps(SslCert); } });
+						}
+
 						// To serve HTTP/2 with gRPC *without* TLS enabled, a separate port for HTTP/2 must be used.
 						// This is useful when having a load balancer in front that terminates TLS.
-						Options.ListenAnyIP(ServerSettings.Http2Port, Configure => { Configure.Protocols = HttpProtocols.Http2; });
-
-						if (SslCert != null)
+						if (ServerSettings.Http2Port != 0)
 						{
-							Options.ListenAnyIP(ServerSettings.HttpsPort, Configure => { Configure.UseHttps(SslCert); });
+							Options.ListenAnyIP(ServerSettings.Http2Port, Configure => { Configure.Protocols = HttpProtocols.Http2; });
 						}
 					});
 					WebBuilder.UseStartup<Startup>(); 
diff --git a/Engine/Source/Programs/Horde/HordeServer/ServerSettings.cs b/Engine/Source/Programs/Horde/HordeServer/ServerSettings.cs
--- a/Engine/Source/Programs/Horde/HordeServer/ServerSettings.cs
+++ b/Engine/Source/Programs/Horde/HordeServer/ServerSettings.cs
@@ -124,19 +124,19 @@ namespace HordeServer
 	public class ServerSettings
 	{
 		/// <summary>
-		/// Main port for serving HTTP
+		/// Main port for serving HTTP. Uses the default Kestrel port (5000) if not specified.
 		/// </summary>
-		public int HttpPort { get; set; } = 80;
-		
+		public int HttpPort { get; set; }
+
 		/// <summary>
-		/// Port for serving only HTTP/2
+		/// Port for serving HTTP with TLS enabled. Uses the default Kestrel port (5001) if not specified.
 		/// </summary>
-		public int Http2Port { get; set; } = 52103;
-		
+		public int HttpsPort { get; set; }
+
 		/// <summary>
-		/// Port for serving HTTP with TLS enabled
+		/// Dedicated port for serving only HTTP/2.
 		/// </summary>
-		public int HttpsPort { get; set; } = 443;
+		public int Http2Port { get; set; }
 
 		/// <summary>
 		/// Whether the server is running as a single instance or with multiple instances, such as in k8s
