Horde: Update Fluent to 8.36.5
#rnx

diff --git a/Engine/Source/Programs/Horde/HordeDashboard/package.json b/Engine/Source/Programs/Horde/HordeDashboard/package.json
--- a/Engine/Source/Programs/Horde/HordeDashboard/package.json
+++ b/Engine/Source/Programs/Horde/HordeDashboard/package.json
@@ -7,7 +7,7 @@
   "dependencies": {
     "@blueprintjs/icons": "^3.13.0",
     "@datadog/browser-logs": "^1.12.4",
-    "@fluentui/react": "^7.155.0",
+    "@fluentui/react": "^8.36.5",
     "@fluentui/react-charting": "^5.0.22",
     "crypto-js": "^4.0.0",
     "formik": "^2.1.1",
@@ -50,6 +50,7 @@
   "devDependencies": {
     "@types/crypto-js": "^4.0.1",
     "@types/markdown-to-jsx": "6.9.1",
+    "@types/react": "^17.0.30",
     "@types/react-custom-scrollbars": "^4.0.6",
     "@types/react-dom": "^17.0.0",
     "@types/react-highlighter": "^0.3.4",
diff --git a/Engine/Source/Programs/Horde/HordeDashboard/src/backend/Api.ts b/Engine/Source/Programs/Horde/HordeDashboard/src/backend/Api.ts
--- a/Engine/Source/Programs/Horde/HordeDashboard/src/backend/Api.ts
+++ b/Engine/Source/Programs/Horde/HordeDashboard/src/backend/Api.ts
@@ -374,6 +374,26 @@ export type UpdateAgentRequest = {
 
 }
 
+export type AuditLogQuery = {
+	minTime?: string;
+	maxTime?: string;
+	index?: number;
+	count?: number;
+}
+
+export enum AuditLogLevel {
+    Information = "Information"
+}
+
+export type AuditLogEntry = {
+	time: string;
+	level: AuditLogLevel;
+	message: string;
+	format: string;
+	properties?: Record<string, string>;
+}
+
+
 /**Parameters to create a pool */
 export type CreatePoolRequest = {
 
diff --git a/Engine/Source/Programs/Horde/HordeDashboard/src/backend/Backend.ts b/Engine/Source/Programs/Horde/HordeDashboard/src/backend/Backend.ts
--- a/Engine/Source/Programs/Horde/HordeDashboard/src/backend/Backend.ts
+++ b/Engine/Source/Programs/Horde/HordeDashboard/src/backend/Backend.ts
@@ -3,7 +3,7 @@
 import { configure } from "mobx";
 import { isNumber } from 'util';
 import templateCache from '../backend/TemplateCache';
-import { AgentData, ArtifactData, BatchUpdatePoolRequest, ChangeSummaryData, CreateDeviceRequest, CreateDeviceResponse, CreateJobRequest, CreateJobResponse, CreatePoolRequest, CreateSoftwareResponse, CreateSubscriptionRequest, CreateSubscriptionResponse, DashboardPreference, EventData, GetArtifactZipRequest, GetDevicePlatformResponse, GetDevicePoolResponse, GetDeviceReservationResponse, GetDeviceResponse, GetGraphResponse, GetIssueStreamResponse, GetJobsTabResponse, GetJobStepRefResponse, GetJobStepTraceResponse, GetJobTimingResponse, GetLogEventResponse, GetNotificationResponse, GetPerforceServerStatusResponse, GetServerInfoResponse, GetServerSettingsResponse, GetSoftwareResponse, GetSubscriptionResponse, GetUserResponse, GetUtilizationTelemetryResponse, GlobalConfig, IssueData, IssueQuery, JobData, JobQuery, JobsTabColumnType, LeaseData, LogData, LogLineData, PerforceCluster, PoolData, ProjectData, ScheduleData, ScheduleQuery, SearchLogFileResponse, ServerUpdateResponse, SessionData, StreamData, TabType, TemplateData, TestData, UpdateAgentRequest, UpdateDeviceRequest, UpdateGlobalConfigRequest, UpdateIssueRequest, UpdateJobRequest, UpdateLeaseRequest, UpdateNotificationsRequest, UpdatePoolRequest, UpdateServerSettingsRequest, UpdateStepRequest, UpdateStepResponse, UpdateUserRequest, UsersQuery } from './Api';
+import { AgentData, AuditLogEntry, AuditLogQuery, ArtifactData, BatchUpdatePoolRequest, ChangeSummaryData, CreateDeviceRequest, CreateDeviceResponse, CreateJobRequest, CreateJobResponse, CreatePoolRequest, CreateSoftwareResponse, CreateSubscriptionRequest, CreateSubscriptionResponse, DashboardPreference, EventData, GetArtifactZipRequest, GetDevicePlatformResponse, GetDevicePoolResponse, GetDeviceReservationResponse, GetDeviceResponse, GetGraphResponse, GetIssueStreamResponse, GetJobsTabResponse, GetJobStepRefResponse, GetJobStepTraceResponse, GetJobTimingResponse, GetLogEventResponse, GetNotificationResponse, GetPerforceServerStatusResponse, GetServerInfoResponse, GetServerSettingsResponse, GetSoftwareResponse, GetSubscriptionResponse, GetUserResponse, GetUtilizationTelemetryResponse, GlobalConfig, IssueData, IssueQuery, JobData, JobQuery, JobsTabColumnType, LeaseData, LogData, LogLineData, PerforceCluster, PoolData, ProjectData, ScheduleData, ScheduleQuery, SearchLogFileResponse, ServerUpdateResponse, SessionData, StreamData, TabType, TemplateData, TestData, UpdateAgentRequest, UpdateDeviceRequest, UpdateGlobalConfigRequest, UpdateIssueRequest, UpdateJobRequest, UpdateLeaseRequest, UpdateNotificationsRequest, UpdatePoolRequest, UpdateServerSettingsRequest, UpdateStepRequest, UpdateStepResponse, UpdateUserRequest, UsersQuery } from './Api';
 import dashboard from './Dashboard';
 import { ChallengeStatus, Fetch } from './Fetch';
 import graphCache, { GraphQuery } from './GraphCache';
@@ -107,6 +107,15 @@ export class Backend {
         });
     }
 
+    getAgentHistory(id: string, query: AuditLogQuery): Promise<AuditLogEntry[]> {
+        return new Promise<AuditLogEntry[]>((resolve, reject) => {
+            this.backend.get(`/api/v1/agents/${encodeURIComponent(id)}/history`, { params: query }).then((response) => {
+                const history = (response.data?.entries ?? []) as AuditLogEntry[];
+                resolve(history);
+            }).catch(reason => { reject(reason); });
+        });
+    }
+
     getAgents(modifiedAfter: Date): Promise<AgentData[]> {
         return new Promise<AgentData[]>((resolve, reject) => {
             // for whenever we bring back modifiedDate, swap to this line
diff --git a/Engine/Source/Programs/Horde/HordeDashboard/src/backend/JobDetails.ts b/Engine/Source/Programs/Horde/HordeDashboard/src/backend/JobDetails.ts
--- a/Engine/Source/Programs/Horde/HordeDashboard/src/backend/JobDetails.ts
+++ b/Engine/Source/Programs/Horde/HordeDashboard/src/backend/JobDetails.ts
@@ -1,7 +1,7 @@
 // Copyright Epic Games, Inc. All Rights Reserved.
 
 import { action, observable } from 'mobx';
-import { getTheme, mergeStyles, mergeStyleSets } from 'office-ui-fabric-react/lib/Styling';
+import { getTheme, mergeStyles, mergeStyleSets } from '@fluentui/react/lib/Styling';
 import backend from '.';
 import { getBatchInitElapsed, getNiceTime, getStepElapsed, getStepETA, getStepFinishTime, getStepTimingDelta } from '../base/utilities/timeUtils';
 import { getBatchText } from '../components/JobDetailCommon';
diff --git a/Engine/Source/Programs/Horde/HordeDashboard/src/base/components/Markdown/Markdown.types.ts b/Engine/Source/Programs/Horde/HordeDashboard/src/base/components/Markdown/Markdown.types.ts
--- a/Engine/Source/Programs/Horde/HordeDashboard/src/base/components/Markdown/Markdown.types.ts
+++ b/Engine/Source/Programs/Horde/HordeDashboard/src/base/components/Markdown/Markdown.types.ts
@@ -1,4 +1,4 @@
-import { ITheme, IStyleFunctionOrObject, IStyle, ILinkStyleProps, IImageStyleProps } from 'office-ui-fabric-react';
+import { ITheme, IStyleFunctionOrObject, IStyle, ILinkStyleProps, IImageStyleProps } from '@fluentui/react';
 import { IMarkdownHeaderStyleProps } from './MarkdownHeader';
 import { IMarkdownParagraphStyleProps } from './MarkdownParagraph';
 import { IMarkdownTableStyleProps } from '../MarkdownTable/index';
diff --git a/Engine/Source/Programs/Horde/HordeDashboard/src/base/components/MarkdownTable/MarkdownTable.styles.ts b/Engine/Source/Programs/Horde/HordeDashboard/src/base/components/MarkdownTable/MarkdownTable.styles.ts
--- a/Engine/Source/Programs/Horde/HordeDashboard/src/base/components/MarkdownTable/MarkdownTable.styles.ts
+++ b/Engine/Source/Programs/Horde/HordeDashboard/src/base/components/MarkdownTable/MarkdownTable.styles.ts
@@ -1,4 +1,4 @@
-import { getGlobalClassNames, IStyle } from 'office-ui-fabric-react';
+import { getGlobalClassNames, IStyle } from '@fluentui/react';
 import { IMarkdownTableStyleProps, IMarkdownTableStyles } from './MarkdownTable.types';
 
 const GlobalClassNames = {
diff --git a/Engine/Source/Programs/Horde/HordeDashboard/src/base/components/MarkdownTable/MarkdownTable.types.ts b/Engine/Source/Programs/Horde/HordeDashboard/src/base/components/MarkdownTable/MarkdownTable.types.ts
--- a/Engine/Source/Programs/Horde/HordeDashboard/src/base/components/MarkdownTable/MarkdownTable.types.ts
+++ b/Engine/Source/Programs/Horde/HordeDashboard/src/base/components/MarkdownTable/MarkdownTable.types.ts
@@ -1,4 +1,4 @@
-import { IStyle, ITheme, IStyleFunctionOrObject } from 'office-ui-fabric-react';
+import { IStyle, ITheme, IStyleFunctionOrObject } from '@fluentui/react';
 
 export interface IMarkdownTableProps {
   className?: string;
diff --git a/Engine/Source/Programs/Horde/HordeDashboard/src/base/components/SideRail/SideRail.styles.ts b/Engine/Source/Programs/Horde/HordeDashboard/src/base/components/SideRail/SideRail.styles.ts
--- a/Engine/Source/Programs/Horde/HordeDashboard/src/base/components/SideRail/SideRail.styles.ts
+++ b/Engine/Source/Programs/Horde/HordeDashboard/src/base/components/SideRail/SideRail.styles.ts
@@ -1,4 +1,4 @@
-import { FontWeights, getFocusOutlineStyle, IStyleFunction } from 'office-ui-fabric-react/lib/index';
+import { FontWeights, getFocusOutlineStyle, IStyleFunction } from '@fluentui/react/lib/index';
 import { ISideRailStyleProps, ISideRailStyles } from './SideRail.types';
 
 
diff --git a/Engine/Source/Programs/Horde/HordeDashboard/src/base/components/SideRail/SideRail.types.ts b/Engine/Source/Programs/Horde/HordeDashboard/src/base/components/SideRail/SideRail.types.ts
--- a/Engine/Source/Programs/Horde/HordeDashboard/src/base/components/SideRail/SideRail.types.ts
+++ b/Engine/Source/Programs/Horde/HordeDashboard/src/base/components/SideRail/SideRail.types.ts
@@ -1,4 +1,4 @@
-import { ITheme, IStyleFunctionOrObject, IStyle } from 'office-ui-fabric-react';
+import { ITheme, IStyleFunctionOrObject, IStyle } from '@fluentui/react';
 
 export interface ISideRailProps {
   /**
