Horde: Add an audit log for issues, showing any related state transitions. Should make debugging issues more straightforward.

diff --git a/Engine/Source/Programs/Horde/HordeServer/Auditing/IAuditLog.cs b/Engine/Source/Programs/Horde/HordeServer/Auditing/IAuditLog.cs
--- a/Engine/Source/Programs/Horde/HordeServer/Auditing/IAuditLog.cs
+++ b/Engine/Source/Programs/Horde/HordeServer/Auditing/IAuditLog.cs
@@ -20,6 +20,7 @@ using EpicGames.Core;
 using System.Buffers;
 using System.Text.Json;
 using System.Text;
+using System.IO.Pipelines;
 
 namespace HordeServer.Collections
 {
@@ -118,4 +119,37 @@ namespace HordeServer.Collections
 		/// <returns></returns>
 		IAuditLog<TSubject> Create(string CollectionName, string SubjectProperty);
 	}
+
+	/// <summary>
+	/// Extension methods for audit log channels
+	/// </summary>
+	static class AuditLogExtensions
+	{
+		/// <summary>
+		/// Retrieve historical information about a specific agent
+		/// </summary>
+		/// <param name="Channel">Channel to query</param>
+		/// <param name="BodyWriter">Writer for Json data</param>
+		/// <param name="MinTime">Minimum time for records to return</param>
+		/// <param name="MaxTime">Maximum time for records to return</param>
+		/// <param name="Index">Offset of the first result</param>
+		/// <param name="Count">Number of records to return</param>
+		/// <returns>Information about the requested agent</returns>
+		public static async Task FindAsync<T>(this IAuditLogChannel<T> Channel, PipeWriter BodyWriter, DateTime? MinTime = null, DateTime? MaxTime = null, int Index = 0, int Count = 50)
+		{
+			string Prefix = "{\n\t\"entries\":\n\t[";
+			await BodyWriter.WriteAsync(Encoding.UTF8.GetBytes(Prefix));
+
+			string Separator = "";
+			await foreach (IAuditLogMessage<AgentId> Message in Channel.FindAsync(MinTime, MaxTime, Index, Count))
+			{
+				string Line = $"{Separator}\n\t\t{Message.Data}";
+				await BodyWriter.WriteAsync(Encoding.UTF8.GetBytes(Line));
+				Separator = ",";
+			}
+
+			string Suffix = "\n\t]\n}";
+			await BodyWriter.WriteAsync(Encoding.UTF8.GetBytes(Suffix));
+		}
+	}
 }
diff --git a/Engine/Source/Programs/Horde/HordeServer/Controllers/AgentsController.cs b/Engine/Source/Programs/Horde/HordeServer/Controllers/AgentsController.cs
--- a/Engine/Source/Programs/Horde/HordeServer/Controllers/AgentsController.cs
+++ b/Engine/Source/Programs/Horde/HordeServer/Controllers/AgentsController.cs
@@ -234,25 +234,10 @@ namespace HordeServer.Controllers
 		[Route("/api/v1/agents/{AgentId}/history")]
 		public async Task GetAgentHistoryAsync(string AgentId, [FromQuery] DateTime? MinTime = null, [FromQuery] DateTime? MaxTime = null, [FromQuery] int Index = 0, [FromQuery] int Count = 50)
 		{
-			IAuditLogChannel<AgentId> Channel = AgentService.Agents.GetLogger(AgentId.ToAgentId());
-
 			Response.ContentType = "application/json";
 			Response.StatusCode = 200;
 			await Response.StartAsync();
-
-			string Prefix = "{\n\t\"entries\":\n\t[";
-			await Response.BodyWriter.WriteAsync(Encoding.UTF8.GetBytes(Prefix));
-
-			string Separator = "";
-			await foreach (IAuditLogMessage<AgentId> Message in Channel.FindAsync(MinTime, MaxTime, Index, Count))
-			{
-				string Line = $"{Separator}\n\t\t{Message.Data}";
-				await Response.Body.WriteAsync(Encoding.UTF8.GetBytes(Line));
-				Separator = ",";
-			}
-
-			string Suffix = "\n\t]\n}";
-			await Response.BodyWriter.WriteAsync(Encoding.UTF8.GetBytes(Suffix));
+			await AgentService.Agents.GetLogger(AgentId.ToAgentId()).FindAsync(Response.BodyWriter, MinTime, MaxTime, Index, Count);
 		}
 
 		/// <summary>
diff --git a/Engine/Source/Programs/Horde/HordeServer/Controllers/IssuesController.cs b/Engine/Source/Programs/Horde/HordeServer/Controllers/IssuesController.cs
--- a/Engine/Source/Programs/Horde/HordeServer/Controllers/IssuesController.cs
+++ b/Engine/Source/Programs/Horde/HordeServer/Controllers/IssuesController.cs
@@ -33,47 +33,20 @@ namespace HordeServer.Controllers
 	[Route("[controller]")]
 	public class IssuesController : ControllerBase
 	{
-		/// <summary>
-		/// Singleton instance of the issue service
-		/// </summary>
+		private readonly IIssueCollection IssueCollection;
 		private readonly IIssueService IssueService;
-
-		/// <summary>
-		/// Singleton instance of the job service
-		/// </summary>
 		private readonly JobService JobService;
-
-		/// <summary>
-		/// Singleton instance of the stream service
-		/// </summary>
 		private readonly StreamService StreamService;
-
-		/// <summary>
-		/// 
-		/// </summary>
 		private readonly IUserCollection UserCollection;
-
-		/// <summary>
-		/// Collection of events
-		/// </summary>
 		private readonly ILogEventCollection LogEventCollection;
-
-		/// <summary>
-		/// The log file service
-		/// </summary>
 		private readonly ILogFileService LogFileService;
 
 		/// <summary>
 		/// Constructor
 		/// </summary>
-		/// <param name="IssueService">The issue service</param>
-		/// <param name="JobService">The job service</param>
-		/// <param name="StreamService">The stream service</param>
-		/// <param name="UserCollection"></param>
-		/// <param name="LogEventCollection">The event collection</param>
-		/// <param name="LogFileService">The log file service</param>
-		public IssuesController(IIssueService IssueService, JobService JobService, StreamService StreamService, IUserCollection UserCollection, ILogEventCollection LogEventCollection, ILogFileService LogFileService)
+		public IssuesController(IIssueCollection IssueCollection, IIssueService IssueService, JobService JobService, StreamService StreamService, IUserCollection UserCollection, ILogEventCollection LogEventCollection, ILogFileService LogFileService)
 		{
+			this.IssueCollection = IssueCollection;
 			this.IssueService = IssueService;
 			this.JobService = JobService;
 			this.StreamService = StreamService;
@@ -184,6 +157,25 @@ namespace HordeServer.Controllers
 			return PropertyFilter.Apply(await CreateIssueResponseAsync(Details, bShowDesktopAlerts), Filter);
 		}
 
+		/// <summary>
+		/// Retrieve historical information about a specific issue
+		/// </summary>
+		/// <param name="IssueId">Id of the agent to get information about</param>
+		/// <param name="MinTime">Minimum time for records to return</param>
+		/// <param name="MaxTime">Maximum time for records to return</param>
+		/// <param name="Index">Offset of the first result</param>
+		/// <param name="Count">Number of records to return</param>
+		/// <returns>Information about the requested agent</returns>
+		[HttpGet]
+		[Route("/api/v1/issues/{IssueId}/history")]
+		public async Task GetAgentHistoryAsync(int IssueId, [FromQuery] DateTime? MinTime = null, [FromQuery] DateTime? MaxTime = null, [FromQuery] int Index = 0, [FromQuery] int Count = 50)
+		{
+			Response.ContentType = "application/json";
+			Response.StatusCode = 200;
+			await Response.StartAsync();
+			await IssueCollection.GetLogger(IssueId).FindAsync(Response.BodyWriter, MinTime, MaxTime, Index, Count);
+		}
+
 		/// <summary>
 		/// Create an issue response object
 		/// </summary>
diff --git a/Engine/Source/Programs/Horde/HordeServer/Issues/IIssueCollection.cs b/Engine/Source/Programs/Horde/HordeServer/Issues/IIssueCollection.cs
--- a/Engine/Source/Programs/Horde/HordeServer/Issues/IIssueCollection.cs
+++ b/Engine/Source/Programs/Horde/HordeServer/Issues/IIssueCollection.cs
@@ -542,5 +542,12 @@ namespace HordeServer.Collections
 		Task<List<IIssueStep>> FindStepsAsync(JobId JobId, SubResourceId? BatchId, SubResourceId? StepId);
 
 		#endregion
+
+		/// <summary>
+		/// Gets the logger for a particular issue
+		/// </summary>
+		/// <param name="IssueId">The issue id</param>
+		/// <returns>Logger for this issue</returns>
+		IAuditLogChannel<int> GetLogger(int IssueId);
 	}
 }
diff --git a/Engine/Source/Programs/Horde/HordeServer/Issues/Impl/IssueCollection.cs b/Engine/Source/Programs/Horde/HordeServer/Issues/Impl/IssueCollection.cs
--- a/Engine/Source/Programs/Horde/HordeServer/Issues/Impl/IssueCollection.cs
+++ b/Engine/Source/Programs/Horde/HordeServer/Issues/Impl/IssueCollection.cs
@@ -436,9 +436,10 @@ namespace HordeServer.Collections.Impl
 		IMongoCollection<IssueSpan> IssueSpans;
 		IMongoCollection<IssueStep> IssueSteps;
 		IMongoCollection<IssueSuspect> IssueSuspects;
+		IAuditLog<int> AuditLog;
 		ILogger Logger;
 
-		public IssueCollection(DatabaseService DatabaseService, ILogger<IssueCollection> Logger)
+		public IssueCollection(DatabaseService DatabaseService, IAuditLogFactory<int> AuditLogFactory, ILogger<IssueCollection> Logger)
 		{
 			this.Logger = Logger;
 
@@ -448,7 +449,8 @@ namespace HordeServer.Collections.Impl
 			IssueSpans = DatabaseService.GetCollection<IssueSpan>("IssuesV2.Spans");
 			IssueSteps = DatabaseService.GetCollection<IssueStep>("IssuesV2.Steps");
 			IssueSuspects = DatabaseService.GetCollection<IssueSuspect>("IssuesV2.Suspects");
-			
+			AuditLog = AuditLogFactory.Create("IssuesV2.History", "IssueId");
+
 			if (!DatabaseService.ReadOnlyMode)
 			{
 				Issues.Indexes.CreateOne(new CreateIndexModel<Issue>(Builders<Issue>.IndexKeys.Ascending(x => x.ResolvedAt)));
@@ -573,10 +575,121 @@ namespace HordeServer.Collections.Impl
 				return null;
 			}
 
+			ILogger IssueLogger = GetLogger(NewIssue.Id);
+			IssueLogger.LogInformation("Created issue {IssueId}", NewIssue.Id);
+			LogSpanInfo(IssueLogger, Span);
+
 			await TokenValue.FlushAsync();
 			return NewIssue;
 		}
 
+		static void LogIssueChanges(ILogger IssueLogger, Issue OldIssue, Issue NewIssue)
+		{
+			if (NewIssue.Severity != OldIssue.Severity)
+			{
+				IssueLogger.LogInformation("Changed severity to {Severity}", NewIssue.Severity);
+			}
+			if (NewIssue.Summary != OldIssue.Summary)
+			{
+				IssueLogger.LogInformation("Changed summary to \"{Summary}\"", NewIssue.Summary);
+			}
+			if (NewIssue.OwnerId != OldIssue.OwnerId)
+			{
+				if (NewIssue.NominatedById != null)
+				{
+					IssueLogger.LogInformation("User {UserId} was nominated by {NominatedByUserId}", NewIssue.OwnerId, NewIssue.NominatedById);
+				}
+				else
+				{
+					IssueLogger.LogInformation("User {UserId} was nominated by default", NewIssue.OwnerId);
+				}
+			}
+			if (NewIssue.AcknowledgedAt != OldIssue.AcknowledgedAt)
+			{
+				if (NewIssue.AcknowledgedAt == null)
+				{
+					IssueLogger.LogInformation("Issue was un-acknowledged by {UserId}", OldIssue.OwnerId);
+				}
+				else
+				{
+					IssueLogger.LogInformation("Issue was acknowledged by {UserId}", OldIssue.OwnerId);
+				}
+			}
+			if (NewIssue.FixChange != OldIssue.FixChange)
+			{
+				if (NewIssue.FixChange == 0)
+				{
+					IssueLogger.LogInformation("Issue was marked as not fixed");
+				}
+				else
+				{
+					IssueLogger.LogInformation("Issue was marked as fixed in {Change}", NewIssue.FixChange);
+				}
+			}
+			if (NewIssue.ResolvedById != OldIssue.ResolvedById)
+			{
+				if (NewIssue.ResolvedById == null)
+				{
+					IssueLogger.LogInformation("Marking as unresolved");
+				}
+				else
+				{
+					IssueLogger.LogInformation("Resolved by {UserId}", NewIssue.ResolvedById);
+				}
+			}
+			if (NewIssue.NotifySuspects != OldIssue.NotifySuspects)
+			{
+				IssueLogger.LogInformation("Notify suspects set to {Value}", NewIssue.NotifySuspects);
+			}
+
+			HashSet<StreamId> OldFixStreams = new HashSet<StreamId>(OldIssue.Streams.Where(x => x.ContainsFix ?? false).Select(x => x.StreamId));
+			HashSet<StreamId> NewFixStreams = new HashSet<StreamId>(NewIssue.Streams.Where(x => x.ContainsFix ?? false).Select(x => x.StreamId));
+			foreach (StreamId StreamId in NewFixStreams.Where(x => !OldFixStreams.Contains(x)))
+			{
+				IssueLogger.LogInformation("Marking stream {StreamId} as fixed", StreamId);
+			}
+			foreach (StreamId StreamId in OldFixStreams.Where(x => !NewFixStreams.Contains(x)))
+			{
+				IssueLogger.LogInformation("Marking stream {StreamId} as not fixed", StreamId);
+			}
+
+			HashSet<(ObjectId, int)> OldSuspects = new HashSet<(ObjectId, int)>(OldIssue.Suspects.Select(x => (x.AuthorId, x.Change)));
+			HashSet<(ObjectId, int)> NewSuspects = new HashSet<(ObjectId, int)>(NewIssue.Suspects.Select(x => (x.AuthorId, x.Change)));
+			foreach ((ObjectId UserId, int Change) in NewSuspects.Where(x => !OldSuspects.Contains(x)))
+			{
+				IssueLogger.LogInformation("Added suspect {UserId} for change {Change}", UserId, Change);
+			}
+			foreach ((ObjectId UserId, int Change) in OldSuspects.Where(x => !NewSuspects.Contains(x)))
+			{
+				IssueLogger.LogInformation("Removed suspect {UserId} for change {Change}", UserId, Change);
+			}
+
+			HashSet<ObjectId> OldDeclinedBy = new HashSet<ObjectId>(NewIssue.Suspects.Where(x => x.DeclinedAt != null).Select(x => x.AuthorId));
+			HashSet<ObjectId> NewDeclinedBy = new HashSet<ObjectId>(NewIssue.Suspects.Where(x => x.DeclinedAt != null).Select(x => x.AuthorId));
+			foreach (ObjectId AddDeclinedBy in NewDeclinedBy.Where(x => !OldDeclinedBy.Contains(x)))
+			{
+				IssueLogger.LogInformation("Declined by {UserId}", AddDeclinedBy);
+			}
+			foreach (ObjectId RemoveDeclinedBy in OldDeclinedBy.Where(x => !NewDeclinedBy.Contains(x)))
+			{
+				IssueLogger.LogInformation("Un-declined by {UserId}", RemoveDeclinedBy);
+			}
+		}
+
+		static void LogSpanInfo(ILogger IssueLogger, IIssueSpan Span)
+		{
+			IssueLogger.LogInformation("Added span {SpanId} in {StreamId}", Span.Id, Span.StreamId);
+			IssueLogger.LogInformation("Span {SpanId} first failure: CL {Change} ({JobId})", Span.Id, Span.FirstFailure.Change, Span.FirstFailure.JobId);
+			if (Span.LastSuccess != null)
+			{
+				IssueLogger.LogInformation("Span {SpanId} last success: CL {Change} ({JobId})", Span.Id, Span.LastSuccess.Change, Span.LastSuccess.JobId);
+			}
+			if (Span.NextSuccess != null)
+			{
+				IssueLogger.LogInformation("Span {SpanId} next success: CL {Change} ({JobId})", Span.Id, Span.NextSuccess.Change, Span.NextSuccess.JobId);
+			}
+		}
+
 		/// <inheritdoc/>
 		public async Task<IIssue?> GetIssueAsync(int IssueId)
 		{
@@ -921,11 +1034,16 @@ namespace HordeServer.Collections.Impl
 			{
 				return IssueDocument;
 			}
-			else
+
+			Issue? NewIssue = await TryUpdateIssueAsync(Issue, Builders<Issue>.Update.Combine(Updates));
+			if(NewIssue == null)
 			{
-				return await TryUpdateIssueAsync(Issue, Builders<Issue>.Update.Combine(Updates));
+				return null;
 			}
 
+			ILogger IssueLogger = GetLogger(Issue.Id);
+			LogIssueChanges(IssueLogger, IssueDocument, NewIssue);
+			return NewIssue;
 		}
 
 		public async Task<IIssue?> UpdateIssueDerivedDataAsync(IIssue Issue)
@@ -1083,6 +1201,10 @@ namespace HordeServer.Collections.Impl
 				return null;
 			}
 
+			// Log any changes
+			ILogger IssueLogger = GetLogger(Issue.Id);
+			LogIssueChanges(IssueLogger, Issue, NewIssue);
+
 			// Clear the modified flag on any spans
 			for (int Idx = 0; Idx < Spans.Count; Idx++)
 			{
@@ -1178,9 +1300,28 @@ namespace HordeServer.Collections.Impl
 			}
 
 			IssueSpan? NewSpan = await TryUpdateSpanAsync(Span, Builders<IssueSpan>.Update.Combine(Updates));
-			if (NewSpan != null && NewFailure != null)
+			if (NewSpan != null)
 			{
-				await AddStepAsync(NewSpan.Id, NewFailure);
+				if (NewFailure != null)
+				{
+					await AddStepAsync(NewSpan.Id, NewFailure);
+				}
+				if (NewSpan.IssueId != null)
+				{
+					ILogger Logger = GetLogger(NewSpan.IssueId.Value);
+					if (NewLastSuccess != null)
+					{
+						Logger.LogInformation("Set last success for span {SpanId} to job {JobId} at CL {Change}", NewSpan.Id, NewLastSuccess.JobId, NewLastSuccess.Change);
+					}
+					if (NewNextSuccess != null)
+					{
+						Logger.LogInformation("Set next success for span {SpanId} to job {JobId} at CL {Change}", NewSpan.Id, NewNextSuccess.JobId, NewNextSuccess.Change);
+					}
+					if (NewFailure != null)
+					{
+						Logger.LogInformation("Added failure for span {SpanId} in job {JobId} at CL {Change}", NewSpan.Id, NewFailure.JobId, NewFailure.Change);
+					}
+				}
 			}
 			return NewSpan;
 		}
@@ -1237,6 +1378,7 @@ namespace HordeServer.Collections.Impl
 				return false;
 			}
 
+			LogSpanInfo(GetLogger(Issue.Id), Span);
 			await TokenValue.FlushAsync();
 			return true;
 		}
@@ -1281,5 +1423,11 @@ namespace HordeServer.Collections.Impl
 		}
 
 		#endregion
+
+		/// <inheritdoc/>
+		public IAuditLogChannel<int> GetLogger(int IssueId)
+		{
+			return AuditLog[IssueId];
+		}
 	}
 }
