Horde: Add a strongly typed version of user id.

diff --git a/Engine/Source/Programs/Horde/HordeServer.Compute.Tests/IssueServiceTests.cs b/Engine/Source/Programs/Horde/HordeServer.Compute.Tests/IssueServiceTests.cs
--- a/Engine/Source/Programs/Horde/HordeServer.Compute.Tests/IssueServiceTests.cs
+++ b/Engine/Source/Programs/Horde/HordeServer.Compute.Tests/IssueServiceTests.cs
@@ -43,6 +43,7 @@ namespace HordeServerTests
 {
 	using JobId = ObjectId<IJob>;
 	using LogId = ObjectId<ILogFile>;
+	using UserId = ObjectId<IUser>;
 
 	[TestClass]
 	public class IssueServiceTests : DatabaseIntegrationTest
@@ -109,9 +110,9 @@ namespace HordeServerTests
 		ILogFileService LogFileService => TestSetup.LogFileService;
 		IIssueService IssueService => TestSetup.IssueService;
 
-		ObjectId TimId;
-		ObjectId JerryId;
-		ObjectId BobId;
+		UserId TimId;
+		UserId JerryId;
+		UserId BobId;
 
 		DirectoryReference WorkspaceDir;
 
@@ -601,7 +602,7 @@ namespace HordeServerTests
 
 				List<IIssueSuspect> Suspects = await IssueService.GetIssueSuspectsAsync(Issues[0]);
 
-				List<ObjectId> PrimarySuspects = Suspects.Select(x => x.AuthorId).ToList();
+				List<UserId> PrimarySuspects = Suspects.Select(x => x.AuthorId).ToList();
 				Assert.AreEqual(2, PrimarySuspects.Count);
 				Assert.IsTrue(PrimarySuspects.Contains(JerryId)); // 115
 				Assert.IsTrue(PrimarySuspects.Contains(TimId)); // 120
@@ -705,7 +706,7 @@ namespace HordeServerTests
 
 				List<IIssueSuspect> Suspects = await IssueService.GetIssueSuspectsAsync(Issues[0]);
 
-				List<ObjectId> PrimarySuspects = Suspects.Select(x => x.AuthorId).ToList();
+				List<UserId> PrimarySuspects = Suspects.Select(x => x.AuthorId).ToList();
 				Assert.AreEqual(1, PrimarySuspects.Count);
 				Assert.AreEqual(JerryId, PrimarySuspects[0]); // 115
 			}
@@ -747,7 +748,7 @@ namespace HordeServerTests
 
 				List<IIssueSuspect> Suspects = await IssueService.GetIssueSuspectsAsync(Issues[0]);
 
-				List<ObjectId> PrimarySuspects = Suspects.Select(x => x.AuthorId).ToList();
+				List<UserId> PrimarySuspects = Suspects.Select(x => x.AuthorId).ToList();
 				Assert.AreEqual(2, PrimarySuspects.Count);
 				Assert.IsTrue(PrimarySuspects.Contains(JerryId)); // 115
 				Assert.IsTrue(PrimarySuspects.Contains(TimId)); // 120
@@ -766,7 +767,7 @@ namespace HordeServerTests
 
 				List<IIssueSuspect> Suspects = await IssueService.GetIssueSuspectsAsync(Issues[0]);
 
-				List<ObjectId> PrimarySuspects = Suspects.Where(x => x.DeclinedAt == null).Select(x => x.AuthorId).ToList();
+				List<UserId> PrimarySuspects = Suspects.Where(x => x.DeclinedAt == null).Select(x => x.AuthorId).ToList();
 				Assert.AreEqual(1, PrimarySuspects.Count);
 				Assert.AreEqual(JerryId, PrimarySuspects[0]); // 115
 			}
@@ -813,7 +814,7 @@ namespace HordeServerTests
 
 				List<IIssueSuspect> Suspects = await IssueService.GetIssueSuspectsAsync(Issues[0]);
 
-				List<ObjectId> PrimarySuspects = Suspects.Select(x => x.AuthorId).ToList();
+				List<UserId> PrimarySuspects = Suspects.Select(x => x.AuthorId).ToList();
 				Assert.AreEqual(1, PrimarySuspects.Count);
 				Assert.AreEqual(JerryId, PrimarySuspects[0]); // 115 = foo.cpp
 			}
diff --git a/Engine/Source/Programs/Horde/HordeServer/Acls/AclService.cs b/Engine/Source/Programs/Horde/HordeServer/Acls/AclService.cs
--- a/Engine/Source/Programs/Horde/HordeServer/Acls/AclService.cs
+++ b/Engine/Source/Programs/Horde/HordeServer/Acls/AclService.cs
@@ -20,6 +20,8 @@ using System.Threading.Tasks;
 
 namespace HordeServer.Services
 {
+	using UserId = ObjectId<IUser>;
+
 	/// <summary>
 	/// Cache of global ACL
 	/// </summary>
@@ -225,9 +227,9 @@ namespace HordeServer.Services
 		/// <param name="User"></param>
 		/// <param name="UserId"></param>
 		/// <returns></returns>
-		public Task<bool> AuthorizeAsUserAsync(ClaimsPrincipal User, ObjectId UserId)
+		public Task<bool> AuthorizeAsUserAsync(ClaimsPrincipal User, UserId UserId)
 		{
-			ObjectId? CurrentUserId = User.GetUserId();
+			UserId? CurrentUserId = User.GetUserId();
 			if (CurrentUserId != null && CurrentUserId.Value == UserId)
 			{
 				return Task.FromResult(true);
diff --git a/Engine/Source/Programs/Horde/HordeServer/Controllers/IssuesController.cs b/Engine/Source/Programs/Horde/HordeServer/Controllers/IssuesController.cs
--- a/Engine/Source/Programs/Horde/HordeServer/Controllers/IssuesController.cs
+++ b/Engine/Source/Programs/Horde/HordeServer/Controllers/IssuesController.cs
@@ -24,6 +24,7 @@ namespace HordeServer.Controllers
 	using JobId = ObjectId<IJob>;
 	using LogId = ObjectId<ILogFile>;
 	using StreamId = StringId<IStream>;
+	using UserId = ObjectId<IUser>;
 
 	/// <summary>
 	/// Controller for the /api/v1/issues endpoint
@@ -83,10 +84,10 @@ namespace HordeServer.Controllers
 				Ids = null;
 			}
 
-			ObjectId? UserIdValue = null;
+			UserId? UserIdValue = null;
 			if (UserId != null)
 			{
-				UserIdValue = new ObjectId(UserId);
+				UserIdValue = new UserId(UserId);
 			}
 
 			List<IIssue> Issues;
@@ -388,28 +389,28 @@ namespace HordeServer.Controllers
 		[Route("/api/v1/issues/{IssueId}")]
 		public async Task<ActionResult> UpdateIssueAsync(int IssueId, [FromBody] UpdateIssueRequest Request)
 		{
-			ObjectId? NewOwnerId = null;
+			UserId? NewOwnerId = null;
 			if (Request.OwnerId != null)
 			{
-				NewOwnerId = Request.OwnerId.Length == 0 ? ObjectId.Empty : new ObjectId(Request.OwnerId);
+				NewOwnerId = Request.OwnerId.Length == 0 ? UserId.Empty : new UserId(Request.OwnerId);
 			}
 
-			ObjectId? NewNominatedById = null;
+			UserId? NewNominatedById = null;
 			if (Request.NominatedById != null)
 			{
-				NewNominatedById = new ObjectId(Request.NominatedById);
+				NewNominatedById = new UserId(Request.NominatedById);
 			}
 
-			ObjectId? NewDeclinedById = null;
+			UserId? NewDeclinedById = null;
 			if (Request.Declined ?? false)
 			{
 				NewDeclinedById = User.GetUserId();
 			}
 
-			ObjectId? NewResolvedById = null;
+			UserId? NewResolvedById = null;
 			if (Request.Resolved.HasValue)
 			{
-				NewResolvedById = Request.Resolved.Value ? User.GetUserId() : ObjectId.Empty;
+				NewResolvedById = Request.Resolved.Value ? User.GetUserId() : UserId.Empty;
 			}
 
 			if (!await IssueService.UpdateIssueAsync(IssueId, Request.Summary, NewOwnerId, NewNominatedById, Request.Acknowledged, NewDeclinedById, Request.FixChange, NewResolvedById))
diff --git a/Engine/Source/Programs/Horde/HordeServer/Controllers/JobsController.cs b/Engine/Source/Programs/Horde/HordeServer/Controllers/JobsController.cs
--- a/Engine/Source/Programs/Horde/HordeServer/Controllers/JobsController.cs
+++ b/Engine/Source/Programs/Horde/HordeServer/Controllers/JobsController.cs
@@ -30,6 +30,7 @@ namespace HordeServer.Controllers
 	using JobId = ObjectId<IJob>;
 	using StreamId = StringId<IStream>;
 	using TemplateRefId = StringId<TemplateRef>;
+	using UserId = ObjectId<IUser>;
 
 	/// <summary>
 	/// Controller for the /api/v1/jobs endpoing
@@ -543,18 +544,18 @@ namespace HordeServer.Controllers
 				PreflightChange = 0;
 			}
 
-			ObjectId? PreflightStartedByUserIdValue = null;
+			UserId? PreflightStartedByUserIdValue = null;
 
 			if (PreflightStartedByUserId != null)
 			{
-				PreflightStartedByUserIdValue = new ObjectId(PreflightStartedByUserId);
+				PreflightStartedByUserIdValue = new UserId(PreflightStartedByUserId);
 			}
 
-			ObjectId? StartedByUserIdValue = null;
+			UserId? StartedByUserIdValue = null;
 
 			if (StartedByUserId != null)
 			{
-				StartedByUserIdValue = new ObjectId(StartedByUserId);
+				StartedByUserIdValue = new UserId(StartedByUserId);
 			}
 
 			List<IJob> Jobs;
diff --git a/Engine/Source/Programs/Horde/HordeServer/Controllers/SubscriptionsController.cs b/Engine/Source/Programs/Horde/HordeServer/Controllers/SubscriptionsController.cs
--- a/Engine/Source/Programs/Horde/HordeServer/Controllers/SubscriptionsController.cs
+++ b/Engine/Source/Programs/Horde/HordeServer/Controllers/SubscriptionsController.cs
@@ -16,6 +16,8 @@ using System.Threading.Tasks;
 
 namespace HordeServer.Controllers
 {
+	using UserId = ObjectId<IUser>;
+
 	/// <summary>
 	/// Controller for the /api/v1/agents endpoint
 	/// </summary>
@@ -56,7 +58,7 @@ namespace HordeServer.Controllers
 		[ProducesResponseType(typeof(List<GetSubscriptionResponse>), 200)]
 		public async Task<ActionResult<List<object>>> GetSubscriptionsAsync([FromQuery] string UserId, [FromQuery] PropertyFilter? Filter = null)
 		{
-			ObjectId UserIdValue;
+			UserId UserIdValue;
 			if (!TryParseUserId(UserId, out UserIdValue))
 			{
 				return BadRequest("Invalid user id");
@@ -128,9 +130,9 @@ namespace HordeServer.Controllers
 		[Route("/api/v1/subscriptions")]
 		public async Task<ActionResult<List<CreateSubscriptionResponse>>> CreateSubscriptionsAsync(List<CreateSubscriptionRequest> Subscriptions)
 		{
-			HashSet<ObjectId> AuthorizedUsers = new HashSet<ObjectId>();
+			HashSet<UserId> AuthorizedUsers = new HashSet<UserId>();
 
-			ObjectId? CurrentUserId = User.GetUserId();
+			UserId? CurrentUserId = User.GetUserId();
 			if(CurrentUserId != null)
 			{
 				AuthorizedUsers.Add(CurrentUserId.Value);
@@ -141,7 +143,7 @@ namespace HordeServer.Controllers
 			List<NewSubscription> NewSubscriptions = new List<NewSubscription>();
 			foreach (CreateSubscriptionRequest Subscription in Subscriptions)
 			{
-				ObjectId NewUserId;
+				UserId NewUserId;
 				if (!TryParseUserId(Subscription.UserId, out NewUserId))
 				{
 					return BadRequest($"Invalid user id: '{Subscription.UserId}'.");
@@ -163,10 +165,10 @@ namespace HordeServer.Controllers
 		/// <param name="UserName"></param>
 		/// <param name="ObjectId"></param>
 		/// <returns></returns>
-		bool TryParseUserId(string UserName, out ObjectId ObjectId)
+		bool TryParseUserId(string UserName, out UserId ObjectId)
 		{
-			ObjectId NewObjectId;
-			if (ObjectId.TryParse(UserName, out NewObjectId))
+			UserId NewObjectId;
+			if (UserId.TryParse(UserName, out NewObjectId))
 			{
 				ObjectId = NewObjectId;
 				return true;
@@ -175,7 +177,7 @@ namespace HordeServer.Controllers
 			string? CurrentUserName = User.GetUserName();
 			if (CurrentUserName != null && String.Equals(UserName, CurrentUserName, StringComparison.OrdinalIgnoreCase))
 			{
-				ObjectId? CurrentUserId = User.GetUserId();
+				UserId? CurrentUserId = User.GetUserId();
 				if (CurrentUserId != null)
 				{
 					ObjectId = CurrentUserId.Value;
diff --git a/Engine/Source/Programs/Horde/HordeServer/Controllers/UgsController.cs b/Engine/Source/Programs/Horde/HordeServer/Controllers/UgsController.cs
--- a/Engine/Source/Programs/Horde/HordeServer/Controllers/UgsController.cs
+++ b/Engine/Source/Programs/Horde/HordeServer/Controllers/UgsController.cs
@@ -25,6 +25,7 @@ namespace HordeServer.Controllers
 {
 	using LogId = ObjectId<ILogFile>;
 	using StreamId = StringId<IStream>;
+	using UserId = ObjectId<IUser>;
 
 	/// <summary>
 	/// Controller for the /api/v1/issues endpoint
@@ -382,32 +383,32 @@ namespace HordeServer.Controllers
 		[Route("/ugs/api/issues/{IssueId}")]
 		public async Task<ActionResult> UpdateIssueAsync(int IssueId, [FromBody] UpdateUgsIssueRequest Request)
 		{
-			ObjectId? NewOwnerId = null;
+			UserId? NewOwnerId = null;
 			if (!String.IsNullOrEmpty(Request.Owner))
 			{
 				NewOwnerId = (await UserCollection.FindOrAddUserByLoginAsync(Request.Owner))?.Id;
 			}
 
-			ObjectId? NewNominatedById = null;
+			UserId? NewNominatedById = null;
 			if (!String.IsNullOrEmpty(Request.NominatedBy))
 			{
 				NewNominatedById = (await UserCollection.FindOrAddUserByLoginAsync(Request.NominatedBy))?.Id;
 			}
 
-			ObjectId? NewDeclinedById = null;
+			UserId? NewDeclinedById = null;
 			if (!String.IsNullOrEmpty(Request.DeclinedBy))
 			{
 				NewDeclinedById = (await UserCollection.FindOrAddUserByLoginAsync(Request.DeclinedBy))?.Id;
 			}
 
-			ObjectId? NewResolvedById = null;
+			UserId? NewResolvedById = null;
 			if (!String.IsNullOrEmpty(Request.ResolvedBy))
 			{
 				NewResolvedById = (await UserCollection.FindOrAddUserByLoginAsync(Request.ResolvedBy))?.Id;
 			}
 			if (NewResolvedById == null && Request.Resolved.HasValue)
 			{
-				NewResolvedById = Request.Resolved.Value ? IIssue.ResolvedByUnknownId : ObjectId.Empty;
+				NewResolvedById = Request.Resolved.Value ? IIssue.ResolvedByUnknownId : UserId.Empty;
 			}
 
 			if (!await IssueService.UpdateIssueAsync(IssueId, null, NewOwnerId, NewNominatedById, Request.Acknowledged, NewDeclinedById, Request.FixChange, NewResolvedById))
diff --git a/Engine/Source/Programs/Horde/HordeServer/Controllers/UserController.cs b/Engine/Source/Programs/Horde/HordeServer/Controllers/UserController.cs
--- a/Engine/Source/Programs/Horde/HordeServer/Controllers/UserController.cs
+++ b/Engine/Source/Programs/Horde/HordeServer/Controllers/UserController.cs
@@ -16,6 +16,7 @@ using System.Threading.Tasks;
 namespace HordeServer.Controllers
 {
 	using JobId = ObjectId<IJob>;
+	using UserId = ObjectId<IUser>;
 
 	/// <summary>
 	/// Controller for the /api/v1/user endpoint
@@ -75,7 +76,7 @@ namespace HordeServer.Controllers
 		[Route("/api/v1/user")]
 		public async Task<ActionResult> UpdateUserAsync(UpdateUserRequest Request)
 		{
-			ObjectId? UserId = User.GetUserId();
+			UserId? UserId = User.GetUserId();
 			if(UserId == null)
 			{
 				return BadRequest("Current user does not have a registered profile");
diff --git a/Engine/Source/Programs/Horde/HordeServer/Controllers/UsersController.cs b/Engine/Source/Programs/Horde/HordeServer/Controllers/UsersController.cs
--- a/Engine/Source/Programs/Horde/HordeServer/Controllers/UsersController.cs
+++ b/Engine/Source/Programs/Horde/HordeServer/Controllers/UsersController.cs
@@ -17,6 +17,8 @@ using System.Threading.Tasks;
 
 namespace HordeServer.Controllers
 {
+	using UserId = ObjectId<IUser>;
+
 	/// <summary>
 	/// Controller for the /api/v1/users endpoint
 	/// </summary>
@@ -83,10 +85,10 @@ namespace HordeServer.Controllers
 			[FromQuery] bool IncludeAvatar = false)
 		{
 
-			ObjectId[]? UserIds = null;
+			UserId[]? UserIds = null;
 			if (Ids != null && Ids.Length > 0)
 			{
-				UserIds = Ids.Select(x => new ObjectId(x)).ToArray();
+				UserIds = Ids.Select(x => new UserId(x)).ToArray();
 			}
 
 			List<IUser> Users = await UserCollection.FindUsersAsync(UserIds, NameRegex, Index, Count);
@@ -105,7 +107,7 @@ namespace HordeServer.Controllers
 
 		async Task<IUser?> GetUserInternalAsync(string Id)
 		{
-			ObjectId? UserId = ParseUserId(Id);
+			UserId? UserId = ParseUserId(Id);
 			if(UserId == null)
 			{
 				return null;
@@ -113,13 +115,13 @@ namespace HordeServer.Controllers
 			return await UserCollection.GetUserAsync(UserId.Value);
 		}
 
-		ObjectId? ParseUserId(string Id)
+		UserId? ParseUserId(string Id)
 		{
 			if (Id.Equals("current", StringComparison.OrdinalIgnoreCase))
 			{
 				return User.GetUserId();
 			}
-			else if(ObjectId.TryParse(Id, out ObjectId Result))
+			else if(UserId.TryParse(Id, out UserId Result))
 			{
 				return Result;
 			}
diff --git a/Engine/Source/Programs/Horde/HordeServer/Devices/IDevice.cs b/Engine/Source/Programs/Horde/HordeServer/Devices/IDevice.cs
--- a/Engine/Source/Programs/Horde/HordeServer/Devices/IDevice.cs
+++ b/Engine/Source/Programs/Horde/HordeServer/Devices/IDevice.cs
@@ -12,6 +12,7 @@ namespace HordeServer.Models
 	using DeviceId = StringId<IDevice>;
 	using DevicePlatformId = StringId<IDevicePlatform>;
 	using DevicePoolId = StringId<IDevicePool>;
+	using UserId = ObjectId<IUser>;
 
 	/// <summary>
 	/// A reservation containing one or more devices
@@ -41,7 +42,7 @@ namespace HordeServer.Models
 		/// <summary>
 		/// Reservations held by a user, requires a token
 		/// </summary>
-		public ObjectId? UserId { get; }
+		public UserId? UserId { get; }
 
 		/// <summary>
 		/// The hostname of machine holding reservation
diff --git a/Engine/Source/Programs/Horde/HordeServer/Devices/IDeviceCollection.cs b/Engine/Source/Programs/Horde/HordeServer/Devices/IDeviceCollection.cs
--- a/Engine/Source/Programs/Horde/HordeServer/Devices/IDeviceCollection.cs
+++ b/Engine/Source/Programs/Horde/HordeServer/Devices/IDeviceCollection.cs
@@ -11,6 +11,7 @@ namespace HordeServer.Collections
 	using DeviceId = StringId<IDevice>;
 	using DevicePlatformId = StringId<IDevicePlatform>;
 	using DevicePoolId = StringId<IDevicePool>;
+	using UserId = ObjectId<IUser>;
 
 	/// <summary>
 	/// Device reservation request data
@@ -124,7 +125,7 @@ namespace HordeServer.Collections
 		/// <param name="Address">The network address or hostname the device can be reached at</param>
 		/// <param name="ModelId">The vendor model id of the device</param>
         /// <param name="UserId">The user adding the device</param>
-		Task<IDevice?> TryAddDeviceAsync(DeviceId Id, string Name, DevicePlatformId PlatformId, DevicePoolId PoolId, bool? Enabled, string? Address, string? ModelId, ObjectId? UserId);
+		Task<IDevice?> TryAddDeviceAsync(DeviceId Id, string Name, DevicePlatformId PlatformId, DevicePoolId PoolId, bool? Enabled, string? Address, string? ModelId, UserId? UserId);
 
 		/// <summary>
 		/// Update a device
@@ -139,7 +140,7 @@ namespace HordeServer.Collections
 		/// <param name="NewProblem">Whether to set or clear problem state</param>
 		/// <param name="NewMaintenance">Whether to set or clear maintenance state</param>
         /// <param name="ModifiedByUserId">The user who is updating the device</param>
-		Task UpdateDeviceAsync(DeviceId DeviceId, DevicePoolId? NewPoolId, string? NewName, string? NewAddress, string? NewModelId, string? NewNotes, bool? NewEnabled, bool? NewProblem, bool? NewMaintenance, ObjectId? ModifiedByUserId = null);
+		Task UpdateDeviceAsync(DeviceId DeviceId, DevicePoolId? NewPoolId, string? NewName, string? NewAddress, string? NewModelId, string? NewNotes, bool? NewEnabled, bool? NewProblem, bool? NewMaintenance, UserId? ModifiedByUserId = null);
 
 		/// <summary>
 		/// Delete a device from the collection
@@ -153,7 +154,7 @@ namespace HordeServer.Collections
         /// <param name="DeviceId"></param>
         /// <param name="CheckedOutByUserId"></param>
         /// <returns></returns>
-        Task CheckoutDeviceAsync(DeviceId DeviceId, ObjectId? CheckedOutByUserId);
+        Task CheckoutDeviceAsync(DeviceId DeviceId, UserId? CheckedOutByUserId);
 
         // RESERVATIONS
 
diff --git a/Engine/Source/Programs/Horde/HordeServer/Devices/Impl/DeviceCollection.cs b/Engine/Source/Programs/Horde/HordeServer/Devices/Impl/DeviceCollection.cs
--- a/Engine/Source/Programs/Horde/HordeServer/Devices/Impl/DeviceCollection.cs
+++ b/Engine/Source/Programs/Horde/HordeServer/Devices/Impl/DeviceCollection.cs
@@ -18,6 +18,7 @@ namespace HordeServer.Collections.Impl
 	using DeviceId = StringId<IDevice>;
 	using DevicePlatformId = StringId<IDevicePlatform>;
 	using DevicePoolId = StringId<IDevicePool>;
+	using UserId = ObjectId<IUser>;
 
 	/// <summary>
 	/// Collection of device documents
@@ -109,7 +110,7 @@ namespace HordeServer.Collections.Impl
             /// Reservations held by a user, requires a token like download code
             /// </summary>
             [BsonIgnoreIfNull]
-			public ObjectId? UserId { get; set; }
+			public UserId? UserId { get; set; }
 
             /// <summary>
             /// The hostname of the machine which has made the reservation
@@ -218,7 +219,7 @@ namespace HordeServer.Collections.Impl
 
 			}
 
-			public DeviceDocument(DeviceId Id, DevicePlatformId PlatformId, DevicePoolId PoolId, string Name, bool Enabled, string? Address, string? ModelId, ObjectId? UserId)
+			public DeviceDocument(DeviceId Id, DevicePlatformId PlatformId, DevicePoolId PoolId, string Name, bool Enabled, string? Address, string? ModelId, UserId? UserId)
 			{
 				this.Id = Id;
 				this.PlatformId = PlatformId;
@@ -259,7 +260,7 @@ namespace HordeServer.Collections.Impl
 		}
 
 		/// <inheritdoc/>
-		public async Task<IDevice?> TryAddDeviceAsync(DeviceId Id, string Name, DevicePlatformId PlatformId, DevicePoolId PoolId, bool? Enabled, string? Address, string? ModelId, ObjectId? UserId)
+		public async Task<IDevice?> TryAddDeviceAsync(DeviceId Id, string Name, DevicePlatformId PlatformId, DevicePoolId PoolId, bool? Enabled, string? Address, string? ModelId, UserId? UserId)
 		{
 			DeviceDocument NewDevice = new DeviceDocument(Id, PlatformId, PoolId, Name, Enabled ?? true, Address, ModelId, UserId);
 			await Devices.InsertOneAsync(NewDevice);
@@ -385,7 +386,7 @@ namespace HordeServer.Collections.Impl
 		}
 
 		/// <inheritdoc/>
-		public async Task CheckoutDeviceAsync(DeviceId DeviceId, ObjectId? CheckedOutByUserId)
+		public async Task CheckoutDeviceAsync(DeviceId DeviceId, UserId? CheckedOutByUserId)
 		{
 			UpdateDefinitionBuilder<DeviceDocument> UpdateBuilder = Builders<DeviceDocument>.Update;
 
@@ -405,7 +406,7 @@ namespace HordeServer.Collections.Impl
 
 
 		/// <inheritdoc/>
-		public async Task UpdateDeviceAsync(DeviceId DeviceId, DevicePoolId? NewPoolId, string? NewName, string? NewAddress, string? NewModelId, string? NewNotes, bool? NewEnabled, bool? NewProblem, bool? NewMaintenance, ObjectId? ModifiedByUserId = null)
+		public async Task UpdateDeviceAsync(DeviceId DeviceId, DevicePoolId? NewPoolId, string? NewName, string? NewAddress, string? NewModelId, string? NewNotes, bool? NewEnabled, bool? NewProblem, bool? NewMaintenance, UserId? ModifiedByUserId = null)
 		{
 			UpdateDefinitionBuilder<DeviceDocument> UpdateBuilder = Builders<DeviceDocument>.Update;
 
diff --git a/Engine/Source/Programs/Horde/HordeServer/Devices/Impl/DeviceService.cs b/Engine/Source/Programs/Horde/HordeServer/Devices/Impl/DeviceService.cs
--- a/Engine/Source/Programs/Horde/HordeServer/Devices/Impl/DeviceService.cs
+++ b/Engine/Source/Programs/Horde/HordeServer/Devices/Impl/DeviceService.cs
@@ -24,6 +24,7 @@ namespace HordeServer.Services
 	using DevicePlatformId = StringId<IDevicePlatform>;
 	using DevicePoolId = StringId<IDevicePool>;
 	using JobId = ObjectId<IJob>;
+	using UserId = ObjectId<IUser>;
 
 	/// <summary>
 	/// Platform map required by V1 API
@@ -220,7 +221,7 @@ namespace HordeServer.Services
 		/// <param name="ModelId">Vendor model id</param>
         /// <param name="UserId">User adding the device</param>
 		/// <returns></returns>
-		public Task<IDevice?> TryCreateDeviceAsync(DeviceId Id, string Name, DevicePlatformId PlatformId, DevicePoolId PoolId, bool? Enabled, string? Address, string? ModelId, ObjectId? UserId = null)
+		public Task<IDevice?> TryCreateDeviceAsync(DeviceId Id, string Name, DevicePlatformId PlatformId, DevicePoolId PoolId, bool? Enabled, string? Address, string? ModelId, UserId? UserId = null)
 		{
 			return Devices.TryAddDeviceAsync(Id, Name, PlatformId, PoolId, Enabled, Address, ModelId, UserId);
 		}
@@ -228,7 +229,7 @@ namespace HordeServer.Services
 		/// <summary>
 		/// Update a device
 		/// </summary>
-		public Task UpdateDeviceAsync(DeviceId DeviceId, DevicePoolId? NewPoolId = null, string? NewName = null, string? NewAddress = null, string? NewModelId = null, string? NewNotes = null, bool? NewEnabled = null, bool? NewProblem = null, bool? NewMaintenance = null, ObjectId? ModifiedByUserId = null)
+		public Task UpdateDeviceAsync(DeviceId DeviceId, DevicePoolId? NewPoolId = null, string? NewName = null, string? NewAddress = null, string? NewModelId = null, string? NewNotes = null, bool? NewEnabled = null, bool? NewProblem = null, bool? NewMaintenance = null, UserId? ModifiedByUserId = null)
 		{
 			return Devices.UpdateDeviceAsync(DeviceId, NewPoolId, NewName, NewAddress, NewModelId, NewNotes, NewEnabled, NewProblem, NewMaintenance, ModifiedByUserId);
 		}
@@ -236,7 +237,7 @@ namespace HordeServer.Services
 		/// <summary>
 		/// Checkout a device
 		/// </summary>
-		public Task CheckoutDeviceAsync(DeviceId DeviceId, ObjectId? UserId)
+		public Task CheckoutDeviceAsync(DeviceId DeviceId, UserId? UserId)
 		{
             return Devices.CheckoutDeviceAsync(DeviceId, UserId);
         }
diff --git a/Engine/Source/Programs/Horde/HordeServer/Issues/IIssue.cs b/Engine/Source/Programs/Horde/HordeServer/Issues/IIssue.cs
--- a/Engine/Source/Programs/Horde/HordeServer/Issues/IIssue.cs
+++ b/Engine/Source/Programs/Horde/HordeServer/Issues/IIssue.cs
@@ -19,6 +19,7 @@ namespace HordeServer.Models
 {
 	using StreamId = StringId<IStream>;
 	using TemplateRefId = StringId<TemplateRef>;
+	using UserId = ObjectId<IUser>;
 
 	/// <summary>
 	/// The severity of an issue
@@ -49,12 +50,12 @@ namespace HordeServer.Models
 		/// <summary>
 		/// Id for resolved by system
 		/// </summary>
-		static ObjectId ResolvedByTimeoutId { get; } = ObjectId.Parse("609592712b5c90b5bcf88c48");
+		static UserId ResolvedByTimeoutId { get; } = UserId.Parse("609592712b5c90b5bcf88c48");
 
 		/// <summary>
 		/// Id for resolved by unknown user
 		/// </summary>
-		static ObjectId ResolvedByUnknownId { get; } = ObjectId.Parse("609593b83e9b0b6dde620cf3");
+		static UserId ResolvedByUnknownId { get; } = UserId.Parse("609593b83e9b0b6dde620cf3");
 
 		/// <summary>
 		/// The unique object id
@@ -84,12 +85,12 @@ namespace HordeServer.Models
 		/// <summary>
 		/// User id of the owner
 		/// </summary>
-		public ObjectId? OwnerId { get; }
+		public UserId? OwnerId { get; }
 
 		/// <summary>
 		/// User id of the person that nominated the owner
 		/// </summary>
-		public ObjectId? NominatedById { get; }
+		public UserId? NominatedById { get; }
 
 		/// <summary>
 		/// Time at which the issue was created
@@ -114,7 +115,7 @@ namespace HordeServer.Models
 		/// <summary>
 		/// User that resolved the issue
 		/// </summary>
-		public ObjectId? ResolvedById { get; }
+		public UserId? ResolvedById { get; }
 
 		/// <summary>
 		/// Time at which the issue was verified fixed
@@ -181,7 +182,7 @@ namespace HordeServer.Models
 		/// <summary>
 		/// The user id of the change's author
 		/// </summary>
-		ObjectId AuthorId { get; }
+		UserId AuthorId { get; }
 
 		/// <summary>
 		/// The change suspected of causing this issue (in the origin stream)
diff --git a/Engine/Source/Programs/Horde/HordeServer/Issues/IIssueCollection.cs b/Engine/Source/Programs/Horde/HordeServer/Issues/IIssueCollection.cs
--- a/Engine/Source/Programs/Horde/HordeServer/Issues/IIssueCollection.cs
+++ b/Engine/Source/Programs/Horde/HordeServer/Issues/IIssueCollection.cs
@@ -27,6 +27,7 @@ namespace HordeServer.Collections
 	using LogId = ObjectId<ILogFile>;
 	using StreamId = StringId<IStream>;
 	using TemplateRefId = StringId<TemplateRef>;
+	using UserId = ObjectId<IUser>;
 
 	/// <summary>
 	/// Fingerprint for an issue
@@ -279,7 +280,7 @@ namespace HordeServer.Collections
 		/// <summary>
 		/// Author of the changelist
 		/// </summary>
-		public ObjectId AuthorId { get; set; }
+		public UserId AuthorId { get; set; }
 
 		/// <summary>
 		/// The submitted changelist
@@ -297,7 +298,7 @@ namespace HordeServer.Collections
 		/// <param name="AuthorId">Author of the change</param>
 		/// <param name="Change">The changelist number</param>
 		/// <param name="DeclinedAt">The time that the user declined this issue</param>
-		public NewIssueSuspectData(ObjectId AuthorId, int Change, DateTime? DeclinedAt)
+		public NewIssueSuspectData(UserId AuthorId, int Change, DateTime? DeclinedAt)
 		{
 			this.AuthorId = AuthorId;
 			this.Change = Change;
@@ -318,7 +319,7 @@ namespace HordeServer.Collections
 		/// <summary>
 		/// Author of the changelist
 		/// </summary>
-		public ObjectId AuthorId { get; set; }
+		public UserId AuthorId { get; set; }
 
 		/// <summary>
 		/// The original changelist number, if merged from another branch. For changes merged between several branches, this is the originally submitted change.
@@ -330,7 +331,7 @@ namespace HordeServer.Collections
 		/// </summary>
 		/// <param name="Change">The changelist number</param>
 		/// <param name="AuthorId">Author of the change</param>
-		public NewIssueSpanSuspectData(int Change, ObjectId AuthorId)
+		public NewIssueSpanSuspectData(int Change, UserId AuthorId)
 		{
 			this.Change = Change;
 			this.AuthorId = AuthorId;
@@ -405,7 +406,7 @@ namespace HordeServer.Collections
 		/// <param name="Index">Index within the results to return</param>
 		/// <param name="Count">Number of results</param>
 		/// <returns>List of streams open in the given stream at the given changelist</returns>
-		Task<List<IIssue>> FindIssuesAsync(IEnumerable<int>? Ids = null, ObjectId? UserId = null, StreamId? StreamId = null, int? MinChange = null, int? MaxChange = null, bool? Resolved = null, int? Index = null, int? Count = null);
+		Task<List<IIssue>> FindIssuesAsync(IEnumerable<int>? Ids = null, UserId? UserId = null, StreamId? StreamId = null, int? MinChange = null, int? MaxChange = null, bool? Resolved = null, int? Index = null, int? Count = null);
 
 		/// <summary>
 		/// Searches for open issues
@@ -431,7 +432,7 @@ namespace HordeServer.Collections
 		/// <param name="NewLastSeenAt"></param>
 		/// <param name="NewNotifySuspects">Whether all suspects should be notified about this issue</param>
 		/// <returns>True if the issue was updated</returns>
-		Task<IIssue?> UpdateIssueAsync(IIssue Issue, IssueSeverity? NewSeverity = null, string? NewSummary = null, string? NewUserSummary = null, ObjectId? NewOwnerId = null, ObjectId? NewNominatedById = null, bool? NewAcknowledged = null, ObjectId? NewDeclinedById = null, int? NewFixChange = null, Dictionary<StreamId, bool>? NewFixStreamIds = null, ObjectId? NewResolvedById = null, DateTime? NewLastSeenAt = null, bool? NewNotifySuspects = null);
+		Task<IIssue?> UpdateIssueAsync(IIssue Issue, IssueSeverity? NewSeverity = null, string? NewSummary = null, string? NewUserSummary = null, UserId? NewOwnerId = null, UserId? NewNominatedById = null, bool? NewAcknowledged = null, UserId? NewDeclinedById = null, int? NewFixChange = null, Dictionary<StreamId, bool>? NewFixStreamIds = null, UserId? NewResolvedById = null, DateTime? NewLastSeenAt = null, bool? NewNotifySuspects = null);
 
 		/// <summary>
 		/// Updates the list of suspects for an issue, and optionally sets the resolved state
diff --git a/Engine/Source/Programs/Horde/HordeServer/Issues/IIssueDetails.cs b/Engine/Source/Programs/Horde/HordeServer/Issues/IIssueDetails.cs
--- a/Engine/Source/Programs/Horde/HordeServer/Issues/IIssueDetails.cs
+++ b/Engine/Source/Programs/Horde/HordeServer/Issues/IIssueDetails.cs
@@ -1,6 +1,7 @@
 // Copyright Epic Games, Inc. All Rights Reserved.
 
 using HordeServer.Services;
+using HordeServer.Utilities;
 using MongoDB.Bson;
 using System;
 using System.Collections.Generic;
@@ -9,6 +10,8 @@ using System.Threading.Tasks;
 
 namespace HordeServer.Models
 {
+	using UserId = ObjectId<IUser>;
+
 	/// <summary>
 	/// Cached issue information
 	/// </summary>
@@ -65,7 +68,7 @@ namespace HordeServer.Models
 		/// </summary>
 		/// <param name="UserId">The user to query</param>
 		/// <returns>True if the issue is relevant to the given user</returns>
-		bool IncludeForUser(ObjectId UserId);
+		bool IncludeForUser(UserId UserId);
 	}
 
 	/// <summary>
diff --git a/Engine/Source/Programs/Horde/HordeServer/Issues/IIssueService.cs b/Engine/Source/Programs/Horde/HordeServer/Issues/IIssueService.cs
--- a/Engine/Source/Programs/Horde/HordeServer/Issues/IIssueService.cs
+++ b/Engine/Source/Programs/Horde/HordeServer/Issues/IIssueService.cs
@@ -13,6 +13,7 @@ namespace HordeServer.Services
 {
 	using LogId = ObjectId<ILogFile>;
 	using StreamId = StringId<IStream>;
+	using UserId = ObjectId<IUser>;
 
 	/// <summary>
 	/// Service which manages build health issues
@@ -97,7 +98,7 @@ namespace HordeServer.Services
 		/// <param name="Index">Index within the results to return</param>
 		/// <param name="Count">Number of results</param>
 		/// <returns>List of streams open in the given stream at the given changelist</returns>
-		Task<List<IIssue>> FindIssuesAsync(IEnumerable<int>? Ids = null, ObjectId? UserId = null, StreamId? StreamId = null, int? MinChange = null, int? MaxChange = null, bool? Resolved = null, int? Index = null, int? Count = null);
+		Task<List<IIssue>> FindIssuesAsync(IEnumerable<int>? Ids = null, UserId? UserId = null, StreamId? StreamId = null, int? MinChange = null, int? MaxChange = null, bool? Resolved = null, int? Index = null, int? Count = null);
 
 		/// <summary>
 		/// Searches for open issues affecting a job
@@ -113,7 +114,7 @@ namespace HordeServer.Services
 		/// <param name="Index">Index within the results to return</param>
 		/// <param name="Count">Number of results</param>
 		/// <returns></returns>
-		Task<List<IIssue>> FindIssuesForJobAsync(int[]? Ids, IJob Job, IGraph Graph, SubResourceId? StepId = null, SubResourceId? BatchId = null, int? LabelIdx = null, ObjectId? UserId = null, bool? Resolved = null, int? Index = null, int? Count = null);
+		Task<List<IIssue>> FindIssuesForJobAsync(int[]? Ids, IJob Job, IGraph Graph, SubResourceId? StepId = null, SubResourceId? BatchId = null, int? LabelIdx = null, UserId? UserId = null, bool? Resolved = null, int? Index = null, int? Count = null);
 
 		/// <summary>
 		/// Find the events for a particular issue
@@ -137,7 +138,7 @@ namespace HordeServer.Services
 		/// <param name="FixChange">Fix changelist for the issue</param>
 		/// <param name="ResolvedById">Whether the issue has been resolved</param>
 		/// <returns>True if the issue was updated</returns>
-		Task<bool> UpdateIssueAsync(int Id, string? Summary = null, ObjectId? OwnerId = null, ObjectId? NominatedById = null, bool? Acknowledged = null, ObjectId? DeclinedById = null, int? FixChange = null, ObjectId? ResolvedById = null);
+		Task<bool> UpdateIssueAsync(int Id, string? Summary = null, UserId? OwnerId = null, UserId? NominatedById = null, bool? Acknowledged = null, UserId? DeclinedById = null, int? FixChange = null, UserId? ResolvedById = null);
 
 		/// <summary>
 		/// Marks a step as complete
diff --git a/Engine/Source/Programs/Horde/HordeServer/Issues/IIssueSpan.cs b/Engine/Source/Programs/Horde/HordeServer/Issues/IIssueSpan.cs
--- a/Engine/Source/Programs/Horde/HordeServer/Issues/IIssueSpan.cs
+++ b/Engine/Source/Programs/Horde/HordeServer/Issues/IIssueSpan.cs
@@ -12,6 +12,7 @@ namespace HordeServer.Models
 {
 	using StreamId = StringId<IStream>;
 	using TemplateRefId = StringId<TemplateRef>;
+	using UserId = ObjectId<IUser>;
 
 	/// <summary>
 	/// Trace of a set of errors in a single step across multiple jobs
@@ -102,7 +103,7 @@ namespace HordeServer.Models
 		/// <summary>
 		/// Id of the changelist's author
 		/// </summary>
-		public ObjectId AuthorId { get; }
+		public UserId AuthorId { get; }
 
 		/// <summary>
 		/// The original changelist number, if merged from another branch. For changes merged between several branches, this is the originally submitted change.
diff --git a/Engine/Source/Programs/Horde/HordeServer/Issues/Impl/IssueCollection.cs b/Engine/Source/Programs/Horde/HordeServer/Issues/Impl/IssueCollection.cs
--- a/Engine/Source/Programs/Horde/HordeServer/Issues/Impl/IssueCollection.cs
+++ b/Engine/Source/Programs/Horde/HordeServer/Issues/Impl/IssueCollection.cs
@@ -24,6 +24,7 @@ namespace HordeServer.Collections.Impl
 	using LogId = ObjectId<ILogFile>;
 	using StreamId = StringId<IStream>;
 	using TemplateRefId = StringId<TemplateRef>;
+	using UserId = ObjectId<IUser>;
 
 	class IssueCollection : IIssueCollection
 	{
@@ -128,10 +129,10 @@ namespace HordeServer.Collections.Impl
 			public IssueSeverity Severity { get; set; }
 
 			[BsonIgnoreIfNull]
-			public ObjectId? OwnerId { get; set; }
+			public UserId? OwnerId { get; set; }
 
 			[BsonIgnoreIfNull]
-			public ObjectId? NominatedById { get; set; }
+			public UserId? NominatedById { get; set; }
 
 			public DateTime CreatedAt { get; set; }
 
@@ -145,7 +146,7 @@ namespace HordeServer.Collections.Impl
 			public DateTime? ResolvedAt { get; set; }
 
 			[BsonIgnoreIfNull]
-			public ObjectId? ResolvedById { get; set; }
+			public UserId? ResolvedById { get; set; }
 
 			[BsonIgnoreIfNull]
 			public DateTime? VerifiedAt { get; set; }
@@ -168,7 +169,7 @@ namespace HordeServer.Collections.Impl
 			public int UpdateIndex { get; set; }
 
 			IIssueFingerprint IIssue.Fingerprint => Fingerprint;
-			ObjectId? IIssue.OwnerId => OwnerId ?? GetDefaultOwner()?.AuthorId;
+			UserId? IIssue.OwnerId => OwnerId ?? GetDefaultOwner()?.AuthorId;
 			IReadOnlyList<IIssueStream> IIssue.Streams => Streams;
 			DateTime IIssue.LastSeenAt => (LastSeenAt == default) ? DateTime.UtcNow : LastSeenAt;
 
@@ -204,7 +205,7 @@ namespace HordeServer.Collections.Impl
 			{
 				if (Suspects.Count > 0)
 				{
-					ObjectId PossibleOwner = Suspects[0].AuthorId;
+					UserId PossibleOwner = Suspects[0].AuthorId;
 					if (Suspects.All(x => x.AuthorId == PossibleOwner) && Suspects.Any(x => x.DeclinedAt == null))
 					{
 						return Suspects[0];
@@ -234,7 +235,7 @@ namespace HordeServer.Collections.Impl
 		{
 			public ObjectId Id { get; set; }
 			public int IssueId { get; set; }
-			public ObjectId AuthorId { get; set; }
+			public UserId AuthorId { get; set; }
 			public int Change { get; set; }
 			public DateTime? DeclinedAt { get; set; }
 			public DateTime? ResolvedAt { get; set; } // Degenerate
@@ -248,7 +249,7 @@ namespace HordeServer.Collections.Impl
 			{
 			}
 
-			public IssueSuspect(int IssueId, ObjectId AuthorId, int Change, DateTime? DeclinedAt, DateTime? ResolvedAt)
+			public IssueSuspect(int IssueId, UserId AuthorId, int Change, DateTime? DeclinedAt, DateTime? ResolvedAt)
 			{
 				this.Id = ObjectId.GenerateNewId();
 				this.IssueId = IssueId;
@@ -367,7 +368,7 @@ namespace HordeServer.Collections.Impl
 		class IssueSpanSuspect : IIssueSpanSuspect
 		{
 			public int Change { get; set; }
-			public ObjectId AuthorId { get; set; }
+			public UserId AuthorId { get; set; }
 			public int? OriginatingChange { get; set; }
 
 			[BsonConstructor]
@@ -375,7 +376,7 @@ namespace HordeServer.Collections.Impl
 			{
 			}
 
-			public IssueSpanSuspect(int Change, ObjectId AuthorId, int? OriginatingChange)
+			public IssueSpanSuspect(int Change, UserId AuthorId, int? OriginatingChange)
 			{
 				this.Change = Change;
 				this.AuthorId = AuthorId;
@@ -653,24 +654,24 @@ namespace HordeServer.Collections.Impl
 				IssueLogger.LogInformation("Marking stream {StreamId} as not fixed", StreamId);
 			}
 
-			HashSet<(ObjectId, int)> OldSuspects = new HashSet<(ObjectId, int)>(OldIssue.Suspects.Select(x => (x.AuthorId, x.Change)));
-			HashSet<(ObjectId, int)> NewSuspects = new HashSet<(ObjectId, int)>(NewIssue.Suspects.Select(x => (x.AuthorId, x.Change)));
-			foreach ((ObjectId UserId, int Change) in NewSuspects.Where(x => !OldSuspects.Contains(x)))
+			HashSet<(UserId, int)> OldSuspects = new HashSet<(UserId, int)>(OldIssue.Suspects.Select(x => (x.AuthorId, x.Change)));
+			HashSet<(UserId, int)> NewSuspects = new HashSet<(UserId, int)>(NewIssue.Suspects.Select(x => (x.AuthorId, x.Change)));
+			foreach ((UserId UserId, int Change) in NewSuspects.Where(x => !OldSuspects.Contains(x)))
 			{
 				IssueLogger.LogInformation("Added suspect {UserId} for change {Change}", UserId, Change);
 			}
-			foreach ((ObjectId UserId, int Change) in OldSuspects.Where(x => !NewSuspects.Contains(x)))
+			foreach ((UserId UserId, int Change) in OldSuspects.Where(x => !NewSuspects.Contains(x)))
 			{
 				IssueLogger.LogInformation("Removed suspect {UserId} for change {Change}", UserId, Change);
 			}
 
-			HashSet<ObjectId> OldDeclinedBy = new HashSet<ObjectId>(NewIssue.Suspects.Where(x => x.DeclinedAt != null).Select(x => x.AuthorId));
-			HashSet<ObjectId> NewDeclinedBy = new HashSet<ObjectId>(NewIssue.Suspects.Where(x => x.DeclinedAt != null).Select(x => x.AuthorId));
-			foreach (ObjectId AddDeclinedBy in NewDeclinedBy.Where(x => !OldDeclinedBy.Contains(x)))
+			HashSet<UserId> OldDeclinedBy = new HashSet<UserId>(NewIssue.Suspects.Where(x => x.DeclinedAt != null).Select(x => x.AuthorId));
+			HashSet<UserId> NewDeclinedBy = new HashSet<UserId>(NewIssue.Suspects.Where(x => x.DeclinedAt != null).Select(x => x.AuthorId));
+			foreach (UserId AddDeclinedBy in NewDeclinedBy.Where(x => !OldDeclinedBy.Contains(x)))
 			{
 				IssueLogger.LogInformation("Declined by {UserId}", AddDeclinedBy);
 			}
-			foreach (ObjectId RemoveDeclinedBy in OldDeclinedBy.Where(x => !NewDeclinedBy.Contains(x)))
+			foreach (UserId RemoveDeclinedBy in OldDeclinedBy.Where(x => !NewDeclinedBy.Contains(x)))
 			{
 				IssueLogger.LogInformation("Un-declined by {UserId}", RemoveDeclinedBy);
 			}
@@ -710,13 +711,13 @@ namespace HordeServer.Collections.Impl
 		}
 
 		/// <inheritdoc/>
-		public async Task<List<IIssue>> FindIssuesAsync(IEnumerable<int>? Ids = null, ObjectId? UserId = null, StreamId? StreamId = null, int? MinChange = null, int? MaxChange = null, bool? Resolved = null, int? Index = null, int? Count = null)
+		public async Task<List<IIssue>> FindIssuesAsync(IEnumerable<int>? Ids = null, UserId? UserId = null, StreamId? StreamId = null, int? MinChange = null, int? MaxChange = null, bool? Resolved = null, int? Index = null, int? Count = null)
 		{
 			List<Issue> Results = await FilterIssuesByUserIdAsync(Ids, UserId, StreamId, MinChange, MaxChange, Resolved ?? false, Index ?? 0, Count);
 			return Results.ConvertAll<IIssue>(x => x);
 		}
 
-		async Task<List<Issue>> FilterIssuesByUserIdAsync(IEnumerable<int>? Ids, ObjectId? UserId, StreamId? StreamId, int? MinChange, int? MaxChange, bool? Resolved, int Index, int? Count)
+		async Task<List<Issue>> FilterIssuesByUserIdAsync(IEnumerable<int>? Ids, UserId? UserId, StreamId? StreamId, int? MinChange, int? MaxChange, bool? Resolved, int Index, int? Count)
 		{
 			if (UserId == null)
 			{
@@ -892,7 +893,7 @@ namespace HordeServer.Collections.Impl
 		}
 
 		/// <inheritdoc/>
-		public async Task<IIssue?> UpdateIssueAsync(IIssue Issue, IssueSeverity? NewSeverity = null, string? NewSummary = null, string? NewUserSummary = null, ObjectId? NewOwnerId = null, ObjectId? NewNominatedById = null, bool? NewAcknowledged = null, ObjectId? NewDeclinedById = null, int? NewFixChange = null, Dictionary<StreamId, bool>? NewFixStreamIds = null, ObjectId? NewResolvedById = null, DateTime? NewLastSeenAt = null, bool? NewNotifySuspects = null)
+		public async Task<IIssue?> UpdateIssueAsync(IIssue Issue, IssueSeverity? NewSeverity = null, string? NewSummary = null, string? NewUserSummary = null, UserId? NewOwnerId = null, UserId? NewNominatedById = null, bool? NewAcknowledged = null, UserId? NewDeclinedById = null, int? NewFixChange = null, Dictionary<StreamId, bool>? NewFixStreamIds = null, UserId? NewResolvedById = null, DateTime? NewLastSeenAt = null, bool? NewNotifySuspects = null)
 		{
 			Issue IssueDocument = (Issue)Issue;
 
@@ -920,7 +921,7 @@ namespace HordeServer.Collections.Impl
 			}
 			if (NewOwnerId != null)
 			{
-				if (NewOwnerId.Value == ObjectId.Empty)
+				if (NewOwnerId.Value == UserId.Empty)
 				{
 					Updates.Add(Builders<Issue>.Update.Unset(x => x.OwnerId!));
 					Updates.Add(Builders<Issue>.Update.Unset(x => x.NominatedAt!));
@@ -1001,7 +1002,7 @@ namespace HordeServer.Collections.Impl
 			}
 			if (NewResolvedById != null)
 			{
-				if (NewResolvedById.Value != ObjectId.Empty)
+				if (NewResolvedById.Value != UserId.Empty)
 				{
 					if (IssueDocument.ResolvedAt == null || IssueDocument.ResolvedById != NewResolvedById)
 					{
@@ -1097,10 +1098,10 @@ namespace HordeServer.Collections.Impl
 			// Find the current list of suspects
 			List<IssueSuspect> CurSuspects = await IssueSuspects.Find(x => x.IssueId == Issue.Id).ToListAsync();
 
-			HashSet<(ObjectId, int)> CurSuspectKeys = new HashSet<(ObjectId, int)>(CurSuspects.Select(x => (x.AuthorId, x.Change)));
+			HashSet<(UserId, int)> CurSuspectKeys = new HashSet<(UserId, int)>(CurSuspects.Select(x => (x.AuthorId, x.Change)));
 			List<IssueSuspect> CreateSuspects = NewSuspects.Where(x => !CurSuspectKeys.Contains((x.AuthorId, x.Change))).ToList();
 
-			HashSet<(ObjectId, int)> NewSuspectKeys = new HashSet<(ObjectId, int)>(NewSuspects.Select(x => (x.AuthorId, x.Change)));
+			HashSet<(UserId, int)> NewSuspectKeys = new HashSet<(UserId, int)>(NewSuspects.Select(x => (x.AuthorId, x.Change)));
 			List<IssueSuspect> DeleteSuspects = CurSuspects.Where(x => !NewSuspectKeys.Contains((x.AuthorId, x.Change))).ToList();
 
 			// Apply the suspect changes
diff --git a/Engine/Source/Programs/Horde/HordeServer/Issues/Impl/IssueService.cs b/Engine/Source/Programs/Horde/HordeServer/Issues/Impl/IssueService.cs
--- a/Engine/Source/Programs/Horde/HordeServer/Issues/Impl/IssueService.cs
+++ b/Engine/Source/Programs/Horde/HordeServer/Issues/Impl/IssueService.cs
@@ -35,6 +35,7 @@ namespace HordeServer.Services.Impl
 	using LogId = ObjectId<ILogFile>;
 	using StreamId = StringId<IStream>;
 	using TemplateRefId = StringId<TemplateRef>;
+	using UserId = ObjectId<IUser>;
 
 	/// <summary>
 	/// Detailed issue information
@@ -69,7 +70,7 @@ namespace HordeServer.Services.Impl
 		bool ShowDesktopAlerts;
 
 		/// <inheritdoc/>
-		HashSet<ObjectId> NotifyUsers;
+		HashSet<UserId> NotifyUsers;
 
 		/// <summary>
 		/// Constructor
@@ -97,11 +98,11 @@ namespace HordeServer.Services.Impl
 
 			if (Issue.OwnerId == null)
 			{
-				NotifyUsers = new HashSet<ObjectId>(Suspects.Select(x => x.AuthorId));
+				NotifyUsers = new HashSet<UserId>(Suspects.Select(x => x.AuthorId));
 			}
 			else
 			{
-				NotifyUsers = new HashSet<ObjectId> { Issue.OwnerId.Value };
+				NotifyUsers = new HashSet<UserId> { Issue.OwnerId.Value };
 			}
 		}
 
@@ -119,7 +120,7 @@ namespace HordeServer.Services.Impl
 		/// </summary>
 		/// <param name="UserId">The user to query</param>
 		/// <returns>True if the issue is relevant to the given user</returns>
-		public bool IncludeForUser(ObjectId UserId)
+		public bool IncludeForUser(UserId UserId)
 		{
 			return NotifyUsers.Contains(UserId);
 		}
@@ -325,7 +326,7 @@ namespace HordeServer.Services.Impl
 			List<IIssueSuspect> Suspects = await IssueCollection.GetSuspectsAsync(Issue);
 
 			List<IUser> SuspectUsers = new List<IUser>();
-			foreach (ObjectId SuspectUserId in Suspects.Select(x => x.AuthorId).Distinct())
+			foreach (UserId SuspectUserId in Suspects.Select(x => x.AuthorId).Distinct())
 			{
 				IUser? SuspectUser = await UserCollection.GetCachedUserAsync(SuspectUserId);
 				if (SuspectUser != null)
@@ -356,13 +357,13 @@ namespace HordeServer.Services.Impl
 		}
 
 		/// <inheritdoc/>
-		public async Task<List<IIssue>> FindIssuesAsync(IEnumerable<int>? Ids = null, ObjectId? UserId = null, StreamId? StreamId = null, int? MinChange = null, int? MaxChange = null, bool? Resolved = null, int? Index = null, int? Count = null)
+		public async Task<List<IIssue>> FindIssuesAsync(IEnumerable<int>? Ids = null, UserId? UserId = null, StreamId? StreamId = null, int? MinChange = null, int? MaxChange = null, bool? Resolved = null, int? Index = null, int? Count = null)
 		{
 			return await IssueCollection.FindIssuesAsync(Ids, UserId, StreamId, MinChange, MaxChange, Resolved, Index, Count);
 		}
 
 		/// <inheritdoc/>
-		public async Task<List<IIssue>> FindIssuesForJobAsync(int[]? Ids, IJob Job, IGraph Graph, SubResourceId? StepId = null, SubResourceId? BatchId = null, int? LabelIdx = null, ObjectId? UserId = null, bool? Resolved = null, int? Index = null, int? Count = null)
+		public async Task<List<IIssue>> FindIssuesForJobAsync(int[]? Ids, IJob Job, IGraph Graph, SubResourceId? StepId = null, SubResourceId? BatchId = null, int? LabelIdx = null, UserId? UserId = null, bool? Resolved = null, int? Index = null, int? Count = null)
 		{
 			List<IIssueStep> Steps = await IssueCollection.FindStepsAsync(Job.Id, BatchId, StepId);
 			List<IIssueSpan> Spans = await IssueCollection.FindSpansAsync(Steps.Select(x => x.SpanId));
@@ -423,7 +424,7 @@ namespace HordeServer.Services.Impl
 		}
 
 		/// <inheritdoc/>
-		public async Task<bool> UpdateIssueAsync(int Id, string? UserSummary = null, ObjectId? OwnerId = null, ObjectId? NominatedById = null, bool? Acknowledged = null, ObjectId? DeclinedById = null, int? FixChange = null, ObjectId? ResolvedById = null)
+		public async Task<bool> UpdateIssueAsync(int Id, string? UserSummary = null, UserId? OwnerId = null, UserId? NominatedById = null, bool? Acknowledged = null, UserId? DeclinedById = null, int? FixChange = null, UserId? ResolvedById = null)
 		{
 			Dictionary<StreamId, bool>? NewFixStreamIds = null;
 			for (; ; )
@@ -980,7 +981,7 @@ namespace HordeServer.Services.Impl
 			// Update the resolved flag
 			if (Issue != null && Issue.ResolvedAt != null && ShouldMarkFixFailed(Issue, Span))
 			{
-				Issue = await IssueCollection.UpdateIssueAsync(Issue, NewResolvedById: ObjectId.Empty);
+				Issue = await IssueCollection.UpdateIssueAsync(Issue, NewResolvedById: UserId.Empty);
 			}
 
 			return Issue;
diff --git a/Engine/Source/Programs/Horde/HordeServer/Jobs/IJob.cs b/Engine/Source/Programs/Horde/HordeServer/Jobs/IJob.cs
--- a/Engine/Source/Programs/Horde/HordeServer/Jobs/IJob.cs
+++ b/Engine/Source/Programs/Horde/HordeServer/Jobs/IJob.cs
@@ -26,6 +26,7 @@ namespace HordeServer.Models
 	using StreamId = StringId<IStream>;
 	using TemplateRefId = StringId<TemplateRef>;
 	using PoolId = StringId<IPool>;
+	using UserId = ObjectId<IUser>;
 
 	/// <summary>
 	/// Report for a job or jobstep
@@ -524,7 +525,7 @@ namespace HordeServer.Models
 		/// <summary>
 		/// Id of the user that started this job
 		/// </summary>
-		public ObjectId? StartedByUserId { get; }
+		public UserId? StartedByUserId { get; }
 
 		/// <summary>
 		/// Name of the user that started this job
diff --git a/Engine/Source/Programs/Horde/HordeServer/Jobs/IJobCollection.cs b/Engine/Source/Programs/Horde/HordeServer/Jobs/IJobCollection.cs
--- a/Engine/Source/Programs/Horde/HordeServer/Jobs/IJobCollection.cs
+++ b/Engine/Source/Programs/Horde/HordeServer/Jobs/IJobCollection.cs
@@ -25,6 +25,7 @@ namespace HordeServer.Collections
 	using PoolId = StringId<IPool>;
 	using StreamId = StringId<IStream>;
 	using TemplateRefId = StringId<TemplateRef>;
+	using UserId = ObjectId<IUser>;
 
 	/// <summary>
 	/// Interface for a collection of job documents
@@ -58,7 +59,7 @@ namespace HordeServer.Collections
 		/// <param name="HelixSwarmCallbackUrl">Callback URL to a Helix Server review, if any</param>
 		/// <param name="Arguments">Arguments for the job</param>
 		/// <returns>The new job document</returns>
-		Task<IJob> AddAsync(JobId JobId, StreamId StreamId, TemplateRefId TemplateRefId, ContentHash TemplateHash, IGraph Graph, string Name, int Change, int CodeChange, int? PreflightChange, int? ClonedPreflightChange, ObjectId? StartedByUserId, string? StartedByUserName, Priority? Priority, bool? AutoSubmit, bool? UpdateIssues, List<ChainedJobTemplate>? JobTriggers, bool ShowUgsBadges, bool ShowUgsAlerts, string? NotificationChannel, string? NotificationChannelFilter, string? HelixSwarmCallbackUrl, List<string>? Arguments);
+		Task<IJob> AddAsync(JobId JobId, StreamId StreamId, TemplateRefId TemplateRefId, ContentHash TemplateHash, IGraph Graph, string Name, int Change, int CodeChange, int? PreflightChange, int? ClonedPreflightChange, UserId? StartedByUserId, string? StartedByUserName, Priority? Priority, bool? AutoSubmit, bool? UpdateIssues, List<ChainedJobTemplate>? JobTriggers, bool ShowUgsBadges, bool ShowUgsAlerts, string? NotificationChannel, string? NotificationChannelFilter, string? HelixSwarmCallbackUrl, List<string>? Arguments);
 
 		/// <summary>
 		/// Gets a job with the given unique id
@@ -106,7 +107,7 @@ namespace HordeServer.Collections
 		/// <param name="Index">Index of the first result to return</param>
 		/// <param name="Count">Number of results to return</param>
 		/// <returns>List of jobs matching the given criteria</returns>
-		Task<List<IJob>> FindAsync(JobId[]? JobIds = null, StreamId? StreamId = null, string? Name = null, TemplateRefId[]? Templates = null, int? MinChange = null, int? MaxChange = null, int? PreflightChange = null, ObjectId? PreflightStartedByUser = null, ObjectId? StartedByUser = null, DateTimeOffset ? MinCreateTime = null, DateTimeOffset? MaxCreateTime = null, DateTimeOffset? ModifiedBefore = null, DateTimeOffset? ModifiedAfter = null, int? Index = null, int? Count = null);
+		Task<List<IJob>> FindAsync(JobId[]? JobIds = null, StreamId? StreamId = null, string? Name = null, TemplateRefId[]? Templates = null, int? MinChange = null, int? MaxChange = null, int? PreflightChange = null, UserId? PreflightStartedByUser = null, UserId? StartedByUser = null, DateTimeOffset? MinCreateTime = null, DateTimeOffset? MaxCreateTime = null, DateTimeOffset? ModifiedBefore = null, DateTimeOffset? ModifiedAfter = null, int? Index = null, int? Count = null);
 
 		/// <summary>
 		/// Updates a new job
diff --git a/Engine/Source/Programs/Horde/HordeServer/Jobs/Impl/JobCollection.cs b/Engine/Source/Programs/Horde/HordeServer/Jobs/Impl/JobCollection.cs
--- a/Engine/Source/Programs/Horde/HordeServer/Jobs/Impl/JobCollection.cs
+++ b/Engine/Source/Programs/Horde/HordeServer/Jobs/Impl/JobCollection.cs
@@ -27,6 +27,7 @@ namespace HordeServer.Collections.Impl
 	using StreamId = StringId<IStream>;
 	using TemplateRefId = StringId<TemplateRef>;
 	using PoolId = StringId<IPool>;
+	using UserId = ObjectId<IUser>;
 
 	/// <summary>
 	/// Wrapper around the jobs collection in a mongo DB
@@ -195,7 +196,7 @@ namespace HordeServer.Collections.Impl
 			public ContentHash GraphHash { get; set; }
 
 			[BsonIgnoreIfNull]
-			public ObjectId? StartedByUserId { get; set; }
+			public UserId? StartedByUserId { get; set; }
 
 			[BsonIgnoreIfNull]
 			public string? StartedByUser { get; set; }
@@ -274,7 +275,7 @@ namespace HordeServer.Collections.Impl
 				GraphHash = null!;
 			}
 
-			public JobDocument(JobId Id, StreamId StreamId, TemplateRefId TemplateId, ContentHash TemplateHash, ContentHash GraphHash, string Name, int Change, int CodeChange, int PreflightChange, int ClonedPreflightChange, ObjectId? StartedByUserId, string? StartedByUserName, Priority? Priority, bool? AutoSubmit, bool? UpdateIssues, DateTime CreateTimeUtc, List<ChainedJobDocument> ChainedJobs, bool ShowUgsBadges, bool ShowUgsAlerts, string? NotificationChannel, string? NotificationChannelFilter, string? HelixSwarmCallbackUrl, List<string>? Arguments)
+			public JobDocument(JobId Id, StreamId StreamId, TemplateRefId TemplateId, ContentHash TemplateHash, ContentHash GraphHash, string Name, int Change, int CodeChange, int PreflightChange, int ClonedPreflightChange, UserId? StartedByUserId, string? StartedByUserName, Priority? Priority, bool? AutoSubmit, bool? UpdateIssues, DateTime CreateTimeUtc, List<ChainedJobDocument> ChainedJobs, bool ShowUgsBadges, bool ShowUgsAlerts, string? NotificationChannel, string? NotificationChannelFilter, string? HelixSwarmCallbackUrl, List<string>? Arguments)
 			{
 				this.Id = Id;
 				this.StreamId = StreamId;
@@ -366,7 +367,7 @@ namespace HordeServer.Collections.Impl
 
 		/// <inheritdoc/>
 		[SuppressMessage("Compiler", "CA1054:URI parameters should not be strings")]
-		public async Task<IJob> AddAsync(JobId JobId, StreamId StreamId, TemplateRefId TemplateRefId, ContentHash TemplateHash, IGraph Graph, string Name, int Change, int CodeChange, int? PreflightChange, int? ClonedPreflightChange, ObjectId? StartedByUserId, string? StartedByUserName, Priority? Priority, bool? AutoSubmit, bool? UpdateIssues, List<ChainedJobTemplate>? ChainedJobs, bool ShowUgsBadges, bool ShowUgsAlerts, string? NotificationChannel, string? NotificationChannelFilter, string? HelixSwarmCallbackUrl, List<string>? Arguments)
+		public async Task<IJob> AddAsync(JobId JobId, StreamId StreamId, TemplateRefId TemplateRefId, ContentHash TemplateHash, IGraph Graph, string Name, int Change, int CodeChange, int? PreflightChange, int? ClonedPreflightChange, UserId? StartedByUserId, string? StartedByUserName, Priority? Priority, bool? AutoSubmit, bool? UpdateIssues, List<ChainedJobTemplate>? ChainedJobs, bool ShowUgsBadges, bool ShowUgsAlerts, string? NotificationChannel, string? NotificationChannelFilter, string? HelixSwarmCallbackUrl, List<string>? Arguments)
 		{
 			List<ChainedJobDocument> JobTriggers = new List<ChainedJobDocument>();
 			if (ChainedJobs == null)
@@ -417,7 +418,7 @@ namespace HordeServer.Collections.Impl
 		}
 
 		/// <inheritdoc/>
-		public async Task<List<IJob>> FindAsync(JobId[]? JobIds, StreamId? StreamId, string? Name, TemplateRefId[]? Templates, int? MinChange, int? MaxChange, int? PreflightChange, ObjectId? PreflightStartedByUser, ObjectId? StartedByUser, DateTimeOffset? MinCreateTime, DateTimeOffset? MaxCreateTime, DateTimeOffset? ModifiedBefore, DateTimeOffset? ModifiedAfter, int? Index, int? Count)
+		public async Task<List<IJob>> FindAsync(JobId[]? JobIds, StreamId? StreamId, string? Name, TemplateRefId[]? Templates, int? MinChange, int? MaxChange, int? PreflightChange, UserId? PreflightStartedByUser, UserId? StartedByUser, DateTimeOffset? MinCreateTime, DateTimeOffset? MaxCreateTime, DateTimeOffset? ModifiedBefore, DateTimeOffset? ModifiedAfter, int? Index, int? Count)
 		{
 			FilterDefinitionBuilder<JobDocument> FilterBuilder = Builders<JobDocument>.Filter;
 
diff --git a/Engine/Source/Programs/Horde/HordeServer/Jobs/Impl/JobService.cs b/Engine/Source/Programs/Horde/HordeServer/Jobs/Impl/JobService.cs
--- a/Engine/Source/Programs/Horde/HordeServer/Jobs/Impl/JobService.cs
+++ b/Engine/Source/Programs/Horde/HordeServer/Jobs/Impl/JobService.cs
@@ -35,6 +35,7 @@ namespace HordeServer.Services
 	using LogId = ObjectId<ILogFile>;
 	using StreamId = StringId<IStream>;
 	using TemplateRefId = StringId<TemplateRef>;
+	using UserId = ObjectId<IUser>;
 
 	/// <summary>
 	/// Exception thrown when attempting to retry executing a node that does not allow retries
@@ -232,7 +233,7 @@ namespace HordeServer.Services
 		/// <param name="HelixSwarmCallbackUrl">Helix Swarm callback URL for review, if any</param>
 		/// <param name="Arguments">Arguments for the job</param>
 		/// <returns>Unique id representing the job</returns>
-		public async Task<IJob> CreateJobAsync(JobId? JobId, IStream Stream, TemplateRefId TemplateRefId, ContentHash TemplateHash, IGraph Graph, string Name, int Change, int CodeChange, int? PreflightChange, int? ClonedPreflightChange, ObjectId? StartedByUserId, string? StartedByUserName, Priority? Priority, bool? AutoSubmit, bool? UpdateIssues, List<ChainedJobTemplate>? JobTriggers, bool ShowUgsBadges, bool ShowUgsAlerts, string? NotificationChannel, string? NotificationChannelFilter, string? HelixSwarmCallbackUrl, IReadOnlyList<string> Arguments)
+		public async Task<IJob> CreateJobAsync(JobId? JobId, IStream Stream, TemplateRefId TemplateRefId, ContentHash TemplateHash, IGraph Graph, string Name, int Change, int CodeChange, int? PreflightChange, int? ClonedPreflightChange, UserId? StartedByUserId, string? StartedByUserName, Priority? Priority, bool? AutoSubmit, bool? UpdateIssues, List<ChainedJobTemplate>? JobTriggers, bool ShowUgsBadges, bool ShowUgsAlerts, string? NotificationChannel, string? NotificationChannelFilter, string? HelixSwarmCallbackUrl, IReadOnlyList<string> Arguments)
 		{
 			JobId JobIdValue = JobId ?? HordeServer.Utilities.ObjectId<IJob>.GenerateNewId();
 			using IDisposable Scope = Logger.BeginScope("CreateJobAsync({JobId})", JobIdValue);
@@ -494,7 +495,7 @@ namespace HordeServer.Services
 		/// <param name="Index">Index of the first result to return</param>
 		/// <param name="Count">Number of results to return</param>
 		/// <returns>List of jobs matching the given criteria</returns>
-		public async Task<List<IJob>> FindJobsAsync(JobId[]? JobIds = null, StreamId? StreamId = null, string? Name = null, TemplateRefId[]? Templates = null, int? MinChange = null, int? MaxChange = null, int? PreflightChange = null, ObjectId ? PreflightStartedByUser = null, ObjectId? StartedByUser = null, DateTimeOffset ? MinCreateTime = null, DateTimeOffset? MaxCreateTime = null, string? Target = null, JobStepState[]? State = null, JobStepOutcome[]? Outcome = null, DateTimeOffset? ModifiedBefore = null, DateTimeOffset? ModifiedAfter = null, int? Index = null, int? Count = null)
+		public async Task<List<IJob>> FindJobsAsync(JobId[]? JobIds = null, StreamId? StreamId = null, string? Name = null, TemplateRefId[]? Templates = null, int? MinChange = null, int? MaxChange = null, int? PreflightChange = null, UserId? PreflightStartedByUser = null, UserId? StartedByUser = null, DateTimeOffset ? MinCreateTime = null, DateTimeOffset? MaxCreateTime = null, string? Target = null, JobStepState[]? State = null, JobStepOutcome[]? Outcome = null, DateTimeOffset? ModifiedBefore = null, DateTimeOffset? ModifiedAfter = null, int? Index = null, int? Count = null)
 		{
 			if (Target == null && (State == null || State.Length == 0) && (Outcome == null || Outcome.Length == 0))
 			{
diff --git a/Engine/Source/Programs/Horde/HordeServer/Notifications/INotificationTrigger.cs b/Engine/Source/Programs/Horde/HordeServer/Notifications/INotificationTrigger.cs
--- a/Engine/Source/Programs/Horde/HordeServer/Notifications/INotificationTrigger.cs
+++ b/Engine/Source/Programs/Horde/HordeServer/Notifications/INotificationTrigger.cs
@@ -1,12 +1,15 @@
 // Copyright Epic Games, Inc. All Rights Reserved.
 
 using HordeServer.Api;
+using HordeServer.Utilities;
 using MongoDB.Bson;
 using MongoDB.Bson.Serialization.Attributes;
 using System.Collections.Generic;
 
 namespace HordeServer.Models
 {
+	using UserId = ObjectId<IUser>;
+
 	/// <summary>
 	/// An individual subscription
 	/// </summary>
@@ -15,7 +18,7 @@ namespace HordeServer.Models
 		/// <summary>
 		/// Name of the user
 		/// </summary>
-		public ObjectId UserId { get; }
+		public UserId UserId { get; }
 
 		/// <summary>
 		/// Whether to receive email notifications
diff --git a/Engine/Source/Programs/Horde/HordeServer/Notifications/INotificationTriggerCollection.cs b/Engine/Source/Programs/Horde/HordeServer/Notifications/INotificationTriggerCollection.cs
--- a/Engine/Source/Programs/Horde/HordeServer/Notifications/INotificationTriggerCollection.cs
+++ b/Engine/Source/Programs/Horde/HordeServer/Notifications/INotificationTriggerCollection.cs
@@ -13,6 +13,8 @@ using System.Threading.Tasks;
 
 namespace HordeServer.Collections
 {
+	using UserId = ObjectId<IUser>;
+
 	/// <summary>
 	/// Collection of notification triggers
 	/// </summary>
@@ -61,6 +63,6 @@ namespace HordeServer.Collections
 		/// <param name="Email">Whether to receive email notifications</param>
 		/// <param name="Slack">Whether to receive Slack notifications</param>
 		/// <returns>The new trigger state</returns>
-		Task<INotificationTrigger?> UpdateSubscriptionsAsync(INotificationTrigger Trigger, ObjectId UserId, bool? Email, bool? Slack);
+		Task<INotificationTrigger?> UpdateSubscriptionsAsync(INotificationTrigger Trigger, UserId UserId, bool? Email, bool? Slack);
 	}
 }
diff --git a/Engine/Source/Programs/Horde/HordeServer/Notifications/ISubscription.cs b/Engine/Source/Programs/Horde/HordeServer/Notifications/ISubscription.cs
--- a/Engine/Source/Programs/Horde/HordeServer/Notifications/ISubscription.cs
+++ b/Engine/Source/Programs/Horde/HordeServer/Notifications/ISubscription.cs
@@ -2,6 +2,7 @@
 
 using HordeCommon;
 using HordeServer.Api;
+using HordeServer.Utilities;
 using MongoDB.Bson;
 using System;
 using System.Collections.Generic;
@@ -10,6 +11,8 @@ using System.Threading.Tasks;
 
 namespace HordeServer.Models
 {
+	using UserId = ObjectId<IUser>;
+
 	/// <summary>
 	/// The type of notification to send
 	/// </summary>
@@ -39,7 +42,7 @@ namespace HordeServer.Models
 		/// <summary>
 		/// User to notify
 		/// </summary>
-		public ObjectId UserId { get; }
+		public UserId UserId { get; }
 
 		/// <summary>
 		/// Type of notification to receive
diff --git a/Engine/Source/Programs/Horde/HordeServer/Notifications/ISubscriptionCollection.cs b/Engine/Source/Programs/Horde/HordeServer/Notifications/ISubscriptionCollection.cs
--- a/Engine/Source/Programs/Horde/HordeServer/Notifications/ISubscriptionCollection.cs
+++ b/Engine/Source/Programs/Horde/HordeServer/Notifications/ISubscriptionCollection.cs
@@ -21,6 +21,7 @@ namespace HordeServer.Collections
 {
 	using StreamId = StringId<IStream>;
 	using TemplateRefId = StringId<TemplateRef>;
+	using UserId = ObjectId<IUser>;
 
 	/// <summary>
 	/// Information for creating a new subscription
@@ -35,7 +36,7 @@ namespace HordeServer.Collections
 		/// <summary>
 		/// The user name
 		/// </summary>
-		public ObjectId UserId { get; set; }
+		public UserId UserId { get; set; }
 
 		/// <summary>
 		/// Type of notification to send
@@ -48,7 +49,7 @@ namespace HordeServer.Collections
 		/// <param name="Event">Name of the event</param>
 		/// <param name="UserId">User name</param>
 		/// <param name="NotificationType">Notification type</param>
-		public NewSubscription(EventRecord Event, ObjectId UserId, NotificationType NotificationType)
+		public NewSubscription(EventRecord Event, UserId UserId, NotificationType NotificationType)
 		{
 			this.Event = Event;
 			this.UserId = UserId;
@@ -94,6 +95,6 @@ namespace HordeServer.Collections
 		/// </summary>
 		/// <param name="UserId">The user to search for</param>
 		/// <returns>List of subscriptions</returns>
-		public Task<List<ISubscription>> FindSubscriptionsAsync(ObjectId UserId);
+		public Task<List<ISubscription>> FindSubscriptionsAsync(UserId UserId);
 	}
 }
diff --git a/Engine/Source/Programs/Horde/HordeServer/Notifications/Impl/NotificationService.cs b/Engine/Source/Programs/Horde/HordeServer/Notifications/Impl/NotificationService.cs
--- a/Engine/Source/Programs/Horde/HordeServer/Notifications/Impl/NotificationService.cs
+++ b/Engine/Source/Programs/Horde/HordeServer/Notifications/Impl/NotificationService.cs
@@ -34,6 +34,8 @@ using HordeServer.Services;
 
 namespace HordeServer.Notifications.Impl
 {
+	using UserId = ObjectId<IUser>;
+
 	/// <summary>
 	/// Wraps funtionality for delivering notifications.
 	/// </summary>
@@ -147,7 +149,7 @@ namespace HordeServer.Notifications.Impl
 		/// <inheritdoc/>
 		public async Task<bool> UpdateSubscriptionsAsync(ObjectId TriggerId, ClaimsPrincipal User, bool? Email, bool? Slack)
 		{
-			ObjectId? UserId = User.GetUserId();
+			UserId? UserId = User.GetUserId();
 			if (UserId == null)
 			{
 				Logger.LogWarning("Unable to find username for principal {User}", User.Identity?.Name);
@@ -162,7 +164,7 @@ namespace HordeServer.Notifications.Impl
 		/// <inheritdoc/>
 		public async Task<INotificationSubscription?> GetSubscriptionsAsync(ObjectId TriggerId, ClaimsPrincipal User)
 		{
-			ObjectId? UserId = User.GetUserId();
+			UserId? UserId = User.GetUserId();
 			if (UserId == null)
 			{
 				return null;
@@ -417,7 +419,7 @@ namespace HordeServer.Notifications.Impl
 
 		private async Task<List<IUser>> GetUsersToNotify(EventRecord? Event, ObjectId? NotificationTriggerId, bool bFireTrigger)
 		{
-			List<ObjectId> UserIds = new List<ObjectId>();
+			List<UserId> UserIds = new List<UserId>();
 
 			// Find the notifications for all steps of this type
 			if (Event != null)
diff --git a/Engine/Source/Programs/Horde/HordeServer/Notifications/Impl/NotificationTriggerCollection.cs b/Engine/Source/Programs/Horde/HordeServer/Notifications/Impl/NotificationTriggerCollection.cs
--- a/Engine/Source/Programs/Horde/HordeServer/Notifications/Impl/NotificationTriggerCollection.cs
+++ b/Engine/Source/Programs/Horde/HordeServer/Notifications/Impl/NotificationTriggerCollection.cs
@@ -13,6 +13,8 @@ using System.Threading.Tasks;
 
 namespace HordeServer.Collections.Impl
 {
+	using UserId = ObjectId<IUser>;
+
 	/// <summary>
 	/// Collection of notification triggers
 	/// </summary>
@@ -23,7 +25,7 @@ namespace HordeServer.Collections.Impl
 			[BsonIgnoreIfNull]
 			public string? User { get; set; }
 
-			public ObjectId UserId { get; set; }
+			public UserId UserId { get; set; }
 			public bool Email { get; set; }
 			public bool Slack { get; set; }
 
@@ -33,7 +35,7 @@ namespace HordeServer.Collections.Impl
 				User = null!;
 			}
 
-			public SubscriptionDocument(ObjectId UserId, bool Email, bool Slack)
+			public SubscriptionDocument(UserId UserId, bool Email, bool Slack)
 			{
 				this.UserId = UserId;
 				this.Email = Email;
@@ -196,7 +198,7 @@ namespace HordeServer.Collections.Impl
 		}
 
 		/// <inheritdoc/>
-		public async Task<INotificationTrigger?> UpdateSubscriptionsAsync(INotificationTrigger Trigger, ObjectId UserId, bool? Email, bool? Slack)
+		public async Task<INotificationTrigger?> UpdateSubscriptionsAsync(INotificationTrigger Trigger, UserId UserId, bool? Email, bool? Slack)
 		{
 			for (; ; )
 			{
diff --git a/Engine/Source/Programs/Horde/HordeServer/Notifications/Impl/SlackNotificationSink.cs b/Engine/Source/Programs/Horde/HordeServer/Notifications/Impl/SlackNotificationSink.cs
--- a/Engine/Source/Programs/Horde/HordeServer/Notifications/Impl/SlackNotificationSink.cs
+++ b/Engine/Source/Programs/Horde/HordeServer/Notifications/Impl/SlackNotificationSink.cs
@@ -36,6 +36,8 @@ using System.Threading.Tasks;
 
 namespace HordeServer.Notifications.Impl
 {
+	using UserId = ObjectId<IUser>;
+
 	/// <summary>
 	/// Maintains a connection to Slack, in order to receive socket-mode notifications of user interactions
 	/// </summary>
@@ -131,7 +133,7 @@ namespace HordeServer.Notifications.Impl
 			public string Recipient { get; set; } = String.Empty;
 
 			[BsonElement("usr")]
-			public ObjectId? UserId { get; set; }
+			public UserId? UserId { get; set; }
 
 			[BsonElement("eid")]
 			public string EventId { get; set; } = String.Empty;
@@ -150,7 +152,7 @@ namespace HordeServer.Notifications.Impl
 		{
 			public const int CurrentVersion = 2;
 			
-			public ObjectId Id { get; set; }
+			public UserId Id { get; set; }
 
 			[BsonElement("u")]
 			public string? SlackUserId { get; set; }
@@ -177,7 +179,7 @@ namespace HordeServer.Notifications.Impl
 			{
 			}
 
-			public SlackUser(ObjectId Id, UserInfo? Info)
+			public SlackUser(UserId Id, UserInfo? Info)
 			{
 				this.Id = Id;
 				this.SlackUserId = Info?.UserName;
@@ -263,7 +265,7 @@ namespace HordeServer.Notifications.Impl
 			return await MessageStates.Find(Filter).FirstOrDefaultAsync();
 		}
 
-		async Task<MessageStateDocument> SetMessageStateAsync(string Recipient, string EventId, ObjectId? UserId, string Channel, string Ts, string Digest)
+		async Task<MessageStateDocument> SetMessageStateAsync(string Recipient, string EventId, UserId? UserId, string Channel, string Ts, string Digest)
 		{
 			FilterDefinition<MessageStateDocument> Filter = Builders<MessageStateDocument>.Filter.Eq(x => x.Recipient, Recipient) & Builders<MessageStateDocument>.Filter.Eq(x => x.EventId, EventId);
 			UpdateDefinition<MessageStateDocument> Update = Builders<MessageStateDocument>.Update.SetOnInsert(x => x.Id, ObjectId.GenerateNewId()).Set(x => x.UserId, UserId).Set(x => x.Channel, Channel).Set(x => x.Ts, Ts).Set(x => x.Digest, Digest);
@@ -535,7 +537,7 @@ namespace HordeServer.Notifications.Impl
 		{
 			IIssueDetails Details = await IssueService.GetIssueDetailsAsync(Issue);
 
-			HashSet<ObjectId> UserIds = new HashSet<ObjectId>();
+			HashSet<UserId> UserIds = new HashSet<UserId>();
 			if (Issue.NotifySuspects)
 			{
 				UserIds.UnionWith(Details.Suspects.Select(x => x.AuthorId));
@@ -582,7 +584,7 @@ namespace HordeServer.Notifications.Impl
 
 			if (UserIds.Count > 0)
 			{
-				foreach (ObjectId UserId in UserIds)
+				foreach (UserId UserId in UserIds)
 				{
 					IUser? User = await UserCollection.GetUserAsync(UserId);
 					if (User == null)
@@ -650,7 +652,7 @@ namespace HordeServer.Notifications.Impl
 			return SampleText.ToString();
 		}
 
-		async Task SendIssueMessageAsync(string Recipient, IIssue Issue, IIssueDetails Details, ObjectId? UserId)
+		async Task SendIssueMessageAsync(string Recipient, IIssue Issue, IIssueDetails Details, UserId? UserId)
 		{
 			using IDisposable Scope = Logger.BeginScope("SendIssueMessageAsync (User: {SlackUser}, Issue: {IssueId})", Recipient, Issue.Id);
 
@@ -788,7 +790,7 @@ namespace HordeServer.Notifications.Impl
 			return $"issue_{Issue.Id}";
 		}
 
-		async Task<string> FormatNameAsync(ObjectId UserId)
+		async Task<string> FormatNameAsync(UserId UserId)
 		{
 			IUser? User = await UserCollection.GetUserAsync(UserId);
 			if (User == null)
@@ -798,7 +800,7 @@ namespace HordeServer.Notifications.Impl
 			return User.Name;
 		}
 
-		async Task<string> FormatMentionAsync(ObjectId UserId)
+		async Task<string> FormatMentionAsync(UserId UserId)
 		{
 			IUser? User = await UserCollection.GetUserAsync(UserId);
 			if (User == null)
@@ -838,7 +840,7 @@ namespace HordeServer.Notifications.Impl
 
 			int IssueId = int.Parse(Match.Groups[1].Value, NumberStyles.Integer, CultureInfo.InvariantCulture);
 			string Verb = Match.Groups[2].Value;
-			ObjectId UserId = Match.Groups[3].Value.ToObjectId();
+			UserId UserId = Match.Groups[3].Value.ToObjectId<IUser>();
 			Logger.LogInformation("Issue {IssueId}: {Action} from {SlackUser} ({UserId})", IssueId, Verb, UserName, UserId);
 
 			if (String.Equals(Verb, "ack", StringComparison.Ordinal))
@@ -1161,7 +1163,7 @@ namespace HordeServer.Notifications.Impl
 			await SendRequestAsync<PostMessageResponse>(PostMessageUrl, Message);
 		}
 
-		private async Task SendOrUpdateMessageAsync(string Recipient, string EventId, ObjectId? UserId, string? Text = null, BlockBase[]? Blocks = null, BlockKitAttachment[]? Attachments = null)
+		private async Task SendOrUpdateMessageAsync(string Recipient, string EventId, UserId? UserId, string? Text = null, BlockBase[]? Blocks = null, BlockKitAttachment[]? Attachments = null)
 		{
 			if (AllowUsers != null && !AllowUsers.Contains(Recipient))
 			{
diff --git a/Engine/Source/Programs/Horde/HordeServer/Notifications/Impl/SubscriptionCollection.cs b/Engine/Source/Programs/Horde/HordeServer/Notifications/Impl/SubscriptionCollection.cs
--- a/Engine/Source/Programs/Horde/HordeServer/Notifications/Impl/SubscriptionCollection.cs
+++ b/Engine/Source/Programs/Horde/HordeServer/Notifications/Impl/SubscriptionCollection.cs
@@ -21,6 +21,7 @@ namespace HordeServer.Collections.Impl
 {
 	using StreamId = StringId<IStream>;
 	using TemplateRefId = StringId<TemplateRef>;
+	using UserId = ObjectId<IUser>;
 
 	/// <summary>
 	/// Collection of subscription documents
@@ -159,7 +160,7 @@ namespace HordeServer.Collections.Impl
 		{
 			public string Id { get; set; }
 			public Event Event { get; set; }
-			public ObjectId UserId { get; set; }
+			public UserId UserId { get; set; }
 			public NotificationType NotificationType { get; set; }
 
 			IEvent ISubscription.Event => Event;
@@ -240,7 +241,7 @@ namespace HordeServer.Collections.Impl
 		}
 
 		/// <inheritdoc/>
-		public async Task<List<ISubscription>> FindSubscriptionsAsync(ObjectId UserId)
+		public async Task<List<ISubscription>> FindSubscriptionsAsync(UserId UserId)
 		{
 			List<Subscription> Results = await Collection.Find(x => x.UserId == UserId).ToListAsync();
 			return Results.ConvertAll<ISubscription>(x => x);
diff --git a/Engine/Source/Programs/Horde/HordeServer/Perforce/ICommit.cs b/Engine/Source/Programs/Horde/HordeServer/Perforce/ICommit.cs
--- a/Engine/Source/Programs/Horde/HordeServer/Perforce/ICommit.cs
+++ b/Engine/Source/Programs/Horde/HordeServer/Perforce/ICommit.cs
@@ -13,6 +13,7 @@ using System.Threading.Tasks;
 namespace HordeServer.Commits
 {
 	using StreamId = StringId<IStream>;
+	using UserId = ObjectId<IUser>;
 
 	/// <summary>
 	/// Stores metadata about a commit
@@ -42,12 +43,12 @@ namespace HordeServer.Commits
 		/// <summary>
 		/// The author user id
 		/// </summary>
-		public ObjectId AuthorId { get; }
+		public UserId AuthorId { get; }
 
 		/// <summary>
 		/// The owner of this change, if different from the author (due to Robomerge)
 		/// </summary>
-		public ObjectId OwnerId { get; }
+		public UserId OwnerId { get; }
 
 		/// <summary>
 		/// Changelist description
@@ -85,10 +86,10 @@ namespace HordeServer.Commits
 		public int OriginalChange { get; set; }
 
 		/// <inheritdoc cref="ICommit.AuthorId"/>
-		public ObjectId AuthorId { get; set; }
+		public UserId AuthorId { get; set; }
 
 		/// <inheritdoc cref="ICommit.OwnerId"/>
-		public ObjectId? OwnerId { get; set; }
+		public UserId? OwnerId { get; set; }
 
 		/// <inheritdoc cref="ICommit.Description"/>
 		public string Description { get; set; }
@@ -114,7 +115,7 @@ namespace HordeServer.Commits
 		/// <summary>
 		/// Constructor
 		/// </summary>
-		public NewCommit(StreamId StreamId, int Change, int OriginalChange, ObjectId AuthorId, ObjectId OwnerId, string Description, string BasePath, DateTime DateUtc)
+		public NewCommit(StreamId StreamId, int Change, int OriginalChange, UserId AuthorId, UserId OwnerId, string Description, string BasePath, DateTime DateUtc)
 		{
 			this.StreamId = StreamId;
 			this.Change = Change;
diff --git a/Engine/Source/Programs/Horde/HordeServer/Perforce/Impl/CommitCollection.cs b/Engine/Source/Programs/Horde/HordeServer/Perforce/Impl/CommitCollection.cs
--- a/Engine/Source/Programs/Horde/HordeServer/Perforce/Impl/CommitCollection.cs
+++ b/Engine/Source/Programs/Horde/HordeServer/Perforce/Impl/CommitCollection.cs
@@ -15,6 +15,7 @@ using MongoDB.Bson.Serialization.Attributes;
 namespace HordeServer.Commits.Impl
 {
 	using StreamId = StringId<IStream>;
+	using UserId = ObjectId<IUser>;
 
 	/// <summary>
 	/// Concrete implementation of ICommitCollection
@@ -36,10 +37,10 @@ namespace HordeServer.Commits.Impl
 			public int? OriginalChange { get; set; }
 
 			[BsonElement("a")]
-			public ObjectId AuthorId { get; set; }
+			public UserId AuthorId { get; set; }
 
 			[BsonElement("o"), BsonIgnoreIfNull]
-			public ObjectId? OwnerId { get; set; }
+			public UserId? OwnerId { get; set; }
 
 			[BsonElement("d")]
 			public string Description { get; set; } = String.Empty;
@@ -54,7 +55,7 @@ namespace HordeServer.Commits.Impl
 			public DateTime DateUtc { get; set; }
 
 			int ICommit.OriginalChange => OriginalChange ?? Change;
-			ObjectId ICommit.OwnerId => OwnerId ?? AuthorId;
+			UserId ICommit.OwnerId => OwnerId ?? AuthorId;
 
 			public Commit()
 			{
diff --git a/Engine/Source/Programs/Horde/HordeServer/Startup.cs b/Engine/Source/Programs/Horde/HordeServer/Startup.cs
--- a/Engine/Source/Programs/Horde/HordeServer/Startup.cs
+++ b/Engine/Source/Programs/Horde/HordeServer/Startup.cs
@@ -90,6 +90,7 @@ namespace HordeServer
 	using JobId = ObjectId<IJob>;
 	using LeaseId = ObjectId<ILease>;
 	using LogId = ObjectId<ILogFile>;
+	using UserId = ObjectId<IUser>;
 	using ILogger = Microsoft.Extensions.Logging.ILogger;
 
 	class Startup
@@ -613,6 +614,7 @@ namespace HordeServer
 			MessageTemplate.RegisterFormatter(typeof(JobId), new MessageTemplateTypeNameFormatter("JobId"));
 			MessageTemplate.RegisterFormatter(typeof(LeaseId), new MessageTemplateTypeNameFormatter("LeaseId"));
 			MessageTemplate.RegisterFormatter(typeof(LogId), new MessageTemplateTypeNameFormatter("LogId"));
+			MessageTemplate.RegisterFormatter(typeof(UserId), new MessageTemplateTypeNameFormatter("UserId"));
 		}
 
 		public static void ConfigureJsonSerializer(JsonSerializerOptions Options)
diff --git a/Engine/Source/Programs/Horde/HordeServer/Users/IUser.cs b/Engine/Source/Programs/Horde/HordeServer/Users/IUser.cs
--- a/Engine/Source/Programs/Horde/HordeServer/Users/IUser.cs
+++ b/Engine/Source/Programs/Horde/HordeServer/Users/IUser.cs
@@ -11,6 +11,7 @@ using System.Threading.Tasks;
 namespace HordeServer.Models
 {
 	using JobId = ObjectId<IJob>;
+	using UserId = ObjectId<IUser>;
 
 	/// <summary>
 	/// Document which collates information about a user, and their personal settings
@@ -20,7 +21,7 @@ namespace HordeServer.Models
 		/// <summary>
 		/// The user id
 		/// </summary>
-		public ObjectId Id { get; }
+		public UserId Id { get; }
 
 		/// <summary>
 		/// Full name of the user
@@ -46,7 +47,7 @@ namespace HordeServer.Models
 		/// <summary>
 		/// The user id
 		/// </summary>
-		public ObjectId UserId { get; }
+		public UserId UserId { get; }
 
 		/// <summary>
 		/// Claims for this user (on last login)
@@ -62,7 +63,7 @@ namespace HordeServer.Models
 		/// <summary>
 		/// The user id
 		/// </summary>
-		public ObjectId UserId { get; }
+		public UserId UserId { get; }
 
 		/// <summary>
 		/// Whether to enable experimental features
diff --git a/Engine/Source/Programs/Horde/HordeServer/Users/IUserCollection.cs b/Engine/Source/Programs/Horde/HordeServer/Users/IUserCollection.cs
--- a/Engine/Source/Programs/Horde/HordeServer/Users/IUserCollection.cs
+++ b/Engine/Source/Programs/Horde/HordeServer/Users/IUserCollection.cs
@@ -13,6 +13,7 @@ using System.Threading.Tasks;
 namespace HordeServer.Collections
 {
 	using JobId = ObjectId<IJob>;
+	using UserId = ObjectId<IUser>;
 
 	/// <summary>
 	/// Manages user documents
@@ -24,14 +25,14 @@ namespace HordeServer.Collections
 		/// </summary>
 		/// <param name="Id">Id of the user</param>
 		/// <returns>The user information</returns>
-		Task<IUser?> GetUserAsync(ObjectId Id);
+		Task<IUser?> GetUserAsync(UserId Id);
 
 		/// <summary>
 		/// Gets a cached user by unique id
 		/// </summary>
 		/// <param name="Id">Id of the user</param>
 		/// <returns>The user information</returns>
-		Task<IUser?> GetCachedUserAsync(ObjectId Id);
+		Task<IUser?> GetCachedUserAsync(UserId Id);
 
 		/// <summary>
 		/// Gets a user by unique id
@@ -41,7 +42,7 @@ namespace HordeServer.Collections
 		/// <param name="Index">Maximum number of results</param>
 		/// <param name="Count">Number of results to return</param>
 		/// <returns>The user information</returns>
-		Task<List<IUser>> FindUsersAsync(IEnumerable<ObjectId>? Ids = null, string? NameRegex = null, int? Index = null, int? Count = null);
+		Task<List<IUser>> FindUsersAsync(IEnumerable<UserId>? Ids = null, string? NameRegex = null, int? Index = null, int? Count = null);
 
 		/// <summary>
 		/// Gets a user by login
@@ -64,7 +65,7 @@ namespace HordeServer.Collections
 		/// </summary>
 		/// <param name="UserId"></param>
 		/// <returns></returns>
-		Task<IUserClaims> GetClaimsAsync(ObjectId UserId);
+		Task<IUserClaims> GetClaimsAsync(UserId UserId);
 
 		/// <summary>
 		/// Update the claims for a user
@@ -72,14 +73,14 @@ namespace HordeServer.Collections
 		/// <param name="UserId"></param>
 		/// <param name="Claims"></param>
 		/// <returns></returns>
-		Task UpdateClaimsAsync(ObjectId UserId, IEnumerable<IUserClaim> Claims);
+		Task UpdateClaimsAsync(UserId UserId, IEnumerable<IUserClaim> Claims);
 
 		/// <summary>
 		/// Get settings for a user
 		/// </summary>
 		/// <param name="UserId"></param>
 		/// <returns></returns>
-		Task<IUserSettings> GetSettingsAsync(ObjectId UserId);
+		Task<IUserSettings> GetSettingsAsync(UserId UserId);
 
 		/// <summary>
 		/// Update a user
@@ -91,7 +92,7 @@ namespace HordeServer.Collections
 		/// <param name="AddPinnedJobIds"></param>
 		/// <param name="RemovePinnedJobIds"></param>
 		/// <returns>Updated user object</returns>
-		Task UpdateSettingsAsync(ObjectId UserId, bool? EnableExperimentalFeatures = null, bool? EnableIssueNotifications = null, BsonValue? DashboardSettings = null, IEnumerable<JobId>? AddPinnedJobIds = null, IEnumerable<JobId>? RemovePinnedJobIds = null);
+		Task UpdateSettingsAsync(UserId UserId, bool? EnableExperimentalFeatures = null, bool? EnableIssueNotifications = null, BsonValue? DashboardSettings = null, IEnumerable<JobId>? AddPinnedJobIds = null, IEnumerable<JobId>? RemovePinnedJobIds = null);
 	}
 
 	/// <summary>
@@ -107,7 +108,7 @@ namespace HordeServer.Collections
 		/// <returns></returns>
 		public static Task<IUser?> GetUserAsync(this IUserCollection UserCollection, ClaimsPrincipal Principal)
 		{
-			ObjectId? UserId = Principal.GetUserId();
+			UserId? UserId = Principal.GetUserId();
 			if (UserId != null)
 			{
 				return UserCollection.GetUserAsync(UserId.Value);
diff --git a/Engine/Source/Programs/Horde/HordeServer/Users/Impl/UserCollection.cs b/Engine/Source/Programs/Horde/HordeServer/Users/Impl/UserCollection.cs
--- a/Engine/Source/Programs/Horde/HordeServer/Users/Impl/UserCollection.cs
+++ b/Engine/Source/Programs/Horde/HordeServer/Users/Impl/UserCollection.cs
@@ -16,6 +16,7 @@ using System.Threading.Tasks;
 namespace HordeServer.Collections.Impl
 {
 	using JobId = ObjectId<IJob>;
+	using UserId = ObjectId<IUser>;
 
 	/// <summary>
 	/// Manages user documents
@@ -24,7 +25,7 @@ namespace HordeServer.Collections.Impl
 	{
 		class UserDocument : IUser, IUserClaims, IUserSettings
 		{
-			public ObjectId Id { get; set; }
+			public UserId Id { get; set; }
 
 			public ClaimDocument PrimaryClaim { get; set; } = null!;
 			public List<ClaimDocument> Claims { get; set; } = new List<ClaimDocument>();
@@ -42,10 +43,10 @@ namespace HordeServer.Collections.Impl
 			string IUser.Login => Claims.FirstOrDefault(x => String.Equals(x.Type, ClaimTypes.Name, StringComparison.Ordinal))?.Value ?? PrimaryClaim.Value;
 			string? IUser.Email => Claims.FirstOrDefault(x => String.Equals(x.Type, ClaimTypes.Email, StringComparison.Ordinal))?.Value;
 
-			ObjectId IUserClaims.UserId => Id;
+			UserId IUserClaims.UserId => Id;
 			IReadOnlyList<IUserClaim> IUserClaims.Claims => Claims;
 
-			ObjectId IUserSettings.UserId => Id;
+			UserId IUserSettings.UserId => Id;
 			IReadOnlyList<JobId> IUserSettings.PinnedJobIds => PinnedJobIds;
 		}
 
@@ -99,19 +100,19 @@ namespace HordeServer.Collections.Impl
 		}
 
 		/// <inheritdoc/>
-		public async Task<IUser?> GetUserAsync(ObjectId Id)
+		public async Task<IUser?> GetUserAsync(UserId Id)
 		{
 			return await Users.Find(x => x.Id == Id).FirstOrDefaultAsync();
 		}
 
 		/// <inheritdoc/>
-		public Task<IUser?> GetCachedUserAsync(ObjectId Id)
+		public Task<IUser?> GetCachedUserAsync(UserId Id)
 		{
 			return GetUserAsync(Id);
 		}
 
 		/// <inheritdoc/>
-		public async Task<List<IUser>> FindUsersAsync(IEnumerable<ObjectId>? Ids, string? NameRegex, int? Index, int? Count)
+		public async Task<List<IUser>> FindUsersAsync(IEnumerable<UserId>? Ids, string? NameRegex, int? Index, int? Count)
 		{
 			FilterDefinition<UserDocument> Filter = Builders<UserDocument>.Filter.In(x => x.Id, Ids);
 			return await Users.Find(Filter).Range(Index, Count).ToListAsync<UserDocument, IUser>();
@@ -128,31 +129,31 @@ namespace HordeServer.Collections.Impl
 		public async Task<IUser> FindOrAddUserByLoginAsync(string Login, string? Name, string? Email)
 		{
 			ClaimDocument NewPrimaryClaim = new ClaimDocument(ClaimTypes.Name, Login);
-			UpdateDefinition<UserDocument> Update = Builders<UserDocument>.Update.SetOnInsert(x => x.Id, ObjectId.GenerateNewId());
+			UpdateDefinition<UserDocument> Update = Builders<UserDocument>.Update.SetOnInsert(x => x.Id, UserId.GenerateNewId());
 			return await Users.FindOneAndUpdateAsync<UserDocument>(x => x.PrimaryClaim == NewPrimaryClaim, Update, new FindOneAndUpdateOptions<UserDocument> { IsUpsert = true, ReturnDocument = ReturnDocument.After });
 		}
 
 		/// <inheritdoc/>
-		public async Task<IUserClaims> GetClaimsAsync(ObjectId UserId)
+		public async Task<IUserClaims> GetClaimsAsync(UserId UserId)
 		{
 			return await Users.Find(x => x.Id == UserId).FirstOrDefaultAsync() ?? new UserDocument { Id = UserId };
 		}
 
 		/// <inheritdoc/>
-		public async Task UpdateClaimsAsync(ObjectId UserId, IEnumerable<IUserClaim> Claims)
+		public async Task UpdateClaimsAsync(UserId UserId, IEnumerable<IUserClaim> Claims)
 		{
 			List<ClaimDocument> NewClaims = Claims.Select(x => new ClaimDocument(x)).ToList();
 			await Users.FindOneAndUpdateAsync(x => x.Id == UserId, Builders<UserDocument>.Update.Set(x => x.Claims, NewClaims));
 		}
 
 		/// <inheritdoc/>
-		public async Task<IUserSettings> GetSettingsAsync(ObjectId UserId)
+		public async Task<IUserSettings> GetSettingsAsync(UserId UserId)
 		{
 			return await Users.Find(x => x.Id == UserId).FirstOrDefaultAsync() ?? new UserDocument { Id = UserId };
 		}
 
 		/// <inheritdoc/>
-		public async Task UpdateSettingsAsync(ObjectId UserId, bool? EnableExperimentalFeatures, bool? EnableIssueNotifications, BsonValue? DashboardSettings = null, IEnumerable<JobId>? AddPinnedJobIds = null, IEnumerable<JobId>? RemovePinnedJobIds = null)
+		public async Task UpdateSettingsAsync(UserId UserId, bool? EnableExperimentalFeatures, bool? EnableIssueNotifications, BsonValue? DashboardSettings = null, IEnumerable<JobId>? AddPinnedJobIds = null, IEnumerable<JobId>? RemovePinnedJobIds = null)
 		{
 			if (AddPinnedJobIds != null)
 			{
diff --git a/Engine/Source/Programs/Horde/HordeServer/Users/Impl/UserCollectionV2.cs b/Engine/Source/Programs/Horde/HordeServer/Users/Impl/UserCollectionV2.cs
--- a/Engine/Source/Programs/Horde/HordeServer/Users/Impl/UserCollectionV2.cs
+++ b/Engine/Source/Programs/Horde/HordeServer/Users/Impl/UserCollectionV2.cs
@@ -18,6 +18,7 @@ using System.Threading.Tasks;
 namespace HordeServer.Collections.Impl
 {
 	using JobId = ObjectId<IJob>;
+	using UserId = ObjectId<IUser>;
 
 	/// <summary>
 	/// Manages user documents
@@ -26,7 +27,7 @@ namespace HordeServer.Collections.Impl
 	{
 		class UserDocument : IUser
 		{
-			public ObjectId Id { get; set; }
+			public UserId Id { get; set; }
 
 			public string Name { get; set; }
 			public string Login { get; set; }
@@ -50,7 +51,7 @@ namespace HordeServer.Collections.Impl
 			{
 			}
 
-			public UserDocument(ObjectId Id, string Name, string Login, string? Email)
+			public UserDocument(UserId Id, string Name, string Login, string? Email)
 			{
 				this.Id = Id;
 				this.Name = Name;
@@ -63,10 +64,10 @@ namespace HordeServer.Collections.Impl
 
 		class UserClaimsDocument : IUserClaims
 		{
-			public ObjectId Id { get; set; }
+			public UserId Id { get; set; }
 			public List<UserClaim> Claims { get; set; } = new List<UserClaim>();
 
-			ObjectId IUserClaims.UserId => Id;
+			UserId IUserClaims.UserId => Id;
 			IReadOnlyList<IUserClaim> IUserClaims.Claims => Claims;
 
 			[BsonConstructor]
@@ -74,7 +75,7 @@ namespace HordeServer.Collections.Impl
 			{
 			}
 
-			public UserClaimsDocument(ObjectId Id)
+			public UserClaimsDocument(UserId Id)
 			{
 				this.Id = Id;
 			}
@@ -88,7 +89,7 @@ namespace HordeServer.Collections.Impl
 
 		class UserSettingsDocument : IUserSettings
 		{
-			public ObjectId Id { get; set; }
+			public UserId Id { get; set; }
 
 			[BsonDefaultValue(false), BsonIgnoreIfDefault]
 			public bool EnableExperimentalFeatures { get; set; }
@@ -99,7 +100,7 @@ namespace HordeServer.Collections.Impl
 			public BsonValue DashboardSettings { get; set; } = BsonNull.Value;
 			public List<JobId> PinnedJobIds { get; set; } = new List<JobId>();
 
-			ObjectId IUserSettings.UserId => Id;
+			UserId IUserSettings.UserId => Id;
 			IReadOnlyList<JobId> IUserSettings.PinnedJobIds => PinnedJobIds;
 
 			[BsonConstructor]
@@ -107,7 +108,7 @@ namespace HordeServer.Collections.Impl
 			{
 			}
 
-			public UserSettingsDocument(ObjectId Id)
+			public UserSettingsDocument(UserId Id)
 			{
 				this.Id = Id;
 			}
@@ -196,7 +197,7 @@ namespace HordeServer.Collections.Impl
 		}
 
 		/// <inheritdoc/>
-		public async Task<IUser?> GetUserAsync(ObjectId Id)
+		public async Task<IUser?> GetUserAsync(UserId Id)
 		{
 			IUser? User = await Users.Find(x => x.Id == Id).FirstOrDefaultAsync();
 			using (ICacheEntry Entry = UserCache.CreateEntry(Id))
@@ -208,7 +209,7 @@ namespace HordeServer.Collections.Impl
 		}
 
 		/// <inheritdoc/>
-		public Task<IUser?> GetCachedUserAsync(ObjectId Id)
+		public Task<IUser?> GetCachedUserAsync(UserId Id)
 		{
 			IUser? User;
 			if (UserCache.TryGetValue(Id, out User))
@@ -222,7 +223,7 @@ namespace HordeServer.Collections.Impl
 		}
 
 		/// <inheritdoc/>
-		public async Task<List<IUser>> FindUsersAsync(IEnumerable<ObjectId>? Ids, string? NameRegex = null, int? Index = null, int? Count = null)
+		public async Task<List<IUser>> FindUsersAsync(IEnumerable<UserId>? Ids, string? NameRegex = null, int? Index = null, int? Count = null)
 		{
 			FilterDefinition<UserDocument> Filter = FilterDefinition<UserDocument>.Empty;
 			if (Ids != null)
@@ -251,7 +252,7 @@ namespace HordeServer.Collections.Impl
 		/// <inheritdoc/>
 		public async Task<IUser> FindOrAddUserByLoginAsync(string Login, string? Name, string? Email)
 		{
-			UpdateDefinition<UserDocument> Update = Builders<UserDocument>.Update.SetOnInsert(x => x.Id, ObjectId.GenerateNewId()).SetOnInsert(x => x.Login, Login).Unset(x => x.Hidden);
+			UpdateDefinition<UserDocument> Update = Builders<UserDocument>.Update.SetOnInsert(x => x.Id, UserId.GenerateNewId()).SetOnInsert(x => x.Login, Login).Unset(x => x.Hidden);
 
 			if (Name == null)
 			{
@@ -272,7 +273,7 @@ namespace HordeServer.Collections.Impl
 		}
 
 		/// <inheritdoc/>
-		public async Task<IUserClaims> GetClaimsAsync(ObjectId UserId)
+		public async Task<IUserClaims> GetClaimsAsync(UserId UserId)
 		{
 			IUserClaims? Claims = await UserClaims.Find(x => x.Id == UserId).FirstOrDefaultAsync();
 			if(Claims == null)
@@ -283,7 +284,7 @@ namespace HordeServer.Collections.Impl
 		}
 
 		/// <inheritdoc/>
-		public async Task UpdateClaimsAsync(ObjectId UserId, IEnumerable<IUserClaim> Claims)
+		public async Task UpdateClaimsAsync(UserId UserId, IEnumerable<IUserClaim> Claims)
 		{
 			UserClaimsDocument NewDocument = new UserClaimsDocument(UserId);
 			NewDocument.Claims.AddRange(Claims.Select(x => new UserClaim(x)));
@@ -291,7 +292,7 @@ namespace HordeServer.Collections.Impl
 		}
 
 		/// <inheritdoc/>
-		public async Task<IUserSettings> GetSettingsAsync(ObjectId UserId)
+		public async Task<IUserSettings> GetSettingsAsync(UserId UserId)
 		{
 			IUserSettings? Settings = await UserSettings.Find(x => x.Id == UserId).FirstOrDefaultAsync();
 			if (Settings == null)
@@ -302,7 +303,7 @@ namespace HordeServer.Collections.Impl
 		}
 
 		/// <inheritdoc/>
-		public async Task UpdateSettingsAsync(ObjectId UserId, bool? EnableExperimentalFeatures = null, bool? EnableIssueNotifications = null, BsonValue? DashboardSettings = null, IEnumerable<JobId>? AddPinnedJobIds = null, IEnumerable<JobId>? RemovePinnedJobIds = null)
+		public async Task UpdateSettingsAsync(UserId UserId, bool? EnableExperimentalFeatures = null, bool? EnableIssueNotifications = null, BsonValue? DashboardSettings = null, IEnumerable<JobId>? AddPinnedJobIds = null, IEnumerable<JobId>? RemovePinnedJobIds = null)
 		{
 			List<UpdateDefinition<UserSettingsDocument>> Updates = new List<UpdateDefinition<UserSettingsDocument>>();
 			if (EnableExperimentalFeatures != null)
diff --git a/Engine/Source/Programs/Horde/HordeServer/Utilities/HordeClaimTypes.cs b/Engine/Source/Programs/Horde/HordeServer/Utilities/HordeClaimTypes.cs
--- a/Engine/Source/Programs/Horde/HordeServer/Utilities/HordeClaimTypes.cs
+++ b/Engine/Source/Programs/Horde/HordeServer/Utilities/HordeClaimTypes.cs
@@ -1,6 +1,7 @@
 // Copyright Epic Games, Inc. All Rights Reserved.
 
 using HordeServer.Collections;
+using HordeServer.Models;
 using MongoDB.Bson;
 using System;
 using System.Collections.Generic;
@@ -10,6 +11,8 @@ using System.Threading.Tasks;
 
 namespace HordeServer.Utilities
 {
+	using UserId = ObjectId<IUser>;
+
 	/// <summary>
 	/// Claim types that are specific to horde
 	/// </summary>
@@ -66,7 +69,7 @@ namespace HordeServer.Utilities
 		/// </summary>
 		/// <param name="Principal"></param>
 		/// <returns></returns>
-		public static ObjectId? GetUserId(this ClaimsPrincipal Principal)
+		public static UserId? GetUserId(this ClaimsPrincipal Principal)
 		{
 			string? IdValue = Principal.FindFirstValue(HordeClaimTypes.UserId);
 			if(IdValue == null)
@@ -75,7 +78,7 @@ namespace HordeServer.Utilities
 			}
 			else
 			{
-				return new ObjectId(IdValue);
+				return new UserId(IdValue);
 			}
 		}
 
