Horde: Return more detailed user information in response objects (includes id, name, email). No longer include the login name, on the assumption that this generally should not be exposed anywhere outside the server.

diff --git a/Engine/Source/Programs/Horde/HordeServer.Compute.Tests/Stubs/Services/PerforceServiceStub.cs b/Engine/Source/Programs/Horde/HordeServer.Compute.Tests/Stubs/Services/PerforceServiceStub.cs
--- a/Engine/Source/Programs/Horde/HordeServer.Compute.Tests/Stubs/Services/PerforceServiceStub.cs
+++ b/Engine/Source/Programs/Horde/HordeServer.Compute.Tests/Stubs/Services/PerforceServiceStub.cs
@@ -3,6 +3,7 @@
 using EpicGames.Perforce;
 using HordeServer.Models;
 using HordeServer.Services;
+using HordeServer.Utilities;
 using System;
 using System.Collections.Generic;
 using System.Diagnostics.CodeAnalysis;
@@ -15,6 +16,8 @@ using StreamId = HordeServer.Utilities.StringId<HordeServer.Models.IStream>;
 
 namespace HordeServerTests.Stubs.Services
 {
+	using UserId = ObjectId<IUser>;
+
 	static class PerforceExtensions
 	{
 		public static ChangeFile CreateChangeFile(string Path)
@@ -30,11 +33,28 @@ namespace HordeServerTests.Stubs.Services
 
 	class PerforceServiceStub : IPerforceService
 	{
+		class User : IUser
+		{
+			public UserId Id { get; }
+			public string Name { get; set; }
+			public string Email { get; set; }
+			public string Login { get; set; }
+
+			public User(string Login)
+			{
+				this.Id = UserId.GenerateNewId();
+				this.Name = Login.ToUpperInvariant();
+				this.Email = $"{Login}@server";
+				this.Login = Login;
+			}
+		}
+
 		class ChangeComparer : IComparer<int>
 		{
 			public int Compare(int X, int Y) => Y.CompareTo(X);
 		}
 
+		Dictionary<string, User> NameToUser = new Dictionary<string, User>();
 		public Dictionary<string, SortedDictionary<int, ChangeDetails>> Changes { get; } = new Dictionary<string, SortedDictionary<int, ChangeDetails>>(StringComparer.OrdinalIgnoreCase);
 
 		public Task<NativePerforceConnection?> GetServiceUserConnection(string? ClusterName)
@@ -42,6 +62,22 @@ namespace HordeServerTests.Stubs.Services
 			throw new NotImplementedException();
 		}
 
+		IUser FindOrAddUser(string UserName)
+		{
+			User? User;
+			if (!NameToUser.TryGetValue(UserName, out User))
+			{
+				User = new User(UserName);
+				NameToUser.Add(UserName, User);
+			}
+			return User;
+		}
+
+		public ValueTask<IUser> FindOrAddUserAsync(string ClusterName, string UserName)
+		{
+			return new ValueTask<IUser>(FindOrAddUser(UserName));
+		}
+
 		public void AddChange(string StreamName, int Number, string Author, string Description, IEnumerable<string> Files)
 		{
 			SortedDictionary<int, ChangeDetails>? StreamChanges;
@@ -50,7 +86,7 @@ namespace HordeServerTests.Stubs.Services
 				StreamChanges = new SortedDictionary<int, ChangeDetails>(new ChangeComparer());
 				Changes[StreamName] = StreamChanges;
 			}
-			StreamChanges.Add(Number, new ChangeDetails(Number, Author, null!, Description, Files.Select(x => PerforceExtensions.CreateChangeFile(x)).ToList(), DateTime.Now));
+			StreamChanges.Add(Number, new ChangeDetails(Number, FindOrAddUser(Author), null!, Description, Files.Select(x => PerforceExtensions.CreateChangeFile(x)).ToList(), DateTime.Now));
 		}
 
 		public Task<List<ChangeSummary>> GetChangesAsync(string ClusterName, string StreamName, int? MinChange, int? MaxChange, int NumResults, string? ImpersonateUser)
@@ -82,7 +118,7 @@ namespace HordeServerTests.Stubs.Services
 
 		public Task<PerforceUserInfo?> GetUserInfoAsync(string ClusterName, string UserName)
 		{
-			return Task.FromResult<PerforceUserInfo?>(new PerforceUserInfo { Name = UserName, Email = $"[REDACTED_EMAIL]" });
+			return Task.FromResult<PerforceUserInfo?>(new PerforceUserInfo { Login = UserName, FullName = UserName, Email = $"[REDACTED_EMAIL]" });
 		}
 
 		public Task<List<ChangeDetails>> GetChangeDetailsAsync(string ClusterName, string StreamName, IReadOnlyList<int> ChangeNumbers, string? ImpersonateUser)
@@ -122,7 +158,7 @@ namespace HordeServerTests.Stubs.Services
 
 		public Task<int> CreateNewChangeAsync(string ClusterName, string StreamName, string Path)
 		{
-			ChangeDetails NewChange = new ChangeDetails(Changes[StreamName].First().Key + 1, "", null!, "", new List<ChangeFile> { PerforceExtensions.CreateChangeFile(Path) }, DateTime.Now);
+			ChangeDetails NewChange = new ChangeDetails(Changes[StreamName].First().Key + 1, FindOrAddUser("system"), null!, "", new List<ChangeFile> { PerforceExtensions.CreateChangeFile(Path) }, DateTime.Now);
 			Changes[StreamName].Add(NewChange.Number, NewChange);
 			return Task.FromResult(NewChange.Number);
 		}
diff --git a/Engine/Source/Programs/Horde/HordeServer/Api/Changes.cs b/Engine/Source/Programs/Horde/HordeServer/Api/Changes.cs
--- a/Engine/Source/Programs/Horde/HordeServer/Api/Changes.cs
+++ b/Engine/Source/Programs/Horde/HordeServer/Api/Changes.cs
@@ -19,10 +19,15 @@ namespace HordeServer.Api
 		public int Number { get; set; }
 
 		/// <summary>
-		/// Name of the user that authored this change
+		/// Name of the user that authored this change [DEPRECATED]
 		/// </summary>
 		public string Author { get; set; }
 
+		/// <summary>
+		/// Information about the change author
+		/// </summary>
+		public GetThinUserInfoResponse AuthorInfo { get; set; }
+
 		/// <summary>
 		/// The description text
 		/// </summary>
@@ -35,7 +40,8 @@ namespace HordeServer.Api
 		public GetChangeSummaryResponse(ChangeSummary ChangeSummary)
 		{
 			this.Number = ChangeSummary.Number;
-			this.Author = ChangeSummary.Author;
+			this.Author = ChangeSummary.Author.Name;
+			this.AuthorInfo = new GetThinUserInfoResponse(ChangeSummary.Author);
 			this.Description = ChangeSummary.Description;
 		}
 	}
@@ -51,10 +57,15 @@ namespace HordeServer.Api
 		public int Number { get; set; }
 
 		/// <summary>
-		/// Name of the user that authored this change
+		/// Name of the user that authored this change [DEPRECATED]
 		/// </summary>
 		public string Author { get; set; }
 
+		/// <summary>
+		/// Information about the user that authored this change
+		/// </summary>
+		public GetThinUserInfoResponse AuthorInfo { get; set; }
+
 		/// <summary>
 		/// The description text
 		/// </summary>
@@ -72,7 +83,8 @@ namespace HordeServer.Api
 		public GetChangeDetailsResponse(ChangeDetails ChangeDetails)
 		{
 			this.Number = ChangeDetails.Number;
-			this.Author = ChangeDetails.Author;
+			this.Author = ChangeDetails.Author.Name;
+			this.AuthorInfo = new GetThinUserInfoResponse(ChangeDetails.Author);
 			this.Description = ChangeDetails.Description;
 			this.Files = ChangeDetails.Files.ConvertAll(x => x.Path);
 		}
diff --git a/Engine/Source/Programs/Horde/HordeServer/Api/Issues.cs b/Engine/Source/Programs/Horde/HordeServer/Api/Issues.cs
--- a/Engine/Source/Programs/Horde/HordeServer/Api/Issues.cs
+++ b/Engine/Source/Programs/Horde/HordeServer/Api/Issues.cs
@@ -381,20 +381,30 @@ namespace HordeServer.Api
 		public IssueSeverity Severity { get; set; }
 
 		/// <summary>
-		/// Owner of the issue
+		/// Owner of the issue [DEPRECATED]
 		/// </summary>
 		public string? Owner { get; set; }
 
 		/// <summary>
-		/// User id of the owner
+		/// User id of the owner [DEPRECATED]
 		/// </summary>
 		public string? OwnerId { get; set; }
 
 		/// <summary>
-		/// User that nominated the current owner
+		/// Owner of the issue
+		/// </summary>
+		public GetThinUserInfoResponse? OwnerInfo { get; set; }
+
+		/// <summary>
+		/// User that nominated the current owner [DEPRECATED]
 		/// </summary>
 		public string? NominatedBy { get; set; }
 
+		/// <summary>
+		/// Owner of the issue
+		/// </summary>
+		public GetThinUserInfoResponse? NominatedByInfo { get; set; }
+
 		/// <summary>
 		/// Time that the issue was acknowledged
 		/// </summary>
@@ -411,15 +421,20 @@ namespace HordeServer.Api
 		public DateTime? ResolvedAt { get; set; }
 
 		/// <summary>
-		/// Name of the user that resolved the issue
+		/// Name of the user that resolved the issue [DEPRECATED]
 		/// </summary>
 		public string? ResolvedBy { get; set; }
 
 		/// <summary>
-		/// User id of the person that resolved the issue
+		/// User id of the person that resolved the issue [DEPRECATED]
 		/// </summary>
 		public string? ResolvedById { get; set; }
 
+		/// <summary>
+		/// User that resolved the issue
+		/// </summary>
+		public GetThinUserInfoResponse? ResolvedByInfo { get; set; }
+
 		/// <summary>
 		/// Time at which the issue was verified
 		/// </summary>
@@ -451,15 +466,20 @@ namespace HordeServer.Api
 		public List<GetIssueAffectedStreamResponse> AffectedStreams { get; set; }
 
 		/// <summary>
-		/// Most likely suspects for causing this issue, deprecated
+		/// Most likely suspects for causing this issue [DEPRECATED]
 		/// </summary>
 		public List<string> PrimarySuspects { get; set; }
 
 		/// <summary>
-		/// User ids of the most likely suspects
+		/// User ids of the most likely suspects [DEPRECATED]
 		/// </summary>
 		public List<string> PrimarySuspectIds { get; set; }
 
+		/// <summary>
+		/// Most likely suspects for causing this issue
+		/// </summary>
+		public List<GetThinUserInfoResponse> PrimarySuspectsInfo { get; set; }
+
 		/// <summary>
 		/// Whether to show alerts for this issue
 		/// </summary>
@@ -481,12 +501,24 @@ namespace HordeServer.Api
 			this.Severity = Issue.Severity;
 			this.Owner = Details.Owner?.Login;
 			this.OwnerId = (Details.Owner == null)? null : Details.Owner.Id.ToString();
+			if(Details.Owner != null)
+			{
+				this.OwnerInfo = new GetThinUserInfoResponse(Details.Owner);
+			}
 			this.NominatedBy = Details.NominatedBy?.Login;
+			if (Details.NominatedBy != null)
+			{
+				this.NominatedByInfo = new GetThinUserInfoResponse(Details.NominatedBy);
+			}
 			this.AcknowledgedAt = Issue.AcknowledgedAt;
 			this.FixChange = Issue.FixChange;
 			this.ResolvedAt = Issue.ResolvedAt;
 			this.ResolvedBy = Details.ResolvedBy?.Login;
 			this.ResolvedById = (Details.ResolvedBy == null) ? null : Details.ResolvedBy.Id.ToString();
+			if (Details.ResolvedBy != null)
+			{
+				this.ResolvedByInfo = new GetThinUserInfoResponse(Details.ResolvedBy);
+			}
 			this.VerifiedAt = Issue.VerifiedAt;
 			this.LastSeenAt = Issue.LastSeenAt;
 			this.Streams = Details.Spans.Select(x => x.StreamName).Distinct().ToList()!;
@@ -506,6 +538,7 @@ namespace HordeServer.Api
 			}
 			this.PrimarySuspects = Details.SuspectUsers.Where(x => x.Login != null).Select(x => x.Login).ToList();
 			this.PrimarySuspectIds= Details.SuspectUsers.Select(x => x.Id.ToString()).ToList();
+			this.PrimarySuspectsInfo = Details.SuspectUsers.ConvertAll(x => new GetThinUserInfoResponse(x));
 			this.ShowDesktopAlerts = ShowDesktopAlerts;
 		}
 	}
diff --git a/Engine/Source/Programs/Horde/HordeServer/Api/Jobs.cs b/Engine/Source/Programs/Horde/HordeServer/Api/Jobs.cs
--- a/Engine/Source/Programs/Horde/HordeServer/Api/Jobs.cs
+++ b/Engine/Source/Programs/Horde/HordeServer/Api/Jobs.cs
@@ -246,20 +246,30 @@ namespace HordeServer.Api
 		public string? GraphHash { get; set; }
 
 		/// <summary>
-		/// The user that started this job
+		/// The user that started this job [DEPRECATED]
 		/// </summary>
 		public string? StartedByUserId { get; set; }
 
 		/// <summary>
-		/// The user that started this job
+		/// The user that started this job [DEPRECATED]
 		/// </summary>
 		public string? StartedByUser { get; set; }
 
 		/// <summary>
-		/// The user that aborted this job
+		/// The user that started this job
+		/// </summary>
+		public GetThinUserInfoResponse? StartedByUserInfo { get; set; }
+
+		/// <summary>
+		/// The user that aborted this job [DEPRECATED]
 		/// </summary>
 		public string? AbortedByUser { get; set; }
 
+		/// <summary>
+		/// The user that aborted this job
+		/// </summary>
+		public GetThinUserInfoResponse? AbortedByUserInfo { get; set; }
+
 		/// <summary>
 		/// Priority of the job
 		/// </summary>
@@ -329,8 +339,10 @@ namespace HordeServer.Api
 		/// Constructor
 		/// </summary>
 		/// <param name="Job">Job to create a response for</param>
+		/// <param name="StartedByUserInfo">User that started this job</param>
+		/// <param name="AbortedByUserInfo">User that aborted this job</param>
 		/// <param name="AclResponse">The ACL response object</param>
-		public GetJobResponse(IJob Job, GetAclResponse? AclResponse)
+		public GetJobResponse(IJob Job, GetThinUserInfoResponse? StartedByUserInfo, GetThinUserInfoResponse? AbortedByUserInfo, GetAclResponse? AclResponse)
 		{
 			this.Id = Job.Id.ToString();
 			this.StreamId = Job.StreamId.ToString();
@@ -344,7 +356,9 @@ namespace HordeServer.Api
 			this.GraphHash = Job.GraphHash.ToString();
 			this.StartedByUserId = Job.StartedByUserId?.ToString();
 			this.StartedByUser = Job.StartedByUser;
+			this.StartedByUserInfo = StartedByUserInfo;
 			this.AbortedByUser = Job.AbortedByUser;
+			this.AbortedByUserInfo = AbortedByUserInfo;
 			this.CreateTime = new DateTimeOffset(Job.CreateTimeUtc);
 			this.State = Job.GetState();
 			this.Priority = Job.Priority;
@@ -473,15 +487,25 @@ namespace HordeServer.Api
 		public bool AbortRequested { get; set; }
 		
 		/// <summary>
-		/// Name of the user that requested the abort of this step
+		/// Name of the user that requested the abort of this step [DEPRECATED]
 		/// </summary>
 		public string? AbortByUser { get; set; }
 
 		/// <summary>
-		/// Name of the user that requested this step be run again
+		/// The user that requested this step be run again 
+		/// </summary>
+		public GetThinUserInfoResponse? AbortedByUserInfo { get; set; }
+
+		/// <summary>
+		/// Name of the user that requested this step be run again [DEPRECATED]
 		/// </summary>
 		public string? RetryByUser { get; set; }
 
+		/// <summary>
+		/// The user that requested this step be run again 
+		/// </summary>
+		public GetThinUserInfoResponse? RetriedByUserInfo { get; set; }
+
 		/// <summary>
 		/// The log id for this step
 		/// </summary>
@@ -516,7 +540,9 @@ namespace HordeServer.Api
 		/// Constructor
 		/// </summary>
 		/// <param name="Step">The step to construct from</param>
-		public GetStepResponse(IJobStep Step)
+		/// <param name="AbortedByUserInfo">User that aborted this step</param>
+		/// <param name="RetriedByUserInfo">User that retried this step</param>
+		public GetStepResponse(IJobStep Step, GetThinUserInfoResponse? AbortedByUserInfo, GetThinUserInfoResponse? RetriedByUserInfo)
 		{
 			this.Id = Step.Id.ToString();
 			this.NodeIdx = Step.NodeIdx;
@@ -524,7 +550,9 @@ namespace HordeServer.Api
 			this.Outcome = Step.Outcome;
 			this.AbortRequested = Step.AbortRequested;
 			this.AbortByUser = Step.AbortByUser;
+			this.AbortedByUserInfo = AbortedByUserInfo;
 			this.RetryByUser = Step.RetryByUser;
+			this.RetriedByUserInfo = RetriedByUserInfo;
 			this.LogId = Step.LogId?.ToString();
 			this.ReadyTime = Step.ReadyTimeUtc;
 			this.StartTime = Step.StartTimeUtc;
@@ -717,15 +745,16 @@ namespace HordeServer.Api
 		/// Converts this batch into a public response object
 		/// </summary>
 		/// <param name="Batch">The batch to construct from</param>
+		/// <param name="Steps">Steps in this batch</param>
 		/// <returns>Response instance</returns>
-		public GetBatchResponse(IJobStepBatch Batch)
+		public GetBatchResponse(IJobStepBatch Batch, List<GetStepResponse> Steps)
 		{
 			this.Id = Batch.Id.ToString();
 			this.LogId = Batch.LogId?.ToString();
 			this.GroupIdx = Batch.GroupIdx;
 			this.State = Batch.State;
 			this.Error = Batch.Error;
-			this.Steps = Batch.Steps.ConvertAll(x => new GetStepResponse(x));
+			this.Steps = Steps;
 			this.AgentId = Batch.AgentId?.ToString();
 			this.SessionId = Batch.SessionId?.ToString();
 			this.LeaseId = Batch.LeaseId?.ToString();
diff --git a/Engine/Source/Programs/Horde/HordeServer/Api/Users.cs b/Engine/Source/Programs/Horde/HordeServer/Api/Users.cs
--- a/Engine/Source/Programs/Horde/HordeServer/Api/Users.cs
+++ b/Engine/Source/Programs/Horde/HordeServer/Api/Users.cs
@@ -103,6 +103,47 @@ namespace HordeServer.Api
 		}
 	}
 
+	/// <summary>
+	/// Basic information about a user. May be embedded in other responses.
+	/// </summary>
+	public class GetThinUserInfoResponse
+	{
+		/// <summary>
+		/// Id of the user
+		/// </summary>
+		public string Id { get; set; }
+
+		/// <summary>
+		/// Name of the user
+		/// </summary>
+		public string Name { get; set; }
+
+		/// <summary>
+		/// The user's email address
+		/// </summary>
+		public string? Email { get; set; }
+
+		/// <summary>
+		/// Constructor
+		/// </summary>
+		/// <param name="User"></param>
+		public GetThinUserInfoResponse(IUser? User)
+		{
+			if (User == null)
+			{
+				this.Id = String.Empty;
+				this.Name = "(Unknown)";
+				this.Email = null;
+			}
+			else
+			{
+				this.Id = User.Id.ToString();
+				this.Name = User.Name;
+				this.Email = User.Email;
+			}
+		}
+	}
+
 	/// <summary>
 	/// Request to update settings for a user
 	/// </summary>
diff --git a/Engine/Source/Programs/Horde/HordeServer/Controllers/DebugController.cs b/Engine/Source/Programs/Horde/HordeServer/Controllers/DebugController.cs
--- a/Engine/Source/Programs/Horde/HordeServer/Controllers/DebugController.cs
+++ b/Engine/Source/Programs/Horde/HordeServer/Controllers/DebugController.cs
@@ -699,7 +699,7 @@ namespace HordeServer.Controllers
 			}
 			else
 			{
-				Results.Add("GetUserInfoAsync", $"{UserInfo.Name} : {UserInfo.Email}");
+				Results.Add("GetUserInfoAsync", $"{UserInfo.Login} : {UserInfo.Email}");
 			}
 
 			// changes
diff --git a/Engine/Source/Programs/Horde/HordeServer/Controllers/JobsController.cs b/Engine/Source/Programs/Horde/HordeServer/Controllers/JobsController.cs
--- a/Engine/Source/Programs/Horde/HordeServer/Controllers/JobsController.cs
+++ b/Engine/Source/Programs/Horde/HordeServer/Controllers/JobsController.cs
@@ -40,55 +40,21 @@ namespace HordeServer.Controllers
 	[Route("[controller]")]
 	public class JobsController : ControllerBase
 	{
-		/// <summary>
-		/// Instance of the ACL service
-		/// </summary>
 		private readonly AclService AclService;
-
-		/// <summary>
-		/// Collection of graphs
-		/// </summary>
 		private readonly IGraphCollection Graphs;
-
-		/// <summary>
-		/// The perforce service instance
-		/// </summary>
 		private readonly IPerforceService Perforce;
-
-		/// <summary>
-		/// Instance of the stream service
-		/// </summary>
 		private readonly StreamService StreamService;
-
-		/// <summary>
-		/// Instance of the JobService singleton
-		/// </summary>
 		private readonly JobService JobService;
-
-		/// <summary>
-		/// Instance of the TemplateService singleton
-		/// </summary>
 		private readonly ITemplateCollection TemplateCollection;
-
-		/// <summary>
-		/// 
-		/// </summary>
 		private readonly IArtifactCollection ArtifactCollection;
-
-		/// <summary>
-		/// Instance of the notification service singleton
-		/// </summary>
+		private readonly IUserCollection UserCollection;
 		private readonly INotificationService NotificationService;
-
-		/// <summary>
-		/// Logger instance
-		/// </summary>
 		private ILogger<JobsController> Logger;
 
 		/// <summary>
 		/// Constructor
 		/// </summary>
-		public JobsController(AclService AclService, IGraphCollection Graphs, IPerforceService Perforce, StreamService StreamService, JobService JobService, ITemplateCollection TemplateCollection, IArtifactCollection ArtifactCollection, INotificationService NotificationService, ILogger<JobsController> Logger)
+		public JobsController(AclService AclService, IGraphCollection Graphs, IPerforceService Perforce, StreamService StreamService, JobService JobService, ITemplateCollection TemplateCollection, IArtifactCollection ArtifactCollection, IUserCollection UserCollection, INotificationService NotificationService, ILogger<JobsController> Logger)
 		{
 			this.AclService = AclService;
 			this.Graphs = Graphs;
@@ -97,6 +63,7 @@ namespace HordeServer.Controllers
 			this.JobService = JobService;
 			this.TemplateCollection = TemplateCollection;
 			this.ArtifactCollection = ArtifactCollection;
+			this.UserCollection = UserCollection;
 			this.NotificationService = NotificationService;
 			this.Logger = Logger;
 		}
@@ -381,7 +348,113 @@ namespace HordeServer.Controllers
 
 			IGraph Graph = await JobService.GetGraphAsync(Job);
 			bool bIncludeAcl = await JobService.AuthorizeAsync(Job, AclAction.ViewPermissions, User, Cache);
-			return Job.ToResponse(Graph, bIncludeAcl, Filter);
+			return await CreateJobResponseAsync(Job, Graph, bIncludeAcl, Filter);
+		}
+
+		/// <summary>
+		/// Creates a response containing information about this object
+		/// </summary>
+		/// <param name="Job">The job document</param>
+		/// <param name="Graph">The graph for this job</param>
+		/// <param name="bIncludeAcl">Whether to include the ACL in the response</param>
+		/// <param name="Filter">Filter for the properties to return</param>
+		/// <returns>Object containing the requested properties</returns>
+		async Task<object> CreateJobResponseAsync(IJob Job, IGraph Graph, bool bIncludeAcl, PropertyFilter? Filter)
+		{
+			if (Filter == null)
+			{
+				return await CreateJobResponseAsync(Job, Graph, true, true, bIncludeAcl);
+			}
+			else
+			{
+				return Filter.ApplyTo(await CreateJobResponseAsync(Job, Graph, Filter.Includes(nameof(GetJobResponse.Batches)), Filter.Includes(nameof(GetJobResponse.Labels)) || Filter.Includes(nameof(GetJobResponse.DefaultLabel)), bIncludeAcl));
+			}
+		}
+
+		/// <summary>
+		/// Creates a response containing information about this object
+		/// </summary>
+		/// <param name="Job">The job document</param>
+		/// <param name="Graph">The graph definition</param>
+		/// <param name="bIncludeBatches">Whether to include the job batches in the response</param>
+		/// <param name="bIncludeLabels">Whether to include the job aggregates in the response</param>
+		/// <param name="bIncludeAcl">Whether to include the ACL in the response</param>
+		/// <returns>The response object</returns>
+		async ValueTask<GetJobResponse> CreateJobResponseAsync(IJob Job, IGraph Graph, bool bIncludeBatches, bool bIncludeLabels, bool bIncludeAcl)
+		{
+			GetThinUserInfoResponse? StartedByUserInfo = null;
+			if (Job.StartedByUserId != null)
+			{
+				StartedByUserInfo = new GetThinUserInfoResponse(await UserCollection.GetCachedUserAsync(Job.StartedByUserId.Value));
+			}
+
+			GetThinUserInfoResponse? AbortedByUserInfo = null;
+			if (Job.AbortedByUser != null)
+			{
+				AbortedByUserInfo = new GetThinUserInfoResponse(await UserCollection.FindOrAddUserByLoginAsync(Job.AbortedByUser));
+			}
+
+			GetAclResponse? AclResponse = null;
+			if (bIncludeAcl && Job.Acl != null)
+			{
+				AclResponse = new GetAclResponse(Job.Acl);
+			}
+
+			GetJobResponse Response = new GetJobResponse(Job, StartedByUserInfo, AbortedByUserInfo, AclResponse);
+			if (bIncludeBatches || bIncludeLabels)
+			{
+				if (bIncludeBatches)
+				{
+					Response.Batches = new List<GetBatchResponse>();
+					foreach (IJobStepBatch Batch in Job.Batches)
+					{
+						Response.Batches.Add(await CreateBatchResponseAsync(Batch));
+					}
+				}
+				if (bIncludeLabels)
+				{
+					Response.Labels = new List<GetLabelStateResponse>();
+					Response.DefaultLabel = Job.GetLabelStateResponses(Graph, Response.Labels);
+				}
+			}
+			return Response;
+		}
+
+		/// <summary>
+		/// Get the response object for a batch
+		/// </summary>
+		/// <param name="Batch"></param>
+		/// <returns></returns>
+		async ValueTask<GetBatchResponse> CreateBatchResponseAsync(IJobStepBatch Batch)
+		{
+			List<GetStepResponse> Steps = new List<GetStepResponse>();
+			foreach (IJobStep Step in Batch.Steps)
+			{
+				Steps.Add(await CreateStepResponseAsync(Step));
+			}
+			return new GetBatchResponse(Batch, Steps);
+		}
+
+		/// <summary>
+		/// Get the response object for a step
+		/// </summary>
+		/// <param name="Step"></param>
+		/// <returns></returns>
+		async ValueTask<GetStepResponse> CreateStepResponseAsync(IJobStep Step)
+		{
+			GetThinUserInfoResponse? AbortedByUserInfo = null;
+			if (Step.AbortByUser != null)
+			{
+				AbortedByUserInfo = new GetThinUserInfoResponse(await UserCollection.FindOrAddUserByLoginAsync(Step.AbortByUser));
+			}
+
+			GetThinUserInfoResponse? RetriedByUserInfo = null;
+			if (Step.RetryByUser != null)
+			{
+				RetriedByUserInfo = new GetThinUserInfoResponse(await UserCollection.FindOrAddUserByLoginAsync(Step.RetryByUser));
+			}
+
+			return new GetStepResponse(Step, AbortedByUserInfo, RetriedByUserInfo);
 		}
 
 		/// <summary>
@@ -596,7 +669,7 @@ namespace HordeServer.Controllers
 
 					using (IScope _ = GlobalTracer.Instance.BuildSpan("CreateResponse").StartActive())
 					{
-						Responses.Add(Job.ToResponse(Graph, bIncludeAcl, Filter));
+						Responses.Add(await CreateJobResponseAsync(Job, Graph, bIncludeAcl, Filter));
 					}
 				}
 			}
@@ -780,7 +853,8 @@ namespace HordeServer.Controllers
 			List<object> Responses = new List<object>();
 			foreach (IJobStepBatch Batch in Job.Batches)
 			{
-				Responses.Add(new GetBatchResponse(Batch).ApplyFilter(Filter));
+				GetBatchResponse Response = await CreateBatchResponseAsync(Batch);
+				Responses.Add(Response.ApplyFilter(Filter));
 			}
 			return Responses;
 		}
@@ -849,7 +923,8 @@ namespace HordeServer.Controllers
 			{
 				if (Batch.Id == BatchIdValue)
 				{
-					return new GetBatchResponse(Batch).ApplyFilter(Filter);
+					GetBatchResponse Response = await CreateBatchResponseAsync(Batch);
+					return Response.ApplyFilter(Filter);
 				}
 			}
 
@@ -890,7 +965,8 @@ namespace HordeServer.Controllers
 					List<object> Responses = new List<object>();
 					foreach (IJobStep Step in Batch.Steps)
 					{
-						Responses.Add(new GetStepResponse(Step).ApplyFilter(Filter));
+						GetStepResponse Response = await CreateStepResponseAsync(Step);
+						Responses.Add(Response.ApplyFilter(Filter));
 					}
 					return Responses;
 				}
@@ -1032,7 +1108,6 @@ namespace HordeServer.Controllers
 				return Forbid();
 			}
 
-			IGraph Graph = await JobService.GetGraphAsync(Job);
 			foreach (IJobStepBatch Batch in Job.Batches)
 			{
 				if (Batch.Id == BatchIdValue)
@@ -1041,7 +1116,8 @@ namespace HordeServer.Controllers
 					{
 						if (Step.Id == StepIdValue)
 						{
-							return new GetStepResponse(Step).ApplyFilter(Filter);
+							GetStepResponse Response = await CreateStepResponseAsync(Step);
+							return Response.ApplyFilter(Filter);
 						}
 					}
 					break;
diff --git a/Engine/Source/Programs/Horde/HordeServer/Issues/Impl/IssueService.cs b/Engine/Source/Programs/Horde/HordeServer/Issues/Impl/IssueService.cs
--- a/Engine/Source/Programs/Horde/HordeServer/Issues/Impl/IssueService.cs
+++ b/Engine/Source/Programs/Horde/HordeServer/Issues/Impl/IssueService.cs
@@ -1093,10 +1093,15 @@ namespace HordeServer.Services.Impl
 						Logger.LogDebug("Suspect CL: {Change}, Author: {Author}, Rank: {Rank}, MaxRank: {MaxRank}", SuspectChange.Details.Number, SuspectChange.Details.Author, SuspectChange.Rank, MaxRank);
 						if (SuspectChange.Rank == MaxRank)
 						{
-							string? Author = ParseRobomergeOwner(SuspectChange.Details.Description) ?? SuspectChange.Details.Author;
-							IUser User = await UserCollection.FindOrAddUserByLoginAsync(Author);
+							IUser? Owner = SuspectChange.Details.Author;
 
-							NewIssueSpanSuspectData Suspect = new NewIssueSpanSuspectData(SuspectChange.Details.Number, User.Id);
+							string? RoboOwnerName = ParseRobomergeOwner(SuspectChange.Details.Description);
+							if (RoboOwnerName != null)
+							{
+								Owner = await Perforce.FindOrAddUserAsync(Stream.ClusterName, RoboOwnerName);
+							}
+
+							NewIssueSpanSuspectData Suspect = new NewIssueSpanSuspectData(SuspectChange.Details.Number, Owner.Id);
 							Suspect.OriginatingChange = ParseRobomergeSource(SuspectChange.Details.Description);
 
 							Suspects.Add(Suspect);
diff --git a/Engine/Source/Programs/Horde/HordeServer/Jobs/IJob.cs b/Engine/Source/Programs/Horde/HordeServer/Jobs/IJob.cs
--- a/Engine/Source/Programs/Horde/HordeServer/Jobs/IJob.cs
+++ b/Engine/Source/Programs/Horde/HordeServer/Jobs/IJob.cs
@@ -1411,54 +1411,6 @@ namespace HordeServer.Models
 			}
 		}
 
-		/// <summary>
-		/// Creates a response containing information about this object
-		/// </summary>
-		/// <param name="Job">The job document</param>
-		/// <param name="Graph">The graph definition</param>
-		/// <param name="bIncludeBatches">Whether to include the job batches in the response</param>
-		/// <param name="bIncludeLabels">Whether to include the job aggregates in the response</param>
-		/// <param name="bIncludeAcl">Whether to include the ACL in the response</param>
-		/// <returns>The response object</returns>
-		public static GetJobResponse ToResponse(this IJob Job, IGraph Graph, bool bIncludeBatches, bool bIncludeLabels, bool bIncludeAcl)
-		{
-			GetAclResponse? AclResponse = (bIncludeAcl && Job.Acl != null) ? new GetAclResponse(Job.Acl) : null;
-			GetJobResponse Response = new GetJobResponse(Job, AclResponse);
-			if (bIncludeBatches || bIncludeLabels)
-			{
-				if (bIncludeBatches)
-				{
-					Response.Batches = Job.Batches.ConvertAll(x => new GetBatchResponse(x));
-				}
-				if (bIncludeLabels)
-				{
-					Response.Labels = new List<GetLabelStateResponse>();
-					Response.DefaultLabel = GetLabelStateResponses(Job, Graph, Response.Labels);
-				}
-			}
-			return Response;
-		}
-
-		/// <summary>
-		/// Creates a response containing information about this object
-		/// </summary>
-		/// <param name="Job">The job document</param>
-		/// <param name="Graph">The graph for this job</param>
-		/// <param name="bIncludeAcl">Whether to include the ACL in the response</param>
-		/// <param name="Filter">Filter for the properties to return</param>
-		/// <returns>Object containing the requested properties</returns>
-		public static object ToResponse(this IJob Job, IGraph Graph, bool bIncludeAcl, PropertyFilter? Filter)
-		{
-			if(Filter == null)
-			{
-				return Job.ToResponse(Graph, true, true, bIncludeAcl);
-			}
-			else
-			{
-				return Filter.ApplyTo(Job.ToResponse(Graph, Filter.Includes(nameof(GetJobResponse.Batches)), Filter.Includes(nameof(GetJobResponse.Labels)) || Filter.Includes(nameof(GetJobResponse.DefaultLabel)), bIncludeAcl));
-			}
-		}
-
 		/// <summary>
 		/// Creates an RPC response object
 		/// </summary>
diff --git a/Engine/Source/Programs/Horde/HordeServer/Notifications/INotificationService.cs b/Engine/Source/Programs/Horde/HordeServer/Notifications/INotificationService.cs
--- a/Engine/Source/Programs/Horde/HordeServer/Notifications/INotificationService.cs
+++ b/Engine/Source/Programs/Horde/HordeServer/Notifications/INotificationService.cs
@@ -61,7 +61,7 @@ namespace HordeServer.Notifications
 		/// <param name="Change">Latest change number of the file</param>
 		/// <param name="Author">Author of the change, if found from p4 service</param>
 		/// <param name="Description">Description of the change, if found from p4 service</param>
-		void NotifyConfigUpdateFailure(string ErrorMessage, string FileName, int? Change = null, string? Author = null, string? Description = null);
+		void NotifyConfigUpdateFailure(string ErrorMessage, string FileName, int? Change = null, IUser? Author = null, string? Description = null);
 
 		/// <summary>
 		/// Sends a notification to a user regarding a build health issue.
diff --git a/Engine/Source/Programs/Horde/HordeServer/Notifications/Impl/NotificationService.cs b/Engine/Source/Programs/Horde/HordeServer/Notifications/Impl/NotificationService.cs
--- a/Engine/Source/Programs/Horde/HordeServer/Notifications/Impl/NotificationService.cs
+++ b/Engine/Source/Programs/Horde/HordeServer/Notifications/Impl/NotificationService.cs
@@ -220,15 +220,9 @@ namespace HordeServer.Notifications.Impl
 		}
 
 		/// <inheritdoc/>
-		public void NotifyConfigUpdateFailure(string ErrorMessage, string FileName, int? Change = null, string? Author = null, string? Description = null)
+		public void NotifyConfigUpdateFailure(string ErrorMessage, string FileName, int? Change = null, IUser? Author = null, string? Description = null)
 		{
-			EnqueueTask(() => NotifyConfigUpdateFailureAsync(ErrorMessage, FileName, Change, Author, Description));
-		}
-
-		async Task NotifyConfigUpdateFailureAsync(string ErrorMessage, string FileName, int? Change = null, string? Author = null, string? Description = null)
-		{
-			IUser? User = (Author != null) ? await UserCollection.FindUserByLoginAsync(Author) : null;
-			EnqueueTasks(Sink => Sink.NotifyConfigUpdateFailureAsync(ErrorMessage, FileName, Change, User, Description));
+			EnqueueTasks(Sink => Sink.NotifyConfigUpdateFailureAsync(ErrorMessage, FileName, Change, Author, Description));
 		}
 
 		/// <inheritdoc/>
diff --git a/Engine/Source/Programs/Horde/HordeServer/Perforce/IPerforceService.cs b/Engine/Source/Programs/Horde/HordeServer/Perforce/IPerforceService.cs
--- a/Engine/Source/Programs/Horde/HordeServer/Perforce/IPerforceService.cs
+++ b/Engine/Source/Programs/Horde/HordeServer/Perforce/IPerforceService.cs
@@ -28,10 +28,15 @@ namespace HordeServer.Services
 	/// </summary>
 	public class PerforceUserInfo
 	{
+		/// <summary>
+		/// Login for the user
+		/// </summary>
+		public string? Login { get; set; }
+
 		/// <summary>
 		/// Full name of the user
 		/// </summary>
-		public string? Name { get; set; }
+		public string? FullName { get; set; }
 
 		/// <summary>
 		/// User's email address
@@ -58,6 +63,14 @@ namespace HordeServer.Services
 	/// </summary>
 	public interface IPerforceService
 	{
+		/// <summary>
+		/// Finds or adds a user from the given Perforce server, adding the user (and populating their profile with Perforce data) if they do not currently exist
+		/// </summary>
+		/// <param name="ClusterName"></param>
+		/// <param name="UserName"></param>
+		/// <returns></returns>
+		public ValueTask<IUser> FindOrAddUserAsync(string ClusterName, string UserName);
+
 		/// <summary>
 		/// 
 		/// </summary>
diff --git a/Engine/Source/Programs/Horde/HordeServer/Server/ConfigService.cs b/Engine/Source/Programs/Horde/HordeServer/Server/ConfigService.cs
--- a/Engine/Source/Programs/Horde/HordeServer/Server/ConfigService.cs
+++ b/Engine/Source/Programs/Horde/HordeServer/Server/ConfigService.cs
@@ -474,7 +474,7 @@ namespace HordeServer.Services
 
 			string FileName = ConfigPath.AbsolutePath;
 			int Change = -1;
-			string? Author = null;
+			IUser? Author = null;
 			string? Description = null;
 
 			if (ConfigPath.Scheme == PerforceScheme)
