Horde: Use user ids everywhere in job model.

diff --git a/Engine/Source/Programs/Horde/HordeServer/Api/Jobs.cs b/Engine/Source/Programs/Horde/HordeServer/Api/Jobs.cs
--- a/Engine/Source/Programs/Horde/HordeServer/Api/Jobs.cs
+++ b/Engine/Source/Programs/Horde/HordeServer/Api/Jobs.cs
@@ -355,9 +355,9 @@ namespace HordeServer.Api
 			this.TemplateHash = Job.TemplateHash?.ToString() ?? String.Empty;
 			this.GraphHash = Job.GraphHash.ToString();
 			this.StartedByUserId = Job.StartedByUserId?.ToString();
-			this.StartedByUser = Job.StartedByUser;
+			this.StartedByUser = StartedByUserInfo?.Login;
 			this.StartedByUserInfo = StartedByUserInfo;
-			this.AbortedByUser = Job.AbortedByUser;
+			this.AbortedByUser = AbortedByUserInfo?.Login;
 			this.AbortedByUserInfo = AbortedByUserInfo;
 			this.CreateTime = new DateTimeOffset(Job.CreateTimeUtc);
 			this.State = Job.GetState();
@@ -549,9 +549,9 @@ namespace HordeServer.Api
 			this.State = Step.State;
 			this.Outcome = Step.Outcome;
 			this.AbortRequested = Step.AbortRequested;
-			this.AbortByUser = Step.AbortByUser;
+			this.AbortByUser = AbortedByUserInfo?.Login;
 			this.AbortedByUserInfo = AbortedByUserInfo;
-			this.RetryByUser = Step.RetryByUser;
+			this.RetryByUser = RetriedByUserInfo?.Login;
 			this.RetriedByUserInfo = RetriedByUserInfo;
 			this.LogId = Step.LogId?.ToString();
 			this.ReadyTime = Step.ReadyTimeUtc;
diff --git a/Engine/Source/Programs/Horde/HordeServer/Api/Users.cs b/Engine/Source/Programs/Horde/HordeServer/Api/Users.cs
--- a/Engine/Source/Programs/Horde/HordeServer/Api/Users.cs
+++ b/Engine/Source/Programs/Horde/HordeServer/Api/Users.cs
@@ -123,6 +123,11 @@ namespace HordeServer.Api
 		/// </summary>
 		public string? Email { get; set; }
 
+		/// <summary>
+		/// The user login [DEPRECATED]
+		/// </summary>
+		internal string? Login { get; set; }
+
 		/// <summary>
 		/// Constructor
 		/// </summary>
@@ -134,12 +139,14 @@ namespace HordeServer.Api
 				this.Id = String.Empty;
 				this.Name = "(Unknown)";
 				this.Email = null;
+				this.Login = null;
 			}
 			else
 			{
 				this.Id = User.Id.ToString();
 				this.Name = User.Name;
 				this.Email = User.Email;
+				this.Login = User.Login;
 			}
 		}
 	}
diff --git a/Engine/Source/Programs/Horde/HordeServer/Controllers/JobsController.cs b/Engine/Source/Programs/Horde/HordeServer/Controllers/JobsController.cs
--- a/Engine/Source/Programs/Horde/HordeServer/Controllers/JobsController.cs
+++ b/Engine/Source/Programs/Horde/HordeServer/Controllers/JobsController.cs
@@ -158,7 +158,7 @@ namespace HordeServer.Controllers
 			List<string> Arguments = Create.Arguments ?? Template.GetDefaultArguments();
 
 			// Create the job
-			IJob Job = await JobService.CreateJobAsync(null, Stream, TemplateRefId, Template.Id, Graph, Name, Change, CodeChange, Create.PreflightChange, null, User.GetUserId(), User.GetUserName(), Priority, Create.AutoSubmit, Create.UpdateIssues, TemplateRef.ChainedJobs, TemplateRef.ShowUgsBadges, TemplateRef.ShowUgsAlerts, TemplateRef.NotificationChannel, TemplateRef.NotificationChannelFilter, null, Arguments);
+			IJob Job = await JobService.CreateJobAsync(null, Stream, TemplateRefId, Template.Id, Graph, Name, Change, CodeChange, Create.PreflightChange, null, User.GetUserId(), Priority, Create.AutoSubmit, Create.UpdateIssues, TemplateRef.ChainedJobs, TemplateRef.ShowUgsBadges, TemplateRef.ShowUgsAlerts, TemplateRef.NotificationChannel, TemplateRef.NotificationChannelFilter, Arguments);
 			await UpdateNotificationsAsync(Job.Id, new UpdateNotificationsRequest { Slack = true });
 			return new CreateJobResponse(Job.Id.ToString());
 		}
@@ -244,13 +244,13 @@ namespace HordeServer.Controllers
 				Request.Arguments = null;
 			}
 
-			string? AbortedByUser = null;
+			UserId? AbortedByUserId = null;
 			if (Request.Aborted ?? false)
 			{
-				AbortedByUser = User.Identity?.Name;
+				AbortedByUserId = User.GetUserId();
 			}
 
-			if (!await JobService.UpdateJobAsync(Job, Name: Request.Name, Priority: Request.Priority, AutoSubmit: Request.AutoSubmit, AbortedByUser: AbortedByUser, Arguments: Request.Arguments))
+			if (!await JobService.UpdateJobAsync(Job, Name: Request.Name, Priority: Request.Priority, AutoSubmit: Request.AutoSubmit, AbortedByUserId: AbortedByUserId, Arguments: Request.Arguments))
 			{
 				return NotFound();
 			}
@@ -389,9 +389,9 @@ namespace HordeServer.Controllers
 			}
 
 			GetThinUserInfoResponse? AbortedByUserInfo = null;
-			if (Job.AbortedByUser != null)
+			if (Job.AbortedByUserId != null)
 			{
-				AbortedByUserInfo = new GetThinUserInfoResponse(await UserCollection.FindOrAddUserByLoginAsync(Job.AbortedByUser));
+				AbortedByUserInfo = new GetThinUserInfoResponse(await UserCollection.GetCachedUserAsync(Job.AbortedByUserId.Value));
 			}
 
 			GetAclResponse? AclResponse = null;
@@ -443,15 +443,15 @@ namespace HordeServer.Controllers
 		async ValueTask<GetStepResponse> CreateStepResponseAsync(IJobStep Step)
 		{
 			GetThinUserInfoResponse? AbortedByUserInfo = null;
-			if (Step.AbortByUser != null)
+			if (Step.AbortedByUserId != null)
 			{
-				AbortedByUserInfo = new GetThinUserInfoResponse(await UserCollection.FindOrAddUserByLoginAsync(Step.AbortByUser));
+				AbortedByUserInfo = new GetThinUserInfoResponse(await UserCollection.GetCachedUserAsync(Step.AbortedByUserId.Value));
 			}
 
 			GetThinUserInfoResponse? RetriedByUserInfo = null;
-			if (Step.RetryByUser != null)
+			if (Step.RetriedByUserId != null)
 			{
-				RetriedByUserInfo = new GetThinUserInfoResponse(await UserCollection.FindOrAddUserByLoginAsync(Step.RetryByUser));
+				RetriedByUserInfo = new GetThinUserInfoResponse(await UserCollection.GetCachedUserAsync(Step.RetriedByUserId.Value));
 			}
 
 			return new GetStepResponse(Step, AbortedByUserInfo, RetriedByUserInfo);
diff --git a/Engine/Source/Programs/Horde/HordeServer/Controllers/StreamsController.cs b/Engine/Source/Programs/Horde/HordeServer/Controllers/StreamsController.cs
--- a/Engine/Source/Programs/Horde/HordeServer/Controllers/StreamsController.cs
+++ b/Engine/Source/Programs/Horde/HordeServer/Controllers/StreamsController.cs
@@ -28,52 +28,19 @@ namespace HordeServer.Controllers
 	[Route("[controller]")]
 	public class StreamsController : ControllerBase
 	{
-		/// <summary>
-		/// Singleton instance of the ACL service
-		/// </summary>
 		private readonly AclService AclService;
-
-		/// <summary>
-		/// Singleton instance of the stream service
-		/// </summary>
 		private readonly ProjectService ProjectService;
-
-		/// <summary>
-		/// Singleton instance of the stream service
-		/// </summary>
 		private readonly StreamService StreamService;
-
-		/// <summary>
-		/// Singleton instance of the template service
-		/// </summary>
 		private readonly ITemplateCollection TemplateCollection;
-
-		/// <summary>
-		/// Singleton instance of the job service
-		/// </summary>
 		private readonly JobService JobService;
-
-		/// <summary>
-		/// Collection of jobstep refs
-		/// </summary>
 		private readonly IJobStepRefCollection JobStepRefCollection;
-
-		/// <summary>
-		/// Singleton instance of the perforce service
-		/// </summary>
 		private readonly IPerforceService PerforceService;
+		private readonly IUserCollection UserCollection;
 
 		/// <summary>
 		/// Constructor
 		/// </summary>
-		/// <param name="AclService">The ACL service</param>
-		/// <param name="ProjectService">The project service</param>
-		/// <param name="StreamService">The stream service</param>
-		/// <param name="TemplateCollection">The template service</param>
-		/// <param name="JobService">The job service</param>
-		/// <param name="JobStepRefCollection">The jobstep ref collection</param>
-		/// <param name="PerforceService">The perforce service</param>
-		public StreamsController(AclService AclService, ProjectService ProjectService, StreamService StreamService, ITemplateCollection TemplateCollection, JobService JobService, IJobStepRefCollection JobStepRefCollection, IPerforceService PerforceService)
+		public StreamsController(AclService AclService, ProjectService ProjectService, StreamService StreamService, ITemplateCollection TemplateCollection, JobService JobService, IJobStepRefCollection JobStepRefCollection, IPerforceService PerforceService, IUserCollection UserCollection)
 		{
 			this.AclService = AclService;
 			this.ProjectService = ProjectService;
@@ -82,6 +49,7 @@ namespace HordeServer.Controllers
 			this.JobService = JobService;
 			this.JobStepRefCollection = JobStepRefCollection;
 			this.PerforceService = PerforceService;
+			this.UserCollection = UserCollection;
 		}
 
 		/// <summary>
@@ -201,7 +169,13 @@ namespace HordeServer.Controllers
 			}
 
 			List<ChangeSummary> Commits = await PerforceService.GetChangesAsync(Stream.ClusterName, Stream.Name, Min, Max, Results, PerforceUser);
-			return Commits.ConvertAll(x => PropertyFilter.Apply(new GetChangeSummaryResponse(x), Filter));
+
+			List<GetChangeSummaryResponse> Responses = new List<GetChangeSummaryResponse>();
+			foreach (ChangeSummary Commit in Commits)
+			{
+				Responses.Add(new GetChangeSummaryResponse(Commit));
+			}
+			return Responses.ConvertAll(x => PropertyFilter.Apply(x, Filter));
 		}
 
 		/// <summary>
diff --git a/Engine/Source/Programs/Horde/HordeServer/Jobs/IJob.cs b/Engine/Source/Programs/Horde/HordeServer/Jobs/IJob.cs
--- a/Engine/Source/Programs/Horde/HordeServer/Jobs/IJob.cs
+++ b/Engine/Source/Programs/Horde/HordeServer/Jobs/IJob.cs
@@ -122,7 +122,7 @@ namespace HordeServer.Models
 		/// <summary>
 		/// If a retry is requested, stores the name of the user that requested it
 		/// </summary>
-		public string? RetryByUser { get; }
+		public UserId? RetriedByUserId { get; }
 		
 		/// <summary>
 		/// Signal if a step should be aborted
@@ -130,9 +130,9 @@ namespace HordeServer.Models
 		public bool AbortRequested { get; }
 		
 		/// <summary>
-		/// If an abort is requested, stores the name of the user that requested it
+		/// If an abort is requested, stores the id of the user that requested it
 		/// </summary>
-		public string? AbortByUser { get; }
+		public UserId? AbortedByUserId { get; }
 
 		/// <summary>
 		/// List of reports for this step
@@ -528,14 +528,9 @@ namespace HordeServer.Models
 		public UserId? StartedByUserId { get; }
 
 		/// <summary>
-		/// Name of the user that started this job
+		/// Id of the user that aborted this job. Set to null if the job is not aborted.
 		/// </summary>
-		public string? StartedByUser { get; }
-
-		/// <summary>
-		/// Name of the user that aborted this job. Set to null if the job is not aborted.
-		/// </summary>
-		public string? AbortedByUser { get; }
+		public UserId? AbortedByUserId { get; }
 
 		/// <summary>
 		/// Name of the job.
@@ -637,11 +632,6 @@ namespace HordeServer.Models
 		/// </summary>
 		public string? NotificationChannelFilter { get; }
 		
-		/// <summary>
-		/// URL to a Helix Swarm server to call once job finishes. Used for adding build outcome to Swarm reviews.
-		/// </summary>
-		public string? HelixSwarmCallbackUrl { get; }
-
 		/// <summary>
 		/// Mapping of label ids to notification trigger ids for notifications
 		/// </summary>
diff --git a/Engine/Source/Programs/Horde/HordeServer/Jobs/IJobCollection.cs b/Engine/Source/Programs/Horde/HordeServer/Jobs/IJobCollection.cs
--- a/Engine/Source/Programs/Horde/HordeServer/Jobs/IJobCollection.cs
+++ b/Engine/Source/Programs/Horde/HordeServer/Jobs/IJobCollection.cs
@@ -47,7 +47,6 @@ namespace HordeServer.Collections
 		/// <param name="PreflightChange">Optional changelist to preflight</param>
 		/// <param name="ClonedPreflightChange">Optional cloned preflight changelist</param>
 		/// <param name="StartedByUserId">User that started the job</param>
-		/// <param name="StartedByUserName">User that started the job</param>
 		/// <param name="Priority">Priority of the job</param>
 		/// <param name="AutoSubmit">Whether to automatically submit the preflighted change on completion</param>
 		/// <param name="UpdateIssues">Whether to update issues based on the outcome of this job</param>
@@ -56,10 +55,9 @@ namespace HordeServer.Collections
 		/// <param name="ShowUgsAlerts">Whether to show alerts in UGS for this job</param>
 		/// <param name="NotificationChannel">Notification channel for this job</param>
 		/// <param name="NotificationChannelFilter">Notification channel filter for this job</param>
-		/// <param name="HelixSwarmCallbackUrl">Callback URL to a Helix Server review, if any</param>
 		/// <param name="Arguments">Arguments for the job</param>
 		/// <returns>The new job document</returns>
-		Task<IJob> AddAsync(JobId JobId, StreamId StreamId, TemplateRefId TemplateRefId, ContentHash TemplateHash, IGraph Graph, string Name, int Change, int CodeChange, int? PreflightChange, int? ClonedPreflightChange, UserId? StartedByUserId, string? StartedByUserName, Priority? Priority, bool? AutoSubmit, bool? UpdateIssues, List<ChainedJobTemplate>? JobTriggers, bool ShowUgsBadges, bool ShowUgsAlerts, string? NotificationChannel, string? NotificationChannelFilter, string? HelixSwarmCallbackUrl, List<string>? Arguments);
+		Task<IJob> AddAsync(JobId JobId, StreamId StreamId, TemplateRefId TemplateRefId, ContentHash TemplateHash, IGraph Graph, string Name, int Change, int CodeChange, int? PreflightChange, int? ClonedPreflightChange, UserId? StartedByUserId, Priority? Priority, bool? AutoSubmit, bool? UpdateIssues, List<ChainedJobTemplate>? JobTriggers, bool ShowUgsBadges, bool ShowUgsAlerts, string? NotificationChannel, string? NotificationChannelFilter, List<string>? Arguments);
 
 		/// <summary>
 		/// Gets a job with the given unique id
@@ -119,13 +117,13 @@ namespace HordeServer.Collections
 		/// <param name="AutoSubmit">Automatically submit the job on completion</param>
 		/// <param name="AutoSubmitChange">Changelist that was automatically submitted</param>
 		/// <param name="AutoSubmitMessage"></param>
-		/// <param name="AbortedByUser">Name of the user that aborted the job</param>
+		/// <param name="AbortedByUserId">Name of the user that aborted the job</param>
 		/// <param name="NotificationTriggerId">Id for a notification trigger</param>
 		/// <param name="Reports">New reports</param>
 		/// <param name="Arguments">New arguments for the job</param>
 		/// <param name="LabelIdxToTriggerId">New trigger ID for a label in the job</param>
 		/// <param name="JobTrigger">New downstream job id</param>
-		Task<bool> TryUpdateJobAsync(IJob Job, IGraph Graph, string? Name = null, Priority? Priority = null, bool? AutoSubmit = null, int? AutoSubmitChange = null, string? AutoSubmitMessage = null, string? AbortedByUser = null, ObjectId? NotificationTriggerId = null, List<Report>? Reports = null, List<string>? Arguments = null, KeyValuePair<int, ObjectId>? LabelIdxToTriggerId = null, KeyValuePair<TemplateRefId, JobId>? JobTrigger = null);
+		Task<bool> TryUpdateJobAsync(IJob Job, IGraph Graph, string? Name = null, Priority? Priority = null, bool? AutoSubmit = null, int? AutoSubmitChange = null, string? AutoSubmitMessage = null, UserId? AbortedByUserId = null, ObjectId? NotificationTriggerId = null, List<Report>? Reports = null, List<string>? Arguments = null, KeyValuePair<int, ObjectId>? LabelIdxToTriggerId = null, KeyValuePair<TemplateRefId, JobId>? JobTrigger = null);
 
 		/// <summary>
 		/// Updates the state of a batch
diff --git a/Engine/Source/Programs/Horde/HordeServer/Jobs/Impl/JobCollection.cs b/Engine/Source/Programs/Horde/HordeServer/Jobs/Impl/JobCollection.cs
--- a/Engine/Source/Programs/Horde/HordeServer/Jobs/Impl/JobCollection.cs
+++ b/Engine/Source/Programs/Horde/HordeServer/Jobs/Impl/JobCollection.cs
@@ -72,10 +72,12 @@ namespace HordeServer.Collections.Impl
 			[BsonIgnoreIfDefault, BsonDefaultValue(false)]
 			public bool Retry { get; set; }
 
+			public UserId? RetriedByUserId { get; set; }
 			public string? RetryByUser { get; set; }
 
 			public bool AbortRequested { get; set; } = false;
 
+			public UserId? AbortedByUserId { get; set; }
 			public string? AbortByUser { get; set; }
 
 			[BsonIgnoreIfNull]
@@ -85,7 +87,6 @@ namespace HordeServer.Collections.Impl
 			[BsonIgnoreIfNull]
 			public Dictionary<string, string>? Properties { get; set; }
 
-			string? IJobStep.RetryByUser => RetryByUser ?? (Retry ? (string?)"Unknown" : null);
 			DateTime? IJobStep.ReadyTimeUtc => ReadyTime?.UtcDateTime;
 			DateTime? IJobStep.StartTimeUtc => StartTime?.UtcDateTime;
 			DateTime? IJobStep.FinishTimeUtc => FinishTime?.UtcDateTime;
@@ -201,6 +202,9 @@ namespace HordeServer.Collections.Impl
 			[BsonIgnoreIfNull]
 			public string? StartedByUser { get; set; }
 
+			[BsonIgnoreIfNull]
+			public UserId? AbortedByUserId { get; set; }
+
 			[BsonIgnoreIfNull]
 			public string? AbortedByUser { get; set; }
 
@@ -240,7 +244,6 @@ namespace HordeServer.Collections.Impl
 			public bool ShowUgsAlerts { get; set; }
 			public string? NotificationChannel { get; set; }
 			public string? NotificationChannelFilter { get; set; }
-			public string? HelixSwarmCallbackUrl { get; set; }
 			public List<LabelNotificationDocument> LabelNotifications = new List<LabelNotificationDocument>();
 			public List<ChainedJobDocument> ChainedJobs { get; set; } = new List<ChainedJobDocument>();
 
@@ -275,7 +278,7 @@ namespace HordeServer.Collections.Impl
 				GraphHash = null!;
 			}
 
-			public JobDocument(JobId Id, StreamId StreamId, TemplateRefId TemplateId, ContentHash TemplateHash, ContentHash GraphHash, string Name, int Change, int CodeChange, int PreflightChange, int ClonedPreflightChange, UserId? StartedByUserId, string? StartedByUserName, Priority? Priority, bool? AutoSubmit, bool? UpdateIssues, DateTime CreateTimeUtc, List<ChainedJobDocument> ChainedJobs, bool ShowUgsBadges, bool ShowUgsAlerts, string? NotificationChannel, string? NotificationChannelFilter, string? HelixSwarmCallbackUrl, List<string>? Arguments)
+			public JobDocument(JobId Id, StreamId StreamId, TemplateRefId TemplateId, ContentHash TemplateHash, ContentHash GraphHash, string Name, int Change, int CodeChange, int PreflightChange, int ClonedPreflightChange, UserId? StartedByUserId, Priority? Priority, bool? AutoSubmit, bool? UpdateIssues, DateTime CreateTimeUtc, List<ChainedJobDocument> ChainedJobs, bool ShowUgsBadges, bool ShowUgsAlerts, string? NotificationChannel, string? NotificationChannelFilter, List<string>? Arguments)
 			{
 				this.Id = Id;
 				this.StreamId = StreamId;
@@ -288,7 +291,6 @@ namespace HordeServer.Collections.Impl
 				this.PreflightChange = PreflightChange;
 				this.ClonedPreflightChange = ClonedPreflightChange;
 				this.StartedByUserId = StartedByUserId;
-				this.StartedByUser = StartedByUserName;
 				this.Priority = Priority ?? HordeCommon.Priority.Normal;
 				this.AutoSubmit = AutoSubmit ?? false;
 				this.UpdateIssues = UpdateIssues ?? (StartedByUserId == null && PreflightChange == 0);
@@ -298,7 +300,6 @@ namespace HordeServer.Collections.Impl
 				this.ShowUgsAlerts = ShowUgsAlerts;
 				this.NotificationChannel = NotificationChannel;
 				this.NotificationChannelFilter = NotificationChannelFilter;
-				this.HelixSwarmCallbackUrl = HelixSwarmCallbackUrl;
 				this.Arguments = Arguments ?? this.Arguments;
 				this.NextSubResourceId = SubResourceId.Random();
 				this.UpdateTimeUtc = CreateTimeUtc;
@@ -367,7 +368,7 @@ namespace HordeServer.Collections.Impl
 
 		/// <inheritdoc/>
 		[SuppressMessage("Compiler", "CA1054:URI parameters should not be strings")]
-		public async Task<IJob> AddAsync(JobId JobId, StreamId StreamId, TemplateRefId TemplateRefId, ContentHash TemplateHash, IGraph Graph, string Name, int Change, int CodeChange, int? PreflightChange, int? ClonedPreflightChange, UserId? StartedByUserId, string? StartedByUserName, Priority? Priority, bool? AutoSubmit, bool? UpdateIssues, List<ChainedJobTemplate>? ChainedJobs, bool ShowUgsBadges, bool ShowUgsAlerts, string? NotificationChannel, string? NotificationChannelFilter, string? HelixSwarmCallbackUrl, List<string>? Arguments)
+		public async Task<IJob> AddAsync(JobId JobId, StreamId StreamId, TemplateRefId TemplateRefId, ContentHash TemplateHash, IGraph Graph, string Name, int Change, int CodeChange, int? PreflightChange, int? ClonedPreflightChange, UserId? StartedByUserId, Priority? Priority, bool? AutoSubmit, bool? UpdateIssues, List<ChainedJobTemplate>? ChainedJobs, bool ShowUgsBadges, bool ShowUgsAlerts, string? NotificationChannel, string? NotificationChannelFilter, List<string>? Arguments)
 		{
 			List<ChainedJobDocument> JobTriggers = new List<ChainedJobDocument>();
 			if (ChainedJobs == null)
@@ -379,7 +380,7 @@ namespace HordeServer.Collections.Impl
 				JobTriggers = ChainedJobs.ConvertAll(x => new ChainedJobDocument(x));
 			}
 
-			JobDocument NewJob = new JobDocument(JobId, StreamId, TemplateRefId, TemplateHash, Graph.Id, Name, Change, CodeChange, PreflightChange ?? 0, ClonedPreflightChange ?? 0, StartedByUserId, StartedByUserName, Priority, AutoSubmit, UpdateIssues, DateTime.UtcNow, JobTriggers, ShowUgsBadges, ShowUgsAlerts, NotificationChannel, NotificationChannelFilter, HelixSwarmCallbackUrl, Arguments);
+			JobDocument NewJob = new JobDocument(JobId, StreamId, TemplateRefId, TemplateHash, Graph.Id, Name, Change, CodeChange, PreflightChange ?? 0, ClonedPreflightChange ?? 0, StartedByUserId, Priority, AutoSubmit, UpdateIssues, DateTime.UtcNow, JobTriggers, ShowUgsBadges, ShowUgsAlerts, NotificationChannel, NotificationChannelFilter, Arguments);
 			CreateBatches(NewJob, Graph, Logger);
 
 			await Jobs.InsertOneAsync(NewJob);
@@ -498,7 +499,7 @@ namespace HordeServer.Collections.Impl
 		}
 
 		/// <inheritdoc/>
-		public async Task<bool> TryUpdateJobAsync(IJob Job, IGraph Graph, string? Name, Priority? Priority, bool? AutoSubmit, int? AutoSubmitChange, string? AutoSubmitMessage, string? AbortedByUser, ObjectId? NotificationTriggerId, List<Report>? Reports, List<string>? Arguments, KeyValuePair<int, ObjectId>? LabelIdxToTriggerId, KeyValuePair<TemplateRefId, JobId>? JobTrigger)
+		public async Task<bool> TryUpdateJobAsync(IJob Job, IGraph Graph, string? Name, Priority? Priority, bool? AutoSubmit, int? AutoSubmitChange, string? AutoSubmitMessage, UserId? AbortedByUserId, ObjectId? NotificationTriggerId, List<Report>? Reports, List<string>? Arguments, KeyValuePair<int, ObjectId>? LabelIdxToTriggerId, KeyValuePair<TemplateRefId, JobId>? JobTrigger)
 		{
 			// Create the update 
 			UpdateDefinitionBuilder<JobDocument> UpdateBuilder = Builders<JobDocument>.Update;
@@ -534,10 +535,10 @@ namespace HordeServer.Collections.Impl
 				JobDocument.AutoSubmitMessage = (AutoSubmitMessage.Length == 0)? null : AutoSubmitMessage;
 				Updates.Add(UpdateBuilder.SetOrUnsetNullRef(x => x.AutoSubmitMessage, Job.AutoSubmitMessage));
 			}
-			if (AbortedByUser != null && Job.AbortedByUser == null)
+			if (AbortedByUserId != null && Job.AbortedByUserId == null)
 			{
-				JobDocument.AbortedByUser = AbortedByUser;
-				Updates.Add(UpdateBuilder.Set(x => x.AbortedByUser, Job.AbortedByUser));
+				JobDocument.AbortedByUserId = AbortedByUserId;
+				Updates.Add(UpdateBuilder.Set(x => x.AbortedByUserId, Job.AbortedByUserId));
 				bUpdateBatches = true;
 			}
 			if (NotificationTriggerId != null)
@@ -1293,7 +1294,7 @@ namespace HordeServer.Collections.Impl
 			{
 				foreach (IJobStep Step in Batch.Steps)
 				{
-					if ((Step.State == JobStepState.Running && Step.RetryByUser == null) || (Step.State == JobStepState.Completed && Step.RetryByUser == null) || (Step.State == JobStepState.Aborted && Step.RetryByUser == null) || Step.State == JobStepState.Skipped)
+					if ((Step.State == JobStepState.Running && Step.RetriedByUserId == null) || (Step.State == JobStepState.Completed && Step.RetriedByUserId == null) || (Step.State == JobStepState.Aborted && Step.RetriedByUserId == null) || Step.State == JobStepState.Skipped)
 					{
 						NewNodesToExecute.Remove(Graph.Groups[Batch.GroupIdx].Nodes[Step.NodeIdx]);
 					}
diff --git a/Engine/Source/Programs/Horde/HordeServer/Jobs/Impl/JobService.cs b/Engine/Source/Programs/Horde/HordeServer/Jobs/Impl/JobService.cs
--- a/Engine/Source/Programs/Horde/HordeServer/Jobs/Impl/JobService.cs
+++ b/Engine/Source/Programs/Horde/HordeServer/Jobs/Impl/JobService.cs
@@ -221,7 +221,6 @@ namespace HordeServer.Services
 		/// <param name="PreflightChange">Optional changelist to preflight</param>
 		/// <param name="ClonedPreflightChange">Duplicated preflight change</param>
 		/// <param name="StartedByUserId">Id of the user that started the job</param>
-		/// <param name="StartedByUserName">User that started the job</param>
 		/// <param name="Priority">Priority of the job</param>
 		/// <param name="AutoSubmit">Whether to automatically submit the preflighted change on completion</param>
 		/// <param name="UpdateIssues">Whether to update issues when this job completes</param>
@@ -230,10 +229,9 @@ namespace HordeServer.Services
 		/// <param name="ShowUgsAlerts">Whether to show alerts in UGS for this job</param>
 		/// <param name="NotificationChannel">Notification Channel for this job</param>
 		/// <param name="NotificationChannelFilter">Notification Channel filter for this job</param>
-		/// <param name="HelixSwarmCallbackUrl">Helix Swarm callback URL for review, if any</param>
 		/// <param name="Arguments">Arguments for the job</param>
 		/// <returns>Unique id representing the job</returns>
-		public async Task<IJob> CreateJobAsync(JobId? JobId, IStream Stream, TemplateRefId TemplateRefId, ContentHash TemplateHash, IGraph Graph, string Name, int Change, int CodeChange, int? PreflightChange, int? ClonedPreflightChange, UserId? StartedByUserId, string? StartedByUserName, Priority? Priority, bool? AutoSubmit, bool? UpdateIssues, List<ChainedJobTemplate>? JobTriggers, bool ShowUgsBadges, bool ShowUgsAlerts, string? NotificationChannel, string? NotificationChannelFilter, string? HelixSwarmCallbackUrl, IReadOnlyList<string> Arguments)
+		public async Task<IJob> CreateJobAsync(JobId? JobId, IStream Stream, TemplateRefId TemplateRefId, ContentHash TemplateHash, IGraph Graph, string Name, int Change, int CodeChange, int? PreflightChange, int? ClonedPreflightChange, UserId? StartedByUserId, Priority? Priority, bool? AutoSubmit, bool? UpdateIssues, List<ChainedJobTemplate>? JobTriggers, bool ShowUgsBadges, bool ShowUgsAlerts, string? NotificationChannel, string? NotificationChannelFilter, IReadOnlyList<string> Arguments)
 		{
 			JobId JobIdValue = JobId ?? HordeServer.Utilities.ObjectId<IJob>.GenerateNewId();
 			using IDisposable Scope = Logger.BeginScope("CreateJobAsync({JobId})", JobIdValue);
@@ -266,7 +264,7 @@ namespace HordeServer.Services
 
 			Name = StringUtils.ExpandProperties(Name, Properties);
 
-			IJob NewJob = await Jobs.AddAsync(JobIdValue, Stream.Id, TemplateRefId, TemplateHash, Graph, Name, Change, CodeChange, PreflightChange, ClonedPreflightChange, StartedByUserId, StartedByUserName, Priority, AutoSubmit, UpdateIssues, JobTriggers, ShowUgsBadges, ShowUgsAlerts, NotificationChannel, NotificationChannelFilter, HelixSwarmCallbackUrl, ExpandedArguments);
+			IJob NewJob = await Jobs.AddAsync(JobIdValue, Stream.Id, TemplateRefId, TemplateHash, Graph, Name, Change, CodeChange, PreflightChange, ClonedPreflightChange, StartedByUserId, Priority, AutoSubmit, UpdateIssues, JobTriggers, ShowUgsBadges, ShowUgsAlerts, NotificationChannel, NotificationChannelFilter, ExpandedArguments);
 			JobTaskSource.UpdateQueuedJob(NewJob, Graph);
 
 			await JobTaskSource.UpdateUgsBadges(NewJob, Graph, new List<(LabelState, LabelOutcome)>());
@@ -295,10 +293,10 @@ namespace HordeServer.Services
 				if (Job.Id == NewJob.Id) continue; // Don't remove the new job
 				if (Job.TemplateId != NewJob.TemplateId) continue;
 				if (Job.TemplateHash != NewJob.TemplateHash) continue;
-				if (Job.StartedByUser != NewJob.StartedByUser) continue;
+				if (Job.StartedByUserId != NewJob.StartedByUserId) continue;
 				if (string.Join(",", Job.Arguments) != string.Join(",", NewJob.Arguments)) continue;
 
-				bool bUpdateSuccessful = await UpdateJobAsync(Job, null, null, null, "horde.duplicated.by.newer.CL", null, null, null);
+				bool bUpdateSuccessful = await UpdateJobAsync(Job, null, null, null, KnownUsers.System, null, null, null);
 				if (!bUpdateSuccessful)
 				{
 					Logger.LogError("Failed marking duplicate job as aborted! Job ID: {JobId}", Job.Id);
@@ -360,12 +358,12 @@ namespace HordeServer.Services
 		/// <param name="Name">Name of the job</param>
 		/// <param name="Priority">Priority of the job</param>
 		/// <param name="AutoSubmit">Whether to automatically submit the preflighted change on completion</param>
-		/// <param name="AbortedByUser">Name of the user that aborted this job</param>
+		/// <param name="AbortedByUserId">Name of the user that aborted this job</param>
 		/// <param name="OnCompleteTriggerId">Object id for a notification trigger</param>
 		/// <param name="Reports">New reports to add</param>
 		/// <param name="Arguments">New arguments for the job</param>
 		/// <param name="LabelIdxToTriggerId">New trigger ID for a label in the job</param>
-		public async Task<bool> UpdateJobAsync(IJob Job, string? Name = null, Priority? Priority = null, bool? AutoSubmit = null, string? AbortedByUser = null, ObjectId? OnCompleteTriggerId = null, List<Report>? Reports = null, List<string>? Arguments = null, KeyValuePair<int, ObjectId>? LabelIdxToTriggerId = null)
+		public async Task<bool> UpdateJobAsync(IJob Job, string? Name = null, Priority? Priority = null, bool? AutoSubmit = null, UserId? AbortedByUserId = null, ObjectId? OnCompleteTriggerId = null, List<Report>? Reports = null, List<string>? Arguments = null, KeyValuePair<int, ObjectId>? LabelIdxToTriggerId = null)
 		{
 			using IDisposable Scope = Logger.BeginScope("UpdateJobAsync({JobId})", Job.Id);
 			for (; ; )
@@ -376,7 +374,7 @@ namespace HordeServer.Services
 				IReadOnlyList<(LabelState, LabelOutcome)> OldLabelStates = Job.GetLabelStates(Graph);
 
 				// Update the new list of job steps
-				if (await Jobs.TryUpdateJobAsync(Job, Graph, Name, Priority, AutoSubmit, null, null, AbortedByUser, OnCompleteTriggerId, Reports, Arguments, LabelIdxToTriggerId))
+				if (await Jobs.TryUpdateJobAsync(Job, Graph, Name, Priority, AutoSubmit, null, null, AbortedByUserId, OnCompleteTriggerId, Reports, Arguments, LabelIdxToTriggerId))
 				{
 					// Update any badges that have been modified
 					await JobTaskSource.UpdateUgsBadges(Job, Graph, OldLabelStates);
@@ -880,7 +878,7 @@ namespace HordeServer.Services
 						(JobStepState State, JobStepOutcome Outcome) = Job.GetTargetState();
 						if (State == JobStepState.Completed)
 						{
-							if (Job.AutoSubmit && Outcome == JobStepOutcome.Success && Job.AbortedByUser == null)
+							if (Job.AutoSubmit && Outcome == JobStepOutcome.Success && Job.AbortedByUserId == null)
 							{
 								Job = await AutoSubmitChangeAsync(Job, Graph);
 							}
@@ -1125,7 +1123,7 @@ namespace HordeServer.Services
 					IGraph TriggerGraph = await Graphs.AddAsync(Template);
 					Logger.LogInformation("Creating downstream job {ChainedJobId} from job {JobId}", ChainedJobId, Job.Id);
 
-					await CreateJobAsync(ChainedJobId, Stream, JobTrigger.TemplateRefId, TemplateRef.Hash, TriggerGraph, TemplateRef.Name, Job.Change, Job.CodeChange, Job.PreflightChange, Job.ClonedPreflightChange, Job.StartedByUserId, Job.StartedByUser, Template.Priority, null, Job.UpdateIssues, TemplateRef.ChainedJobs, false, false, TemplateRef.NotificationChannel, TemplateRef.NotificationChannelFilter, null, Template.Arguments);
+					await CreateJobAsync(ChainedJobId, Stream, JobTrigger.TemplateRefId, TemplateRef.Hash, TriggerGraph, TemplateRef.Name, Job.Change, Job.CodeChange, Job.PreflightChange, Job.ClonedPreflightChange, Job.StartedByUserId, Template.Priority, null, Job.UpdateIssues, TemplateRef.ChainedJobs, false, false, TemplateRef.NotificationChannel, TemplateRef.NotificationChannelFilter, Template.Arguments);
 					break;
 				}
 
diff --git a/Engine/Source/Programs/Horde/HordeServer/Jobs/Impl/JobTaskSource.cs b/Engine/Source/Programs/Horde/HordeServer/Jobs/Impl/JobTaskSource.cs
--- a/Engine/Source/Programs/Horde/HordeServer/Jobs/Impl/JobTaskSource.cs
+++ b/Engine/Source/Programs/Horde/HordeServer/Jobs/Impl/JobTaskSource.cs
@@ -438,7 +438,7 @@ namespace HordeServer.Tasks.Impl
 					await Jobs.TryRemoveFromDispatchQueueAsync(NewJob);
 					continue;
 				}
-				if (NewJob.AbortedByUser != null)
+				if (NewJob.AbortedByUserId != null)
 				{
 					Logger.LogError("Job {JobId} was aborted but not removed from dispatch queue", NewJob.Id);
 					await Jobs.TryRemoveFromDispatchQueueAsync(NewJob);
diff --git a/Engine/Source/Programs/Horde/HordeServer/Perforce/Change.cs b/Engine/Source/Programs/Horde/HordeServer/Perforce/Change.cs
--- a/Engine/Source/Programs/Horde/HordeServer/Perforce/Change.cs
+++ b/Engine/Source/Programs/Horde/HordeServer/Perforce/Change.cs
@@ -16,6 +16,8 @@ using System.Threading.Tasks;
 
 namespace HordeServer.Models
 {
+	using UserId = ObjectId<IUser>;
+
 	/// <summary>
 	/// Summary information for a change
 	/// </summary>
@@ -30,7 +32,7 @@ namespace HordeServer.Models
 		/// <summary>
 		/// Author of the change
 		/// </summary>
-		public string Author { get; set; }
+		public IUser Author { get; set; }
 
 		/// <summary>
 		/// The base path for modified files
@@ -49,7 +51,7 @@ namespace HordeServer.Models
 		/// <param name="Author">Author of the change</param>
 		/// <param name="Path">Base path for modified files</param>
 		/// <param name="Description">Changelist description</param>
-		public ChangeSummary(int Number, string Author, string Path, string Description)
+		public ChangeSummary(int Number, IUser Author, string Path, string Description)
 		{
 			this.Number = Number;
 			this.Author = Author;
@@ -167,7 +169,7 @@ namespace HordeServer.Models
 		/// <summary>
 		/// Name of the user that authored this change
 		/// </summary>
-		public string Author { get; set; }
+		public IUser Author { get; set; }
 
 		/// <summary>
 		/// The base path for modified files
@@ -189,17 +191,6 @@ namespace HordeServer.Models
 		/// </summary>
 		public DateTime Date { get; set; }
 
-		/// <summary>
-		/// Private constructor for serialization
-		/// </summary>
-		private ChangeDetails()
-		{
-			Author = null!;
-			Path = null!;
-			Description = null!;
-			Files = null!;
-		}
-
 		/// <summary>
 		/// Constructor
 		/// </summary>
@@ -209,7 +200,7 @@ namespace HordeServer.Models
 		/// <param name="Description">Changelist description</param>
 		/// <param name="Files">List of files modified, relative to the stream base</param>
 		/// <param name="Date">Date that the change was submitted</param>
-		public ChangeDetails(int Number, string Author, string Path, string Description, List<ChangeFile> Files, DateTime Date)
+		public ChangeDetails(int Number, IUser Author, string Path, string Description, List<ChangeFile> Files, DateTime Date)
 		{
 			this.Number = Number;
 			this.Author = Author;
diff --git a/Engine/Source/Programs/Horde/HordeServer/Perforce/Impl/PerforceService.cs b/Engine/Source/Programs/Horde/HordeServer/Perforce/Impl/PerforceService.cs
--- a/Engine/Source/Programs/Horde/HordeServer/Perforce/Impl/PerforceService.cs
+++ b/Engine/Source/Programs/Horde/HordeServer/Perforce/Impl/PerforceService.cs
@@ -2,8 +2,10 @@
 
 using EpicGames.Core;
 using EpicGames.Perforce;
+using HordeServer.Collections;
 using HordeServer.Models;
 using HordeServer.Utilities;
+using Microsoft.Extensions.Caching.Memory;
 using Microsoft.Extensions.Logging;
 using Microsoft.Extensions.Options;
 using System;
@@ -22,6 +24,7 @@ using System.Threading.Tasks;
 namespace HordeServer.Services
 {
 	using P4 = Perforce.P4;
+	using UserId = ObjectId<IUser>;
 
 	static class PerforceExtensions
 	{
@@ -39,7 +42,7 @@ namespace HordeServer.Services
 	/// <summary>
 	/// P4API implementation of the Perforce service
 	/// </summary>
-	class PerforceService : IPerforceService
+	class PerforceService : IPerforceService, IDisposable
 	{
 		class CachedTicketInfo
 		{
@@ -81,13 +84,17 @@ namespace HordeServer.Services
 
 		Dictionary<string, Dictionary<string, CachedTicketInfo>> ClusterTickets = new Dictionary<string, Dictionary<string, CachedTicketInfo>>(StringComparer.OrdinalIgnoreCase);
 
+		IUserCollection UserCollection;
+		MemoryCache UserCache = new MemoryCache(new MemoryCacheOptions { SizeLimit = 2000 });
+
 		/// <summary>
 		/// Constructor
 		/// </summary>
-		public PerforceService(PerforceLoadBalancer LoadBalancer, DatabaseService DatabaseService, IOptions<ServerSettings> Settings, ILogger<PerforceService> Logger)
+		public PerforceService(PerforceLoadBalancer LoadBalancer, DatabaseService DatabaseService, IUserCollection UserCollection, IOptions<ServerSettings> Settings, ILogger<PerforceService> Logger)
 		{
 			this.LoadBalancer = LoadBalancer;
 			this.CachedGlobals = new LazyCachedValue<Task<Globals>>(() => DatabaseService.GetGlobalsAsync(), TimeSpan.FromSeconds(30.0));
+			this.UserCollection = UserCollection;
 			this.Settings = Settings.Value;
 			this.Logger = Logger;
 
@@ -97,6 +104,33 @@ namespace HordeServer.Services
 			P4.LogFile.SetLoggingFunction(LogPerforce);
 		}
 
+		public void Dispose()
+		{
+			UserCache.Dispose();
+		}
+
+		public async ValueTask<IUser> FindOrAddUserAsync(string ClusterName, string UserName)
+		{
+			IUser? User;
+			if (!UserCache.TryGetValue((ClusterName, UserName), out User))
+			{
+				User = await UserCollection.FindUserByLoginAsync(UserName);
+				if (User == null)
+				{
+					PerforceUserInfo? UserInfo = await GetUserInfoAsync(ClusterName, UserName);
+					User = await UserCollection.FindOrAddUserByLoginAsync(UserName, UserInfo?.FullName, UserInfo?.Email);
+				}
+
+				using (ICacheEntry Entry = UserCache.CreateEntry((ClusterName, UserName)))
+				{
+					Entry.SetValue(User);
+					Entry.SetSize(1);
+					Entry.SetAbsoluteExpiration(TimeSpan.FromDays(1.0));
+				}
+			}
+			return User!;
+		}
+
 		async Task<PerforceCluster> GetClusterAsync(string ClusterName)
 		{
 			Globals Globals = await CachedGlobals.GetCached();
@@ -692,7 +726,8 @@ namespace HordeServer.Services
 				{
 					foreach (P4.Changelist Changelist in Changelists)
 					{
-						Changes.Add(new ChangeSummary(Changelist.Id, Changelist.OwnerName, Changelist.GetPath(), Changelist.Description));
+						IUser User = await FindOrAddUserAsync(ClusterName, Changelist.OwnerName);
+						Changes.Add(new ChangeSummary(Changelist.Id, User, Changelist.GetPath(), Changelist.Description));
 					}
 				}
 				return Changes;
@@ -720,7 +755,8 @@ namespace HordeServer.Services
 				{
 					foreach (P4.Changelist Changelist in Changelists)
 					{
-						Changes.Add(new ChangeSummary(Changelist.Id, Changelist.OwnerName, Changelist.GetPath(), Changelist.Description));
+						IUser User = await FindOrAddUserAsync(ClusterName, Changelist.OwnerName);
+						Changes.Add(new ChangeSummary(Changelist.Id, User, Changelist.GetPath(), Changelist.Description));
 					}
 				}
 
@@ -775,7 +811,8 @@ namespace HordeServer.Services
 							}
 						}
 
-						return new ChangeDetails(Change.Id, Change.OwnerName, Change.GetPath(), Change.Description, Files, Change.ModifiedDate);
+						IUser User = await FindOrAddUserAsync(ClusterName, Change.OwnerName);
+						return new ChangeDetails(Change.Id, User, Change.GetPath(), Change.Description, Files, Change.ModifiedDate);
 					}
 				}
 			}
@@ -847,7 +884,8 @@ namespace HordeServer.Services
 
 						}
 
-						Results.Add(new ChangeDetails(Change.Id, Change.OwnerName, Change.GetPath(), Change.Description, Files, Change.ModifiedDate));
+						IUser User = await FindOrAddUserAsync(ClusterName, Change.OwnerName);
+						Results.Add(new ChangeDetails(Change.Id, User, Change.GetPath(), Change.Description, Files, Change.ModifiedDate));
 					}
 
 				}
@@ -1322,7 +1360,7 @@ namespace HordeServer.Services
 					return null;
 				}
 
-				return new PerforceUserInfo { Name = UserName, Email = User.EmailAddress };
+				return new PerforceUserInfo { Login = UserName, FullName = User.FullName, Email = User.EmailAddress };
 			}
 		}
 
diff --git a/Engine/Source/Programs/Horde/HordeServer/Projects/ScheduleService.cs b/Engine/Source/Programs/Horde/HordeServer/Projects/ScheduleService.cs
--- a/Engine/Source/Programs/Horde/HordeServer/Projects/ScheduleService.cs
+++ b/Engine/Source/Programs/Horde/HordeServer/Projects/ScheduleService.cs
@@ -362,7 +362,7 @@ namespace HordeServer.Services
 			foreach ((int Change, int CodeChange) in TriggerChanges.OrderBy(x => x.Change))
 			{
 				List<string> DefaultArguments = Template.GetDefaultArguments();
-				IJob NewJob = await JobService.CreateJobAsync(null, Stream, TemplateRefId, Template.Id, Graph, Template.Name, Change, CodeChange, null, null, null, null, Template.Priority, null, null, TemplateRef.ChainedJobs, TemplateRef.ShowUgsBadges, TemplateRef.ShowUgsAlerts, TemplateRef.NotificationChannel, TemplateRef.NotificationChannelFilter, null, DefaultArguments);
+				IJob NewJob = await JobService.CreateJobAsync(null, Stream, TemplateRefId, Template.Id, Graph, Template.Name, Change, CodeChange, null, null, null, Template.Priority, null, null, TemplateRef.ChainedJobs, TemplateRef.ShowUgsBadges, TemplateRef.ShowUgsAlerts, TemplateRef.NotificationChannel, TemplateRef.NotificationChannelFilter, DefaultArguments);
 				Logger.LogInformation("Started new job for {StreamName} template {TemplateName} at CL {Change} (Code CL {CodeChange}): {JobId}", Stream.Id, TemplateRef.Name, Change, CodeChange, NewJob.Id);
 				await StreamService.UpdateScheduleTriggerAsync(Stream, TemplateRefId, Now, Change, new List<JobId> { NewJob.Id }, new List<JobId>());
 			}
diff --git a/Engine/Source/Programs/Horde/HordeServer/Users/IUser.cs b/Engine/Source/Programs/Horde/HordeServer/Users/IUser.cs
--- a/Engine/Source/Programs/Horde/HordeServer/Users/IUser.cs
+++ b/Engine/Source/Programs/Horde/HordeServer/Users/IUser.cs
@@ -13,6 +13,17 @@ namespace HordeServer.Models
 	using JobId = ObjectId<IJob>;
 	using UserId = ObjectId<IUser>;
 
+	/// <summary>
+	/// Known user ids
+	/// </summary>
+	public static class KnownUsers
+	{
+		/// <summary>
+		/// The system user. Used for automated processes.
+		/// </summary>
+		public static UserId System { get; } = new UserId("6170b423b94a2c7c2d6b6f87");
+	}
+
 	/// <summary>
 	/// Document which collates information about a user, and their personal settings
 	/// </summary>
