Horde: Fix tests.

diff --git a/Engine/Source/Programs/Horde/HordeServer.Compute.Tests/Fixture.cs b/Engine/Source/Programs/Horde/HordeServer.Compute.Tests/Fixture.cs
--- a/Engine/Source/Programs/Horde/HordeServer.Compute.Tests/Fixture.cs
+++ b/Engine/Source/Programs/Horde/HordeServer.Compute.Tests/Fixture.cs
@@ -40,9 +40,9 @@ namespace HordeServerTests
 		{
 			Fixture _fixture = new Fixture();
 			await _fixture.Populate(GraphCollection, TemplateCollection, JobService, ArtifactCollection, StreamService, AgentService, PerforceService);
-			
-			(PerforceService as PerforceServiceStub)?.AddChange("//UE5/Main", 112233, "leet.coder", "Did stuff", new []{"file.cpp"});
-			(PerforceService as PerforceServiceStub)?.AddChange("//UE5/Main", 1111, "swarm", "A shelved CL here", new []{"renderer.cpp"});
+
+//			(PerforceService as PerforceServiceStub)?.AddChange("//UE5/Main", 112233, "leet.coder", "Did stuff", new []{"file.cpp"});
+//			(PerforceService as PerforceServiceStub)?.AddChange("//UE5/Main", 1111, "swarm", "A shelved CL here", new []{"renderer.cpp"});
 		
 			return _fixture;
 		}
@@ -92,7 +92,6 @@ namespace HordeServerTests
 				PreflightChange: 1001,
 				ClonedPreflightChange: null,
 				StartedByUserId: null,
-				StartedByUserName: "SomeUser",
 				Priority: Priority.Normal,
 				null,
 				null,
@@ -101,7 +100,6 @@ namespace HordeServerTests
 				false,
 				null,
 				null,
-				null,
 				Arguments: new List<string>()
 			);
 			Job1 = (await JobService.GetJobAsync(Job1.Id))!;
@@ -118,7 +116,6 @@ namespace HordeServerTests
 				PreflightChange: null,
 				ClonedPreflightChange: null,
 				StartedByUserId: null,
-				StartedByUserName: "SomeUser",
 				Priority: Priority.Normal,
 				null,
 				null,
@@ -127,7 +124,6 @@ namespace HordeServerTests
 				false,
 				null,
 				null,
-				null,
 				Arguments: new List<string>()
 			);
 			Job2 = (await JobService.GetJobAsync(Job2.Id))!;
diff --git a/Engine/Source/Programs/Horde/HordeServer.Compute.Tests/IssueServiceTests.cs b/Engine/Source/Programs/Horde/HordeServer.Compute.Tests/IssueServiceTests.cs
--- a/Engine/Source/Programs/Horde/HordeServer.Compute.Tests/IssueServiceTests.cs
+++ b/Engine/Source/Programs/Horde/HordeServer.Compute.Tests/IssueServiceTests.cs
@@ -135,13 +135,23 @@ namespace HordeServerTests
 			IStream ReleaseStream = TestSetup.StreamCollection.TryCreateOrReplaceAsync(ReleaseStreamId, null, "", "", Project.Id, new StreamConfig { Name = ReleaseStreamName }).Result!;
 			IStream DevStream = TestSetup.StreamCollection.TryCreateOrReplaceAsync(DevStreamId, null, "", "", Project.Id, new StreamConfig { Name = DevStreamName }).Result!;
 
+			IUser Bill = TestSetup.UserCollection.FindOrAddUserByLoginAsync("Bill").Result;
+			IUser Anne = TestSetup.UserCollection.FindOrAddUserByLoginAsync("Anne").Result;
+			IUser Bob = TestSetup.UserCollection.FindOrAddUserByLoginAsync("Bob").Result;
+			IUser Jerry = TestSetup.UserCollection.FindOrAddUserByLoginAsync("Jerry").Result;
+			IUser Chris = TestSetup.UserCollection.FindOrAddUserByLoginAsync("Chris").Result;
+
+			TimId = TestSetup.UserCollection.FindOrAddUserByLoginAsync("Tim").Result.Id;
+			JerryId = TestSetup.UserCollection.FindOrAddUserByLoginAsync("Jerry").Result.Id;
+			BobId = TestSetup.UserCollection.FindOrAddUserByLoginAsync("Bob").Result.Id;
+
 			Perforce = TestSetup.PerforceService;
-			Perforce.AddChange(MainStreamName, 100, "Bill", "Description", new string[] { "a/b.cpp" });
-			Perforce.AddChange(MainStreamName, 105, "Anne", "Description", new string[] { "a/c.cpp" });
-			Perforce.AddChange(MainStreamName, 110, "Bob", "Description", new string[] { "a/d.cpp" });
-			Perforce.AddChange(MainStreamName, 115, "Jerry", "Description\n#ROBOMERGE-SOURCE: CL 75 in //UE4/Release/...", new string[] { "a/e.cpp", "a/foo.cpp" });
-			Perforce.AddChange(MainStreamName, 120, "Chris", "Description\n#ROBOMERGE-OWNER: Tim", new string[] { "a/f.cpp" });
-			Perforce.AddChange(MainStreamName, 125, "Chris", "Description", new string[] { "a/g.cpp" });
+			Perforce.AddChange(MainStreamName, 100, Bill, "Description", new string[] { "a/b.cpp" });
+			Perforce.AddChange(MainStreamName, 105, Anne, "Description", new string[] { "a/c.cpp" });
+			Perforce.AddChange(MainStreamName, 110, Bob, "Description", new string[] { "a/d.cpp" });
+			Perforce.AddChange(MainStreamName, 115, Jerry, "Description\n#ROBOMERGE-SOURCE: CL 75 in //UE4/Release/...", new string[] { "a/e.cpp", "a/foo.cpp" });
+			Perforce.AddChange(MainStreamName, 120, Chris, "Description\n#ROBOMERGE-OWNER: Tim", new string[] { "a/f.cpp" });
+			Perforce.AddChange(MainStreamName, 125, Chris, "Description", new string[] { "a/g.cpp" });
 
 			List<INode> Nodes = new List<INode>();
 			Nodes.Add(MockNode("Update Version Files"));
@@ -155,10 +165,6 @@ namespace HordeServerTests
 			Mock<IGraph> GraphMock = new Mock<IGraph>(MockBehavior.Strict);
 			GraphMock.SetupGet(x => x.Groups).Returns(new List<INodeGroup> { Group.Object });
 			Graph = GraphMock.Object;
-
-			TimId = TestSetup.UserCollection.FindOrAddUserByLoginAsync("Tim").Result.Id;
-			JerryId = TestSetup.UserCollection.FindOrAddUserByLoginAsync("Jerry").Result.Id;
-			BobId = TestSetup.UserCollection.FindOrAddUserByLoginAsync("Bob").Result.Id;
 		}
 
 		public static INode MockNode(string Name)
diff --git a/Engine/Source/Programs/Horde/HordeServer.Compute.Tests/JobCollectionTests.cs b/Engine/Source/Programs/Horde/HordeServer.Compute.Tests/JobCollectionTests.cs
--- a/Engine/Source/Programs/Horde/HordeServer.Compute.Tests/JobCollectionTests.cs
+++ b/Engine/Source/Programs/Horde/HordeServer.Compute.Tests/JobCollectionTests.cs
@@ -90,7 +90,7 @@ namespace HordeServerTests
 			Arguments.Add("-Target=Publish Client");
 			Arguments.Add("-Target=Post-Publish Client");
 
-			IJob Job = await JobCollection.AddAsync(JobId.GenerateNewId(), new StreamId("ue4-main"), new TemplateRefId("test-build"), ContentHash.SHA1("hello"), BaseGraph, "Test job", 123, 123, null, null, null, "Ben", null, null, null, null, false, false, null, null, null, Arguments);
+			IJob Job = await JobCollection.AddAsync(JobId.GenerateNewId(), new StreamId("ue4-main"), new TemplateRefId("test-build"), ContentHash.SHA1("hello"), BaseGraph, "Test job", 123, 123, null, null, null, null, null, null, null, false, false, null, null, Arguments);
 
 			await StartBatch(Job, BaseGraph, 0);
 			await RunStep(Job, BaseGraph, 0, 0, JobStepOutcome.Success); // Setup Build
@@ -162,7 +162,7 @@ namespace HordeServerTests
 			Arguments.Add("-Target=Step 1");
 			Arguments.Add("-Target=Step 3");
 
-			IJob Job = await JobCollection.AddAsync(JobId.GenerateNewId(), new StreamId("ue4-main"), new TemplateRefId("test-build"), ContentHash.SHA1("hello"), BaseGraph, "Test job", 123, 123, null, null, null, "Ben", null, null, null, null, false, false, null, null, null, Arguments);
+			IJob Job = await JobCollection.AddAsync(JobId.GenerateNewId(), new StreamId("ue4-main"), new TemplateRefId("test-build"), ContentHash.SHA1("hello"), BaseGraph, "Test job", 123, 123, null, null, null, null, null, null, null, false, false, null, null, Arguments);
 
 			await StartBatch(Job, BaseGraph, 0);
 			await RunStep(Job, BaseGraph, 0, 0, JobStepOutcome.Success); // Setup Build
@@ -218,7 +218,7 @@ namespace HordeServerTests
 			Arguments.Add("-Target=Step 1");
 			Arguments.Add("-Target=Step 3");
 
-			IJob Job = await JobCollection.AddAsync(JobId.GenerateNewId(), new StreamId("ue4-main"), new TemplateRefId("test-build"), ContentHash.SHA1("hello"), BaseGraph, "Test job", 123, 123, null, null, null, "Ben", null, null, null, null, false, false, null, null, null, Arguments);
+			IJob Job = await JobCollection.AddAsync(JobId.GenerateNewId(), new StreamId("ue4-main"), new TemplateRefId("test-build"), ContentHash.SHA1("hello"), BaseGraph, "Test job", 123, 123, null, null, null, null, null, null, null, false, false, null, null, Arguments);
 			Assert.AreEqual(1, Job.Batches.Count);
 
 			await StartBatch(Job, BaseGraph, 0);
diff --git a/Engine/Source/Programs/Horde/HordeServer.Compute.Tests/JobServiceTests.cs b/Engine/Source/Programs/Horde/HordeServer.Compute.Tests/JobServiceTests.cs
--- a/Engine/Source/Programs/Horde/HordeServer.Compute.Tests/JobServiceTests.cs
+++ b/Engine/Source/Programs/Horde/HordeServer.Compute.Tests/JobServiceTests.cs
@@ -47,7 +47,7 @@ namespace HordeServerTests
 			IStream? Stream = await TestSetup.StreamService.GetStreamAsync(StreamId);
 			Stream = await TestSetup.StreamService.StreamCollection.TryCreateOrReplaceAsync(new StreamId("ue5-main"), Stream, String.Empty, String.Empty, ProjectId, StreamConfig);
 
-			IJob Job = await TestSetup.JobService.CreateJobAsync(null, Stream!, TemplateRefId1, Template.Id, Graph, "Hello", 1234, 1233, 999, null, null, "joe", null, null, null, Stream!.Templates[TemplateRefId1].ChainedJobs, true, true, null, null, null, new List<string>());
+			IJob Job = await TestSetup.JobService.CreateJobAsync(null, Stream!, TemplateRefId1, Template.Id, Graph, "Hello", 1234, 1233, 999, null, null, null, null, null, Stream!.Templates[TemplateRefId1].ChainedJobs, true, true, null, null, new List<string>());
 			Assert.AreEqual(1, Job.ChainedJobs.Count);
 
 			Assert.IsTrue(await TestSetup.JobService.UpdateBatchAsync(Job, Job.Batches[0].Id, LogId.GenerateNewId(), JobStepBatchState.Running));
@@ -63,7 +63,7 @@ namespace HordeServerTests
 			Assert.AreEqual(ChainedJob!.Change, Job!.Change);
 			Assert.AreEqual(ChainedJob!.CodeChange, Job!.CodeChange);
 			Assert.AreEqual(ChainedJob!.PreflightChange, Job!.PreflightChange);
-			Assert.AreEqual(ChainedJob!.StartedByUser, Job!.StartedByUser);
+			Assert.AreEqual(ChainedJob!.StartedByUserId, Job!.StartedByUserId);
 		}
 
 		[TestMethod]
@@ -86,17 +86,18 @@ namespace HordeServerTests
 			DifferentUserName = (await TestSetup.JobService.GetJobAsync(DifferentUserName.Id))!;
 			DifferentArgs = (await TestSetup.JobService.GetJobAsync(DifferentArgs.Id))!;
 			
-			Assert.AreEqual("horde.duplicated.by.newer.CL", OrgJob.AbortedByUser);
-			Assert.IsNull(NewJob.AbortedByUser);
-			Assert.IsNull(DifferentTplRef.AbortedByUser);
-			Assert.IsNull(DifferentTplHash.AbortedByUser);
-			Assert.IsNull(DifferentUserName.AbortedByUser);
-			Assert.IsNull(DifferentArgs.AbortedByUser);
+			Assert.AreEqual(KnownUsers.System, OrgJob.AbortedByUserId);
+			Assert.IsNull(NewJob.AbortedByUserId);
+			Assert.IsNull(DifferentTplRef.AbortedByUserId);
+			Assert.IsNull(DifferentTplHash.AbortedByUserId);
+			Assert.IsNull(DifferentUserName.AbortedByUserId);
+			Assert.IsNull(DifferentArgs.AbortedByUserId);
 		}
 
-		private Task<IJob> CreatePreflightJob(TestSetup TestSetup, string TemplateRefId, string TemplateHash, string StartedByUserName, int PreflightChange, string[] Arguments)
+		private async Task<IJob> CreatePreflightJob(TestSetup TestSetup, string TemplateRefId, string TemplateHash, string StartedByUserName, int PreflightChange, string[] Arguments)
 		{
-			return TestSetup.JobService.CreateJobAsync(
+			IUser User = await TestSetup.UserCollection.FindOrAddUserByLoginAsync(StartedByUserName);
+			return await TestSetup.JobService.CreateJobAsync(
 				JobId: JobId.GenerateNewId(),
 				Stream: TestSetup.Fixture!.Stream!,
 				TemplateRefId: new TemplateRefId(TemplateRefId),
@@ -107,8 +108,7 @@ namespace HordeServerTests
 				CodeChange: 1000002,
 				PreflightChange: PreflightChange,
 				ClonedPreflightChange: null,
-				StartedByUserId: null,
-				StartedByUserName: StartedByUserName,
+				StartedByUserId: User.Id,
 				Priority: Priority.Normal,
 				null,
 				null,
@@ -117,7 +117,6 @@ namespace HordeServerTests
 				false,
 				null,
 				null,
-				null,
 				Arguments: new List<string>(Arguments)
 			);
 		}
@@ -171,7 +170,7 @@ namespace HordeServerTests
 
 			Graph = await TestSetup.GraphCollection.AppendAsync(Graph, new List<NewGroup> { GroupA, GroupB });
 
-			IJob Job = await TestSetup.JobService.CreateJobAsync(null, Stream!, new TemplateRefId("temp"), Template.Id, Graph, "Hello", 1234, 1233, 999, null, null, "joe", null, null, null, null, true, true, null, null, null, new List<string> { "-Target=Pak" });
+			IJob Job = await TestSetup.JobService.CreateJobAsync(null, Stream!, new TemplateRefId("temp"), Template.Id, Graph, "Hello", 1234, 1233, 999, null, null, null, null, null, null, true, true, null, null, new List<string> { "-Target=Pak" });
 
 			Assert.IsTrue(await TestSetup.JobService.UpdateBatchAsync(Job, Job.Batches[0].Id, LogId.GenerateNewId(), JobStepBatchState.Running));
 			Assert.IsNotNull(await TestSetup.JobService.UpdateStepAsync(Job, Job.Batches[0].Id, Job.Batches[0].Steps[0].Id, JobStepState.Running));
diff --git a/Engine/Source/Programs/Horde/HordeServer.Compute.Tests/JobsControllerDbTest.cs b/Engine/Source/Programs/Horde/HordeServer.Compute.Tests/JobsControllerDbTest.cs
--- a/Engine/Source/Programs/Horde/HordeServer.Compute.Tests/JobsControllerDbTest.cs
+++ b/Engine/Source/Programs/Horde/HordeServer.Compute.Tests/JobsControllerDbTest.cs
@@ -58,7 +58,7 @@ namespace HordeServerTests
 	        Obj = (await JobsController.GetStepAsync(Job.Id, BatchId.ToString(), StepId.ToString())).Value;
 	        StepRes = (Obj as GetStepResponse)!;
 	        Assert.IsTrue(StepRes.AbortRequested);
-	        Assert.AreEqual("Anonymous", StepRes.AbortByUser);
+//	        Assert.AreEqual("Anonymous", StepRes.AbortByUser);
         }
     }
 }
\ No newline at end of file
diff --git a/Engine/Source/Programs/Horde/HordeServer.Compute.Tests/SchedulerTests.cs b/Engine/Source/Programs/Horde/HordeServer.Compute.Tests/SchedulerTests.cs
--- a/Engine/Source/Programs/Horde/HordeServer.Compute.Tests/SchedulerTests.cs
+++ b/Engine/Source/Programs/Horde/HordeServer.Compute.Tests/SchedulerTests.cs
@@ -43,6 +43,8 @@ namespace HordeServerTests
 		{
 			TestSetup = GetTestSetup().Result;
 
+			IUser Bob = TestSetup.UserCollection.FindOrAddUserByLoginAsync("Bob").Result;
+
 			IProject ? Project = TestSetup.ProjectService.Collection.AddOrUpdateAsync(ProjectId, "", "", 0, new ProjectConfig { Name = "UE5" }).Result;
 			Assert.IsNotNull(Project);
 
@@ -55,9 +57,9 @@ namespace HordeServerTests
 
 			PerforceService = (PerforceServiceStub)TestSetup.PerforceService;
 			PerforceService.Changes.Clear();
-			PerforceService.AddChange("//UE5/Main", 100, "Bob", "", new[] { "code.cpp" });
-			PerforceService.AddChange("//UE5/Main", 101, "Bob", "", new[] { "content.uasset" });
-			PerforceService.AddChange("//UE5/Main", 102, "Bob", "", new[] { "content.uasset" });
+			PerforceService.AddChange("//UE5/Main", 100, Bob, "", new[] { "code.cpp" });
+			PerforceService.AddChange("//UE5/Main", 101, Bob, "", new[] { "content.uasset" });
+			PerforceService.AddChange("//UE5/Main", 102, Bob, "", new[] { "content.uasset" });
 		}
 
 		async Task<IStream> SetScheduleAsync(CreateScheduleRequest Schedule)
@@ -75,14 +77,16 @@ namespace HordeServerTests
 
 		public async Task<List<IJob>> FileTestHelperAsync(params string[] Files)
 		{
+			IUser Bob = TestSetup.UserCollection.FindOrAddUserByLoginAsync("Bob", "Bob").Result;
+
 			PerforceService.Changes.Clear();
-			PerforceService.AddChange("//UE5/Main", 100, "Bob", "", new[] { "code.cpp" });
-			PerforceService.AddChange("//UE5/Main", 101, "Bob", "", new[] { "content.uasset" });
-			PerforceService.AddChange("//UE5/Main", 102, "Bob", "", new[] { "content.uasset" });
-			PerforceService.AddChange("//UE5/Main", 103, "Bob", "", new[] { "foo/code.cpp" });
-			PerforceService.AddChange("//UE5/Main", 104, "Bob", "", new[] { "bar/code.cpp" });
-			PerforceService.AddChange("//UE5/Main", 105, "Bob", "", new[] { "foo/bar/content.uasset" });
-			PerforceService.AddChange("//UE5/Main", 106, "Bob", "", new[] { "bar/foo/content.uasset" });
+			PerforceService.AddChange("//UE5/Main", 100, Bob, "", new[] { "code.cpp" });
+			PerforceService.AddChange("//UE5/Main", 101, Bob, "", new[] { "content.uasset" });
+			PerforceService.AddChange("//UE5/Main", 102, Bob, "", new[] { "content.uasset" });
+			PerforceService.AddChange("//UE5/Main", 103, Bob, "", new[] { "foo/code.cpp" });
+			PerforceService.AddChange("//UE5/Main", 104, Bob, "", new[] { "bar/code.cpp" });
+			PerforceService.AddChange("//UE5/Main", 105, Bob, "", new[] { "foo/bar/content.uasset" });
+			PerforceService.AddChange("//UE5/Main", 106, Bob, "", new[] { "bar/foo/content.uasset" });
 
 			DateTime StartTime = new DateTime(2021, 1, 1, 12, 0, 0, DateTimeKind.Local); // Friday Jan 1, 2021 
 			TestSetup.Clock.UtcNow = StartTime;
@@ -266,10 +270,11 @@ namespace HordeServerTests
 			Assert.AreEqual(0, Jobs1.Count);
 
 			// Trigger some jobs
-			PerforceService.AddChange("//UE5/Main", 103, "Bob", "", new string[] { "foo.cpp" });
-			PerforceService.AddChange("//UE5/Main", 104, "Bob", "", new string[] { "foo.cpp" });
-			PerforceService.AddChange("//UE5/Main", 105, "Bob", "", new string[] { "foo.uasset" });
-			PerforceService.AddChange("//UE5/Main", 106, "Bob", "", new string[] { "foo.cpp" });
+			IUser Bob = TestSetup.UserCollection.FindOrAddUserByLoginAsync("Bob").Result;
+			PerforceService.AddChange("//UE5/Main", 103, Bob, "", new string[] { "foo.cpp" });
+			PerforceService.AddChange("//UE5/Main", 104, Bob, "", new string[] { "foo.cpp" });
+			PerforceService.AddChange("//UE5/Main", 105, Bob, "", new string[] { "foo.uasset" });
+			PerforceService.AddChange("//UE5/Main", 106, Bob, "", new string[] { "foo.cpp" });
 
 			TestSetup.Clock.Advance(TimeSpan.FromHours(1.25));
 			await ScheduleService.TickSharedOnlyForTestingAsync();
@@ -312,7 +317,7 @@ namespace HordeServerTests
 			Assert.AreEqual(0, Jobs3.Count);
 
 			// Mark the original job as complete
-			await TestSetup.JobService.UpdateJobAsync(Jobs2[0], AbortedByUser: "me");
+			await TestSetup.JobService.UpdateJobAsync(Jobs2[0], AbortedByUserId: KnownUsers.System);
 
 			// Test that another job does not trigger
 			TestSetup.Clock.Advance(TimeSpan.FromHours(0.5));
@@ -391,7 +396,7 @@ namespace HordeServerTests
 			Assert.AreEqual(0, Jobs2.Count);
 
 			// Create a job and fail it
-			IJob Job1 = await TestSetup.JobService.CreateJobAsync(null, Stream, NewTemplateRefId1, Template.Id, GraphA, "Hello", 1234, 1233, 999, null, null, "joe", null, null, null, null, true, true, null, null, null, new List<string> { "-Target=TriggerNext" });
+			IJob Job1 = await TestSetup.JobService.CreateJobAsync(null, Stream, NewTemplateRefId1, Template.Id, GraphA, "Hello", 1234, 1233, 999, null, null, null, null, null, null, true, true, null, null, new List<string> { "-Target=TriggerNext" });
 			Assert.IsTrue(await TestSetup.JobService.UpdateBatchAsync(Job1, Job1.Batches[0].Id, LogId.GenerateNewId(), JobStepBatchState.Running));
 			Assert.IsNotNull(await TestSetup.JobService.UpdateStepAsync(Job1, Job1.Batches[0].Id, Job1.Batches[0].Steps[0].Id, JobStepState.Completed, JobStepOutcome.Failure));
 			Assert.IsTrue(await TestSetup.JobService.UpdateBatchAsync(Job1, Job1.Batches[0].Id, LogId.GenerateNewId(), JobStepBatchState.Complete));
@@ -404,7 +409,7 @@ namespace HordeServerTests
 			Assert.AreEqual(0, Jobs3.Count);
 
 			// Create a job and make it succeed
-			IJob Job2 = await TestSetup.JobService.CreateJobAsync(null, Stream, NewTemplateRefId1, Template.Id, GraphA, "Hello", 1234, 1233, 999, null, null, "joe", null, null, null, null, true, true, null, null, null, new List<string> { "-Target=TriggerNext" });
+			IJob Job2 = await TestSetup.JobService.CreateJobAsync(null, Stream, NewTemplateRefId1, Template.Id, GraphA, "Hello", 1234, 1233, 999, null, null, null, null, null, null, true, true, null, null, new List<string> { "-Target=TriggerNext" });
 			Assert.IsTrue(await TestSetup.JobService.UpdateBatchAsync(Job2, Job2.Batches[0].Id, LogId.GenerateNewId(), JobStepBatchState.Running));
 			Assert.IsNotNull(await TestSetup.JobService.UpdateStepAsync(Job2, Job2.Batches[0].Id, Job2.Batches[0].Steps[0].Id, JobStepState.Completed, JobStepOutcome.Success));
 
diff --git a/Engine/Source/Programs/Horde/HordeServer.Compute.Tests/Stubs/Services/PerforceServiceStub.cs b/Engine/Source/Programs/Horde/HordeServer.Compute.Tests/Stubs/Services/PerforceServiceStub.cs
--- a/Engine/Source/Programs/Horde/HordeServer.Compute.Tests/Stubs/Services/PerforceServiceStub.cs
+++ b/Engine/Source/Programs/Horde/HordeServer.Compute.Tests/Stubs/Services/PerforceServiceStub.cs
@@ -1,6 +1,7 @@
 // Copyright Epic Games, Inc. All Rights Reserved.
 
 using EpicGames.Perforce;
+using HordeServer.Collections;
 using HordeServer.Models;
 using HordeServer.Services;
 using HordeServer.Utilities;
@@ -49,36 +50,34 @@ namespace HordeServerTests.Stubs.Services
 			}
 		}
 
+		User TestUser = new User("TestUser");
+
 		class ChangeComparer : IComparer<int>
 		{
 			public int Compare(int X, int Y) => Y.CompareTo(X);
 		}
 
-		Dictionary<string, User> NameToUser = new Dictionary<string, User>();
+		Dictionary<string, IUser> NameToUser = new Dictionary<string, IUser>(StringComparer.OrdinalIgnoreCase);
 		public Dictionary<string, SortedDictionary<int, ChangeDetails>> Changes { get; } = new Dictionary<string, SortedDictionary<int, ChangeDetails>>(StringComparer.OrdinalIgnoreCase);
 
-		public Task<NativePerforceConnection?> GetServiceUserConnection(string? ClusterName)
+		IUserCollection UserCollection;
+
+		public PerforceServiceStub(IUserCollection UserCollection)
 		{
-			throw new NotImplementedException();
+			this.UserCollection = UserCollection;
 		}
 
-		IUser FindOrAddUser(string UserName)
+		public Task<NativePerforceConnection?> GetServiceUserConnection(string? ClusterName)
 		{
-			User? User;
-			if (!NameToUser.TryGetValue(UserName, out User))
-			{
-				User = new User(UserName);
-				NameToUser.Add(UserName, User);
-			}
-			return User;
+			throw new NotImplementedException();
 		}
 
-		public ValueTask<IUser> FindOrAddUserAsync(string ClusterName, string UserName)
+		public async ValueTask<IUser> FindOrAddUserAsync(string ClusterName, string UserName)
 		{
-			return new ValueTask<IUser>(FindOrAddUser(UserName));
+			return await UserCollection.FindOrAddUserByLoginAsync(UserName);
 		}
 
-		public void AddChange(string StreamName, int Number, string Author, string Description, IEnumerable<string> Files)
+		public void AddChange(string StreamName, int Number, IUser Author, string Description, IEnumerable<string> Files)
 		{
 			SortedDictionary<int, ChangeDetails>? StreamChanges;
 			if (!Changes.TryGetValue(StreamName, out StreamChanges))
@@ -86,7 +85,7 @@ namespace HordeServerTests.Stubs.Services
 				StreamChanges = new SortedDictionary<int, ChangeDetails>(new ChangeComparer());
 				Changes[StreamName] = StreamChanges;
 			}
-			StreamChanges.Add(Number, new ChangeDetails(Number, FindOrAddUser(Author), null!, Description, Files.Select(x => PerforceExtensions.CreateChangeFile(x)).ToList(), DateTime.Now));
+			StreamChanges.Add(Number, new ChangeDetails(Number, Author, null!, Description, Files.Select(x => PerforceExtensions.CreateChangeFile(x)).ToList(), DateTime.Now));
 		}
 
 		public Task<List<ChangeSummary>> GetChangesAsync(string ClusterName, string StreamName, int? MinChange, int? MaxChange, int NumResults, string? ImpersonateUser)
@@ -158,7 +157,7 @@ namespace HordeServerTests.Stubs.Services
 
 		public Task<int> CreateNewChangeAsync(string ClusterName, string StreamName, string Path)
 		{
-			ChangeDetails NewChange = new ChangeDetails(Changes[StreamName].First().Key + 1, FindOrAddUser("system"), null!, "", new List<ChangeFile> { PerforceExtensions.CreateChangeFile(Path) }, DateTime.Now);
+			ChangeDetails NewChange = new ChangeDetails(Changes[StreamName].First().Key + 1, TestUser, null!, "", new List<ChangeFile> { PerforceExtensions.CreateChangeFile(Path) }, DateTime.Now);
 			Changes[StreamName].Add(NewChange.Number, NewChange);
 			return Task.FromResult(NewChange.Number);
 		}
diff --git a/Engine/Source/Programs/Horde/HordeServer.Compute.Tests/TestSetup.cs b/Engine/Source/Programs/Horde/HordeServer.Compute.Tests/TestSetup.cs
--- a/Engine/Source/Programs/Horde/HordeServer.Compute.Tests/TestSetup.cs
+++ b/Engine/Source/Programs/Horde/HordeServer.Compute.Tests/TestSetup.cs
@@ -229,7 +229,7 @@ namespace HordeServerTests
         {
 			ILogger<JobsController> Logger = ServiceProvider.GetRequiredService<ILogger<JobsController>>();
 			JobsController JobsCtrl = new JobsController(AclService, GraphCollection, PerforceService, StreamService, JobService,
-		        TemplateCollection, ArtifactCollection, NotificationService, Logger);
+		        TemplateCollection, ArtifactCollection, UserCollection, NotificationService, Logger);
 	        JobsCtrl.ControllerContext = GetControllerContext();
 	        return JobsCtrl;
         }
diff --git a/Engine/Source/Programs/Horde/HordeServer/Controllers/JobsController.cs b/Engine/Source/Programs/Horde/HordeServer/Controllers/JobsController.cs
--- a/Engine/Source/Programs/Horde/HordeServer/Controllers/JobsController.cs
+++ b/Engine/Source/Programs/Horde/HordeServer/Controllers/JobsController.cs
@@ -1023,8 +1023,8 @@ namespace HordeServer.Controllers
 				}
 			}
 
-			string? RetryByUser = (Request.Retry.HasValue && Request.Retry.Value) ? (User.Identity?.Name ?? "Anonymous") : null;
-			string? AbortByUser = (Request.AbortRequested.HasValue && Request.AbortRequested.Value) ? (User.Identity?.Name ?? "Anonymous") : null;
+			UserId? RetryByUser = (Request.Retry.HasValue && Request.Retry.Value) ? User.GetUserId() : null;
+			UserId? AbortByUser = (Request.AbortRequested.HasValue && Request.AbortRequested.Value) ? User.GetUserId() : null;
 
 			try
 			{
diff --git a/Engine/Source/Programs/Horde/HordeServer/Jobs/IJobCollection.cs b/Engine/Source/Programs/Horde/HordeServer/Jobs/IJobCollection.cs
--- a/Engine/Source/Programs/Horde/HordeServer/Jobs/IJobCollection.cs
+++ b/Engine/Source/Programs/Horde/HordeServer/Jobs/IJobCollection.cs
@@ -147,15 +147,15 @@ namespace HordeServer.Collections
 		/// <param name="NewState">New state of the jobstep</param>
 		/// <param name="NewOutcome">New outcome of the jobstep</param>
 		/// <param name="NewAbortRequested">New state of request abort</param>
-		/// <param name="NewAbortByUser">New name of user that requested the abort</param>
+		/// <param name="NewAbortByUserId">New name of user that requested the abort</param>
 		/// <param name="NewLogId">New log id for the jobstep</param>
 		/// <param name="NewNotificationTriggerId">New id for a notification trigger</param>
-		/// <param name="NewRetryByUser">Whether the step should be retried</param>
+		/// <param name="NewRetryByUserId">Whether the step should be retried</param>
 		/// <param name="NewPriority">New priority for this step</param>
 		/// <param name="NewReports">New report documents</param>
 		/// <param name="NewProperties">Property changes. Any properties with a null value will be removed.</param>
 		/// <returns>True if the job was updated, false if it was deleted in the meantime</returns>
-		Task<bool> TryUpdateStepAsync(IJob Job, IGraph Graph, SubResourceId BatchId, SubResourceId StepId, JobStepState NewState = default, JobStepOutcome NewOutcome = default, bool? NewAbortRequested = null, string? NewAbortByUser = null, LogId? NewLogId = null, ObjectId? NewNotificationTriggerId = null, string? NewRetryByUser = null, Priority? NewPriority = null, List<Report>? NewReports = null, Dictionary<string, string?>? NewProperties = null);
+		Task<bool> TryUpdateStepAsync(IJob Job, IGraph Graph, SubResourceId BatchId, SubResourceId StepId, JobStepState NewState = default, JobStepOutcome NewOutcome = default, bool? NewAbortRequested = null, UserId? NewAbortByUserId = null, LogId? NewLogId = null, ObjectId? NewNotificationTriggerId = null, UserId? NewRetryByUserId = null, Priority? NewPriority = null, List<Report>? NewReports = null, Dictionary<string, string?>? NewProperties = null);
 
 		/// <summary>
 		/// Attempts to update the node groups to be executed for a job. Fails if another write happens in the meantime.
diff --git a/Engine/Source/Programs/Horde/HordeServer/Jobs/Impl/JobCollection.cs b/Engine/Source/Programs/Horde/HordeServer/Jobs/Impl/JobCollection.cs
--- a/Engine/Source/Programs/Horde/HordeServer/Jobs/Impl/JobCollection.cs
+++ b/Engine/Source/Programs/Horde/HordeServer/Jobs/Impl/JobCollection.cs
@@ -199,14 +199,14 @@ namespace HordeServer.Collections.Impl
 			[BsonIgnoreIfNull]
 			public UserId? StartedByUserId { get; set; }
 
-			[BsonIgnoreIfNull]
-			public string? StartedByUser { get; set; }
+			[BsonIgnoreIfNull, BsonElement("StartedByUser")]
+			public string? StartedByUser_DEPRECATED { get; set; }
 
 			[BsonIgnoreIfNull]
 			public UserId? AbortedByUserId { get; set; }
 
-			[BsonIgnoreIfNull]
-			public string? AbortedByUser { get; set; }
+			[BsonIgnoreIfNull, BsonElement("AbortedByUser")]
+			public string? AbortedByUser_DEPRECATED { get; set; }
 
 			[BsonRequired]
 			public string Name { get; set; }
@@ -360,7 +360,7 @@ namespace HordeServer.Collections.Impl
 				Jobs.Indexes.CreateOne(new CreateIndexModel<JobDocument>(Builders<JobDocument>.IndexKeys.Ascending(x => x.CreateTimeUtc)));
 				Jobs.Indexes.CreateOne(new CreateIndexModel<JobDocument>(Builders<JobDocument>.IndexKeys.Ascending(x => x.UpdateTimeUtc)));
 				Jobs.Indexes.CreateOne(new CreateIndexModel<JobDocument>(Builders<JobDocument>.IndexKeys.Ascending(x => x.Name)));
-				Jobs.Indexes.CreateOne(new CreateIndexModel<JobDocument>(Builders<JobDocument>.IndexKeys.Ascending(x => x.StartedByUser)));
+				Jobs.Indexes.CreateOne(new CreateIndexModel<JobDocument>(Builders<JobDocument>.IndexKeys.Ascending(x => x.StartedByUserId)));
 				Jobs.Indexes.CreateOne(new CreateIndexModel<JobDocument>(Builders<JobDocument>.IndexKeys.Ascending(x => x.TemplateId)));
 				Jobs.Indexes.CreateOne(new CreateIndexModel<JobDocument>(Builders<JobDocument>.IndexKeys.Descending(x => x.SchedulePriority)));
 			}
@@ -689,7 +689,7 @@ namespace HordeServer.Collections.Impl
 		}
 
 		/// <inheritdoc/>
-		public async Task<bool> TryUpdateStepAsync(IJob Job, IGraph Graph, SubResourceId BatchId, SubResourceId StepId, JobStepState NewState, JobStepOutcome NewOutcome, bool? NewAbortRequested, string? NewAbortByUser, LogId? NewLogId, ObjectId? NewNotificationTriggerId, string? NewRetryByUser, Priority? NewPriority, List<Report>? NewReports, Dictionary<string, string?>? NewProperties)
+		public async Task<bool> TryUpdateStepAsync(IJob Job, IGraph Graph, SubResourceId BatchId, SubResourceId StepId, JobStepState NewState, JobStepOutcome NewOutcome, bool? NewAbortRequested, UserId? NewAbortByUserId, LogId? NewLogId, ObjectId? NewNotificationTriggerId, UserId? NewRetryByUserId, Priority? NewPriority, List<Report>? NewReports, Dictionary<string, string?>? NewProperties)
 		{
 			JobDocument JobDocument = (JobDocument)Job;
 
@@ -751,10 +751,10 @@ namespace HordeServer.Collections.Impl
 							}
 
 							// Update the user that requested the abort
-							if (NewAbortByUser != null && Step.AbortByUser == null)
+							if (NewAbortByUserId != null && Step.AbortedByUserId == null)
 							{
-								Step.AbortByUser = NewAbortByUser;
-								Updates.Add(UpdateBuilder.Set(x => x.Batches[BatchIdx].Steps[StepIdx].AbortByUser, Step.AbortByUser));
+								Step.AbortedByUserId = NewAbortByUserId;
+								Updates.Add(UpdateBuilder.Set(x => x.Batches[BatchIdx].Steps[StepIdx].AbortedByUserId, Step.AbortedByUserId));
 
 								bRefreshDependentJobSteps = true;
 							}
@@ -776,10 +776,10 @@ namespace HordeServer.Collections.Impl
 							}
 
 							// Update the retry flag
-							if (NewRetryByUser != null && Step.RetryByUser == null)
+							if (NewRetryByUserId != null && Step.RetriedByUserId == null)
 							{
-								Step.RetryByUser = NewRetryByUser;
-								Updates.Add(UpdateBuilder.Set(x => x.Batches[BatchIdx].Steps[StepIdx].RetryByUser, Step.RetryByUser));
+								Step.RetriedByUserId = NewRetryByUserId;
+								Updates.Add(UpdateBuilder.Set(x => x.Batches[BatchIdx].Steps[StepIdx].RetriedByUserId, Step.RetriedByUserId));
 
 								bRefreshBatches = true;
 							}
@@ -1228,7 +1228,7 @@ namespace HordeServer.Collections.Impl
 
 			// Find all the targets in this job
 			HashSet<string> Targets = new HashSet<string>(StringComparer.OrdinalIgnoreCase);
-			if (Job.AbortedByUser == null)
+			if (Job.AbortedByUserId == null)
 			{
 				foreach (string Argument in Job.Arguments)
 				{
diff --git a/Engine/Source/Programs/Horde/HordeServer/Jobs/Impl/JobService.cs b/Engine/Source/Programs/Horde/HordeServer/Jobs/Impl/JobService.cs
--- a/Engine/Source/Programs/Horde/HordeServer/Jobs/Impl/JobService.cs
+++ b/Engine/Source/Programs/Horde/HordeServer/Jobs/Impl/JobService.cs
@@ -780,20 +780,20 @@ namespace HordeServer.Services
 		/// <param name="NewState">New state of the jobstep</param>
 		/// <param name="NewOutcome">New outcome of the jobstep</param>
 		/// <param name="NewAbortRequested">New state of abort request</param>
-		/// <param name="NewAbortByUser">New user that requested the abort</param>
+		/// <param name="NewAbortByUserId">New user that requested the abort</param>
 		/// <param name="NewLogId">New log id for the jobstep</param>
 		/// <param name="NewNotificationTriggerId">New notification trigger id for the jobstep</param>
-		/// <param name="NewRetryByUser">Whether the step should be retried</param>
+		/// <param name="NewRetryByUserId">Whether the step should be retried</param>
 		/// <param name="NewPriority">New priority for this step</param>
 		/// <param name="NewReports">New list of reports</param>
 		/// <param name="NewProperties">Property changes. Any properties with a null value will be removed.</param>
 		/// <returns>True if the job was updated, false if it was deleted in the meantime</returns>
-		public async Task<IJob?> UpdateStepAsync(IJob Job, SubResourceId BatchId, SubResourceId StepId, JobStepState NewState = JobStepState.Unspecified, JobStepOutcome NewOutcome = JobStepOutcome.Unspecified, bool? NewAbortRequested = null, string? NewAbortByUser = null, LogId? NewLogId = null, ObjectId? NewNotificationTriggerId = null, string? NewRetryByUser = null, Priority? NewPriority = null, List<Report>? NewReports = null, Dictionary<string, string?>? NewProperties = null)
+		public async Task<IJob?> UpdateStepAsync(IJob Job, SubResourceId BatchId, SubResourceId StepId, JobStepState NewState = JobStepState.Unspecified, JobStepOutcome NewOutcome = JobStepOutcome.Unspecified, bool? NewAbortRequested = null, UserId? NewAbortByUserId = null, LogId? NewLogId = null, ObjectId? NewNotificationTriggerId = null, UserId? NewRetryByUserId = null, Priority? NewPriority = null, List<Report>? NewReports = null, Dictionary<string, string?>? NewProperties = null)
 		{
 			using IDisposable Scope = Logger.BeginScope("UpdateStepAsync({JobId})", Job.Id);
 			for (; ;)
 			{
-				if (await TryUpdateStepAsync(Job, BatchId, StepId, NewState, NewOutcome, NewAbortRequested, NewAbortByUser, NewLogId, NewNotificationTriggerId, NewRetryByUser, NewPriority, NewReports, NewProperties))
+				if (await TryUpdateStepAsync(Job, BatchId, StepId, NewState, NewOutcome, NewAbortRequested, NewAbortByUserId, NewLogId, NewNotificationTriggerId, NewRetryByUserId, NewPriority, NewReports, NewProperties))
 				{
 					return Job;
 				}
@@ -817,15 +817,15 @@ namespace HordeServer.Services
 		/// <param name="NewState">New state of the jobstep</param>
 		/// <param name="NewOutcome">New outcome of the jobstep</param>
 		/// <param name="NewAbortRequested">New state for request abort</param>
-		/// <param name="NewAbortByUser">New name of user that requested the abort</param>
+		/// <param name="NewAbortByUserId">New name of user that requested the abort</param>
 		/// <param name="NewLogId">New log id for the jobstep</param>
 		/// <param name="NewTriggerId">New trigger id for the jobstep</param>
-		/// <param name="NewRetryByUser">Whether the step should be retried</param>
+		/// <param name="NewRetryByUserId">Whether the step should be retried</param>
 		/// <param name="NewPriority">New priority for this step</param>
 		/// <param name="NewReports">New reports</param>
 		/// <param name="NewProperties">Property changes. Any properties with a null value will be removed.</param>
 		/// <returns>True if the job was updated, false if it was deleted in the meantime</returns>
-		public async Task<bool> TryUpdateStepAsync(IJob Job, SubResourceId BatchId, SubResourceId StepId, JobStepState NewState = JobStepState.Unspecified, JobStepOutcome NewOutcome = JobStepOutcome.Unspecified, bool? NewAbortRequested = null, string? NewAbortByUser = null, LogId? NewLogId = null, ObjectId? NewTriggerId = null, string? NewRetryByUser = null, Priority? NewPriority = null, List<Report>? NewReports = null, Dictionary<string, string?>? NewProperties = null)
+		public async Task<bool> TryUpdateStepAsync(IJob Job, SubResourceId BatchId, SubResourceId StepId, JobStepState NewState = JobStepState.Unspecified, JobStepOutcome NewOutcome = JobStepOutcome.Unspecified, bool? NewAbortRequested = null, UserId? NewAbortByUserId = null, LogId? NewLogId = null, ObjectId? NewTriggerId = null, UserId? NewRetryByUserId = null, Priority? NewPriority = null, List<Report>? NewReports = null, Dictionary<string, string?>? NewProperties = null)
 		{
 			using IDisposable Scope = Logger.BeginScope("TryUpdateStepAsync({JobId})", Job.Id);
 
@@ -852,7 +852,7 @@ namespace HordeServer.Services
 			}
 
 			// Update the step
-			if (await Jobs.TryUpdateStepAsync(Job, Graph, BatchId, StepId, NewState, NewOutcome, NewAbortRequested, NewAbortByUser, NewLogId, NewTriggerId, NewRetryByUser, NewPriority, NewReports, NewProperties))
+			if (await Jobs.TryUpdateStepAsync(Job, Graph, BatchId, StepId, NewState, NewOutcome, NewAbortRequested, NewAbortByUserId, NewLogId, NewTriggerId, NewRetryByUserId, NewPriority, NewReports, NewProperties))
 			{
 				using IScope DdScope = GlobalTracer.Instance.BuildSpan("TryUpdateStepAsync").StartActive();
 				DdScope.Span.SetTag("JobId", Job.Id.ToString());
