Horde: Automatically upgrade job documents that do not have the correct user id values set.

diff --git a/Engine/Source/Programs/Horde/HordeServer.Compute.Tests/JobCollectionTests.cs b/Engine/Source/Programs/Horde/HordeServer.Compute.Tests/JobCollectionTests.cs
--- a/Engine/Source/Programs/Horde/HordeServer.Compute.Tests/JobCollectionTests.cs
+++ b/Engine/Source/Programs/Horde/HordeServer.Compute.Tests/JobCollectionTests.cs
@@ -34,13 +34,15 @@ namespace HordeServerTests
 	public class JobCollectionTests : DatabaseIntegrationTest
 	{
 		IGraphCollection GraphCollection;
+		IUserCollection UserCollection;
 		IJobCollection JobCollection;
 
 		public JobCollectionTests()
 		{
 			DatabaseService DatabaseService = GetDatabaseService();
 			this.GraphCollection = new GraphCollection(DatabaseService);
-			this.JobCollection = new JobCollection(DatabaseService, NullLogger<JobCollection>.Instance);
+			this.UserCollection = new UserCollectionV2(DatabaseService, NullLogger<UserCollectionV2>.Instance);
+			this.JobCollection = new JobCollection(DatabaseService, UserCollection, NullLogger<JobCollection>.Instance);
 		}
 
 		NewGroup AddGroup(List<NewGroup> Groups)
diff --git a/Engine/Source/Programs/Horde/HordeServer/Jobs/Impl/JobCollection.cs b/Engine/Source/Programs/Horde/HordeServer/Jobs/Impl/JobCollection.cs
--- a/Engine/Source/Programs/Horde/HordeServer/Jobs/Impl/JobCollection.cs
+++ b/Engine/Source/Programs/Horde/HordeServer/Jobs/Impl/JobCollection.cs
@@ -73,12 +73,16 @@ namespace HordeServer.Collections.Impl
 			public bool Retry { get; set; }
 
 			public UserId? RetriedByUserId { get; set; }
-			public string? RetryByUser { get; set; }
+
+			[BsonElement("RetryByUser")]
+			public string? RetriedByUser_DEPRECATED { get; set; }
 
 			public bool AbortRequested { get; set; } = false;
 
 			public UserId? AbortedByUserId { get; set; }
-			public string? AbortByUser { get; set; }
+			
+			[BsonElement("AbortByUser")]
+			public string? AbortedByUser_DEPRECATED { get; set; }
 
 			[BsonIgnoreIfNull]
 			public List<Report>? Reports { get; set; }
@@ -304,14 +308,6 @@ namespace HordeServer.Collections.Impl
 				this.NextSubResourceId = SubResourceId.Random();
 				this.UpdateTimeUtc = CreateTimeUtc;
 			}
-
-			public void PostLoad()
-			{
-				if (GraphHash == ContentHash.Empty)
-				{
-					Batches.Clear();
-				}
-			}
 		}
 
 		/// <summary>
@@ -336,6 +332,11 @@ namespace HordeServer.Collections.Impl
 		/// </summary>
 		IMongoCollection<JobDocument> Jobs;
 
+		/// <summary>
+		/// The user collection
+		/// </summary>
+		IUserCollection UserCollection;
+
 		/// <summary>
 		/// Logger for output
 		/// </summary>
@@ -345,9 +346,11 @@ namespace HordeServer.Collections.Impl
 		/// Constructor
 		/// </summary>
 		/// <param name="DatabaseService">The database service singleton</param>
+		/// <param name="UserCollection"></param>
 		/// <param name="Logger">The logger instance</param>
-		public JobCollection(DatabaseService DatabaseService, ILogger<JobCollection> Logger)
+		public JobCollection(DatabaseService DatabaseService, IUserCollection UserCollection, ILogger<JobCollection> Logger)
 		{
+			this.UserCollection = UserCollection;
 			this.Logger = Logger;
 
 			Jobs = DatabaseService.GetCollection<JobDocument>("Jobs");
@@ -366,6 +369,66 @@ namespace HordeServer.Collections.Impl
 			}
 		}
 
+		async Task PostLoadAsync(JobDocument Job)
+		{
+			if (Job.GraphHash == ContentHash.Empty)
+			{
+				Job.Batches.Clear();
+			}
+
+			List<UpdateDefinition<JobDocument>> Updates = new List<UpdateDefinition<JobDocument>>();
+			if (Job.StartedByUser_DEPRECATED != null)
+			{
+				IUser StartedByUser = await UserCollection.FindOrAddUserByLoginAsync(Job.StartedByUser_DEPRECATED);
+				Job.StartedByUserId = StartedByUser.Id;
+				Job.StartedByUser_DEPRECATED = null;
+				Updates.Add(Builders<JobDocument>.Update.Set(x => x.StartedByUserId, StartedByUser.Id).Unset(x => x.StartedByUser_DEPRECATED));
+			}
+
+			if (Job.AbortedByUser_DEPRECATED != null)
+			{
+				IUser AbortedByUser = await UserCollection.FindOrAddUserByLoginAsync(Job.AbortedByUser_DEPRECATED);
+				Job.AbortedByUserId = AbortedByUser.Id;
+				Job.AbortedByUser_DEPRECATED = null;
+				Updates.Add(Builders<JobDocument>.Update.Set(x => x.AbortedByUserId, AbortedByUser.Id).Unset(x => x.AbortedByUser_DEPRECATED));
+			}
+
+
+			for (int BatchIdx = 0; BatchIdx < Job.Batches.Count; BatchIdx++)
+			{
+				int LocalBatchIdx = BatchIdx;
+				JobStepBatchDocument Batch = Job.Batches[BatchIdx];
+
+				for (int StepIdx = 0; StepIdx < Batch.Steps.Count; StepIdx++)
+				{
+					int LocalStepIdx = StepIdx;
+					JobStepDocument Step = Batch.Steps[StepIdx];
+
+					if (Step.AbortedByUser_DEPRECATED != null)
+					{
+						IUser AbortedByUser = await UserCollection.FindOrAddUserByLoginAsync(Step.AbortedByUser_DEPRECATED);
+						Step.AbortedByUserId = AbortedByUser.Id;
+						Step.AbortedByUser_DEPRECATED = null;
+						Updates.Add(Builders<JobDocument>.Update.Set(x => x.Batches[LocalBatchIdx].Steps[LocalStepIdx].AbortedByUserId, AbortedByUser.Id).Unset(x => x.Batches[LocalBatchIdx].Steps[LocalStepIdx].AbortedByUser_DEPRECATED));
+					}
+
+					if (Step.RetriedByUser_DEPRECATED != null)
+					{
+						IUser RetriedByUser = await UserCollection.FindOrAddUserByLoginAsync(Step.RetriedByUser_DEPRECATED);
+						Step.RetriedByUserId = RetriedByUser.Id;
+						Step.RetriedByUser_DEPRECATED = null;
+						Updates.Add(Builders<JobDocument>.Update.Set(x => x.Batches[LocalBatchIdx].Steps[LocalStepIdx].RetriedByUserId, RetriedByUser.Id).Unset(x => x.Batches[LocalBatchIdx].Steps[LocalStepIdx].RetriedByUser_DEPRECATED));
+					}
+				}
+			}
+
+
+			if (Updates.Count > 0)
+			{
+				await TryUpdateAsync(Job, Builders<JobDocument>.Update.Combine(Updates));
+			}
+		}
+
 		/// <inheritdoc/>
 		[SuppressMessage("Compiler", "CA1054:URI parameters should not be strings")]
 		public async Task<IJob> AddAsync(JobId JobId, StreamId StreamId, TemplateRefId TemplateRefId, ContentHash TemplateHash, IGraph Graph, string Name, int Change, int CodeChange, int? PreflightChange, int? ClonedPreflightChange, UserId? StartedByUserId, Priority? Priority, bool? AutoSubmit, bool? UpdateIssues, List<ChainedJobTemplate>? ChainedJobs, bool ShowUgsBadges, bool ShowUgsAlerts, string? NotificationChannel, string? NotificationChannelFilter, List<string>? Arguments)
@@ -394,7 +457,7 @@ namespace HordeServer.Collections.Impl
 			JobDocument? Job = await Jobs.Find<JobDocument>(x => x.Id == JobId).FirstOrDefaultAsync();
 			if (Job != null)
 			{
-				Job.PostLoad();
+				await PostLoadAsync(Job);
 			}
 			return Job;
 		}
@@ -493,7 +556,7 @@ namespace HordeServer.Collections.Impl
 			}
 			foreach (JobDocument Result in Results)
 			{
-				Result.PostLoad();
+				await PostLoadAsync(Result);
 			}
 			return Results.ConvertAll<JobDocument, IJob>(x => x);
 		}
@@ -934,7 +997,7 @@ namespace HordeServer.Collections.Impl
 			List<JobDocument> NewJobs = await Jobs.Find(x => x.SchedulePriority > 0).SortByDescending(x => x.SchedulePriority).ThenBy(x => x.CreateTimeUtc).ToListAsync();
 			foreach (JobDocument Result in NewJobs)
 			{
-				Result.PostLoad();
+				await PostLoadAsync(Result);
 			}
 			return NewJobs.ConvertAll<JobDocument, IJob>(x => x);
 		}
