Horde: Only log lease properties that are explicitly listed for logging (they may contain sensitive info).

diff --git a/Engine/Source/Programs/Horde/HordeServer/Agents/ITaskSource.cs b/Engine/Source/Programs/Horde/HordeServer/Agents/ITaskSource.cs
--- a/Engine/Source/Programs/Horde/HordeServer/Agents/ITaskSource.cs
+++ b/Engine/Source/Programs/Horde/HordeServer/Agents/ITaskSource.cs
@@ -12,6 +12,7 @@ using MongoDB.Bson;
 using System;
 using System.Collections.Generic;
 using System.Linq;
+using System.Linq.Expressions;
 using System.Reflection;
 using System.Text;
 using System.Threading;
@@ -80,11 +81,45 @@ namespace HordeServer.Tasks
 	/// <typeparam name="TMessage"></typeparam>
 	public abstract class TaskSourceBase<TMessage> : ITaskSource where TMessage : IMessage, new()
 	{
+		/// <summary>
+		/// List of properties that can be output as log messages
+		/// </summary>
+		protected class PropertyList
+		{
+			internal StringBuilder FormatString = new StringBuilder();
+			internal List<Func<TMessage, object>> Accessors = new List<Func<TMessage, object>>();
+
+			/// <summary>
+			/// Adds a new property to the list
+			/// </summary>
+			/// <param name="Expr">Accessor for the property</param>
+			public PropertyList Add(Expression<Func<TMessage, object>> Expr)
+			{
+				MemberInfo Member = ((MemberExpression)Expr.Body).Member;
+				return Add(Member.Name, Expr.Compile());
+			}
+
+			/// <summary>
+			/// Adds a new property to the list
+			/// </summary>
+			/// <param name="Name">Name of the property</param>
+			/// <param name="Accessor">Accessor for the property</param>
+			public PropertyList Add(string Name, Func<TMessage, object> Accessor)
+			{
+				FormatString.Append($", {Name}={{{Name}}}");
+				Accessors.Add(Accessor);
+				return this;
+			}
+		}
+
 		static TMessage Message = new TMessage();
 
 		/// <inheritdoc/>
 		public abstract string Type { get; }
 
+		/// <inheritdoc/>
+		protected virtual PropertyList OnLeaseStartedProperties { get; } = new PropertyList();
+
 		/// <inheritdoc/>
 		public MessageDescriptor Descriptor => Message.Descriptor;
 
@@ -106,7 +141,14 @@ namespace HordeServer.Tasks
 		/// <inheritdoc cref="ITaskSource.OnLeaseStartedAsync(IAgent, LeaseId, Any, ILogger)"/>
 		public virtual Task OnLeaseStartedAsync(IAgent Agent, LeaseId LeaseId, TMessage Payload, ILogger Logger)
 		{
-			LogLeaseStartedInfo<TMessage>(LeaseId, Payload, Logger);
+			object[] Arguments = new object[2 + OnLeaseStartedProperties.Accessors.Count];
+			Arguments[0] = LeaseId;
+			Arguments[1] = Type;
+			for (int Idx = 0; Idx < OnLeaseStartedProperties.Accessors.Count; Idx++)
+			{
+				Arguments[Idx + 2] = OnLeaseStartedProperties.Accessors[Idx](Payload);
+			}
+			Logger.LogInformation($"Lease {{LeaseId}} started (Type={{Type}}{OnLeaseStartedProperties.FormatString})", Arguments);
 			return Task.CompletedTask;
 		}
 
@@ -116,44 +158,5 @@ namespace HordeServer.Tasks
 			Logger.LogInformation("Lease {LeaseId} complete", LeaseId);
 			return Task.CompletedTask;
 		}
-
-		class TypeInfo<T>
-		{
-			public static readonly PropertyInfo[] Properties = typeof(T).GetProperties(BindingFlags.Public | BindingFlags.Instance).ToArray();
-			public static readonly string FormatString = CreateFormatString(Properties);
-
-			static string CreateFormatString(PropertyInfo[] Properties)
-			{
-				StringBuilder Message = new StringBuilder($"Lease {{LeaseId}} started (Type={{Type}}");
-				foreach (PropertyInfo Property in Properties)
-				{
-					Message.Append($", {Property.Name}={{{Property.Name}}}");
-				}
-				Message.Append(")");
-				return Message.ToString();
-			}
-		}
-
-		/// <summary>
-		/// Helper method to log information about a lease starting
-		/// </summary>
-		/// <typeparam name="T"></typeparam>
-		/// <param name="LeaseId"></param>
-		/// <param name="Payload"></param>
-		/// <param name="Logger"></param>
-		protected void LogLeaseStartedInfo<T>(LeaseId LeaseId, T Payload, ILogger Logger)
-		{
-			PropertyInfo[] Properties = TypeInfo<T>.Properties;
-
-			object?[] Values = new object?[Properties.Length + 2];
-			Values[0] = LeaseId;
-			Values[1] = Type;
-			for (int Idx = 0; Idx < Properties.Length; Idx++)
-			{
-				Values[Idx + 2] = Properties[Idx].GetValue(Payload);
-			}
-
-			Logger.LogInformation(TypeInfo<T>.FormatString, Values);
-		}
 	}
 }
diff --git a/Engine/Source/Programs/Horde/HordeServer/Compute/Impl/ComputeService.cs b/Engine/Source/Programs/Horde/HordeServer/Compute/Impl/ComputeService.cs
--- a/Engine/Source/Programs/Horde/HordeServer/Compute/Impl/ComputeService.cs
+++ b/Engine/Source/Programs/Horde/HordeServer/Compute/Impl/ComputeService.cs
@@ -161,6 +161,8 @@ namespace HordeServer.Compute.Impl
 			this.RequirementsCache = new MemoryCache(new MemoryCacheOptions());
 			this.CachedPools = new LazyCachedValue<Task<List<IPool>>>(() => PoolCollection.GetAsync(), TimeSpan.FromSeconds(30.0));
 			this.Logger = Logger;
+
+			OnLeaseStartedProperties.Add(x => x.Task);
 		}
 
 		/// <summary>
diff --git a/Engine/Source/Programs/Horde/HordeServer/Jobs/Impl/JobTaskSource.cs b/Engine/Source/Programs/Horde/HordeServer/Jobs/Impl/JobTaskSource.cs
--- a/Engine/Source/Programs/Horde/HordeServer/Jobs/Impl/JobTaskSource.cs
+++ b/Engine/Source/Programs/Horde/HordeServer/Jobs/Impl/JobTaskSource.cs
@@ -322,6 +322,8 @@ namespace HordeServer.Tasks.Impl
 			this.Ticker = new BackgroundTick(TickAsync, RefreshInterval, Logger);
 			this.Settings = Settings;
 			this.Logger = Logger;
+
+			OnLeaseStartedProperties.Add(nameof(ExecuteJobTask.JobId), x => new JobId(x.JobId)).Add(nameof(ExecuteJobTask.BatchId), x => SubResourceId.Parse(x.BatchId));
 		}
 
 		/// <inheritdoc/>
