Horde: Add a "GET" version of the /api/v1/objects/xxx/tree endpoint which does not take a request body.

diff --git a/Engine/Source/Programs/Horde/HordeServer/Storage/Controllers/ObjectsController.cs b/Engine/Source/Programs/Horde/HordeServer/Storage/Controllers/ObjectsController.cs
--- a/Engine/Source/Programs/Horde/HordeServer/Storage/Controllers/ObjectsController.cs
+++ b/Engine/Source/Programs/Horde/HordeServer/Storage/Controllers/ObjectsController.cs
@@ -203,6 +203,18 @@ namespace HordeServer.Storage.Controllers
 			return Response;
 		}
 
+		/// <summary>
+		/// Gets an object tree underneath a given root
+		/// </summary>
+		/// <param name="NamespaceId">Namespace for the operation</param>
+		/// <param name="Hash">The root object hash</param>
+		[HttpGet]
+		[Route("/api/v1/objects/{NamespaceId}/{Hash}/tree")]
+		public Task GetObjectTreeAsync(NamespaceId NamespaceId, IoHash Hash)
+		{
+			return GetObjectTreeAsync(NamespaceId, Hash, new GetObjectTreeRequest());
+		}
+
 		/// <summary>
 		/// Gets an object tree underneath a given root
 		/// </summary>
