Horde: Take a lock when updating an issue, to prevent triggering multiple server pods doing the same thing.

diff --git a/Engine/Source/Programs/Horde/HordeServer/Issues/IIssueCollection.cs b/Engine/Source/Programs/Horde/HordeServer/Issues/IIssueCollection.cs
--- a/Engine/Source/Programs/Horde/HordeServer/Issues/IIssueCollection.cs
+++ b/Engine/Source/Programs/Horde/HordeServer/Issues/IIssueCollection.cs
@@ -478,7 +478,7 @@ namespace HordeServer.Collections
 		/// <param name="NewSuspects">New suspects for the span</param>
 		/// <param name="NewModified">New modified flag for the span</param>
 		/// <returns>The updated span, or null on failure</returns>
-		Task<IIssueSpan?> UpdateSpanAsync(IIssueSpan Span, IssueSeverity? NewSeverity = null, NewIssueStepData? NewLastSuccess = null, NewIssueStepData? NewFailure = null, NewIssueStepData? NewNextSuccess = null, List<NewIssueSpanSuspectData>? NewSuspects = null, bool? NewModified = null);
+		Task<IIssueSpan?> TryUpdateSpanAsync(IIssueSpan Span, IssueSeverity? NewSeverity = null, NewIssueStepData? NewLastSuccess = null, NewIssueStepData? NewFailure = null, NewIssueStepData? NewNextSuccess = null, List<NewIssueSpanSuspectData>? NewSuspects = null, bool? NewModified = null);
 
 		/// <summary>
 		/// Updates the given span
@@ -520,7 +520,7 @@ namespace HordeServer.Collections
 		/// Finds a span marked as modified
 		/// </summary>
 		/// <returns></returns>
-		Task<IIssueSpan?> FindModifiedSpanAsync();
+		Task<List<IIssueSpan>> FindModifiedSpansAsync();
 
 		#endregion
 
diff --git a/Engine/Source/Programs/Horde/HordeServer/Issues/Impl/IssueCollection.cs b/Engine/Source/Programs/Horde/HordeServer/Issues/Impl/IssueCollection.cs
--- a/Engine/Source/Programs/Horde/HordeServer/Issues/Impl/IssueCollection.cs
+++ b/Engine/Source/Programs/Horde/HordeServer/Issues/Impl/IssueCollection.cs
@@ -1254,7 +1254,7 @@ namespace HordeServer.Collections.Impl
 		}
 
 		/// <inheritdoc/>
-		public async Task<IIssueSpan?> UpdateSpanAsync(IIssueSpan Span, IssueSeverity? NewSeverity = null, NewIssueStepData? NewLastSuccess = null, NewIssueStepData? NewFailure = null, NewIssueStepData? NewNextSuccess = null, List<NewIssueSpanSuspectData>? NewSuspects = null, bool? NewModified = null)
+		public async Task<IIssueSpan?> TryUpdateSpanAsync(IIssueSpan Span, IssueSeverity? NewSeverity = null, NewIssueStepData? NewLastSuccess = null, NewIssueStepData? NewFailure = null, NewIssueStepData? NewNextSuccess = null, List<NewIssueSpanSuspectData>? NewSuspects = null, bool? NewModified = null)
 		{
 			List<UpdateDefinition<IssueSpan>> Updates = new List<UpdateDefinition<IssueSpan>>();
 			if (NewSeverity != null)
@@ -1385,10 +1385,9 @@ namespace HordeServer.Collections.Impl
 		}
 
 		/// <inheritdoc/>
-		public async Task<IIssueSpan?> FindModifiedSpanAsync()
+		public async Task<List<IIssueSpan>> FindModifiedSpansAsync()
 		{
-			IssueSpan? IssueSpan = await IssueSpans.Find(x => x.Modified).FirstOrDefaultAsync();
-			return IssueSpan;
+			return await IssueSpans.Find(x => x.Modified).ToListAsync<IssueSpan, IIssueSpan>();
 		}
 
 		#endregion
diff --git a/Engine/Source/Programs/Horde/HordeServer/Issues/Impl/IssueService.cs b/Engine/Source/Programs/Horde/HordeServer/Issues/Impl/IssueService.cs
--- a/Engine/Source/Programs/Horde/HordeServer/Issues/Impl/IssueService.cs
+++ b/Engine/Source/Programs/Horde/HordeServer/Issues/Impl/IssueService.cs
@@ -1,6 +1,7 @@
 // Copyright Epic Games, Inc. All Rights Reserved.
 
 using EpicGames.Core;
+using EpicGames.Redis.Utility;
 using HordeCommon;
 using HordeServer.Collections;
 using HordeServer.IssueHandlers;
@@ -136,39 +137,13 @@ namespace HordeServer.Services.Impl
 		/// </summary>
 		const int MaxChanges = 250;
 
-		/// <summary>
-		/// Collection of job documents
-		/// </summary>
+		RedisService RedisService;
 		IJobCollection Jobs;
-
-		/// <summary>
-		/// Collection of job step ref documents
-		/// </summary>
 		IJobStepRefCollection JobStepRefs;
-
-		/// <summary>
-		/// Collection of issue documents
-		/// </summary>
 		IIssueCollection IssueCollection;
-
-		/// <summary>
-		/// Collection of stream documents
-		/// </summary>
 		StreamService Streams;
-
-		/// <summary>
-		/// Collection of user documents
-		/// </summary>
 		IUserCollection UserCollection;
-
-		/// <summary>
-		/// Collection of commit documents
-		/// </summary>
 		IPerforceService Perforce;
-
-		/// <summary>
-		/// The log file service instance
-		/// </summary>
 		ILogFileService LogFileService;
 
 		/// <summary>
@@ -209,18 +184,11 @@ namespace HordeServer.Services.Impl
 		/// <summary>
 		/// Constructor
 		/// </summary>
-		/// <param name="DatabaseService"></param>
-		/// <param name="IssueCollection">Collection of issue documents</param>
-		/// <param name="Jobs">Collection of job documents</param>
-		/// <param name="JobStepRefs">Collection of jobstepref documents</param>
-		/// <param name="Streams">Collection of stream documents</param>
-		/// <param name="UserCollection"></param>
-		/// <param name="Perforce">The perforce service instance</param>
-		/// <param name="LogFileService">Collection of event documents</param>
-		/// <param name="Logger">The logger instance</param>
-		public IssueService(DatabaseService DatabaseService, IIssueCollection IssueCollection, IJobCollection Jobs, IJobStepRefCollection JobStepRefs, StreamService Streams, IUserCollection UserCollection, IPerforceService Perforce, ILogFileService LogFileService, ILogger<IssueService> Logger)
+		public IssueService(DatabaseService DatabaseService, RedisService RedisService, IIssueCollection IssueCollection, IJobCollection Jobs, IJobStepRefCollection JobStepRefs, StreamService Streams, IUserCollection UserCollection, IPerforceService Perforce, ILogFileService LogFileService, ILogger<IssueService> Logger)
 			: base(DatabaseService, ObjectId.Parse("609542152fb0794700a6c3df"), Logger)
 		{
+			this.RedisService = RedisService;
+
 			Type[] IssueTypes = Assembly.GetExecutingAssembly().GetTypes().Where(x => !x.IsAbstract && typeof(IIssue).IsAssignableFrom(x)).ToArray();
 			foreach (Type IssueType in IssueTypes)
 			{
@@ -619,7 +587,7 @@ namespace HordeServer.Services.Impl
 			}
 
 			// Update all the unassigned spans
-			await AttachSpansToIssues();
+			await AttachSpansToIssuesAsync();
 		}
 
 		/// <summary>
@@ -741,7 +709,7 @@ namespace HordeServer.Services.Impl
 						}
 
 						// Update the span
-						IIssueSpan? NewSpan = await IssueCollection.UpdateSpanAsync(OpenSpan, NewSeverity: NewSeverity, NewFailure: NewFailure, NewModified: (bMarkAsModified? (bool?)true : null));
+						IIssueSpan? NewSpan = await IssueCollection.TryUpdateSpanAsync(OpenSpan, NewSeverity: NewSeverity, NewFailure: NewFailure, NewModified: (bMarkAsModified? (bool?)true : null));
 						if (NewSpan == null)
 						{
 							return false;
@@ -890,73 +858,93 @@ namespace HordeServer.Services.Impl
 		/// Find all the spans which are not currently attached to an issue, and try to attach them
 		/// </summary>
 		/// <returns></returns>
-		async Task AttachSpansToIssues()
+		async Task AttachSpansToIssuesAsync()
 		{
 			IIssueSequenceToken Token = await IssueCollection.GetSequenceTokenAsync();
-			for (; ; )
+			for(; ;)
 			{
-				// Find an unassigned span
-				IIssueSpan? Span = await IssueCollection.FindModifiedSpanAsync();
-				if (Span == null)
+				bool bFoundSpan = false;
+
+				List<IIssueSpan> Spans = await IssueCollection.FindModifiedSpansAsync();
+				foreach (IIssueSpan Span in Spans)
+				{
+					await using RedisLock Lock = new RedisLock(RedisService.Database, $"issues/locks/{Span.Id}");
+					if (await Lock.AcquireAsync(TimeSpan.FromMinutes(1)))
+					{
+						await TryUpdateModifiedSpanAsync(Span, Token);
+						bFoundSpan = true;
+					}
+				}
+
+				if (!bFoundSpan)
 				{
 					break;
 				}
+			}
+		}
 
-				// Try to find an issue for it
-				IIssue? Issue;
-				if (Span.IssueId == null)
+		/// <summary>
+		/// Updates information derived from an updated span, and clear the modified flag on it
+		/// </summary>
+		/// <param name="Span"></param>
+		/// <param name="Token"></param>
+		/// <returns></returns>
+		async ValueTask<IIssueSpan?> TryUpdateModifiedSpanAsync(IIssueSpan Span, IIssueSequenceToken Token)
+		{
+			// Try to find an issue for it
+			IIssue? Issue;
+			if (Span.IssueId == null)
+			{
+				Issue = await FindIssueForSpanAsync(Span);
+				if (Issue == null)
 				{
-					Issue = await FindIssueForSpanAsync(Span);
+					Issue = await IssueCollection.AddIssueAsync(Token, GetSummary(Span.Fingerprint, Span.Severity), Span);
 					if (Issue == null)
 					{
-						Issue = await IssueCollection.AddIssueAsync(Token, GetSummary(Span.Fingerprint, Span.Severity), Span);
-						if (Issue == null)
-						{
-							await Token.ResetAsync();
-							continue;
-						}
-					}
-					else
-					{
-						NewIssueFingerprint NewFingerprint = NewIssueFingerprint.Merge(Issue.Fingerprint, Span.Fingerprint);
-						IssueSeverity NewSeverity = (Issue.Severity > Span.Severity) ? Issue.Severity : Span.Severity;
-						string NewSummary = GetSummary(NewFingerprint, NewSeverity);
-
-						if (!await IssueCollection.AttachSpanToIssueAsync(Token, Span, Issue, NewSeverity, NewFingerprint, NewSummary))
-						{
-							await Token.ResetAsync();
-							continue;
-						}
+						await Token.ResetAsync();
+						return null;
 					}
 				}
 				else
 				{
-					// Update the issue that it belongs to
-					Issue = await GetIssueAsync(Span.IssueId.Value);
-					if (Issue != null)
+					NewIssueFingerprint NewFingerprint = NewIssueFingerprint.Merge(Issue.Fingerprint, Span.Fingerprint);
+					IssueSeverity NewSeverity = (Issue.Severity > Span.Severity) ? Issue.Severity : Span.Severity;
+					string NewSummary = GetSummary(NewFingerprint, NewSeverity);
+
+					if (!await IssueCollection.AttachSpanToIssueAsync(Token, Span, Issue, NewSeverity, NewFingerprint, NewSummary))
 					{
-						Issue = await UpdateIssueDerivedDataAsync(Issue);
-						if (Issue == null)
-						{
-							continue;
-						}
+						await Token.ResetAsync();
+						return null;
 					}
 				}
-
-				// Perform any shared issue updates
+			}
+			else
+			{
+				// Update the issue that it belongs to
+				Issue = await GetIssueAsync(Span.IssueId.Value);
 				if (Issue != null)
 				{
-					Issue = await UpdateResolvedStateAsync(Issue, Span);
-					if(Issue == null)
+					Issue = await UpdateIssueDerivedDataAsync(Issue);
+					if (Issue == null)
 					{
-						continue;
+						return null;
 					}
-					OnIssueUpdated?.Invoke(Issue);
 				}
+			}
 
-				// Clear the modified flag
-				await IssueCollection.UpdateSpanAsync(Span, NewModified: false);
+			// Perform any shared issue updates
+			if (Issue != null)
+			{
+				Issue = await UpdateResolvedStateAsync(Issue, Span);
+				if(Issue == null)
+				{
+					return null;
+				}
+				OnIssueUpdated?.Invoke(Issue);
 			}
+
+			// Clear the modified flag
+			return await IssueCollection.TryUpdateSpanAsync(Span, NewModified: false);
 		}
 
 		async Task<IIssue?> UpdateResolvedStateAsync(IIssue? Issue, IIssueSpan Span)
@@ -1187,7 +1175,7 @@ namespace HordeServer.Services.Impl
 					NewIssueStepData NewLastSuccess = new NewIssueStepData(Job.Change, IssueSeverity.Unspecified, Job.Name, Job.Id, Batch.Id, Step.Id, Step.StartTimeUtc ?? default, Step.LogId, false);
 					List<NewIssueSpanSuspectData> NewSuspects = await FindSuspectsForSpanAsync(Stream, Span.Fingerprint, Job.Change + 1, Span.FirstFailure.Change);
 
-					if (await IssueCollection.UpdateSpanAsync(Span, NewLastSuccess: NewLastSuccess, NewSuspects: NewSuspects, NewModified: true) == null)
+					if (await IssueCollection.TryUpdateSpanAsync(Span, NewLastSuccess: NewLastSuccess, NewSuspects: NewSuspects, NewModified: true) == null)
 					{
 						return false;
 					}
@@ -1197,7 +1185,7 @@ namespace HordeServer.Services.Impl
 				else if (Job.Change > Span.LastFailure.Change && (Span.NextSuccess == null || Job.Change < Span.NextSuccess.Change))
 				{
 					NewIssueStepData NewNextSucccess = new NewIssueStepData(Job.Change, IssueSeverity.Unspecified, Job.Name, Job.Id, Batch.Id, Step.Id, Step.StartTimeUtc ?? default, Step.LogId, false);
-					if (await IssueCollection.UpdateSpanAsync(Span, NewNextSuccess: NewNextSucccess, NewModified: true) == null)
+					if (await IssueCollection.TryUpdateSpanAsync(Span, NewNextSuccess: NewNextSucccess, NewModified: true) == null)
 					{
 						return false;
 					}
diff --git a/Engine/Source/Programs/Shared/EpicGames.Redis/Utility/RedisLock.cs b/Engine/Source/Programs/Shared/EpicGames.Redis/Utility/RedisLock.cs
new file mode 100644
--- /dev/null
+++ b/Engine/Source/Programs/Shared/EpicGames.Redis/Utility/RedisLock.cs
@@ -0,0 +1,106 @@
+using StackExchange.Redis;
+using System;
+using System.Collections.Generic;
+using System.Text;
+using System.Threading;
+using System.Threading.Tasks;
+
+namespace EpicGames.Redis.Utility
+{
+	/// <summary>
+	/// Implements a named single-entry lock which expires after a period of time if the process terminates.
+	/// </summary>
+	public class RedisLock : IAsyncDisposable, IDisposable
+	{
+		IDatabase Database;
+		RedisKey Key;
+		CancellationTokenSource? CancellationSource;
+		Task? BackgroundTask;
+
+		/// <summary>
+		/// Constructor
+		/// </summary>
+		/// <param name="Database"></param>
+		/// <param name="Key"></param>
+		public RedisLock(IDatabase Database, RedisKey Key)
+		{
+			this.Database = Database;
+			this.Key = Key;
+		}
+
+		/// <inheritdoc/>
+		public void Dispose()
+		{
+			DisposeAsync().AsTask().Wait();
+		}
+
+		/// <inheritdoc/>
+		public async ValueTask DisposeAsync()
+		{
+			if (BackgroundTask != null)
+			{
+				CancellationSource!.Cancel();
+				await BackgroundTask;
+				CancellationSource.Dispose();
+				BackgroundTask = null;
+			}
+		}
+
+		/// <summary>
+		/// Attempts to acquire the lock for the given period of time. The lock will be renewed once half of this interval has elapsed.
+		/// </summary>
+		/// <param name="Duration">Time after which the lock expires</param>
+		/// <returns>True if the lock was acquired, false if another service already has it</returns>
+		public async ValueTask<bool> AcquireAsync(TimeSpan Duration)
+		{
+			if (await Database.StringSetAsync(Key, RedisValue.EmptyString, Duration, When.NotExists))
+			{
+				CancellationSource = new CancellationTokenSource();
+				BackgroundTask = Task.Run(() => RenewAsync(Duration, CancellationSource.Token));
+				return true;
+			}
+			return false;
+		}
+
+		/// <summary>
+		/// Background task which renews the lock while the service is running
+		/// </summary>
+		/// <param name="Duration"></param>
+		/// <param name="CancellationToken"></param>
+		/// <returns></returns>
+		async Task RenewAsync(TimeSpan Duration, CancellationToken CancellationToken)
+		{
+			for (; ; )
+			{
+				await Task.Delay(Duration / 2, CancellationToken).ContinueWith(x => { }); // Do not throw
+				if (CancellationToken.IsCancellationRequested)
+				{
+					await Database.StringSetAsync(Key, RedisValue.Null);
+					break;
+				}
+				if (!await Database.StringSetAsync(Key, RedisValue.EmptyString, Duration, When.Exists))
+				{
+					break;
+				}
+			}
+		}
+	}
+
+	/// <summary>
+	/// Implements a named single-entry lock which expires after a period of time if the process terminates.
+	/// </summary>
+	/// <typeparam name="T">Type of the value identifying the lock uniqueness</typeparam>
+	public class RedisLock<T> : RedisLock
+	{
+		/// <summary>
+		/// Constructor
+		/// </summary>
+		/// <param name="Database"></param>
+		/// <param name="BaseKey"></param>
+		/// <param name="Value"></param>
+		public RedisLock(IDatabase Database, RedisKey BaseKey, T Value)
+			: base(Database, BaseKey.Append(RedisSerializer.Serialize<T>(Value).AsKey()))
+		{
+		}
+	}
+}
