Horde: Prevent posting Slack messages until we've updated the message state in the database. Should prevent messages being sent multiple times.

diff --git a/Engine/Source/Programs/Horde/HordeServer/Notifications/Impl/SlackNotificationSink.cs b/Engine/Source/Programs/Horde/HordeServer/Notifications/Impl/SlackNotificationSink.cs
--- a/Engine/Source/Programs/Horde/HordeServer/Notifications/Impl/SlackNotificationSink.cs
+++ b/Engine/Source/Programs/Horde/HordeServer/Notifications/Impl/SlackNotificationSink.cs
@@ -259,17 +259,30 @@ namespace HordeServer.Notifications.Impl
 
 		#region Message state 
 
-		async Task<MessageStateDocument?> GetMessageStateAsync(string Recipient, string EventId)
+		async Task<MessageStateDocument> AddOrUpdateMessageStateAsync(string Recipient, string EventId, UserId? UserId, string Digest)
 		{
+			ObjectId NewId = ObjectId.GenerateNewId();
+
 			FilterDefinition<MessageStateDocument> Filter = Builders<MessageStateDocument>.Filter.Eq(x => x.Recipient, Recipient) & Builders<MessageStateDocument>.Filter.Eq(x => x.EventId, EventId);
-			return await MessageStates.Find(Filter).FirstOrDefaultAsync();
+			UpdateDefinition<MessageStateDocument> Update = Builders<MessageStateDocument>.Update.SetOnInsert(x => x.Id, NewId).Set(x => x.UserId, UserId).Set(x => x.Digest, Digest);
+
+			MessageStateDocument? PrevState = await MessageStates.FindOneAndUpdateAsync(Filter, Update, new FindOneAndUpdateOptions<MessageStateDocument> { IsUpsert = true });
+			if (PrevState == null)
+			{
+				PrevState = new MessageStateDocument();
+				PrevState.Id = NewId;
+				PrevState.Recipient = Recipient;
+				PrevState.EventId = EventId;
+				PrevState.UserId = UserId;
+			}
+			return PrevState;
 		}
 
-		async Task<MessageStateDocument> SetMessageStateAsync(string Recipient, string EventId, UserId? UserId, string Channel, string Ts, string Digest)
+		async Task SetMessageTimestampAsync(ObjectId MessageId, string Channel, string Ts)
 		{
-			FilterDefinition<MessageStateDocument> Filter = Builders<MessageStateDocument>.Filter.Eq(x => x.Recipient, Recipient) & Builders<MessageStateDocument>.Filter.Eq(x => x.EventId, EventId);
-			UpdateDefinition<MessageStateDocument> Update = Builders<MessageStateDocument>.Update.SetOnInsert(x => x.Id, ObjectId.GenerateNewId()).Set(x => x.UserId, UserId).Set(x => x.Channel, Channel).Set(x => x.Ts, Ts).Set(x => x.Digest, Digest);
-			return await MessageStates.FindOneAndUpdateAsync(Filter, Update, new FindOneAndUpdateOptions<MessageStateDocument> { IsUpsert = true, ReturnDocument = ReturnDocument.After });
+			FilterDefinition<MessageStateDocument> Filter = Builders<MessageStateDocument>.Filter.Eq(x => x.Id, MessageId) & Builders<MessageStateDocument>.Filter.Eq(x => x.Ts, String.Empty);
+			UpdateDefinition<MessageStateDocument> Update = Builders<MessageStateDocument>.Update.Set(x => x.Channel, Channel).Set(x => x.Ts, Ts);
+			await MessageStates.FindOneAndUpdateAsync(Filter, Update);
 		}
 
 		#endregion
@@ -1149,7 +1162,7 @@ namespace HordeServer.Notifications.Impl
 			}
 
 			BlockKitMessage Message = new BlockKitMessage();
-			Message.Recipient = Recipient;
+			Message.Channel = Recipient;
 			Message.Text = Text;
 			if (Blocks != null)
 			{
@@ -1184,35 +1197,26 @@ namespace HordeServer.Notifications.Impl
 
 			string RequestDigest = ContentHash.MD5(JsonSerializer.Serialize(Message, new JsonSerializerOptions { IgnoreNullValues = true })).ToString();
 
-			PostMessageResponse Response;
-
-			MessageStateDocument? CurrentState = await GetMessageStateAsync(Recipient, EventId);
-			if (CurrentState == null)
+			MessageStateDocument PrevState = await AddOrUpdateMessageStateAsync(Recipient, EventId, UserId, RequestDigest);
+			if (String.IsNullOrEmpty(PrevState.Ts))
 			{
 				Logger.LogInformation("Sending new slack message to {SlackUser}", Recipient);
-				Message.Recipient = Recipient;
+				Message.Channel = Recipient;
 				Message.Ts = null;
-				Response = await SendRequestAsync<PostMessageResponse>(PostMessageUrl, Message);
-			}
-			else if (RequestDigest == CurrentState.Digest)
-			{
-				return;
-			}
-			else
-			{
-				Logger.LogInformation("Updating existing slack message for user {SlackUser} ({Channel}, {MessageTs})", Recipient, CurrentState.Channel, CurrentState.Ts);
-				Message.Recipient = CurrentState.Channel;
-				Message.Ts = CurrentState.Ts;
-				Response = await SendRequestAsync<PostMessageResponse>(UpdateMessageUrl, Message);
-			}
 
-			if (Response.Ok && Response.Channel != null && Response.Ts != null)
-			{
-				await SetMessageStateAsync(Recipient, EventId, UserId, Response.Channel, Response.Ts, RequestDigest);
+				PostMessageResponse Response = await SendRequestAsync<PostMessageResponse>(PostMessageUrl, Message);
+				if (String.IsNullOrEmpty(Response.Ts) || String.IsNullOrEmpty(Response.Channel))
+				{
+					Logger.LogWarning("Missing 'ts' or 'channel' field on slack response");
+				}
+				await SetMessageTimestampAsync(PrevState.Id, Response.Channel ?? String.Empty, Response.Ts ?? String.Empty);
 			}
 			else
 			{
-				Logger.LogWarning("Missing fields on slack response (channel: {Channel}, ts: {MessageTs})", Response.Channel, Response.Ts);
+				Logger.LogInformation("Updating existing slack message for user {SlackUser} ({Channel}, {MessageTs})", Recipient, PrevState.Channel, PrevState.Ts);
+				Message.Channel = PrevState.Channel;
+				Message.Ts = PrevState.Ts;
+				await SendRequestAsync<PostMessageResponse>(UpdateMessageUrl, Message);
 			}
 		}
 
diff --git a/Engine/Source/Programs/Horde/HordeServer/Utilities/BlockKit/BlockKitMessage.cs b/Engine/Source/Programs/Horde/HordeServer/Utilities/BlockKit/BlockKitMessage.cs
--- a/Engine/Source/Programs/Horde/HordeServer/Utilities/BlockKit/BlockKitMessage.cs
+++ b/Engine/Source/Programs/Horde/HordeServer/Utilities/BlockKit/BlockKitMessage.cs
@@ -14,7 +14,7 @@ namespace HordeServer.Utilities.Slack.BlockKit
 		/// The channel ID or user ID to deliver the notification to.
 		/// </summary>
 		[JsonPropertyName("channel")]
-		public string? Recipient { get; set; }
+		public string? Channel { get; set; }
 
 		/// <summary>
 		/// Timestamp of the message to update
