Horde: Add a 'Description' field to issues, allowing users to add additional information about it.

diff --git a/Engine/Source/Programs/Horde/HordeServer.Compute.Tests/IssueServiceTests.cs b/Engine/Source/Programs/Horde/HordeServer.Compute.Tests/IssueServiceTests.cs
--- a/Engine/Source/Programs/Horde/HordeServer.Compute.Tests/IssueServiceTests.cs
+++ b/Engine/Source/Programs/Horde/HordeServer.Compute.Tests/IssueServiceTests.cs
@@ -330,6 +330,17 @@ namespace HordeServerTests
 				List<IIssue> Issues = await IssueService.FindIssuesAsync();
 				Assert.AreEqual(4, Issues.Count);
 			}
+
+			// #5
+			// Add a description to the issue
+			{
+				List<IIssue> Issues = await IssueService.FindIssuesAsync();
+
+				IIssue Issue = Issues[0];
+				await IssueService.UpdateIssueAsync(Issue.Id, Description: "Hello world!");
+				IIssue? NewIssue = await IssueService.GetIssueAsync(Issue.Id);
+				Assert.AreEqual(NewIssue?.Description, "Hello world!");
+			}
 		}
 
 		[TestMethod]
diff --git a/Engine/Source/Programs/Horde/HordeServer/Api/Issues.cs b/Engine/Source/Programs/Horde/HordeServer/Api/Issues.cs
--- a/Engine/Source/Programs/Horde/HordeServer/Api/Issues.cs
+++ b/Engine/Source/Programs/Horde/HordeServer/Api/Issues.cs
@@ -371,9 +371,9 @@ namespace HordeServer.Api
 		public string Summary { get; set; }
 
 		/// <summary>
-		/// Details text describing the issue
+		/// Detailed description text
 		/// </summary>
-		public string? Details { get; set; }
+		public string? Description { get; set; }
 
 		/// <summary>
 		/// Severity of this issue
@@ -498,6 +498,7 @@ namespace HordeServer.Api
 			this.CreatedAt = Issue.CreatedAt;
 			this.RetrievedAt = DateTime.UtcNow;
 			this.Summary = String.IsNullOrEmpty(Issue.UserSummary)? Issue.Summary : Issue.UserSummary;
+			this.Description = Issue.Description;
 			this.Severity = Issue.Severity;
 			this.Owner = Details.Owner?.Login;
 			this.OwnerId = (Details.Owner == null)? null : Details.Owner.Id.ToString();
@@ -553,6 +554,11 @@ namespace HordeServer.Api
 		/// </summary>
 		public string? Summary { get; set; }
 
+		/// <summary>
+		/// Description of the issue
+		/// </summary>
+		public string? Description { get; set; }
+
 		/// <summary>
 		/// New user id for owner of the issue, can be cleared by passing empty string
 		/// </summary>
diff --git a/Engine/Source/Programs/Horde/HordeServer/Controllers/IssuesController.cs b/Engine/Source/Programs/Horde/HordeServer/Controllers/IssuesController.cs
--- a/Engine/Source/Programs/Horde/HordeServer/Controllers/IssuesController.cs
+++ b/Engine/Source/Programs/Horde/HordeServer/Controllers/IssuesController.cs
@@ -413,7 +413,7 @@ namespace HordeServer.Controllers
 				NewResolvedById = Request.Resolved.Value ? User.GetUserId() : UserId.Empty;
 			}
 
-			if (!await IssueService.UpdateIssueAsync(IssueId, Request.Summary, NewOwnerId, NewNominatedById, Request.Acknowledged, NewDeclinedById, Request.FixChange, NewResolvedById))
+			if (!await IssueService.UpdateIssueAsync(IssueId, Request.Summary, Request.Description, NewOwnerId, NewNominatedById, Request.Acknowledged, NewDeclinedById, Request.FixChange, NewResolvedById))
 			{
 				return NotFound();
 			}
diff --git a/Engine/Source/Programs/Horde/HordeServer/Controllers/UgsController.cs b/Engine/Source/Programs/Horde/HordeServer/Controllers/UgsController.cs
--- a/Engine/Source/Programs/Horde/HordeServer/Controllers/UgsController.cs
+++ b/Engine/Source/Programs/Horde/HordeServer/Controllers/UgsController.cs
@@ -411,7 +411,7 @@ namespace HordeServer.Controllers
 				NewResolvedById = Request.Resolved.Value ? IIssue.ResolvedByUnknownId : UserId.Empty;
 			}
 
-			if (!await IssueService.UpdateIssueAsync(IssueId, null, NewOwnerId, NewNominatedById, Request.Acknowledged, NewDeclinedById, Request.FixChange, NewResolvedById))
+			if (!await IssueService.UpdateIssueAsync(IssueId, null, null, NewOwnerId, NewNominatedById, Request.Acknowledged, NewDeclinedById, Request.FixChange, NewResolvedById))
 			{
 				return NotFound();
 			}
diff --git a/Engine/Source/Programs/Horde/HordeServer/Issues/IIssue.cs b/Engine/Source/Programs/Horde/HordeServer/Issues/IIssue.cs
--- a/Engine/Source/Programs/Horde/HordeServer/Issues/IIssue.cs
+++ b/Engine/Source/Programs/Horde/HordeServer/Issues/IIssue.cs
@@ -72,6 +72,11 @@ namespace HordeServer.Models
 		/// </summary>
 		public string? UserSummary { get; }
 
+		/// <summary>
+		/// Notes set by the user
+		/// </summary>
+		public string? Description { get; }
+
 		/// <summary>
 		/// Fingerprint for this issue
 		/// </summary>
diff --git a/Engine/Source/Programs/Horde/HordeServer/Issues/IIssueCollection.cs b/Engine/Source/Programs/Horde/HordeServer/Issues/IIssueCollection.cs
--- a/Engine/Source/Programs/Horde/HordeServer/Issues/IIssueCollection.cs
+++ b/Engine/Source/Programs/Horde/HordeServer/Issues/IIssueCollection.cs
@@ -422,6 +422,7 @@ namespace HordeServer.Collections
 		/// <param name="NewSeverity">New severity for the issue</param>
 		/// <param name="NewSummary">New summary for the issue</param>
 		/// <param name="NewUserSummary">New user summary for the issue</param>
+		/// <param name="NewDescription">New description for the issue</param>
 		/// <param name="NewOwnerId">New owner of the issue</param>
 		/// <param name="NewNominatedById">Person that nominated the new owner</param>
 		/// <param name="NewAcknowledged">Whether the issue has been acknowledged</param>
@@ -432,7 +433,7 @@ namespace HordeServer.Collections
 		/// <param name="NewLastSeenAt"></param>
 		/// <param name="NewNotifySuspects">Whether all suspects should be notified about this issue</param>
 		/// <returns>True if the issue was updated</returns>
-		Task<IIssue?> UpdateIssueAsync(IIssue Issue, IssueSeverity? NewSeverity = null, string? NewSummary = null, string? NewUserSummary = null, UserId? NewOwnerId = null, UserId? NewNominatedById = null, bool? NewAcknowledged = null, UserId? NewDeclinedById = null, int? NewFixChange = null, Dictionary<StreamId, bool>? NewFixStreamIds = null, UserId? NewResolvedById = null, DateTime? NewLastSeenAt = null, bool? NewNotifySuspects = null);
+		Task<IIssue?> UpdateIssueAsync(IIssue Issue, IssueSeverity? NewSeverity = null, string? NewSummary = null, string? NewUserSummary = null, string? NewDescription = null, UserId? NewOwnerId = null, UserId? NewNominatedById = null, bool? NewAcknowledged = null, UserId? NewDeclinedById = null, int? NewFixChange = null, Dictionary<StreamId, bool>? NewFixStreamIds = null, UserId? NewResolvedById = null, DateTime? NewLastSeenAt = null, bool? NewNotifySuspects = null);
 
 		/// <summary>
 		/// Updates the list of suspects for an issue, and optionally sets the resolved state
diff --git a/Engine/Source/Programs/Horde/HordeServer/Issues/IIssueService.cs b/Engine/Source/Programs/Horde/HordeServer/Issues/IIssueService.cs
--- a/Engine/Source/Programs/Horde/HordeServer/Issues/IIssueService.cs
+++ b/Engine/Source/Programs/Horde/HordeServer/Issues/IIssueService.cs
@@ -129,8 +129,9 @@ namespace HordeServer.Services
 		/// <summary>
 		/// Updates the state of an issue
 		/// </summary>
-		/// <param name="Id">The current issue id</param>
+		/// <param name="Id">The current issue id</param>s
 		/// <param name="Summary">New summary for the issue</param>
+		/// <param name="Description">New description for the issue</param>
 		/// <param name="OwnerId">New owner of the issue</param>
 		/// <param name="NominatedById">Person that nominated the new owner</param>
 		/// <param name="Acknowledged">Whether the issue has been acknowledged</param>
@@ -138,7 +139,7 @@ namespace HordeServer.Services
 		/// <param name="FixChange">Fix changelist for the issue</param>
 		/// <param name="ResolvedById">Whether the issue has been resolved</param>
 		/// <returns>True if the issue was updated</returns>
-		Task<bool> UpdateIssueAsync(int Id, string? Summary = null, UserId? OwnerId = null, UserId? NominatedById = null, bool? Acknowledged = null, UserId? DeclinedById = null, int? FixChange = null, UserId? ResolvedById = null);
+		Task<bool> UpdateIssueAsync(int Id, string? Summary = null, string? Description = null, UserId? OwnerId = null, UserId? NominatedById = null, bool? Acknowledged = null, UserId? DeclinedById = null, int? FixChange = null, UserId? ResolvedById = null);
 
 		/// <summary>
 		/// Marks a step as complete
diff --git a/Engine/Source/Programs/Horde/HordeServer/Issues/Impl/IssueService.cs b/Engine/Source/Programs/Horde/HordeServer/Issues/Impl/IssueService.cs
--- a/Engine/Source/Programs/Horde/HordeServer/Issues/Impl/IssueService.cs
+++ b/Engine/Source/Programs/Horde/HordeServer/Issues/Impl/IssueService.cs
@@ -392,7 +392,7 @@ namespace HordeServer.Services.Impl
 		}
 
 		/// <inheritdoc/>
-		public async Task<bool> UpdateIssueAsync(int Id, string? UserSummary = null, UserId? OwnerId = null, UserId? NominatedById = null, bool? Acknowledged = null, UserId? DeclinedById = null, int? FixChange = null, UserId? ResolvedById = null)
+		public async Task<bool> UpdateIssueAsync(int Id, string? UserSummary = null, string? Description = null, UserId? OwnerId = null, UserId? NominatedById = null, bool? Acknowledged = null, UserId? DeclinedById = null, int? FixChange = null, UserId? ResolvedById = null)
 		{
 			Dictionary<StreamId, bool>? NewFixStreamIds = null;
 			for (; ; )
@@ -408,7 +408,7 @@ namespace HordeServer.Services.Impl
 					NewFixStreamIds = await GetFixStreamsAsync(Issue.Id, FixChange.Value, NewFixStreamIds ?? ((FixChange == Issue.FixChange)? Issue.GetFixStreamIds() : null));
 				}
 
-				Issue = await IssueCollection.UpdateIssueAsync(Issue, NewUserSummary: UserSummary, NewOwnerId: OwnerId ?? ResolvedById, NewNominatedById: NominatedById, NewAcknowledged: Acknowledged, NewDeclinedById: DeclinedById, NewFixChange: FixChange, NewFixStreamIds: NewFixStreamIds, NewResolvedById: ResolvedById);
+				Issue = await IssueCollection.UpdateIssueAsync(Issue, NewUserSummary: UserSummary, NewDescription: Description, NewOwnerId: OwnerId ?? ResolvedById, NewNominatedById: NominatedById, NewAcknowledged: Acknowledged, NewDeclinedById: DeclinedById, NewFixChange: FixChange, NewFixStreamIds: NewFixStreamIds, NewResolvedById: ResolvedById);
 
 				if(Issue != null)
 				{
