Mini refactor of issue service to allow for merging/splitting issues.

- Issue updates are now performed while holding a lock created in Redis.
- Issues can now contain multiple fingerprints, determined as a superset of all the spans merged into it.
- Suspects are now stored outside the primary issue document to minimize document size.
- Derived data is now computed by the issue service, allowing it to decide things with more context than is available inside the collection.
- Spans are always attached to an issue, rather being updated periodically whenever a flag is set.

diff --git a/Engine/Source/Programs/Horde/HordeServer.Compute.Tests/IssueServiceTests.cs b/Engine/Source/Programs/Horde/HordeServer.Compute.Tests/IssueServiceTests.cs
--- a/Engine/Source/Programs/Horde/HordeServer.Compute.Tests/IssueServiceTests.cs
+++ b/Engine/Source/Programs/Horde/HordeServer.Compute.Tests/IssueServiceTests.cs
@@ -419,7 +419,8 @@ namespace HordeServerTests
 				IIssue Issue = Issues[0];
 				List<IIssueSpan> Spans = await IssueService.GetIssueSpansAsync(Issue);
 				Assert.AreEqual(1, Spans.Count);
-				Assert.AreEqual("Compile", Issue.Fingerprint.Type);
+				Assert.AreEqual(1, Issue.Fingerprints.Count);
+				Assert.AreEqual("Compile", Issue.Fingerprints[0].Type);
 				Assert.AreEqual("Compile warnings in GameChat2Impl.h", Issue.Summary);
 			}
 		}
@@ -460,7 +461,8 @@ namespace HordeServerTests
 				IIssue Issue = Issues[0];
 				List<IIssueSpan> Spans = await IssueService.GetIssueSpansAsync(Issue);
 				Assert.AreEqual(1, Spans.Count);
-				Assert.AreEqual("Compile", Issue.Fingerprint.Type);
+				Assert.AreEqual(1, Issue.Fingerprints.Count);
+				Assert.AreEqual("Compile", Issue.Fingerprints[0].Type);
 				Assert.AreEqual("Compile errors in fog.cpp", Issue.Summary);
 
 				IIssueSpan Stream = Spans[0];
@@ -468,12 +470,9 @@ namespace HordeServerTests
 				Assert.AreEqual(null, Stream.NextSuccess?.Change);
 
 				IReadOnlyList<IIssueSuspect> Suspects = await IssueService.GetIssueSuspectsAsync(Issue);
+				Suspects = Suspects.OrderBy(x => x.Id).ToList();
 				Assert.AreEqual(3, Suspects.Count);
 
-				Assert.AreEqual(120, Suspects[0].Change);
-				Assert.AreEqual(TimId, Suspects[0].AuthorId);
-				//				Assert.AreEqual(null, Suspects[0].OriginatingChange);
-
 				Assert.AreEqual(75 /*115*/, Suspects[1].Change);
 				Assert.AreEqual(JerryId, Suspects[1].AuthorId);
 				//				Assert.AreEqual(75, Suspects[1].OriginatingChange);
@@ -482,6 +481,10 @@ namespace HordeServerTests
 				Assert.AreEqual(BobId, Suspects[2].AuthorId);
 				//			Assert.AreEqual(null, Suspects[2].OriginatingChange);
 
+				Assert.AreEqual(120, Suspects[0].Change);
+				Assert.AreEqual(TimId, Suspects[0].AuthorId);
+				//				Assert.AreEqual(null, Suspects[0].OriginatingChange);
+
 				List<IIssue> OpenIssues = await IssueService.FindIssuesAsync(Resolved: false);
 				Assert.AreEqual(1, OpenIssues.Count);
 			}
@@ -676,10 +679,11 @@ namespace HordeServerTests
 				Assert.AreEqual(1, Issues.Count);
 
 				IIssue Issue = Issues[0];
+				Assert.AreEqual(1, Issue.Fingerprints.Count);
 
 				CompileIssueHandler Handler = new CompileIssueHandler();
 
-				IIssueFingerprint Fingerprint = Issue.Fingerprint;
+				IIssueFingerprint Fingerprint = Issue.Fingerprints[0];
 				Assert.AreEqual(Handler.Type, Fingerprint.Type);
 				Assert.AreEqual(1, Fingerprint.Keys.Count);
 				Assert.AreEqual("FOO.CPP", Fingerprint.Keys.First());
@@ -823,7 +827,8 @@ namespace HordeServerTests
 				IIssue Issue = Issues[0];
 				List<IIssueSpan> Spans = await IssueService.GetIssueSpansAsync(Issue);
 				Assert.AreEqual(Spans.Count, 1);
-				Assert.AreEqual(Issue.Fingerprint.Type, "Symbol");
+				Assert.AreEqual(Issue.Fingerprints.Count, 1);
+				Assert.AreEqual(Issue.Fingerprints[0].Type, "Symbol");
 
 				IIssueSpan Stream = Spans[0];
 				Assert.AreEqual(105, Stream.LastSuccess?.Change);
@@ -885,7 +890,8 @@ namespace HordeServerTests
 				IIssue Issue = Issues[0];
 				List<IIssueSpan> Spans = await IssueService.GetIssueSpansAsync(Issue);
 				Assert.AreEqual(Spans.Count, 2);
-				Assert.AreEqual(Issue.Fingerprint.Type, "Symbol");
+				Assert.AreEqual(Issue.Fingerprints.Count, 1);
+				Assert.AreEqual(Issue.Fingerprints[0].Type, "Symbol");
 
 				IIssueSpan Span1 = Spans[0];
 				Assert.AreEqual(105, Span1.LastSuccess?.Change);
@@ -931,7 +937,8 @@ namespace HordeServerTests
 				IIssue Issue = Issues[0];
 				List<IIssueSpan> Spans = await IssueService.GetIssueSpansAsync(Issue);
 				Assert.AreEqual(Spans.Count, 1);
-				Assert.AreEqual(Issue.Fingerprint.Type, "Symbol");
+				Assert.AreEqual(Issue.Fingerprints.Count, 1);
+				Assert.AreEqual(Issue.Fingerprints[0].Type, "Symbol");
 
 				IIssueSpan Span = Spans[0];
 				Assert.AreEqual(105, Span.LastSuccess?.Change);
@@ -973,7 +980,8 @@ namespace HordeServerTests
 				IIssue Issue = Issues[0];
 				List<IIssueSpan> Spans = await IssueService.GetIssueSpansAsync(Issue);
 				Assert.AreEqual(1, Spans.Count);
-				Assert.AreEqual("Compile", Issue.Fingerprint.Type);
+				Assert.AreEqual(1, Issue.Fingerprints.Count);
+				Assert.AreEqual("Compile", Issue.Fingerprints[0].Type);
 			}
 		}
 
@@ -995,7 +1003,8 @@ namespace HordeServerTests
 			IIssue Issue = Issues[0];
 			List<IIssueSpan> Spans = await IssueService.GetIssueSpansAsync(Issue);
 			Assert.AreEqual(1, Spans.Count);
-			Assert.AreEqual("Copyright", Issue.Fingerprint.Type);
+			Assert.AreEqual(1, Issue.Fingerprints.Count);
+			Assert.AreEqual("Copyright", Issue.Fingerprints[0].Type);
 			Assert.AreEqual("Missing copyright notice in ToolchainInfo.cs", Issue.Summary);
 		}
 
@@ -1065,7 +1074,8 @@ namespace HordeServerTests
 				IIssue Issue = Issues[0];
 				List<IIssueSpan> Spans = await IssueService.GetIssueSpansAsync(Issue);
 				Assert.AreEqual(Spans.Count, 1);
-				Assert.AreEqual(Issue.Fingerprint.Type, "Gauntlet");
+				Assert.AreEqual(Issue.Fingerprints.Count, 1);
+				Assert.AreEqual(Issue.Fingerprints[0].Type, "Gauntlet");
 				Assert.AreEqual(Issue.Summary, "HLOD test failures: SectionFlags, SimpleMerge and SingleLODMerge");
 			}
 		}
diff --git a/Engine/Source/Programs/Horde/HordeServer/Issues/IIssue.cs b/Engine/Source/Programs/Horde/HordeServer/Issues/IIssue.cs
--- a/Engine/Source/Programs/Horde/HordeServer/Issues/IIssue.cs
+++ b/Engine/Source/Programs/Horde/HordeServer/Issues/IIssue.cs
@@ -80,7 +80,7 @@ namespace HordeServer.Models
 		/// <summary>
 		/// Fingerprint for this issue
 		/// </summary>
-		public IIssueFingerprint Fingerprint { get; }
+		public IReadOnlyList<IIssueFingerprint> Fingerprints { get; }
 
 		/// <summary>
 		/// Severity of this issue
diff --git a/Engine/Source/Programs/Horde/HordeServer/Issues/IIssueCollection.cs b/Engine/Source/Programs/Horde/HordeServer/Issues/IIssueCollection.cs
--- a/Engine/Source/Programs/Horde/HordeServer/Issues/IIssueCollection.cs
+++ b/Engine/Source/Programs/Horde/HordeServer/Issues/IIssueCollection.cs
@@ -32,21 +32,15 @@
 	/// <summary>
 	/// Fingerprint for an issue
 	/// </summary>
-	public class NewIssueFingerprint : IIssueFingerprint
+	public class NewIssueFingerprint : IIssueFingerprint, IEquatable<IIssueFingerprint>
 	{
-		/// <summary>
-		/// The type of issue
-		/// </summary>
+		/// <inheritdoc/>
 		public string Type { get; }
 
-		/// <summary>
-		/// List of keys which identify this issue.
-		/// </summary>
+		/// <inheritdoc/>
 		public CaseInsensitiveStringSet Keys { get; set; }
 
-		/// <summary>
-		/// Set of keys which should trigger a negative match
-		/// </summary>
+		/// <inheritdoc/>
 		public CaseInsensitiveStringSet? RejectKeys { get; set; }
 
 		/// <summary>
@@ -97,9 +91,11 @@
 		}
 
 		/// <inheritdoc/>
-		public override bool Equals(object? Other)
+		public override bool Equals(object? Other) => Equals(Other as IIssueFingerprint);
+
+		/// <inheritdoc/>
+		public bool Equals(IIssueFingerprint? OtherFingerprint)
 		{
-			NewIssueFingerprint? OtherFingerprint = Other as NewIssueFingerprint;
 			if(OtherFingerprint == null || !Type.Equals(OtherFingerprint.Type, StringComparison.Ordinal) || !Keys.SetEquals(OtherFingerprint.Keys))
 			{
 				return false;
@@ -175,53 +171,81 @@
 	}
 
 	/// <summary>
+	/// Information about a new span
+	/// </summary>
+	public class NewIssueSpanData
+	{
+		/// <inheritdoc cref="IIssueSpan.StreamId"/>
+		public StreamId StreamId { get; set; }
+
+		/// <inheritdoc cref="IIssueSpan.StreamName"/>
+		public string StreamName { get; set; }
+
+		/// <inheritdoc cref="IIssueSpan.TemplateRefId"/>
+		public TemplateRefId TemplateRefId { get; set; }
+
+		/// <inheritdoc cref="IIssueSpan.NodeName"/>
+		public string NodeName { get; set; }
+
+		/// <inheritdoc cref="IIssueSpan.Fingerprint"/>
+		public NewIssueFingerprint Fingerprint { get; set; }
+
+		/// <inheritdoc cref="IIssueSpan.LastSuccess"/>
+		public NewIssueStepData? LastSuccess { get; set; }
+
+		/// <inheritdoc cref="IIssueSpan.FirstFailure"/>
+		public NewIssueStepData FirstFailure { get; set; }
+
+		/// <inheritdoc cref="IIssueSpan.NextSuccess"/>
+		public NewIssueStepData? NextSuccess { get; set; }
+
+		/// <inheritdoc cref="IIssueSpan.Suspects"/>
+		public List<NewIssueSpanSuspectData> Suspects { get; set; } = new List<NewIssueSpanSuspectData>();
+
+		/// <summary>
+		/// Constructor
+		/// </summary>
+		public NewIssueSpanData(StreamId StreamId, string StreamName, TemplateRefId TemplateRefId, string NodeName, NewIssueFingerprint Fingerprint, NewIssueStepData FirstFailure)
+		{
+			this.StreamId = StreamId;
+			this.StreamName = StreamName;
+			this.TemplateRefId = TemplateRefId;
+			this.NodeName = NodeName;
+			this.Fingerprint = Fingerprint;
+			this.FirstFailure = FirstFailure;
+		}
+	}
+
+	/// <summary>
 	/// Identifies a particular changelist and job
 	/// </summary>
 	public class NewIssueStepData
 	{
-		/// <summary>
-		/// The changelist number
-		/// </summary>
+		/// <inheritdoc cref="IIssueStep.Change"/>
 		public int Change { get; set; }
 
-		/// <summary>
-		/// Severity of the issue
-		/// </summary>
+		/// <inheritdoc cref="IIssueStep.Severity"/>
 		public IssueSeverity Severity { get; set; }
 
-		/// <summary>
-		/// Name of the job
-		/// </summary>
+		/// <inheritdoc cref="IIssueStep.JobName"/>
 		public string JobName { get; set; }
 
-		/// <summary>
-		/// The unique job id
-		/// </summary>
+		/// <inheritdoc cref="IIssueStep.JobId"/>
 		public JobId JobId { get; set; }
 
-		/// <summary>
-		/// Batch id for the step
-		/// </summary>
+		/// <inheritdoc cref="IIssueStep.BatchId"/>
 		public SubResourceId BatchId { get; set; }
 
-		/// <summary>
-		/// Id of the step
-		/// </summary>
+		/// <inheritdoc cref="IIssueStep.StepId"/>
 		public SubResourceId StepId { get; set; }
 
-		/// <summary>
-		/// Time that the step started
-		/// </summary>
+		/// <inheritdoc cref="IIssueStep.StepTime"/>
 		public DateTime StepTime { get; set; }
 
-		/// <summary>
-		/// The log id for this step
-		/// </summary>
+		/// <inheritdoc cref="IIssueStep.LogId"/>
 		public LogId? LogId { get; set; }
 
-		/// <summary>
-		/// Whether this step is promoted
-		/// </summary>
+		/// <inheritdoc cref="IIssueStep.PromoteByDefault"/>
 		public bool PromoteByDefault { get; set; }
 
 		/// <summary>
@@ -277,32 +301,21 @@
 	/// </summary>
 	public class NewIssueSuspectData
 	{
-		/// <summary>
-		/// Author of the changelist
-		/// </summary>
+		/// <inheritdoc cref="IIssueSuspect.AuthorId"/>
 		public UserId AuthorId { get; set; }
 
-		/// <summary>
-		/// The submitted changelist
-		/// </summary>
+		/// <inheritdoc cref="IIssueSuspect.Change"/>
 		public int Change { get; set; }
 
 		/// <summary>
-		/// The time that the user declined causing this issue
-		/// </summary>
-		public DateTime? DeclinedAt { get; set; }
-
-		/// <summary>
 		/// Constructor
 		/// </summary>
 		/// <param name="AuthorId">Author of the change</param>
 		/// <param name="Change">The changelist number</param>
-		/// <param name="DeclinedAt">The time that the user declined this issue</param>
-		public NewIssueSuspectData(UserId AuthorId, int Change, DateTime? DeclinedAt)
+		public NewIssueSuspectData(UserId AuthorId, int Change)
 		{
 			this.AuthorId = AuthorId;
 			this.Change = Change;
-			this.DeclinedAt = DeclinedAt;
 		}
 	}
 
@@ -311,19 +324,13 @@
 	/// </summary>
 	public class NewIssueSpanSuspectData
 	{
-		/// <summary>
-		/// The submitted changelist
-		/// </summary>
+		/// <inheritdoc cref="IIssueSpanSuspect.Change"/>
 		public int Change { get; set; }
 
-		/// <summary>
-		/// Author of the changelist
-		/// </summary>
+		/// <inheritdoc cref="IIssueSpanSuspect.AuthorId"/>
 		public UserId AuthorId { get; set; }
 
-		/// <summary>
-		/// The original changelist number, if merged from another branch. For changes merged between several branches, this is the originally submitted change.
-		/// </summary>
+		/// <inheritdoc cref="IIssueSpanSuspect.OriginatingChange"/>
 		public int? OriginatingChange { get; set; }
 
 		/// <summary>
@@ -339,16 +346,24 @@
 	}
 
 	/// <summary>
-	/// Interface for a token used to indicate a position in the transaction queue. Transactions are opportunistic, and will fail if the token is longer valid at the time that it is used.
+	/// Information about a stream containing an issue
 	/// </summary>
-	[SuppressMessage("Design", "CA1040:Avoid empty interfaces")]
-	public interface IIssueSequenceToken
+	public class NewIssueStream : IIssueStream
 	{
+		/// <inheritdoc cref="IIssueStream.StreamId"/>
+		public StreamId StreamId { get; set; }
+
+		/// <inheritdoc cref="IIssueStream.ContainsFix"/>
+		public bool? ContainsFix { get; set; }
+
 		/// <summary>
-		/// Reset the token to the last operation
+		/// Constructor
 		/// </summary>
-		/// <returns>Async task</returns>
-		public Task ResetAsync();
+		/// <param name="StreamId"></param>
+		public NewIssueStream(StreamId StreamId)
+		{
+			this.StreamId = StreamId;
+		}
 	}
 
 	/// <summary>
@@ -357,28 +372,19 @@
 	public interface IIssueCollection
 	{
 		/// <summary>
-		/// Gets a token used to sequence operations to the issue collection
+		/// Attempts to enters a critical section over the issue collection. Used for performing operations that require coordination of several documents (eg. attaching spans to issues).
 		/// </summary>
-		/// <returns>Sequence token instance</returns>
-		Task<IIssueSequenceToken> GetSequenceTokenAsync();
+		/// <returns>The disposable critical section</returns>
+		Task<IAsyncDisposable> EnterCriticalSectionAsync();
 
-		/// <summary>
-		/// Assign a unique issue id. For compatibility with the legacy issue system.
-		/// </summary>
-		/// <param name="Token">The current sequence token</param>
-		/// <returns>Sequence token</returns>
-		Task<int?> ReserveUniqueIdAsync(IIssueSequenceToken Token);
-
 		#region Issues
 
 		/// <summary>
 		/// Creates a new issue
 		/// </summary>
-		/// <param name="Token">Token for the transaction</param>
 		/// <param name="Summary">Summary text for the issue</param>
-		/// <param name="Span">The initial span for the issue</param>
 		/// <returns>The new issue instance</returns>
-		Task<IIssue?> AddIssueAsync(IIssueSequenceToken Token, string Summary, IIssueSpan Span);
+		Task<IIssue> AddIssueAsync(string Summary);
 
 		/// <summary>
 		/// Retrieves and issue by id
@@ -388,11 +394,11 @@
 		Task<IIssue?> GetIssueAsync(int IssueId);
 
 		/// <summary>
-		/// Gets the suspects for an issue
+		/// Finds the suspects for an issue
 		/// </summary>
 		/// <param name="Issue">The issue to retrieve suspects for</param>
 		/// <returns>List of suspects</returns>
-		Task<List<IIssueSuspect>> GetSuspectsAsync(IIssue Issue);
+		Task<List<IIssueSuspect>> FindSuspectsAsync(IIssue Issue);
 
 		/// <summary>
 		/// Searches for open issues
@@ -414,10 +420,10 @@
 		/// </summary>
 		/// <param name="Changes">List of suspect changes</param>
 		/// <returns>List of issues that are affected by the given changes</returns>
-		Task<List<IIssue>> FindIssuesForSuspectsAsync(List<int> Changes);
+		Task<List<IIssue>> FindIssuesForChangesAsync(List<int> Changes);
 
 		/// <summary>
-		/// Updates the state of an issue
+		/// Try to update the state of an issue
 		/// </summary>
 		/// <param name="Issue">The issue to update</param>
 		/// <param name="NewSeverity">New severity for the issue</param>
@@ -430,18 +436,25 @@
 		/// <param name="NewAcknowledged">Whether the issue has been acknowledged</param>
 		/// <param name="NewDeclinedById">Name of a user that has declined the issue</param>
 		/// <param name="NewFixChange">Fix changelist for the issue. Pass 0 to clear the fix changelist, -1 for systemic issue.</param>
-		/// <param name="NewFixStreamIds">List of stream ids with a bool indicating whether the fix changelist is in this stream</param>
 		/// <param name="NewResolvedById">User that resolved the issue (may be ObjectId.Empty to clear)</param>
 		/// <param name="NewLastSeenAt"></param>
 		/// <returns>True if the issue was updated</returns>
-		Task<IIssue?> UpdateIssueAsync(IIssue Issue, IssueSeverity? NewSeverity = null, string? NewSummary = null, string? NewUserSummary = null, string? NewDescription = null, bool? NewPromoted = null, UserId? NewOwnerId = null, UserId? NewNominatedById = null, bool? NewAcknowledged = null, UserId? NewDeclinedById = null, int? NewFixChange = null, Dictionary<StreamId, bool>? NewFixStreamIds = null, UserId? NewResolvedById = null, DateTime? NewLastSeenAt = null);
+		Task<IIssue?> TryUpdateIssueAsync(IIssue Issue, IssueSeverity? NewSeverity = null, string? NewSummary = null, string? NewUserSummary = null, string? NewDescription = null, bool? NewPromoted = null, UserId? NewOwnerId = null, UserId? NewNominatedById = null, bool? NewAcknowledged = null, UserId? NewDeclinedById = null, int? NewFixChange = null, UserId? NewResolvedById = null, DateTime? NewLastSeenAt = null);
 
 		/// <summary>
-		/// Updates the list of suspects for an issue, and optionally sets the resolved state
+		/// Updates derived data for an issue (ie. data computed from the spans attached to it). Also clears the issue's 'modified' state.
 		/// </summary>
-		/// <param name="Issue">The issue to update</param>
-		/// <returns>True if the issue was updated</returns>
-		Task<IIssue?> UpdateIssueDerivedDataAsync(IIssue Issue);
+		/// <param name="Issue">Issue to update</param>
+		/// <param name="NewSummary">New summary for the issue</param>
+		/// <param name="NewSeverity">New severity for the issue</param>
+		/// <param name="NewFingerprints">New fingerprints for the issue</param>
+		/// <param name="NewStreams">New streams for the issue</param>
+		/// <param name="NewSuspects">New suspects for the issue</param>
+		/// <param name="NewResolvedAt">Time for the last resolved change</param>
+		/// <param name="NewVerifiedAt">Time that the issue was resolved</param>
+		/// <param name="NewLastSeenAt">Last time the issue was seen</param>
+		/// <returns>Updated issue, or null if the issue is modified in the interim</returns>
+		Task<IIssue?> TryUpdateIssueDerivedDataAsync(IIssue Issue, string NewSummary, IssueSeverity NewSeverity, List<NewIssueFingerprint> NewFingerprints, List<NewIssueStream> NewStreams, List<NewIssueSuspectData> NewSuspects, DateTime? NewResolvedAt, DateTime? NewVerifiedAt, DateTime NewLastSeenAt);
 
 		#endregion
 
@@ -450,17 +463,10 @@
 		/// <summary>
 		/// Creates a new span from the given failure
 		/// </summary>
-		/// <param name="Token">Token to sequence the addition of new spans</param>
-		/// <param name="Stream">The stream containing the span</param>
-		/// <param name="TemplateRefId">The template being executed</param>
-		/// <param name="NodeName">Name of the step</param>
-		/// <param name="Fingerprint">Fingerprint for this issue</param>
-		/// <param name="LastSuccess">Last time the build succeeded before the span</param>
-		/// <param name="Failure">List of failing steps</param>
-		/// <param name="NextSuccess">First time the build succeeded after the span</param>
-		/// <param name="Suspects">Suspects for this span</param>
+		/// <param name="IssueId">The issue that the span belongs to</param>
+		/// <param name="NewSpan">Information about the new span</param>
 		/// <returns>New span, or null if the sequence token is not valid</returns>
-		Task<IIssueSpan?> AddSpanAsync(IIssueSequenceToken Token, IStream Stream, TemplateRefId TemplateRefId, string NodeName, NewIssueFingerprint Fingerprint, NewIssueStepData? LastSuccess, NewIssueStepData Failure, NewIssueStepData? NextSuccess, List<NewIssueSpanSuspectData> Suspects);
+		Task<IIssueSpan> AddSpanAsync(int IssueId, NewIssueSpanData NewSpan);
 
 		/// <summary>
 		/// Gets a particular span
@@ -473,26 +479,13 @@
 		/// Updates the given span. Note that data in the span's issue may be derived from this, and the issue should be updated afterwards.
 		/// </summary>
 		/// <param name="Span">Span to update</param>
-		/// <param name="NewSeverity">New severity for the span</param>
 		/// <param name="NewLastSuccess">New last successful step</param>
 		/// <param name="NewFailure">New failed step</param>
 		/// <param name="NewNextSuccess">New next successful step</param>
 		/// <param name="NewSuspects">New suspects for the span</param>
-		/// <param name="NewModified">New modified flag for the span</param>
+		/// <param name="NewIssueId">The new issue id for this span</param>
 		/// <returns>The updated span, or null on failure</returns>
-		Task<IIssueSpan?> TryUpdateSpanAsync(IIssueSpan Span, IssueSeverity? NewSeverity = null, NewIssueStepData? NewLastSuccess = null, NewIssueStepData? NewFailure = null, NewIssueStepData? NewNextSuccess = null, List<NewIssueSpanSuspectData>? NewSuspects = null, bool? NewModified = null);
+		Task<IIssueSpan?> TryUpdateSpanAsync(IIssueSpan Span, NewIssueStepData? NewLastSuccess = null, NewIssueStepData? NewFailure = null, NewIssueStepData? NewNextSuccess = null, List<NewIssueSpanSuspectData>? NewSuspects = null, int? NewIssueId = null);
 
 		/// <summary>
-		/// Updates the given span
-		/// </summary>
-		/// <param name="Token">The token</param>
-		/// <param name="Span">The span to update</param>
-		/// <param name="Issue">The issue to attach to</param>
-		/// <param name="NewSeverity">New severity for the span</param>
-		/// <param name="NewFingerprint">New fingerprint for the span</param>
-		/// <param name="NewSummary">New summary for the issue</param>
-		/// <returns>The updated span, or null on failure</returns>
-		Task<bool> AttachSpanToIssueAsync(IIssueSequenceToken Token, IIssueSpan Span, IIssue Issue, IssueSeverity? NewSeverity = null, NewIssueFingerprint? NewFingerprint = null, string? NewSummary = null);
-
-		/// <summary>
 		/// Gets all the spans for a particular issue
@@ -518,16 +511,18 @@
 		/// <returns>List of open issues</returns>
 		Task<List<IIssueSpan>> FindOpenSpansAsync(StreamId StreamId, TemplateRefId TemplateId, string Name, int Change);
 
-		/// <summary>
-		/// Finds a span marked as modified
-		/// </summary>
-		/// <returns></returns>
-		Task<List<IIssueSpan>> FindModifiedSpansAsync();
-
 		#endregion
 
 		#region Steps
-			
+
 		/// <summary>
+		/// Adds a new step
+		/// </summary>
+		/// <param name="SpanId">Initial span for the step</param>
+		/// <param name="NewStep">Information about the new step</param>
+		/// <returns>New step object</returns>
+		Task<IIssueStep> AddStepAsync(ObjectId SpanId, NewIssueStepData NewStep);
+
+		/// <summary>
 		/// Find steps for the given spans
 		/// </summary>
diff --git a/Engine/Source/Programs/Horde/HordeServer/Issues/IIssueSpan.cs b/Engine/Source/Programs/Horde/HordeServer/Issues/IIssueSpan.cs
--- a/Engine/Source/Programs/Horde/HordeServer/Issues/IIssueSpan.cs
+++ b/Engine/Source/Programs/Horde/HordeServer/Issues/IIssueSpan.cs
@@ -44,11 +44,6 @@ namespace HordeServer.Models
 		/// </summary>
 		public string NodeName { get; }
 
-		/// <summary>
-		/// Severity of the failures in this span
-		/// </summary>
-		public IssueSeverity Severity { get; }
-
 		/// <summary>
 		/// Fingerprint for this issue
 		/// </summary>
@@ -87,7 +82,7 @@ namespace HordeServer.Models
 		/// <summary>
 		/// The issue that this span is assigned to
 		/// </summary>
-		public int? IssueId { get; }
+		public int IssueId { get; }
 	}
 
 	/// <summary>
diff --git a/Engine/Source/Programs/Horde/HordeServer/Issues/Impl/IssueCollection.cs b/Engine/Source/Programs/Horde/HordeServer/Issues/Impl/IssueCollection.cs
--- a/Engine/Source/Programs/Horde/HordeServer/Issues/Impl/IssueCollection.cs
+++ b/Engine/Source/Programs/Horde/HordeServer/Issues/Impl/IssueCollection.cs
@@ -1,6 +1,7 @@
 // Copyright Epic Games, Inc. All Rights Reserved.
 
 using EpicGames.Core;
+using EpicGames.Redis.Utility;
 using HordeServer.Models;
 using HordeServer.Services;
 using HordeServer.Utilities;
@@ -13,6 +14,7 @@ using MongoDB.Driver;
 using System;
 using System.Collections.Concurrent;
 using System.Collections.Generic;
+using System.Diagnostics;
 using System.Linq;
 using System.Linq.Expressions;
 using System.Security.Claims;
@@ -28,90 +30,10 @@ namespace HordeServer.Collections.Impl
 
 	class IssueCollection : IIssueCollection
 	{
-		class UpdateTransaction
-		{
-			class Statics<TDocument>
-			{
-				public static IBsonSerializer<TDocument> Serializer { get; } = BsonSerializer.LookupSerializer<TDocument>();
-			}
-
-			[BsonRequired]
-			public string Filter { get; set; }
-
-			[BsonRequired]
-			public string Update { get; set; }
-
-			private UpdateTransaction()
-			{
-				this.Filter = null!;
-				this.Update = null!;
-			}
-
-			public UpdateTransaction(BsonDocument Filter, BsonDocument Update)
-			{
-				this.Filter = Filter.ToJson();
-				this.Update = Update.ToJson();
-			}
-
-			public static UpdateTransaction Create<TDocument>(Expression<Func<TDocument, bool>> FilterFunc, UpdateDefinition<TDocument> Update)
-			{
-				return Create((FilterDefinition<TDocument>)FilterFunc, Update);
-			}
-
-			public static UpdateTransaction Create<TDocument>(FilterDefinition<TDocument> Filter, UpdateDefinition<TDocument> Update)
-			{
-				IBsonSerializer<TDocument> Serializer = Statics<TDocument>.Serializer;
-
-				BsonDocument FilterDocument = Filter.Render(Serializer, BsonSerializer.SerializerRegistry);
-				BsonDocument UpdateDocument = (BsonDocument)Update.Render(Serializer, BsonSerializer.SerializerRegistry);
-
-				return new UpdateTransaction(FilterDocument, UpdateDocument);
-			}
-		}
-
 		[SingletonDocument("5e4c226440ce25fa3207a9af")]
 		class IssueLedger : SingletonBase
 		{
 			public int NextId { get; set; }
-
-			[BsonIgnoreIfNull]
-			public Issue? InsertIssue { get; set; }
-
-			[BsonIgnoreIfNull]
-			public IssueSpan? InsertIssueSpan { get; set; }
-
-			[BsonIgnoreIfNull]
-			public List<IssueSuspect>? InsertIssueSuspects { get; set; }
-
-			[BsonIgnoreIfNull]
-			public UpdateTransaction? UpdateIssue { get; set; }
-
-			[BsonIgnoreIfNull]
-			public UpdateTransaction? UpdateIssueSpan { get; set; }
-
-			public bool HasPendingTransaction => InsertIssue != null || InsertIssueSpan != null || InsertIssueSuspects != null || UpdateIssue != null || UpdateIssueSpan != null;
-		}
-
-		class IssueSequenceToken : IIssueSequenceToken
-		{
-			IssueCollection Collection;
-			public IssueLedger Ledger { get; private set; }
-
-			public IssueSequenceToken(IssueCollection Collection, IssueLedger Ledger)
-			{
-				this.Collection = Collection;
-				this.Ledger = Ledger;
-			}
-
-			public async Task ResetAsync()
-			{
-				Ledger = await Collection.GetLedgerAsync();
-			}
-
-			public async Task FlushAsync()
-			{
-				Ledger = await Collection.FlushLedgerAsync(Ledger);
-			}
 		}
 
 		class Issue : IIssue
@@ -127,8 +49,10 @@ namespace HordeServer.Collections.Impl
 			[BsonIgnoreIfNull]
 			public string? Description { get; set; }
 
-			[BsonRequired]
-			public IssueFingerprint Fingerprint { get; set; }
+			[BsonIgnoreIfNull]
+			public IssueFingerprint? Fingerprint { get; set; }
+			public List<IssueFingerprint>? Fingerprints { get; set; }
+
 			public IssueSeverity Severity { get; set; }
 
 			public bool? Promoted { get; set; }
@@ -136,6 +60,9 @@ namespace HordeServer.Collections.Impl
 			[BsonIgnoreIfNull]
 			public UserId? OwnerId { get; set; }
 
+			[BsonIgnoreIfNull]
+			public UserId? DefaultOwnerId { get; set; }
+
 			[BsonIgnoreIfNull]
 			public UserId? NominatedById { get; set; }
 
@@ -169,14 +96,14 @@ namespace HordeServer.Collections.Impl
 			[BsonElement("NotifySuspects"), BsonIgnoreIfDefault(false)]
 			public bool NotifySuspects_DEPRECATED { get; set; }
 
-			[BsonRequired]
-			public List<IssueSuspect> Suspects { get; set; }
+			[BsonElement("Suspects"), BsonIgnoreIfNull]
+			public List<IssueSuspect>? Suspects_DEPRECATED { get; set; }
 
 			public int UpdateIndex { get; set; }
 
 			bool IIssue.Promoted => Promoted ?? NotifySuspects_DEPRECATED;
-			IIssueFingerprint IIssue.Fingerprint => Fingerprint;
-			UserId? IIssue.OwnerId => OwnerId ?? GetDefaultOwner()?.AuthorId;
+			IReadOnlyList<IIssueFingerprint> IIssue.Fingerprints => Fingerprints ?? ((Fingerprint == null)? new List<IssueFingerprint>() : new List<IssueFingerprint> { Fingerprint });
+			UserId? IIssue.OwnerId => OwnerId ?? DefaultOwnerId ?? GetDefaultOwnerId();
 			IReadOnlyList<IIssueStream> IIssue.Streams => Streams;
 			DateTime IIssue.LastSeenAt => (LastSeenAt == default) ? DateTime.UtcNow : LastSeenAt;
 
@@ -185,37 +112,24 @@ namespace HordeServer.Collections.Impl
 			{
 				Summary = String.Empty;
 				Fingerprint = null!;
-				Suspects = null!;
 			}
 
-			public Issue(int Id, string Summary, IIssueSpan Span)
+			public Issue(int Id, string Summary)
 			{
 				this.Id = Id;
 				this.Summary = Summary;
 				this.CreatedAt = DateTime.UtcNow;
 				this.LastSeenAt = DateTime.UtcNow;
-				this.Fingerprint = new IssueFingerprint(Span.Fingerprint);
-				this.Promoted = Span.PromoteByDefault;
-				this.Severity = Span.Severity;
-				this.ResolvedAt = Span.NextSuccess?.StepTime;
-				this.VerifiedAt = this.ResolvedAt;
-				this.Suspects = Span.Suspects.ConvertAll(x => new IssueSuspect(Id, x.AuthorId, x.OriginatingChange ?? x.Change, null, this.ResolvedAt));
-
-				if (Suspects.Count > 0)
-				{
-					MinSuspectChange = Suspects.Min(x => x.Change);
-					MaxSuspectChange = Suspects.Max(x => x.Change);
-				}
 			}
 
-			IssueSuspect? GetDefaultOwner()
+			UserId? GetDefaultOwnerId()
 			{
-				if (Suspects.Count > 0)
+				if (Suspects_DEPRECATED != null && Suspects_DEPRECATED.Count > 0)
 				{
-					UserId PossibleOwner = Suspects[0].AuthorId;
-					if (Suspects.All(x => x.AuthorId == PossibleOwner) && Suspects.Any(x => x.DeclinedAt == null))
+					UserId PossibleOwner = Suspects_DEPRECATED[0].AuthorId;
+					if (Suspects_DEPRECATED.All(x => x.AuthorId == PossibleOwner) && Suspects_DEPRECATED.Any(x => x.DeclinedAt == null))
 					{
-						return Suspects[0];
+						return Suspects_DEPRECATED[0].AuthorId;
 					}
 				}
 				return null;
@@ -251,6 +165,15 @@ namespace HordeServer.Collections.Impl
 			{
 			}
 
+			public IssueSuspect(int IssueId, NewIssueSuspectData NewSuspect, DateTime? ResolvedAt)
+			{
+				this.Id = ObjectId.GenerateNewId();
+				this.IssueId = IssueId;
+				this.AuthorId = NewSuspect.AuthorId;
+				this.Change = NewSuspect.Change;
+				this.ResolvedAt = ResolvedAt;
+			}
+
 			public IssueSuspect(int IssueId, IIssueSpanSuspect Suspect)
 				: this(IssueId, Suspect.AuthorId, Suspect.OriginatingChange ?? Suspect.Change, null, null)
 			{
@@ -303,7 +226,6 @@ namespace HordeServer.Collections.Impl
 
 			[BsonRequired]
 			public string NodeName { get; set; }
-			public IssueSeverity Severity { get; set; }
 			public DateTime? ResolvedAt { get; set; } // Propagated from the owning issue
 
 			[BsonRequired]
@@ -330,8 +252,7 @@ namespace HordeServer.Collections.Impl
 			bool IIssueSpan.PromoteByDefault => PromoteByDefault ?? NotifySuspects_DEPRECATED;
 
 			public List<IssueSpanSuspect> Suspects { get; set; }
-			public int? IssueId { get; set; }
-			public bool Modified { get; set; }
+			public int IssueId { get; set; }
 			public int UpdateIndex { get; set; }
 
 			IIssueStep? IIssueSpan.LastSuccess => LastSuccess;
@@ -351,30 +272,29 @@ namespace HordeServer.Collections.Impl
 				this.Suspects = new List<IssueSpanSuspect>();
 			}
 
-			public IssueSpan(StreamId StreamId, string StreamName, TemplateRefId TemplateRefId, string NodeName, IIssueFingerprint Fingerprint, NewIssueStepData? LastSuccess, NewIssueStepData Failure, NewIssueStepData? NextSuccess, List<IssueSpanSuspect> Suspects)
+			public IssueSpan(int IssueId, NewIssueSpanData NewSpan)
 			{
 				this.Id = ObjectId.GenerateNewId();
-				this.StreamId = StreamId;
-				this.StreamName = StreamName;
-				this.TemplateRefId = TemplateRefId;
-				this.NodeName = NodeName;
-				this.Severity = Failure.Severity;
-				this.Fingerprint = new IssueFingerprint(Fingerprint);
-				if (LastSuccess != null)
+				this.StreamId = NewSpan.StreamId;
+				this.StreamName = NewSpan.StreamName;
+				this.TemplateRefId = NewSpan.TemplateRefId;
+				this.NodeName = NewSpan.NodeName;
+				this.Fingerprint = new IssueFingerprint(NewSpan.Fingerprint);
+				if (NewSpan.LastSuccess != null)
 				{
-					this.MinChange = LastSuccess.Change;
-					this.LastSuccess = new IssueStep(Id, LastSuccess);
+					this.MinChange = NewSpan.LastSuccess.Change;
+					this.LastSuccess = new IssueStep(Id, NewSpan.LastSuccess);
 				}
-				this.FirstFailure = new IssueStep(Id, Failure);
-				this.LastFailure = new IssueStep(Id, Failure);
-				if (NextSuccess != null)
+				this.FirstFailure = new IssueStep(Id, NewSpan.FirstFailure);
+				this.LastFailure = new IssueStep(Id, NewSpan.FirstFailure);
+				if (NewSpan.NextSuccess != null)
 				{
-					this.MaxChange = NextSuccess.Change;
-					this.NextSuccess = new IssueStep(Id, NextSuccess);
+					this.MaxChange = NewSpan.NextSuccess.Change;
+					this.NextSuccess = new IssueStep(Id, NewSpan.NextSuccess);
 				}
-				this.PromoteByDefault = Failure.PromoteByDefault;
-				this.Suspects = Suspects;
-				this.Modified = true;
+				this.PromoteByDefault = NewSpan.FirstFailure.PromoteByDefault;
+				this.Suspects = NewSpan.Suspects.ConvertAll(x => new IssueSpanSuspect(x));
+				this.IssueId = IssueId;
 			}
 		}
 
@@ -389,11 +309,11 @@ namespace HordeServer.Collections.Impl
 			{
 			}
 
-			public IssueSpanSuspect(int Change, UserId AuthorId, int? OriginatingChange)
+			public IssueSpanSuspect(NewIssueSpanSuspectData NewSuspectData)
 			{
-				this.Change = Change;
-				this.AuthorId = AuthorId;
-				this.OriginatingChange = OriginatingChange;
+				this.Change = NewSuspectData.Change;
+				this.AuthorId = NewSuspectData.AuthorId;
+				this.OriginatingChange = NewSuspectData.OriginatingChange;
 			}
 		}
 
@@ -450,6 +370,7 @@ namespace HordeServer.Collections.Impl
 			}
 		}
 
+		RedisService RedisService;
 		ISingletonDocument<IssueLedger> LedgerSingleton;
 		IMongoCollection<Issue> Issues;
 		IMongoCollection<IssueSpan> IssueSpans;
@@ -458,8 +379,9 @@ namespace HordeServer.Collections.Impl
 		IAuditLog<int> AuditLog;
 		ILogger Logger;
 
-		public IssueCollection(DatabaseService DatabaseService, IAuditLogFactory<int> AuditLogFactory, ILogger<IssueCollection> Logger)
+		public IssueCollection(DatabaseService DatabaseService, RedisService RedisService, IAuditLogFactory<int> AuditLogFactory, ILogger<IssueCollection> Logger)
 		{
+			this.RedisService = RedisService;
 			this.Logger = Logger;
 
 			LedgerSingleton = new SingletonDocument<IssueLedger>(DatabaseService);
@@ -475,7 +397,6 @@ namespace HordeServer.Collections.Impl
 				Issues.Indexes.CreateOne(new CreateIndexModel<Issue>(Builders<Issue>.IndexKeys.Ascending(x => x.ResolvedAt)));
 				Issues.Indexes.CreateOne(new CreateIndexModel<Issue>(Builders<Issue>.IndexKeys.Ascending(x => x.VerifiedAt)));
 
-				IssueSpans.Indexes.CreateOne(new CreateIndexModel<IssueSpan>(Builders<IssueSpan>.IndexKeys.Ascending(x => x.Modified)));
 				IssueSpans.Indexes.CreateOne(new CreateIndexModel<IssueSpan>(Builders<IssueSpan>.IndexKeys.Ascending(x => x.IssueId)));
 				IssueSpans.Indexes.CreateOne(new CreateIndexModel<IssueSpan>(Builders<IssueSpan>.IndexKeys.Ascending(x => x.StreamId).Ascending(x => x.MinChange).Ascending(x => x.MaxChange)));
 				IssueSpans.Indexes.CreateOne(new CreateIndexModel<IssueSpan>(Builders<IssueSpan>.IndexKeys.Ascending(x => x.StreamId).Ascending(x => x.TemplateRefId).Ascending(x => x.NodeName).Ascending(x => x.MinChange).Ascending(x => x.MaxChange), new CreateIndexOptions { Name = "StreamChanges" }));
@@ -483,11 +404,31 @@ namespace HordeServer.Collections.Impl
 				IssueSteps.Indexes.CreateOne(new CreateIndexModel<IssueStep>(Builders<IssueStep>.IndexKeys.Ascending(x => x.SpanId)));
 				IssueSteps.Indexes.CreateOne(new CreateIndexModel<IssueStep>(Builders<IssueStep>.IndexKeys.Ascending(x => x.JobId).Ascending(x => x.BatchId).Ascending(x => x.StepId)));
 
+				IssueSuspects.Indexes.CreateOne(new CreateIndexModel<IssueSuspect>(Builders<IssueSuspect>.IndexKeys.Ascending(x => x.Change)));
 				IssueSuspects.Indexes.CreateOne(new CreateIndexModel<IssueSuspect>(Builders<IssueSuspect>.IndexKeys.Ascending(x => x.AuthorId).Ascending(x => x.ResolvedAt)));
 				IssueSuspects.Indexes.CreateOne(new CreateIndexModel<IssueSuspect>(Builders<IssueSuspect>.IndexKeys.Ascending(x => x.IssueId).Ascending(x => x.Change), new CreateIndexOptions { Unique = true }));
 			}
 		}
 
+		/// <inheritdoc/>
+		public async Task<IAsyncDisposable> EnterCriticalSectionAsync()
+		{
+			Stopwatch Timer = Stopwatch.StartNew();
+			TimeSpan NextNotifyTime = TimeSpan.FromSeconds(2.0);
+
+			RedisLock Lock = new RedisLock(RedisService.Database, "issues/lock");
+			while (!await Lock.AcquireAsync(TimeSpan.FromMinutes(1)))
+			{
+				if (Timer.Elapsed > NextNotifyTime)
+				{
+					Logger.LogWarning("Waiting on lock over issue collection for {TimeSpan}", Timer.Elapsed);
+					NextNotifyTime *= 2;
+				}
+				await Task.Delay(TimeSpan.FromMilliseconds(100));
+			}
+			return Lock;
+		}
+
 		async Task<Issue?> TryUpdateIssueAsync(IIssue Issue, UpdateDefinition<Issue> Update)
 		{
 			Issue IssueDocument = (Issue)Issue;
@@ -510,95 +451,19 @@ namespace HordeServer.Collections.Impl
 			return await IssueSpans.FindOneAndUpdateAsync<IssueSpan>(x => x.Id == IssueSpanDocument.Id && x.UpdateIndex == PrevUpdateIndex, Update, Options);
 		}
 
-		/// <inheritdoc/>
-		public async Task<IIssueSequenceToken> GetSequenceTokenAsync()
-		{
-			IssueLedger Ledger = await GetLedgerAsync();
-			return new IssueSequenceToken(this, Ledger);
-		}
-
-		async Task<IssueLedger> GetLedgerAsync()
-		{
-			IssueLedger Ledger = await LedgerSingleton.GetAsync();
-			Ledger = await FlushLedgerAsync(Ledger);
-			return Ledger;
-		}
-
-		async Task<IssueLedger> FlushLedgerAsync(IssueLedger Ledger)
-		{
-			while (Ledger.HasPendingTransaction)
-			{
-				if (Ledger.InsertIssue != null)
-				{
-					await Issues.InsertOneIgnoreDuplicatesAsync(Ledger.InsertIssue);
-					Ledger.InsertIssue = null;
-				}
-				if (Ledger.InsertIssueSpan != null)
-				{
-					await IssueSpans.InsertOneIgnoreDuplicatesAsync(Ledger.InsertIssueSpan);
-					Ledger.InsertIssueSpan = null;
-				}
-				if (Ledger.InsertIssueSuspects != null)
-				{
-					await IssueSuspects.InsertManyIgnoreDuplicatesAsync(Ledger.InsertIssueSuspects);
-					Ledger.InsertIssueSuspects = null;
-				}
-				if (Ledger.UpdateIssue != null)
-				{
-					await Issues.UpdateOneAsync(Ledger.UpdateIssue.Filter, Ledger.UpdateIssue.Update);
-					Ledger.UpdateIssue = null;
-				}
-				if (Ledger.UpdateIssueSpan != null)
-				{
-					await IssueSpans.UpdateOneAsync(Ledger.UpdateIssueSpan.Filter, Ledger.UpdateIssueSpan.Update);
-					Ledger.UpdateIssueSpan = null;
-				}
-
-				if (await LedgerSingleton.TryUpdateAsync(Ledger))
-				{
-					break;
-				}
-
-				Ledger = await LedgerSingleton.GetAsync();
-			}
-			return Ledger;
-		}
-
 		#region Issues
 
 		/// <inheritdoc/>
-		public async Task<int?> ReserveUniqueIdAsync(IIssueSequenceToken Token)
+		public async Task<IIssue> AddIssueAsync(string Summary)
 		{
-			IssueSequenceToken TokenValue = ((IssueSequenceToken)Token);
-			int NewIssueId = ++TokenValue.Ledger.NextId;
+			IssueLedger Ledger = await LedgerSingleton.UpdateAsync(x => x.NextId++);
 
-			if (!await LedgerSingleton.TryUpdateAsync(TokenValue.Ledger))
-			{
-				return null;
-			}
-
-			return NewIssueId;
-		}
-
-		/// <inheritdoc/>
-		public async Task<IIssue?> AddIssueAsync(IIssueSequenceToken Token, string Summary, IIssueSpan Span)
-		{
-			IssueSequenceToken TokenValue = ((IssueSequenceToken)Token);
-			Issue NewIssue = new Issue(++TokenValue.Ledger.NextId, Summary, Span);
-			TokenValue.Ledger.InsertIssue = NewIssue;
-			TokenValue.Ledger.InsertIssueSuspects = NewIssue.Suspects;
-			TokenValue.Ledger.UpdateIssueSpan = UpdateTransaction.Create<IssueSpan>(x => x.Id == Span.Id, Builders<IssueSpan>.Update.Set(x => x.IssueId, NewIssue.Id));
-
-			if (!await LedgerSingleton.TryUpdateAsync(TokenValue.Ledger))
-			{
-				return null;
-			}
+			Issue NewIssue = new Issue(Ledger.NextId, Summary);
+			await Issues.InsertOneAsync(NewIssue);
 
 			ILogger IssueLogger = GetLogger(NewIssue.Id);
 			IssueLogger.LogInformation("Created issue {IssueId}", NewIssue.Id);
-			LogSpanInfo(IssueLogger, Span);
 
-			await TokenValue.FlushAsync();
 			return NewIssue;
 		}
 
@@ -675,9 +540,12 @@ namespace HordeServer.Collections.Impl
 			{
 				IssueLogger.LogInformation("Marking stream {StreamId} as not fixed", StreamId);
 			}
+		}
 
-			HashSet<(UserId, int)> OldSuspects = new HashSet<(UserId, int)>(OldIssue.Suspects.Select(x => (x.AuthorId, x.Change)));
-			HashSet<(UserId, int)> NewSuspects = new HashSet<(UserId, int)>(NewIssue.Suspects.Select(x => (x.AuthorId, x.Change)));
+		static void LogIssueSuspectChanges(ILogger IssueLogger, List<IssueSuspect> OldIssueSuspects, List<IssueSuspect> NewIssueSuspects)
+		{
+			HashSet<(UserId, int)> OldSuspects = new HashSet<(UserId, int)>(OldIssueSuspects.Select(x => (x.AuthorId, x.Change)));
+			HashSet<(UserId, int)> NewSuspects = new HashSet<(UserId, int)>(NewIssueSuspects.Select(x => (x.AuthorId, x.Change)));
 			foreach ((UserId UserId, int Change) in NewSuspects.Where(x => !OldSuspects.Contains(x)))
 			{
 				IssueLogger.LogInformation("Added suspect {UserId} for change {Change}", UserId, Change);
@@ -687,8 +555,8 @@ namespace HordeServer.Collections.Impl
 				IssueLogger.LogInformation("Removed suspect {UserId} for change {Change}", UserId, Change);
 			}
 
-			HashSet<UserId> OldDeclinedBy = new HashSet<UserId>(NewIssue.Suspects.Where(x => x.DeclinedAt != null).Select(x => x.AuthorId));
-			HashSet<UserId> NewDeclinedBy = new HashSet<UserId>(NewIssue.Suspects.Where(x => x.DeclinedAt != null).Select(x => x.AuthorId));
+			HashSet<UserId> OldDeclinedBy = new HashSet<UserId>(OldIssueSuspects.Where(x => x.DeclinedAt != null).Select(x => x.AuthorId));
+			HashSet<UserId> NewDeclinedBy = new HashSet<UserId>(NewIssueSuspects.Where(x => x.DeclinedAt != null).Select(x => x.AuthorId));
 			foreach (UserId AddDeclinedBy in NewDeclinedBy.Where(x => !OldDeclinedBy.Contains(x)))
 			{
 				IssueLogger.LogInformation("Declined by {UserId}", AddDeclinedBy);
@@ -721,10 +589,9 @@ namespace HordeServer.Collections.Impl
 		}
 
 		/// <inheritdoc/>
-		public Task<List<IIssueSuspect>> GetSuspectsAsync(IIssue Issue)
+		public Task<List<IIssueSuspect>> FindSuspectsAsync(IIssue Issue)
 		{
-			Issue IssueDocument = (Issue)Issue;
-			return Task.FromResult(IssueDocument.Suspects.ConvertAll<IIssueSuspect>(x => x));
+			return IssueSuspects.Find(x => x.IssueId == Issue.Id).ToListAsync<IssueSuspect, IIssueSuspect>();
 		}
 
 		class ProjectedIssueId
@@ -901,32 +768,14 @@ namespace HordeServer.Collections.Impl
 		}
 
 		/// <inheritdoc/>
-		public async Task<List<IIssue>> FindIssuesForSuspectsAsync(List<int> Changes)
+		public async Task<List<IIssue>> FindIssuesForChangesAsync(List<int> Changes)
 		{
-			List<IIssue> Results = new List<IIssue>();
-			if (Changes.Count > 0)
-			{
-				int MinChange = Changes.Min();
-				int MaxChange = Changes.Max();
-
-				List<Issue> BaseResults = await Issues.Find(x => x.MinSuspectChange <= MaxChange && x.MaxSuspectChange >= MinChange).ToListAsync();
-				Results.AddRange(BaseResults.Where(x => x.Suspects.Any(y => Changes.Contains(y.Change))));
-			}
-			return Results;
-		}
-
-		static List<IssueSpanSuspect> CreateSpanSuspects(List<NewIssueSpanSuspectData> Suspects)
-		{
-			List<IssueSpanSuspect> NewSuspects = new List<IssueSpanSuspect>();
-			foreach (NewIssueSpanSuspectData Suspect in Suspects)
-			{
-				NewSuspects.Add(new IssueSpanSuspect(Suspect.Change, Suspect.AuthorId, Suspect.OriginatingChange));
-			}
-			return NewSuspects;
+			List<int> IssueIds = await (await IssueSuspects.DistinctAsync(x => x.IssueId, Builders<IssueSuspect>.Filter.In(x => x.Change, Changes))).ToListAsync();
+			return await Issues.Find(Builders<Issue>.Filter.In(x => x.Id, IssueIds)).ToListAsync<Issue, IIssue>();
 		}
 
 		/// <inheritdoc/>
-		public async Task<IIssue?> UpdateIssueAsync(IIssue Issue, IssueSeverity? NewSeverity = null, string? NewSummary = null, string? NewUserSummary = null, string? NewDescription = null, bool? NewPromoted = null, UserId? NewOwnerId = null, UserId? NewNominatedById = null, bool? NewAcknowledged = null, UserId? NewDeclinedById = null, int? NewFixChange = null, Dictionary<StreamId, bool>? NewFixStreamIds = null, UserId? NewResolvedById = null, DateTime? NewLastSeenAt = null)
+		public async Task<IIssue?> TryUpdateIssueAsync(IIssue Issue, IssueSeverity? NewSeverity = null, string? NewSummary = null, string? NewUserSummary = null, string? NewDescription = null, bool? NewPromoted = null, UserId? NewOwnerId = null, UserId? NewNominatedById = null, bool? NewAcknowledged = null, UserId? NewDeclinedById = null, int? NewFixChange = null, UserId? NewResolvedById = null, DateTime? NewLastSeenAt = null)
 		{
 			Issue IssueDocument = (Issue)Issue;
 
@@ -1008,18 +857,6 @@ namespace HordeServer.Collections.Impl
 					}
 				}
 			}
-			if (NewDeclinedById != null)
-			{
-				for (int SuspectIdx = 0; SuspectIdx < IssueDocument.Suspects.Count; SuspectIdx++)
-				{
-					IssueSuspect Suspect = IssueDocument.Suspects[SuspectIdx];
-					if (Suspect.AuthorId == NewDeclinedById.Value && Suspect.DeclinedAt == null)
-					{
-						int SuspectIdxCopy = SuspectIdx;
-						Updates.Add(Builders<Issue>.Update.Set(x => x.Suspects[SuspectIdxCopy].DeclinedAt, UtcNow));
-					}
-				}
-			}
 			if (NewFixChange != null)
 			{
 				if (NewFixChange == 0)
@@ -1031,23 +868,6 @@ namespace HordeServer.Collections.Impl
 					Updates.Add(Builders<Issue>.Update.Set(x => x.FixChange, NewFixChange));
 				}
 			}
-			if (NewFixStreamIds != null)
-			{
-				List<IssueStream> NewStreams = new List<IssueStream>();
-				foreach (IssueStream Stream in Issue.Streams)
-				{
-					IssueStream NewStream = new IssueStream(Stream);
-					if (NewFixStreamIds.TryGetValue(Stream.StreamId, out bool bContainsFix))
-					{
-						NewStream.ContainsFix = bContainsFix;
-					}
-					else
-					{
-						NewStream.ContainsFix = null;
-					}
-				}
-				Updates.Add(Builders<Issue>.Update.Set(x => x.Streams, NewStreams));
-			}
 			if (NewResolvedById != null)
 			{
 				if (NewResolvedById.Value != UserId.Empty)
@@ -1075,6 +895,10 @@ namespace HordeServer.Collections.Impl
 				Updates.Add(Builders<Issue>.Update.Set(x => x.LastSeenAt, NewLastSeenAt.Value));
 			}
 
+			if (NewDeclinedById != null)
+			{
+				await IssueSuspects.UpdateManyAsync(x => x.IssueId == Issue.Id && x.AuthorId == NewDeclinedById.Value, Builders<IssueSuspect>.Update.Set(x => x.DeclinedAt, DateTime.UtcNow));
+			}
 			if (Updates.Count == 0)
 			{
 				return IssueDocument;
@@ -1091,180 +915,120 @@ namespace HordeServer.Collections.Impl
 			return NewIssue;
 		}
 
-		public async Task<IIssue?> UpdateIssueDerivedDataAsync(IIssue Issue)
+		/// <inheritdoc/>
+		public async Task<IIssue?> TryUpdateIssueDerivedDataAsync(IIssue Issue, string NewSummary, IssueSeverity NewSeverity, List<NewIssueFingerprint> NewFingerprints, List<NewIssueStream> NewStreams, List<NewIssueSuspectData> NewSuspects, DateTime? NewResolvedAt, DateTime? NewVerifiedAt, DateTime NewLastSeenAt)
 		{
-			for (; ; )
-			{
-				Issue IssueDocument = (Issue)Issue;
+			Issue IssueImpl = (Issue)Issue;
 
-				// Find all the spans that are attached to the issue
-				List<IssueSpan> Spans = await IssueSpans.Find(x => x.IssueId == Issue.Id).ToListAsync();
+			// Update all the suspects for this issue
+			List<IssueSuspect> OldSuspectImpls = await IssueSuspects.Find(x => x.IssueId == Issue.Id).ToListAsync();
+			List<IssueSuspect> NewSuspectImpls = await UpdateIssueSuspectsAsync(Issue.Id, OldSuspectImpls, NewSuspects, NewResolvedAt);
 
-				// Update the issue to reflect information in the spans
-				Issue? Result = await TryUpdateIssueDerivedDataAsync(IssueDocument, Spans);
-				if (Result != null)
-				{
-					return Result;
-				}
-
-				// Update the issue
-				IIssue? NewIssue = await GetIssueAsync(Issue.Id);
-				if(NewIssue == null)
+			// Find the default owner
+			UserId? NewDefaultOwnerId = null;
+			if (NewSuspectImpls.Count > 0)
+			{
+				UserId PossibleOwnerId = NewSuspectImpls[0].AuthorId;
+				if (NewSuspectImpls.All(x => x.AuthorId == PossibleOwnerId) && NewSuspectImpls.Any(x => x.DeclinedAt == null))
 				{
-					return null;
+					NewDefaultOwnerId = PossibleOwnerId;
 				}
-				Issue = NewIssue;
 			}
-		}
 
-		async Task<Issue?> TryUpdateIssueDerivedDataAsync(Issue Issue, List<IssueSpan> Spans)
-		{
-			// Find the list of suspects for this issue. If the issue has an assigned owner, the other suspects can be eliminated.
-			List<IssueSuspect> NewSuspects = new List<IssueSuspect>();
-			if (Spans.Count > 0)
+			// Update all the spans attached to this issue
+			await IssueSpans.UpdateManyAsync(x => x.IssueId == Issue.Id && x.ResolvedAt != NewResolvedAt, Builders<IssueSpan>.Update.Set(x => x.ResolvedAt, NewResolvedAt));
+
+			// Get the range of suspect changes
+			int NewMinSuspectChange = (NewSuspects.Count > 0) ? NewSuspects.Min(x => x.Change) : 0;
+			int NewMaxSuspectChange = (NewSuspects.Count > 0) ? NewSuspects.Min(x => x.Change) : 0;
+
+			// Perform the actual update with this data
+			List<UpdateDefinition<Issue>> Updates = new List<UpdateDefinition<Issue>>();
+			if (!String.Equals(Issue.Summary, NewSummary, StringComparison.Ordinal))
 			{
-				HashSet<int> SuspectChanges = new HashSet<int>(Spans[0].Suspects.Where(x => Issue.OwnerId == null || x.AuthorId == Issue.OwnerId).Select(x => x.OriginatingChange ?? x.Change));
-				for (int SpanIdx = 1; SpanIdx < Spans.Count; SpanIdx++)
-				{
-					SuspectChanges.IntersectWith(Spans[SpanIdx].Suspects.Select(x => x.OriginatingChange ?? x.Change));
-				}
-				NewSuspects = Spans[0].Suspects.Where(x => SuspectChanges.Contains(x.OriginatingChange ?? x.Change)).Select(x => new IssueSuspect(Issue.Id, x)).ToList();
+				Updates.Add(Builders<Issue>.Update.Set(x => x.Summary, NewSummary));
 			}
-			if (Issue.OwnerId != null)
+			if (Issue.Severity != NewSeverity)
 			{
-				NewSuspects.RemoveAll(x => x.AuthorId != Issue.OwnerId);
-				if (NewSuspects.Count == 0)
-				{
-					NewSuspects.Add(new IssueSuspect(Issue.Id, Issue.OwnerId.Value, 0, null, null));
-				}
+				Updates.Add(Builders<Issue>.Update.Set(x => x.Severity, NewSeverity));
 			}
-
-			// Find the current list of suspects
-			List<IssueSuspect> CurSuspects = await IssueSuspects.Find(x => x.IssueId == Issue.Id).ToListAsync();
-
-			HashSet<(UserId, int)> CurSuspectKeys = new HashSet<(UserId, int)>(CurSuspects.Select(x => (x.AuthorId, x.Change)));
-			List<IssueSuspect> CreateSuspects = NewSuspects.Where(x => !CurSuspectKeys.Contains((x.AuthorId, x.Change))).ToList();
-
-			HashSet<(UserId, int)> NewSuspectKeys = new HashSet<(UserId, int)>(NewSuspects.Select(x => (x.AuthorId, x.Change)));
-			List<IssueSuspect> DeleteSuspects = CurSuspects.Where(x => !NewSuspectKeys.Contains((x.AuthorId, x.Change))).ToList();
-
-			// Apply the suspect changes
-			if (CreateSuspects.Count > 0)
+			if (Issue.Fingerprints.Count != NewFingerprints.Count || !NewFingerprints.Zip(Issue.Fingerprints).All(x => x.First.Equals(x.Second)))
 			{
-				await IssueSuspects.InsertManyIgnoreDuplicatesAsync(CreateSuspects);
-				CurSuspects.AddRange(CreateSuspects);
+				Updates.Add(Builders<Issue>.Update.Set(x => x.Fingerprints, NewFingerprints.Select(x => new IssueFingerprint(x))));
 			}
-			if (DeleteSuspects.Count > 0)
+			if (Issue.Streams.Count != NewStreams.Count || !NewStreams.Zip(Issue.Streams).All(x => x.First.StreamId == x.Second.StreamId && x.First.ContainsFix == x.Second.ContainsFix))
 			{
-				await IssueSuspects.DeleteManyAsync(Builders<IssueSuspect>.Filter.In(x => x.Id, DeleteSuspects.Select(y => y.Id)));
-				CurSuspects.RemoveAll(x => !NewSuspectKeys.Contains((x.AuthorId, x.Change)));
+				Updates.Add(Builders<Issue>.Update.Set(x => x.Streams, NewStreams.Select(x => new IssueStream(x))));
 			}
-
-			// Calculate the last time that this issue was seen
-			DateTime NewLastSeenAt = Issue.CreatedAt;
-			foreach (IIssueSpan Span in Spans)
+			if (IssueImpl.MinSuspectChange != NewMinSuspectChange)
 			{
-				if (Span.LastFailure.StepTime > NewLastSeenAt)
-				{
-					NewLastSeenAt = Span.LastFailure.StepTime;
-				}
+				Updates.Add(Builders<Issue>.Update.Set(x => x.MinSuspectChange, NewMinSuspectChange));
 			}
-
-			// Calculate the time at which the issue is verified fixed
-			DateTime? VerifiedAt = null;
-			foreach (IIssueSpan Span in Spans)
+			if (IssueImpl.MaxSuspectChange != NewMaxSuspectChange)
 			{
-				if (Span.NextSuccess == null)
-				{
-					break;
-				}
-				else if (VerifiedAt == null || Span.NextSuccess.StepTime < VerifiedAt.Value)
-				{
-					VerifiedAt = Span.NextSuccess.StepTime;
-				}
+				Updates.Add(Builders<Issue>.Update.Set(x => x.MaxSuspectChange, NewMaxSuspectChange));
 			}
-			if (Spans.Count == 0)
+			if (IssueImpl.DefaultOwnerId != NewDefaultOwnerId)
 			{
-				VerifiedAt = Issue.VerifiedAt ?? DateTime.UtcNow;
+				Updates.Add(Builders<Issue>.Update.Set(x => x.DefaultOwnerId, NewDefaultOwnerId));
 			}
-
-			// Update the spans with the correct resolved time
-			DateTime? ResolvedAt = Issue.ResolvedAt ?? VerifiedAt;
-
-			for(int Idx = 0; Idx < Spans.Count; Idx++)
+			if (Issue.ResolvedAt != NewResolvedAt)
 			{
-				IssueSpan? Span = Spans[Idx];
-				if (Span.ResolvedAt != ResolvedAt)
-				{
-					Span = await TryUpdateSpanAsync(Span, Builders<IssueSpan>.Update.Set(x => x.ResolvedAt, ResolvedAt));
-					if (Span == null)
-					{
-						return null;
-					}
-					Spans[Idx] = Span;
-				}
+				Updates.Add(Builders<Issue>.Update.SetOrUnsetNull(x => x.ResolvedAt, NewResolvedAt));
 			}
-
-			List<ObjectId> UpdateSuspectIds = CurSuspects.Where(x => x.ResolvedAt != ResolvedAt).Select(x => x.Id).ToList();
-			if (UpdateSuspectIds.Count > 0)
+			if (NewResolvedAt == null && Issue.ResolvedById != null)
 			{
-				await IssueSuspects.UpdateManyAsync(Builders<IssueSuspect>.Filter.In(x => x.Id, UpdateSuspectIds), Builders<IssueSuspect>.Update.Set(x => x.ResolvedAt, ResolvedAt));
+				Updates.Add(Builders<Issue>.Update.Unset(x => x.ResolvedById));
 			}
-
-			// Build a list of streams
-			List<IssueStream> Streams = new List<IssueStream>();
-			foreach (IGrouping<StreamId, IIssueSpan> Group in Spans.GroupBy(x => x.StreamId))
+			if (Issue.VerifiedAt != NewVerifiedAt)
 			{
-				IssueStream? Stream = Issue.Streams.FirstOrDefault(x => x.StreamId == Group.Key);
-				if (Stream == null)
-				{
-					Stream = new IssueStream { StreamId = Group.Key };
-				}
-				else
-				{
-					Stream = new IssueStream(Stream);
-				}
-				Streams.Add(Stream);
+				Updates.Add(Builders<Issue>.Update.SetOrUnsetNull(x => x.VerifiedAt, NewVerifiedAt));
+			}
+			if (Issue.LastSeenAt != NewLastSeenAt)
+			{
+				Updates.Add(Builders<Issue>.Update.Set(x => x.LastSeenAt, NewLastSeenAt));
 			}
 
-			// Update the issue document itself
-			UpdateDefinition<Issue> Update = Builders<Issue>.Update
-				.Set(x => x.Severity, Spans.Any(x => x.Severity == IssueSeverity.Error) ? IssueSeverity.Error : IssueSeverity.Warning)
-				.Set(x => x.LastSeenAt, NewLastSeenAt)
-				.Set(x => x.ResolvedAt, ResolvedAt)
-				.Set(x => x.VerifiedAt, VerifiedAt)
-				.Set(x => x.Suspects, CurSuspects)
-				.Set(x => x.Streams, Streams)
-				.Set(x => x.MinSuspectChange, (CurSuspects.Count > 0) ? CurSuspects.Min(x => x.Change) : 0)
-				.Set(x => x.MaxSuspectChange, (CurSuspects.Count > 0) ? CurSuspects.Max(x => x.Change) : 0);
-
-			// Apply the update
-			Issue? NewIssue = await TryUpdateIssueAsync(Issue, Update);
-			if (NewIssue == null)
+			Issue? NewIssue = await TryUpdateIssueAsync(Issue, Builders<Issue>.Update.Combine(Updates));
+			if(NewIssue != null)
 			{
-				return null;
+				ILogger IssueLogger = GetLogger(Issue.Id);
+				LogIssueChanges(GetLogger(Issue.Id), IssueImpl, NewIssue);
+				LogIssueSuspectChanges(IssueLogger, OldSuspectImpls, NewSuspectImpls);
+				return NewIssue;
 			}
+			return null;
+		}
 
-			// Log any changes
-			ILogger IssueLogger = GetLogger(Issue.Id);
-			LogIssueChanges(IssueLogger, Issue, NewIssue);
+		async Task<List<IssueSuspect>> UpdateIssueSuspectsAsync(int IssueId, List<IssueSuspect> OldSuspectImpls, List<NewIssueSuspectData> NewSuspects, DateTime? ResolvedAt)
+		{
+			List<IssueSuspect> NewSuspectImpls = new List<IssueSuspect>(OldSuspectImpls);
+
+			// Find the current list of suspects
+			HashSet<(UserId, int)> CurSuspectKeys = new HashSet<(UserId, int)>(OldSuspectImpls.Select(x => (x.AuthorId, x.Change)));
+			List<IssueSuspect> CreateSuspects = NewSuspects.Where(x => !CurSuspectKeys.Contains((x.AuthorId, x.Change))).Select(x => new IssueSuspect(IssueId, x, ResolvedAt)).ToList();
+
+			HashSet<(UserId, int)> NewSuspectKeys = new HashSet<(UserId, int)>(NewSuspects.Select(x => (x.AuthorId, x.Change)));
+			List<IssueSuspect> DeleteSuspects = OldSuspectImpls.Where(x => !NewSuspectKeys.Contains((x.AuthorId, x.Change))).ToList();
 
-			// Clear the modified flag on any spans
-			for (int Idx = 0; Idx < Spans.Count; Idx++)
+			// Apply the suspect changes
+			if (CreateSuspects.Count > 0)
 			{
-				IssueSpan? Span = Spans[Idx];
-				if (Span.Modified)
-				{
-					Span = await TryUpdateSpanAsync(Span, Builders<IssueSpan>.Update.Set(x => x.Modified, false));
-					if (Span == null)
-					{
-						return null;
-					}
-					Spans[Idx] = Span;
-				}
+				await IssueSuspects.InsertManyIgnoreDuplicatesAsync(CreateSuspects);
+				NewSuspectImpls.AddRange(CreateSuspects);
+			}
+			if (DeleteSuspects.Count > 0)
+			{
+				await IssueSuspects.DeleteManyAsync(Builders<IssueSuspect>.Filter.In(x => x.Id, DeleteSuspects.Select(y => y.Id)));
+				NewSuspectImpls.RemoveAll(x => !NewSuspectKeys.Contains((x.AuthorId, x.Change)));
 			}
 
-			return NewIssue;
+			// Make sure all the remaining suspects have the correct resolved time
+			if (NewSuspectImpls.Any(x => x.ResolvedAt != ResolvedAt))
+			{
+				await IssueSuspects.UpdateManyAsync(Builders<IssueSuspect>.Filter.Eq(x => x.IssueId, IssueId), Builders<IssueSuspect>.Update.Set(x => x.ResolvedAt, ResolvedAt));
+			}
+			return NewSuspectImpls;
 		}
 
 		#endregion
@@ -1272,22 +1036,11 @@ namespace HordeServer.Collections.Impl
 		#region Spans
 
 		/// <inheritdoc/>
-		public async Task<IIssueSpan?> AddSpanAsync(IIssueSequenceToken Token, IStream Stream, TemplateRefId TemplateRefId, string NodeName, NewIssueFingerprint Fingerprint, NewIssueStepData? LastSuccess, NewIssueStepData Failure, NewIssueStepData? NextSuccess, List<NewIssueSpanSuspectData> Suspects)
+		public async Task<IIssueSpan> AddSpanAsync(int IssueId, NewIssueSpanData NewSpan)
 		{
-			List<IssueSpanSuspect> NewSuspects = CreateSpanSuspects(Suspects);
-			IssueSpan NewSpan = new IssueSpan(Stream.Id, Stream.Name, TemplateRefId, NodeName, Fingerprint, LastSuccess, Failure, NextSuccess, NewSuspects);
-
-			IssueLedger Ledger = ((IssueSequenceToken)Token).Ledger;
-			Ledger.InsertIssueSpan = NewSpan;
-
-			if (!await LedgerSingleton.TryUpdateAsync(Ledger))
-			{
-				return null;
-			}
-
-			await FlushLedgerAsync(Ledger);
-			await AddStepAsync(NewSpan.Id, Failure);
-			return NewSpan;
+			IssueSpan Span = new IssueSpan(IssueId, NewSpan);
+			await IssueSpans.InsertOneAsync(Span);
+			return Span;
 		}
 
 		/// <inheritdoc/>
@@ -1297,13 +1050,9 @@ namespace HordeServer.Collections.Impl
 		}
 
 		/// <inheritdoc/>
-		public async Task<IIssueSpan?> TryUpdateSpanAsync(IIssueSpan Span, IssueSeverity? NewSeverity = null, NewIssueStepData? NewLastSuccess = null, NewIssueStepData? NewFailure = null, NewIssueStepData? NewNextSuccess = null, List<NewIssueSpanSuspectData>? NewSuspects = null, bool? NewModified = null)
+		public async Task<IIssueSpan?> TryUpdateSpanAsync(IIssueSpan Span, NewIssueStepData? NewLastSuccess = null, NewIssueStepData? NewFailure = null, NewIssueStepData? NewNextSuccess = null, List<NewIssueSpanSuspectData>? NewSuspects = null, int? NewIssueId = null)
 		{
 			List<UpdateDefinition<IssueSpan>> Updates = new List<UpdateDefinition<IssueSpan>>();
-			if (NewSeverity != null)
-			{
-				Updates.Add(Builders<IssueSpan>.Update.Set(x => x.Severity, NewSeverity.Value));
-			}
 			if (NewLastSuccess != null)
 			{
 				Updates.Add(Builders<IssueSpan>.Update.Set(x => x.MinChange, NewLastSuccess.Change));
@@ -1331,11 +1080,11 @@ namespace HordeServer.Collections.Impl
 			}
 			if (NewSuspects != null)
 			{
-				Updates.Add(Builders<IssueSpan>.Update.Set(x => x.Suspects, CreateSpanSuspects(NewSuspects)));
+				Updates.Add(Builders<IssueSpan>.Update.Set(x => x.Suspects, NewSuspects.ConvertAll(x => new IssueSpanSuspect(x))));
 			}
-			if (NewModified != null)
+			if (NewIssueId != null)
 			{
-				Updates.Add(Builders<IssueSpan>.Update.Set(x => x.Modified, NewModified.Value));
+				Updates.Add(Builders<IssueSpan>.Update.Set(x => x.IssueId, NewIssueId.Value));
 			}
 
 			if (Updates.Count == 0)
@@ -1346,25 +1095,18 @@ namespace HordeServer.Collections.Impl
 			IssueSpan? NewSpan = await TryUpdateSpanAsync(Span, Builders<IssueSpan>.Update.Combine(Updates));
 			if (NewSpan != null)
 			{
-				if (NewFailure != null)
+				ILogger Logger = GetLogger(NewSpan.IssueId);
+				if (NewLastSuccess != null)
 				{
-					await AddStepAsync(NewSpan.Id, NewFailure);
+					Logger.LogInformation("Set last success for span {SpanId} to job {JobId} at CL {Change}", NewSpan.Id, NewLastSuccess.JobId, NewLastSuccess.Change);
 				}
-				if (NewSpan.IssueId != null)
+				if (NewNextSuccess != null)
 				{
-					ILogger Logger = GetLogger(NewSpan.IssueId.Value);
-					if (NewLastSuccess != null)
-					{
-						Logger.LogInformation("Set last success for span {SpanId} to job {JobId} at CL {Change}", NewSpan.Id, NewLastSuccess.JobId, NewLastSuccess.Change);
-					}
-					if (NewNextSuccess != null)
-					{
-						Logger.LogInformation("Set next success for span {SpanId} to job {JobId} at CL {Change}", NewSpan.Id, NewNextSuccess.JobId, NewNextSuccess.Change);
-					}
-					if (NewFailure != null)
-					{
-						Logger.LogInformation("Added failure for span {SpanId} in job {JobId} at CL {Change}", NewSpan.Id, NewFailure.JobId, NewFailure.Change);
-					}
+					Logger.LogInformation("Set next success for span {SpanId} to job {JobId} at CL {Change}", NewSpan.Id, NewNextSuccess.JobId, NewNextSuccess.Change);
+				}
+				if (NewFailure != null)
+				{
+					Logger.LogInformation("Added failure for span {SpanId} in job {JobId} at CL {Change}", NewSpan.Id, NewFailure.JobId, NewFailure.Change);
 				}
 			}
 			return NewSpan;
@@ -1389,62 +1131,16 @@ namespace HordeServer.Collections.Impl
 			return Spans.ConvertAll<IIssueSpan>(x => x);
 		}
 
-		/// <inheritdoc/>
-		public async Task<bool> AttachSpanToIssueAsync(IIssueSequenceToken Token, IIssueSpan Span, IIssue Issue, IssueSeverity? NewSeverity = null, NewIssueFingerprint? NewFingerprint = null, string? NewSummary = null)
-		{
-			List<UpdateDefinition<Issue>> IssueUpdates = new List<UpdateDefinition<Issue>>();
-			if (NewSeverity != null)
-			{
-				IssueUpdates.Add(Builders<Issue>.Update.Set(x => x.Severity, NewSeverity.Value));
-			}
-			if (NewFingerprint != null)
-			{
-				IssueUpdates.Add(Builders<Issue>.Update.Set(x => x.Fingerprint, new IssueFingerprint(NewFingerprint)));
-			}
-			if (NewSummary != null)
-			{
-				IssueUpdates.Add(Builders<Issue>.Update.Set(x => x.Summary, NewSummary));
-			}
-			if (Span.LastFailure.StepTime > Issue.LastSeenAt)
-			{
-				IssueUpdates.Add(Builders<Issue>.Update.Set(x => x.LastSeenAt, Span.LastFailure.StepTime));
-			}
-			if (Span.PromoteByDefault && !Issue.Promoted)
-			{
-				IssueUpdates.Add(Builders<Issue>.Update.Set(x => x.Promoted, true));
-			}
-
-			IssueSequenceToken TokenValue = ((IssueSequenceToken)Token);
-			if (IssueUpdates.Count > 0)
-			{
-				TokenValue.Ledger.UpdateIssue = UpdateTransaction.Create(Builders<Issue>.Filter.Eq(x => x.Id, Issue.Id), Builders<Issue>.Update.Combine(IssueUpdates));
-			}
-			TokenValue.Ledger.UpdateIssueSpan = UpdateTransaction.Create(Builders<IssueSpan>.Filter.Eq(x => x.Id, Span.Id), Builders<IssueSpan>.Update.Set(x => x.IssueId, Issue.Id));
-
-			if (!await LedgerSingleton.TryUpdateAsync(TokenValue.Ledger))
-			{
-				return false;
-			}
-
-			LogSpanInfo(GetLogger(Issue.Id), Span);
-			await TokenValue.FlushAsync();
-			return true;
-		}
-
-		/// <inheritdoc/>
-		public async Task<List<IIssueSpan>> FindModifiedSpansAsync()
-		{
-			return await IssueSpans.Find(x => x.Modified).ToListAsync<IssueSpan, IIssueSpan>();
-		}
-
 		#endregion
 
 		#region Steps
 
 		/// <inheritdoc/>
-		async Task AddStepAsync(ObjectId SpanId, NewIssueStepData NewStep)
+		public async Task<IIssueStep> AddStepAsync(ObjectId SpanId, NewIssueStepData NewStep)
 		{
-			await IssueSteps.InsertOneIgnoreDuplicatesAsync(new IssueStep(SpanId, NewStep));
+			IssueStep Step = new IssueStep(SpanId, NewStep);
+			await IssueSteps.InsertOneAsync(Step);
+			return Step;
 		}
 
 		/// <inheritdoc/>
diff --git a/Engine/Source/Programs/Horde/HordeServer/Issues/Impl/IssueService.cs b/Engine/Source/Programs/Horde/HordeServer/Issues/Impl/IssueService.cs
--- a/Engine/Source/Programs/Horde/HordeServer/Issues/Impl/IssueService.cs
+++ b/Engine/Source/Programs/Horde/HordeServer/Issues/Impl/IssueService.cs
@@ -113,7 +113,7 @@ namespace HordeServer.Services.Impl
 		/// <returns>True if the user should be notified for this change</returns>
 		public bool ShowNotifications()
 		{
-			return ShowDesktopAlerts && Issue.Fingerprint.Type != DefaultIssueHandler.Type;
+			return ShowDesktopAlerts && Issue.Fingerprints.Any(x => x.Type != DefaultIssueHandler.Type);
 		}
 
 		/// <summary>
@@ -251,7 +251,7 @@ namespace HordeServer.Services.Impl
 				IIssue OpenIssue = OpenIssues[Idx];
 				if (OpenIssue.LastSeenAt < UtcNow - TimeSpan.FromDays(7.0))
 				{
-					await IssueCollection.UpdateIssueAsync(OpenIssue, NewResolvedById: IIssue.ResolvedByTimeoutId);
+					await IssueCollection.TryUpdateIssueAsync(OpenIssue, NewResolvedById: IIssue.ResolvedByTimeoutId);
 					OpenIssues.RemoveAt(Idx--);
 				}
 			}
@@ -291,7 +291,7 @@ namespace HordeServer.Services.Impl
 
 			List<IIssueSpan> Spans = await IssueCollection.FindSpansAsync(Issue.Id);
 			List<IIssueStep> Steps = await IssueCollection.FindStepsAsync(Spans.Select(x => x.Id));
-			List<IIssueSuspect> Suspects = await IssueCollection.GetSuspectsAsync(Issue);
+			List<IIssueSuspect> Suspects = await IssueCollection.FindSuspectsAsync(Issue);
 
 			List<IUser> SuspectUsers = new List<IUser>();
 			foreach (UserId SuspectUserId in Suspects.Select(x => x.AuthorId).Distinct())
@@ -321,7 +321,7 @@ namespace HordeServer.Services.Impl
 		/// <inheritdoc/>
 		public Task<List<IIssueSuspect>> GetIssueSuspectsAsync(IIssue Issue)
 		{
-			return IssueCollection.GetSuspectsAsync(Issue);
+			return IssueCollection.FindSuspectsAsync(Issue);
 		}
 
 		/// <inheritdoc/>
@@ -342,15 +342,7 @@ namespace HordeServer.Services.Impl
 				Spans.RemoveAll(x => !NodeNames.Contains(x.NodeName));
 			}
 
-			List<int> IssueIds = new List<int>();
-			foreach (IIssueSpan Span in Spans)
-			{
-				if (Span.IssueId != null)
-				{
-					IssueIds.Add(Span.IssueId.Value);
-				}
-			}
-
+			List<int> IssueIds = new List<int>(Spans.Select(x => x.IssueId));
 			if(IssueIds.Count == 0)
 			{
 				return new List<IIssue>();
@@ -394,7 +386,6 @@ namespace HordeServer.Services.Impl
 		/// <inheritdoc/>
 		public async Task<bool> UpdateIssueAsync(int Id, string? UserSummary = null, string? Description = null, bool? Promoted = null, UserId? OwnerId = null, UserId? NominatedById = null, bool? Acknowledged = null, UserId? DeclinedById = null, int? FixChange = null, UserId? ResolvedById = null)
 		{
-			Dictionary<StreamId, bool>? NewFixStreamIds = null;
 			for (; ; )
 			{
 				IIssue? Issue = await IssueCollection.GetIssueAsync(Id);
@@ -403,16 +394,10 @@ namespace HordeServer.Services.Impl
 					return false;
 				}
 
-				if (FixChange != null && FixChange.Value > 0)
-				{
-					NewFixStreamIds = await GetFixStreamsAsync(Issue.Id, FixChange.Value, NewFixStreamIds ?? ((FixChange == Issue.FixChange)? Issue.GetFixStreamIds() : null));
-				}
-
-				Issue = await IssueCollection.UpdateIssueAsync(Issue, NewUserSummary: UserSummary, NewDescription: Description, NewPromoted: Promoted, NewOwnerId: OwnerId ?? ResolvedById, NewNominatedById: NominatedById, NewAcknowledged: Acknowledged, NewDeclinedById: DeclinedById, NewFixChange: FixChange, NewFixStreamIds: NewFixStreamIds, NewResolvedById: ResolvedById);
-
+				Issue = await IssueCollection.TryUpdateIssueAsync(Issue, NewUserSummary: UserSummary, NewDescription: Description, NewPromoted: Promoted, NewOwnerId: OwnerId ?? ResolvedById, NewNominatedById: NominatedById, NewDeclinedById: DeclinedById, NewAcknowledged: Acknowledged, NewFixChange: FixChange, NewResolvedById: ResolvedById);
 				if(Issue != null)
 				{
-					OnIssueUpdated?.Invoke(Issue);
+					await UpdateIssueDerivedDataAsync(Issue);
 					return true;
 				}
 			}
@@ -555,39 +540,37 @@ namespace HordeServer.Services.Impl
 			}
 
 			// Try to update all the events. We may need to restart this due to optimistic transactions, so keep track of any existing spans we do not need to check against.
-			HashSet<ObjectId> CheckedSpans = new HashSet<ObjectId>();
-			for(; ;)
+			await using(IAsyncDisposable Lock = await IssueCollection.EnterCriticalSectionAsync())
 			{
-				// Get a token for transactional operations on the issue collections
-				IIssueSequenceToken Token = await IssueCollection.GetSequenceTokenAsync();
-
-				// Get the spans that are currently open
-				List<IIssueSpan> OpenSpans = await IssueCollection.FindOpenSpansAsync(Job.StreamId, Job.TemplateId, Node.Name, Job.Change);
-				Logger.LogDebug("{NumSpans} spans are open in {StreamId} at CL {Change} for template {TemplateId}, node {Node}", OpenSpans.Count, Job.StreamId, Job.Change, Job.TemplateId, Node.Name);
-
-				// Add the events to existing issues, and create new issues for everything else
-				if (EventGroups.Count > 0)
+				HashSet<ObjectId> CheckedSpanIds = new HashSet<ObjectId>();
+				for (; ; )
 				{
-					if (!await AddEventsToExistingSpans(Job, Batch, Step, EventGroups, OpenSpans, CheckedSpans, bPromoteByDefault))
+					// Get the spans that are currently open
+					List<IIssueSpan> OpenSpans = await IssueCollection.FindOpenSpansAsync(Job.StreamId, Job.TemplateId, Node.Name, Job.Change);
+					Logger.LogDebug("{NumSpans} spans are open in {StreamId} at CL {Change} for template {TemplateId}, node {Node}", OpenSpans.Count, Job.StreamId, Job.Change, Job.TemplateId, Node.Name);
+
+					// Add the events to existing issues, and create new issues for everything else
+					if (EventGroups.Count > 0)
 					{
-						continue;
+						if (!await AddEventsToExistingSpans(Job, Batch, Step, EventGroups, OpenSpans, CheckedSpanIds, bPromoteByDefault))
+						{
+							continue;
+						}
+						if (!await AddEventsToNewSpans(Stream, Job, Batch, Step, Node, OpenSpans, EventGroups, bPromoteByDefault))
+						{
+							continue;
+						}
 					}
-					if (!await AddEventsToNewSpans(Token, Stream, Job, Batch, Step, Node, OpenSpans, EventGroups, bPromoteByDefault))
+
+					// Try to update the sentinels for any other open steps
+					if (!await TryUpdateSentinelsAsync(OpenSpans, Stream, Job, Batch, Step))
 					{
 						continue;
 					}
-				}
 
-				// Try to update the sentinels for any other open steps
-				if (!await TryUpdateSentinelsAsync(OpenSpans, Stream, Job, Batch, Step))
-				{
-					continue;
+					break;
 				}
-				break;
 			}
-
-			// Update all the unassigned spans
-			await AttachSpansToIssuesAsync();
 		}
 
 		/// <summary>
@@ -684,35 +667,20 @@ namespace HordeServer.Services.Impl
 					List<NewEventGroup> MatchEventGroups = NewEventGroups.Where(x => OpenSpan.Fingerprint.IsMatch(x.Fingerprint)).ToList();
 					if (MatchEventGroups.Count > 0)
 					{
-						// Get the severity of these events
-						IssueSeverity Severity = MatchEventGroups.SelectMany(x => x.Events).Any(x => x.Event.Severity == EventSeverity.Error) ? IssueSeverity.Error : IssueSeverity.Warning;
-
-						// Get the new severity
-						IssueSeverity? NewSeverity = null;
-						if (OpenSpan.Severity == IssueSeverity.Warning && Severity == IssueSeverity.Error)
-						{
-							NewSeverity = IssueSeverity.Error;
-						}
+						// Add the new step data
+						NewIssueStepData NewFailure = new NewIssueStepData(Job, Batch, Step, GetIssueSeverity(MatchEventGroups.SelectMany(x => x.Events)), PromoteByDefault);
+						await IssueCollection.AddStepAsync(OpenSpan.Id, NewFailure);
 
-						// Get the new step data
-						NewIssueStepData NewFailure = new NewIssueStepData(Job, Batch, Step, Severity, PromoteByDefault);
-
-						// Figure out whether we need to update the issue
-						bool bMarkAsModified = NewSeverity != null;
-						if (OpenSpan.IssueId != null)
+						// Update the span if this changes the current range
+						IIssueSpan? NewSpan = OpenSpan;
+						if (NewFailure.Change <= OpenSpan.FirstFailure.Change || NewFailure.Change >= OpenSpan.LastFailure.Change)
 						{
-							IIssue? Issue = await IssueCollection.GetIssueAsync(OpenSpan.IssueId.Value);
-							if (Issue != null)
+							NewSpan = await IssueCollection.TryUpdateSpanAsync(OpenSpan, NewFailure: NewFailure);
+							if (NewSpan == null)
 							{
-								bMarkAsModified |= Issue.ResolvedAt != null;
+								return false;
 							}
-						}
-
-						// Update the span
-						IIssueSpan? NewSpan = await IssueCollection.TryUpdateSpanAsync(OpenSpan, NewSeverity: NewSeverity, NewFailure: NewFailure, NewModified: (bMarkAsModified? (bool?)true : null));
-						if (NewSpan == null)
-						{
-							return false;
+							await UpdateIssueDerivedDataAsync(NewSpan.IssueId);
 						}
 
 						// Write out all the merged events
@@ -761,7 +729,6 @@ namespace HordeServer.Services.Impl
 		/// <summary>
 		/// Adds new spans for the given events
 		/// </summary>
-		/// <param name="Token">Token for the update</param>
 		/// <param name="Stream">The stream containing the job</param>
 		/// <param name="Job">The job instance</param>
 		/// <param name="Batch">The job batch</param>
@@ -771,7 +738,7 @@ namespace HordeServer.Services.Impl
 		/// <param name="NewEventGroups">Set of remaining events</param>
 		/// <param name="PromoteByDefault"></param>
 		/// <returns>True if all events were added</returns>
-		async Task<bool> AddEventsToNewSpans(IIssueSequenceToken Token, IStream Stream, IJob Job, IJobStepBatch Batch, IJobStep Step, INode Node, List<IIssueSpan> OpenSpans, HashSet<NewEventGroup> NewEventGroups, bool PromoteByDefault)
+		async Task<bool> AddEventsToNewSpans(IStream Stream, IJob Job, IJobStepBatch Batch, IJobStep Step, INode Node, List<IIssueSpan> OpenSpans, HashSet<NewEventGroup> NewEventGroups, bool PromoteByDefault)
 		{
 			while (NewEventGroups.Count > 0)
 			{
@@ -795,13 +762,27 @@ namespace HordeServer.Services.Impl
 				// Get the step data
 				NewIssueStepData StepData = new NewIssueStepData(Job, Batch, Step, GetIssueSeverity(EventGroup.Events), PromoteByDefault);
 
-				// Create the new span
-				IIssueSpan? NewSpan = await AddEventToNewSpanAsync(Token, Stream, Job, Node, EventGroup.Fingerprint, StepData);
-				if(NewSpan == null)
+				// Get the span data
+				NewIssueSpanData SpanData = new NewIssueSpanData(Stream.Id, Stream.Name, Job.TemplateId, Node.Name, EventGroup.Fingerprint, StepData);
+
+				IJobStepRef? PrevJob = await JobStepRefs.GetPrevStepForNodeAsync(Job.StreamId, Job.TemplateId, Node.Name, Job.Change);
+				if (PrevJob != null)
 				{
-					return false;
+					SpanData.LastSuccess = new NewIssueStepData(PrevJob);
+					SpanData.Suspects = await FindSuspectsForSpanAsync(Stream, SpanData.Fingerprint, SpanData.LastSuccess.Change + 1, SpanData.FirstFailure.Change);
+				}
+
+				IJobStepRef? NextJob = await JobStepRefs.GetNextStepForNodeAsync(Job.StreamId, Job.TemplateId, Node.Name, Job.Change);
+				if (NextJob != null)
+				{
+					SpanData.NextSuccess = new NewIssueStepData(NextJob);
 				}
-				OpenSpans.Add(NewSpan);
+
+				// Add all the new objects
+				IIssue NewIssue = await FindOrAddIssueForSpanAsync(SpanData);
+				IIssueSpan NewSpan = await IssueCollection.AddSpanAsync(NewIssue.Id, SpanData);
+				IIssueStep NewStep = await IssueCollection.AddStepAsync(NewSpan.Id, StepData);
+				await UpdateIssueDerivedDataAsync(NewIssue);
 
 				// Update the log events
 				Logger.LogDebug("Created new span {SpanId} from event group {Group}", NewSpan.Id, EventGroup.Digest.ToString());
@@ -809,192 +790,188 @@ namespace HordeServer.Services.Impl
 
 				// Remove the events from the remaining list of events to match
 				NewEventGroups.ExceptWith(SourceEventGroups);
+				OpenSpans.Add(NewSpan);
 			}
 			return true;
 		}
 
-		/// <summary>
-		/// Adds an event to a new span
-		/// </summary>
-		/// <param name="Token">The sequence token to insert a new span</param>
-		/// <param name="Stream">Stream containing the job</param>
-		/// <param name="Job">Job that spawned the event</param>
-		/// <param name="Node">Node that was running</param>
-		/// <param name="Fingerprint">Fingerprint for the event</param>
-		/// <param name="NewFailure">Information about the step that failed</param>
-		/// <returns>The new span, or null if the insert failed</returns>
-		async Task<IIssueSpan?> AddEventToNewSpanAsync(IIssueSequenceToken Token, IStream Stream, IJob Job, INode Node, NewIssueFingerprint Fingerprint, NewIssueStepData NewFailure)
+		async Task<IIssue?> UpdateIssueDerivedDataAsync(int IssueId)
 		{
-			// Get the previous job in this span
-			IJobStepRef? PrevJob = await JobStepRefs.GetPrevStepForNodeAsync(Job.StreamId, Job.TemplateId, Node.Name, Job.Change);
-
-			NewIssueStepData? LastSuccess = null;
-			if (PrevJob != null)
+			IIssue? Issue = await IssueCollection.GetIssueAsync(IssueId);
+			if(Issue != null)
 			{
-				LastSuccess = new NewIssueStepData(PrevJob);
+				Issue = await UpdateIssueDerivedDataAsync(Issue);
 			}
-
-			// Get the next successful job in this stream
-			IJobStepRef? NextJob = await JobStepRefs.GetNextStepForNodeAsync(Job.StreamId, Job.TemplateId, Node.Name, Job.Change);
-
-			NewIssueStepData? NextSuccess = null;
-			if (NextJob != null)
-			{
-				NextSuccess = new NewIssueStepData(NextJob);
-			}
-
-			// Find all the suspects for this span
-			List<NewIssueSpanSuspectData> Suspects = new List<NewIssueSpanSuspectData>();
-			if (LastSuccess != null)
-			{
-				Suspects = await FindSuspectsForSpanAsync(Stream, Fingerprint, LastSuccess.Change + 1, NewFailure.Change);
-			}
-
-			// Add the span
-			return await IssueCollection.AddSpanAsync(Token, Stream, Job.TemplateId, Node.Name, Fingerprint, LastSuccess, NewFailure, NextSuccess, Suspects);
+			return Issue;
 		}
 
-		/// <summary>
-		/// Find all the spans which are not currently attached to an issue, and try to attach them
-		/// </summary>
-		/// <returns></returns>
-		async Task AttachSpansToIssuesAsync()
+		async Task<IIssue?> UpdateIssueDerivedDataAsync(IIssue Issue)
 		{
-			IIssueSequenceToken Token = await IssueCollection.GetSequenceTokenAsync();
-			for(; ;)
+			Dictionary<(StreamId, int), bool> FixChangeCache = new Dictionary<(StreamId, int), bool>();
+			for (; ; )
 			{
-				bool bFoundSpan = false;
+				// Find all the spans that are attached to the issue
+				List<IIssueSpan> Spans = await IssueCollection.FindSpansAsync(Issue.Id);
 
-				List<IIssueSpan> Spans = await IssueCollection.FindModifiedSpansAsync();
-				foreach (IIssueSpan Span in Spans)
+				// Update the suspects for this issue
+				List<NewIssueSuspectData> NewSuspects = new List<NewIssueSuspectData>();
+				if (Spans.Count > 0)
 				{
-					await using RedisLock Lock = new RedisLock(RedisService.Database, $"issues/locks/{Span.Id}");
-					if (await Lock.AcquireAsync(TimeSpan.FromMinutes(1)))
+					HashSet<int> SuspectChanges = new HashSet<int>(Spans[0].Suspects.Where(x => Issue.OwnerId == null || x.AuthorId == Issue.OwnerId).Select(x => x.OriginatingChange ?? x.Change));
+					for (int SpanIdx = 1; SpanIdx < Spans.Count; SpanIdx++)
 					{
-						await TryUpdateModifiedSpanAsync(Span, Token);
-						bFoundSpan = true;
+						SuspectChanges.IntersectWith(Spans[SpanIdx].Suspects.Select(x => x.OriginatingChange ?? x.Change));
 					}
+					NewSuspects = Spans[0].Suspects.Where(x => SuspectChanges.Contains(x.OriginatingChange ?? x.Change)).Select(x => new NewIssueSuspectData(x.AuthorId, x.OriginatingChange ?? x.Change)).ToList();
 				}
-
-				if (!bFoundSpan)
+				if (Issue.OwnerId != null)
 				{
-					break;
+					NewSuspects.RemoveAll(x => x.AuthorId != Issue.OwnerId);
+					if (NewSuspects.Count == 0)
+					{
+						NewSuspects.Add(new NewIssueSuspectData(Issue.OwnerId.Value, 0));
+					}
 				}
-			}
-		}
 
-		/// <summary>
-		/// Updates information derived from an updated span, and clear the modified flag on it
-		/// </summary>
-		/// <param name="Span"></param>
-		/// <param name="Token"></param>
-		/// <returns></returns>
-		async ValueTask<IIssueSpan?> TryUpdateModifiedSpanAsync(IIssueSpan Span, IIssueSequenceToken Token)
-		{
-			// Try to find an issue for it
-			IIssue? Issue;
-			if (Span.IssueId == null)
-			{
-				Issue = await FindIssueForSpanAsync(Span);
-				if (Issue == null)
+				// Create the combined fingerprint for this issue
+				List<NewIssueFingerprint> NewFingerprints = new List<NewIssueFingerprint>();
+				foreach (IIssueSpan Span in Spans)
 				{
-					Issue = await IssueCollection.AddIssueAsync(Token, GetSummary(Span.Fingerprint, Span.Severity), Span);
-					if (Issue == null)
+					NewIssueFingerprint? Fingerprint = NewFingerprints.FirstOrDefault(x => x.Type == Span.Fingerprint.Type);
+					if (Fingerprint == null)
+					{
+						NewFingerprints.Add(new NewIssueFingerprint(Span.Fingerprint));
+					}
+					else
 					{
-						await Token.ResetAsync();
-						return null;
+						Fingerprint.MergeWith(Span.Fingerprint);
 					}
 				}
+
+				// Find the severity for this issue
+				IssueSeverity NewSeverity;
+				if (Spans.Count == 0 || Spans.All(x => x.NextSuccess != null))
+				{
+					NewSeverity = Issue.Severity;
+				}
 				else
 				{
-					NewIssueFingerprint NewFingerprint = NewIssueFingerprint.Merge(Issue.Fingerprint, Span.Fingerprint);
-					IssueSeverity NewSeverity = (Issue.Severity > Span.Severity) ? Issue.Severity : Span.Severity;
-					string NewSummary = GetSummary(NewFingerprint, NewSeverity);
+					NewSeverity = Spans.Any(x => x.NextSuccess == null && x.LastFailure.Severity == IssueSeverity.Error) ? IssueSeverity.Error : IssueSeverity.Warning;
+				}
 
-					if (!await IssueCollection.AttachSpanToIssueAsync(Token, Span, Issue, NewSeverity, NewFingerprint, NewSummary))
-					{
-						await Token.ResetAsync();
-						return null;
-					}
+				// Find the default summary for this issue
+				string NewSummary;
+				if (NewFingerprints.Count == 0)
+				{
+					NewSummary = Issue.Summary;
 				}
-			}
-			else
-			{
-				// Update the issue that it belongs to
-				Issue = await GetIssueAsync(Span.IssueId.Value);
-				if (Issue != null)
+				else if (NewFingerprints.Count == 1)
+				{
+					NewSummary = GetSummary(NewFingerprints[0], NewSeverity);
+				}
+				else
 				{
-					Issue = await UpdateIssueDerivedDataAsync(Issue);
-					if (Issue == null)
+					NewSummary = (NewSeverity == IssueSeverity.Error) ? "Errors in multiple steps" : "Warnings in multiple steps";
+				}
+
+				// Calculate the last time that this issue was seen
+				DateTime NewLastSeenAt = Issue.CreatedAt;
+				foreach (IIssueSpan Span in Spans)
+				{
+					if (Span.LastFailure.StepTime > NewLastSeenAt)
 					{
-						return null;
+						NewLastSeenAt = Span.LastFailure.StepTime;
 					}
 				}
-			}
 
-			// Perform any shared issue updates
-			if (Issue != null)
-			{
-				Issue = await UpdateResolvedStateAsync(Issue, Span);
-				if(Issue == null)
+				// Calculate the time at which the issue is verified fixed
+				DateTime? NewVerifiedAt = null;
+				if (Spans.Count == 0)
 				{
-					return null;
+					NewVerifiedAt = Issue.VerifiedAt ?? DateTime.UtcNow;
+				}
+				else
+				{
+					foreach (IIssueSpan Span in Spans)
+					{
+						if (Span.NextSuccess == null)
+						{
+							break;
+						}
+						else if (NewVerifiedAt == null || Span.NextSuccess.StepTime < NewVerifiedAt.Value)
+						{
+							NewVerifiedAt = Span.NextSuccess.StepTime;
+						}
+					}
 				}
-				OnIssueUpdated?.Invoke(Issue);
-			}
 
-			// Clear the modified flag
-			return await IssueCollection.TryUpdateSpanAsync(Span, NewModified: false);
-		}
+				// Get the new resolved timestamp
+				DateTime? NewResolvedAt = Issue.ResolvedAt ?? NewVerifiedAt;
 
-		async Task<IIssue?> UpdateResolvedStateAsync(IIssue? Issue, IIssueSpan Span)
-		{
-			// Update the fix streams
-			if (Issue != null && Issue.FixChange != null && Issue.FixChange.Value > 0)
-			{
-				Dictionary<StreamId, bool>? FixStreamIds = await GetFixStreamsAsync(Issue.Id, Issue.FixChange.Value, Issue.GetFixStreamIds());
-				if (FixStreamIds != null)
+				// Get the new list of streams
+				List<NewIssueStream> NewStreams = Spans.Select(x => x.StreamId).OrderBy(x => x.ToString(), StringComparer.Ordinal).Distinct().Select(x => new NewIssueStream(x)).ToList();
+				if (Issue.FixChange != null)
 				{
-					if(FixStreamIds.Count > Issue.Streams.Count)
+					foreach (NewIssueStream NewStream in NewStreams)
 					{
-						Issue = await UpdateIssueDerivedDataAsync(Issue);
+						IIssueStream? Stream = Issue.Streams.FirstOrDefault(x => x.StreamId == NewStream.StreamId);
+						if (Stream != null)
+						{
+							NewStream.ContainsFix = Stream.ContainsFix;
+						}
+						if (NewStream.ContainsFix == null)
+						{
+							NewStream.ContainsFix = await ContainsFixChange(NewStream.StreamId, Issue.FixChange.Value, FixChangeCache);
+						}
 					}
-					if (Issue != null)
+				}
+
+				// If a fix changelist has been specified and it's not valid, clear it out
+				if (NewVerifiedAt == null && NewResolvedAt != null)
+				{
+					if (Spans.Any(x => HasFixFailed(x, Issue.FixChange, NewResolvedAt.Value, NewStreams)))
 					{
-						Issue = await IssueCollection.UpdateIssueAsync(Issue, NewFixStreamIds: FixStreamIds);
+						NewStreams.ForEach(x => x.ContainsFix = null);
+						NewResolvedAt = null;
 					}
 				}
-			}
 
-			// Update the resolved flag
-			if (Issue != null && Issue.ResolvedAt != null && ShouldMarkFixFailed(Issue, Span))
-			{
-				Issue = await IssueCollection.UpdateIssueAsync(Issue, NewResolvedById: UserId.Empty);
-			}
+				// Update the issue
+				IIssue? NewIssue = await IssueCollection.TryUpdateIssueDerivedDataAsync(Issue, NewSummary, NewSeverity, NewFingerprints, NewStreams, NewSuspects, NewResolvedAt, NewVerifiedAt, NewLastSeenAt);
+				if (NewIssue != null)
+				{
+					OnIssueUpdated?.Invoke(NewIssue);
+					return NewIssue;
+				}
 
-			return Issue;
+				// Fetch the issue and try again
+				NewIssue = await IssueCollection.GetIssueAsync(Issue.Id);
+				if (NewIssue == null)
+				{
+					return null;
+				}
+				Issue = NewIssue;
+			}
 		}
 
 		/// <summary>
 		/// Determine if an issue should be marked as fix failed
 		/// </summary>
-		/// <param name="Issue"></param>
-		/// <param name="Span"></param>
-		/// <returns></returns>
-		static bool ShouldMarkFixFailed(IIssue Issue, IIssueSpan Span)
+		static bool HasFixFailed(IIssueSpan Span, int? FixChange, DateTime ResolvedAt, IReadOnlyList<NewIssueStream> Streams)
 		{
-			if (Issue.ResolvedAt != null && Span.NextSuccess == null)
+			if (Span.NextSuccess == null)
 			{
-				if (Issue.FixChange != null && (Issue.Streams.FirstOrDefault(x => x.StreamId == Span.StreamId)?.ContainsFix ?? false))
+				NewIssueStream? Stream = Streams.FirstOrDefault(x => x.StreamId == Span.StreamId);
+				if (Stream != null && FixChange != null && (Stream.ContainsFix ?? false))
 				{
-					if (Span.LastFailure.Change > Issue.FixChange.Value)
+					if (Span.LastFailure.Change > FixChange.Value)
 					{
 						return true;
 					}
 				}
 				else
 				{
-					if (Span.LastFailure.StepTime > Issue.ResolvedAt.Value + TimeSpan.FromHours(24.0))
+					if (Span.LastFailure.StepTime > ResolvedAt + TimeSpan.FromHours(24.0))
 					{
 						return true;
 					}
@@ -1004,49 +981,23 @@ namespace HordeServer.Services.Impl
 		}
 
 		/// <summary>
-		/// Finds the streams affected by an issue
+		/// Figure out if a stream contains a fix changelist
 		/// </summary>
-		/// <param name="IssueId"></param>
-		/// <param name="FixChange"></param>
-		/// <param name="PrevFixStreamIds"></param>
-		/// <returns></returns>
-		async Task<Dictionary<StreamId, bool>?> GetFixStreamsAsync(int IssueId, int FixChange, IReadOnlyDictionary<StreamId, bool>? PrevFixStreamIds = null)
+		async ValueTask<bool> ContainsFixChange(StreamId StreamId, int FixChange, Dictionary<(StreamId, int), bool> CachedContainsFixChange)
 		{
-			List<IIssueSpan> Spans = await IssueCollection.FindSpansAsync(IssueId);
-			if(Spans.Count == 0)
+			bool bContainsFixChange;
+			if (!CachedContainsFixChange.TryGetValue((StreamId, FixChange), out bContainsFixChange))
 			{
-				return null;
-			}
-
-			Dictionary<StreamId, bool> NextFixStreamIds = new Dictionary<StreamId, bool>();
-			foreach (IIssueSpan Span in Spans)
-			{
-				if (!NextFixStreamIds.ContainsKey(Span.StreamId))
+				IStream? Stream = await Streams.GetCachedStream(StreamId);
+				if (Stream != null)
 				{
-					bool ContainsFix = false;
-					if (PrevFixStreamIds == null || !PrevFixStreamIds.TryGetValue(Span.StreamId, out ContainsFix))
-					{
-						IStream? Stream = await Streams.GetCachedStream(Span.StreamId);
-						if (Stream != null)
-						{
-							Logger.LogInformation("Querying fix changelist {FixChange} for issue {IssueId} in {StreamId} (prev={HavePrev})", FixChange, IssueId, Span.StreamId, PrevFixStreamIds != null);
-							List<ChangeSummary> Changes = await Perforce.GetChangesAsync(Stream.ClusterName, Span.StreamName, FixChange, FixChange, 1, null);
-							ContainsFix = Changes.Count > 0;
-						}
-					}
-					NextFixStreamIds[Span.StreamId] = ContainsFix;
+					Logger.LogInformation("Querying fix changelist {FixChange} in {StreamId}", FixChange, StreamId);
+					List<ChangeSummary> Changes = await Perforce.GetChangesAsync(Stream.ClusterName, Stream.Name, FixChange, FixChange, 1, null);
+					bContainsFixChange = Changes.Count > 0;
 				}
+				CachedContainsFixChange[(StreamId, FixChange)] = bContainsFixChange;
 			}
-
-			if (PrevFixStreamIds != null)
-			{
-				if(PrevFixStreamIds.Count != NextFixStreamIds.Count || PrevFixStreamIds.Any(x => !NextFixStreamIds.TryGetValue(x.Key, out bool Value) || x.Value != Value))
-				{
-					return null;
-				}
-			}
-
-			return NextFixStreamIds;
+			return bContainsFixChange;
 		}
 
 		/// <summary>
@@ -1105,7 +1056,7 @@ namespace HordeServer.Services.Impl
 		/// </summary>
 		/// <param name="Span">The span to find an issue for</param>
 		/// <returns>The matching issue</returns>
-		async Task<IIssue?> FindIssueForSpanAsync(IIssueSpan Span)
+		async Task<IIssue> FindOrAddIssueForSpanAsync(NewIssueSpanData Span)
 		{
 			List<IIssue> ExistingIssues;
 			if (Span.LastSuccess == null)
@@ -1114,29 +1065,16 @@ namespace HordeServer.Services.Impl
 			}
 			else
 			{
-				List<int> SourceChanges = Span.Suspects.ConvertAll(x => x.OriginatingChange ?? x.Change);
-				ExistingIssues = await IssueCollection.FindIssuesForSuspectsAsync(SourceChanges);
+				ExistingIssues = await IssueCollection.FindIssuesForChangesAsync(Span.Suspects.ConvertAll(x => x.OriginatingChange ?? x.Change));
 			}
-			return ExistingIssues.FirstOrDefault(x => x.Fingerprint.IsMatch(Span.Fingerprint));
-		}
 
-		/// <summary>
-		/// 
-		/// </summary>
-		/// <param name="Issue"></param>
-		/// <returns></returns>
-		async Task<IIssue?> UpdateIssueDerivedDataAsync(IIssue Issue)
-		{
-			IIssue? NewIssue = await IssueCollection.UpdateIssueDerivedDataAsync(Issue);
-			if (NewIssue != null)
+			IIssue? Issue = ExistingIssues.FirstOrDefault(x => x.Fingerprints.Any(y => y.IsMatch(Span.Fingerprint)));
+			if (Issue == null)
 			{
-				string NewSummary = GetSummary(NewIssue.Fingerprint, NewIssue.Severity);
-				if (!NewSummary.Equals(NewIssue.Summary, StringComparison.Ordinal))
-				{
-					NewIssue = await IssueCollection.UpdateIssueAsync(NewIssue, NewSummary: NewSummary);
-				}
+				string Summary = GetSummary(Span.Fingerprint, Span.FirstFailure.Severity);
+				Issue = await IssueCollection.AddIssueAsync(Summary);
 			}
-			return NewIssue;
+			return Issue;
 		}
 
 		/// <summary>
@@ -1175,21 +1113,25 @@ namespace HordeServer.Services.Impl
 					NewIssueStepData NewLastSuccess = new NewIssueStepData(Job.Change, IssueSeverity.Unspecified, Job.Name, Job.Id, Batch.Id, Step.Id, Step.StartTimeUtc ?? default, Step.LogId, false);
 					List<NewIssueSpanSuspectData> NewSuspects = await FindSuspectsForSpanAsync(Stream, Span.Fingerprint, Job.Change + 1, Span.FirstFailure.Change);
 
-					if (await IssueCollection.TryUpdateSpanAsync(Span, NewLastSuccess: NewLastSuccess, NewSuspects: NewSuspects, NewModified: true) == null)
+					if (await IssueCollection.TryUpdateSpanAsync(Span, NewLastSuccess: NewLastSuccess, NewSuspects: NewSuspects) == null)
 					{
 						return false;
 					}
 
 					Logger.LogInformation("Set last success for issue {IssueId}, template {TemplateId}, node {Node} as job {JobId}, cl {Change}", Span.IssueId, Job.TemplateId, Span.NodeName, Job.Id, Job.Change);
+					await UpdateIssueDerivedDataAsync(Span.IssueId);
 				}
 				else if (Job.Change > Span.LastFailure.Change && (Span.NextSuccess == null || Job.Change < Span.NextSuccess.Change))
 				{
 					NewIssueStepData NewNextSucccess = new NewIssueStepData(Job.Change, IssueSeverity.Unspecified, Job.Name, Job.Id, Batch.Id, Step.Id, Step.StartTimeUtc ?? default, Step.LogId, false);
-					if (await IssueCollection.TryUpdateSpanAsync(Span, NewNextSuccess: NewNextSucccess, NewModified: true) == null)
+
+					if (await IssueCollection.TryUpdateSpanAsync(Span, NewNextSuccess: NewNextSucccess) == null)
 					{
 						return false;
 					}
+
 					Logger.LogInformation("Set next success for issue {IssueId}, template {TemplateId}, node {Node} as job {JobId}, cl {Change}", Span.IssueId, Job.TemplateId, Span.NodeName, Job.Id, Job.Change);
+					await UpdateIssueDerivedDataAsync(Span.IssueId);
 				}
 			}
 			return true;
diff --git a/Engine/Source/Programs/Horde/HordeServer/Utilities/MongoExtensions.cs b/Engine/Source/Programs/Horde/HordeServer/Utilities/MongoExtensions.cs
--- a/Engine/Source/Programs/Horde/HordeServer/Utilities/MongoExtensions.cs
+++ b/Engine/Source/Programs/Horde/HordeServer/Utilities/MongoExtensions.cs
@@ -156,6 +156,27 @@ namespace HordeServer.Utilities
 			}
 		}
 
+		/// <summary>
+		/// Sets a field to a value, or unsets it if the value is null
+		/// </summary>
+		/// <typeparam name="TDocument">The document type</typeparam>
+		/// <typeparam name="TField">Type of the field to set</typeparam>
+		/// <param name="Update">Update builder</param>
+		/// <param name="Field">Expression for the field to set</param>
+		/// <param name="Value">New value to set</param>
+		/// <returns>Update defintiion</returns>
+		public static UpdateDefinition<TDocument> SetOrUnsetNull<TDocument, TField>(this UpdateDefinition<TDocument> Update, Expression<Func<TDocument, TField?>> Field, TField? Value) where TField : struct
+		{
+			if (Value.HasValue)
+			{
+				return Update.Set(Field, Value.Value);
+			}
+			else
+			{
+				return Update.Unset(new ExpressionFieldDefinition<TDocument, TField?>(Field));
+			}
+		}
+
 		/// <summary>
 		/// Sets a field to a value, or unsets it if the value is null
 		/// </summary>
