Horde: Add support for manually adding/removing spans from an issue.

- Issue fingerprints will merge together when the same type, and create separate fingerprints when not.
- Each issue records a list of spans it should NOT automatically merge with. Removing a span from an issue adds to that list, then automatically reassigns the span to any other (or a new) issue.
- Spans now include an id field in responses.

diff --git a/Engine/Source/Programs/Horde/HordeServer.Compute.Tests/IssueServiceTests.cs b/Engine/Source/Programs/Horde/HordeServer.Compute.Tests/IssueServiceTests.cs
--- a/Engine/Source/Programs/Horde/HordeServer.Compute.Tests/IssueServiceTests.cs
+++ b/Engine/Source/Programs/Horde/HordeServer.Compute.Tests/IssueServiceTests.cs
@@ -1008,6 +1008,98 @@ namespace HordeServerTests
 			Assert.AreEqual("Missing copyright notice in ToolchainInfo.cs", Issue.Summary);
 		}
 
+
+		[TestMethod]
+		public async Task AddSpanToIssueTest()
+		{
+			// Create the first issue
+			IIssue IssueA;
+			IIssueSpan SpanA;
+			{
+				IJob Job = CreateJob(MainStreamId, 120, "Test Build", Graph);
+
+				string[] Lines =
+				{
+					@"  DatasmithDirectLink.cpp.obj : error LNK2019: unresolved external symbol ""enum DirectLink::ECommunicationStatus __cdecl DirectLink::ValidateCommunicationStatus(void)"" (?ValidateCommunicationStatus@DirectLink@@YA?AW4ECommunicationStatus@1@XZ) referenced in function ""public: static int __cdecl FDatasmithDirectLink::ValidateCommunicationSetup(void)"" (?ValidateCommunicationSetup@FDatasmithDirectLink@@SAHXZ)",
+				};
+				await ParseEventsAsync(Job, 0, 0, Lines);
+				await UpdateCompleteStep(Job, 0, 0, JobStepOutcome.Failure);
+
+				List<IIssue> Issues = await IssueService.FindIssuesAsync();
+				Assert.AreEqual(1, Issues.Count);
+
+				IssueA = Issues[0];
+				List<IIssueSpan> Spans = await IssueService.GetIssueSpansAsync(IssueA);
+				Assert.AreEqual(Spans.Count, 1);
+				SpanA = Spans[0];
+			}
+
+			// Create the second issue
+			IIssue IssueB;
+			IIssueSpan SpanB;
+			{
+				string[] Lines =
+				{
+					@"WARNING: Engine\Source\Programs\UnrealBuildTool\ProjectFiles\Rider\ToolchainInfo.cs: Missing copyright boilerplate"
+				};
+
+				IJob Job = CreateJob(MainStreamId, 120, "Test Build", Graph);
+				await ParseEventsAsync(Job, 0, 0, Lines);
+				await UpdateCompleteStep(Job, 0, 0, JobStepOutcome.Failure);
+
+				List<IIssue> Issues = await IssueService.FindIssuesAsync();
+				Issues.RemoveAll(x => x.Id == IssueA.Id);
+				Assert.AreEqual(1, Issues.Count);
+
+				IssueB = Issues[0];
+				List<IIssueSpan> Spans = await IssueService.GetIssueSpansAsync(IssueB);
+				Assert.AreEqual(1, Spans.Count);
+				SpanB = Spans[0];
+			}
+
+			// Add SpanB to IssueA
+			{
+				await IssueService.UpdateIssueAsync(IssueA.Id, AddSpanIds: new List<ObjectId> { SpanB.Id });
+
+				IIssue NewIssueA = (await IssueService.GetIssueAsync(IssueA.Id))!;
+				Assert.IsNull(NewIssueA.VerifiedAt);
+				Assert.IsNull(NewIssueA.ResolvedAt);
+				Assert.AreEqual(2, NewIssueA.Fingerprints.Count);
+				List<IIssueSpan> NewSpansA = await IssueService.GetIssueSpansAsync(NewIssueA!);
+				Assert.AreEqual(2, NewSpansA.Count);
+				Assert.AreEqual(NewIssueA.Id, NewSpansA[0].IssueId);
+				Assert.AreEqual(NewIssueA.Id, NewSpansA[1].IssueId);
+
+				IIssue NewIssueB = (await IssueService.GetIssueAsync(IssueB.Id))!;
+				Assert.IsNotNull(NewIssueB.VerifiedAt);
+				Assert.IsNotNull(NewIssueB.ResolvedAt);
+				Assert.AreEqual(0, NewIssueB.Fingerprints.Count);
+				List<IIssueSpan> NewSpansB = await IssueService.GetIssueSpansAsync(NewIssueB);
+				Assert.AreEqual(0, NewSpansB.Count);
+			}
+
+			// Add SpanA and SpanB to IssueB
+			{
+				await IssueService.UpdateIssueAsync(IssueB.Id, AddSpanIds: new List<ObjectId> { SpanA.Id, SpanB.Id });
+
+				IIssue NewIssueA = (await IssueService.GetIssueAsync(IssueA.Id))!;
+				Assert.IsNotNull(NewIssueA.VerifiedAt);
+				Assert.IsNotNull(NewIssueA.ResolvedAt);
+				Assert.AreEqual(0, NewIssueA.Fingerprints.Count);
+				List<IIssueSpan> NewSpansA = await IssueService.GetIssueSpansAsync(NewIssueA);
+				Assert.AreEqual(0, NewSpansA.Count);
+
+				IIssue NewIssueB = (await IssueService.GetIssueAsync(IssueB.Id))!;
+				Assert.IsNull(NewIssueB.VerifiedAt);
+				Assert.IsNull(NewIssueB.ResolvedAt);
+				Assert.AreEqual(2, NewIssueB.Fingerprints.Count);
+				List<IIssueSpan> NewSpansB = await IssueService.GetIssueSpansAsync(NewIssueB!);
+				Assert.AreEqual(2, NewSpansB.Count);
+				Assert.AreEqual(NewIssueB.Id, NewSpansB[0].IssueId);
+				Assert.AreEqual(NewIssueB.Id, NewSpansB[1].IssueId);
+			}
+		}
+
 		[TestMethod]
 		public async Task GauntletIssueTest()
 		{
diff --git a/Engine/Source/Programs/Horde/HordeServer/Api/Issues.cs b/Engine/Source/Programs/Horde/HordeServer/Api/Issues.cs
--- a/Engine/Source/Programs/Horde/HordeServer/Api/Issues.cs
+++ b/Engine/Source/Programs/Horde/HordeServer/Api/Issues.cs
@@ -81,8 +81,13 @@ namespace HordeServer.Api
 	/// <summary>
 	/// Trace of a set of node failures across multiple steps
 	/// </summary>
-	public class GetIssueNodeResponse
+	public class GetIssueSpanResponse
 	{
+		/// <summary>
+		/// Unique id of this span
+		/// </summary>
+		public string Id { get; set; }
+
 		/// <summary>
 		/// The template containing this step
 		/// </summary>
@@ -113,8 +118,9 @@ namespace HordeServer.Api
 		/// </summary>
 		/// <param name="Span">The node to construct from</param>
 		/// <param name="Steps">Failing steps for this span</param>
-		public GetIssueNodeResponse(IIssueSpan Span, List<IIssueStep> Steps)
+		public GetIssueSpanResponse(IIssueSpan Span, List<IIssueStep> Steps)
 		{
+			this.Id = Span.Id.ToString();
 			this.Name = Span.NodeName;
 			this.TemplateId = Span.TemplateRefId.ToString();
 			this.LastSuccess = (Span.LastSuccess != null) ? new GetIssueStepResponse(Span.LastSuccess) : null;
@@ -146,7 +152,7 @@ namespace HordeServer.Api
 		/// <summary>
 		/// Map of steps to (event signature id -> trace id)
 		/// </summary>
-		public List<GetIssueNodeResponse> Nodes { get; set; }
+		public List<GetIssueSpanResponse> Nodes { get; set; }
 
 		/// <summary>
 		/// Constructor
@@ -170,7 +176,7 @@ namespace HordeServer.Api
 				}
 			}
 
-			this.Nodes = Spans.ConvertAll(x => new GetIssueNodeResponse(x, Steps.Where(y => y.SpanId == x.Id).ToList()));
+			this.Nodes = Spans.ConvertAll(x => new GetIssueSpanResponse(x, Steps.Where(y => y.SpanId == x.Id).ToList()));
 		}
 	}
 
@@ -599,5 +605,15 @@ namespace HordeServer.Api
 		/// Whether the issue should be marked as resolved
 		/// </summary>
 		public bool? Resolved { get; set; }
+
+		/// <summary>
+		/// List of spans to add to this issue
+		/// </summary>
+		public List<string>? AddSpans { get; set; }
+
+		/// <summary>
+		/// List of spans to remove from this issue
+		/// </summary>
+		public List<string>? RemoveSpans { get; set; }
 	}
 }
diff --git a/Engine/Source/Programs/Horde/HordeServer/Controllers/IssuesController.cs b/Engine/Source/Programs/Horde/HordeServer/Controllers/IssuesController.cs
--- a/Engine/Source/Programs/Horde/HordeServer/Controllers/IssuesController.cs
+++ b/Engine/Source/Programs/Horde/HordeServer/Controllers/IssuesController.cs
@@ -414,7 +414,19 @@ namespace HordeServer.Controllers
 				NewResolvedById = Request.Resolved.Value ? User.GetUserId() : UserId.Empty;
 			}
 
-			if (!await IssueService.UpdateIssueAsync(IssueId, Request.Summary, Request.Description, Request.Promoted, NewOwnerId, NewNominatedById, Request.Acknowledged, NewDeclinedById, Request.FixChange, NewResolvedById))
+			List<ObjectId>? AddSpans = null;
+			if (Request.AddSpans != null && Request.AddSpans.Count > 0)
+			{
+				AddSpans = Request.AddSpans.ConvertAll(x => ObjectId.Parse(x));
+			}
+
+			List<ObjectId>? RemoveSpans = null;
+			if (Request.RemoveSpans != null && Request.RemoveSpans.Count > 0)
+			{
+				RemoveSpans = Request.RemoveSpans.ConvertAll(x => ObjectId.Parse(x));
+			}
+
+			if (!await IssueService.UpdateIssueAsync(IssueId, Request.Summary, Request.Description, Request.Promoted, NewOwnerId, NewNominatedById, Request.Acknowledged, NewDeclinedById, Request.FixChange, NewResolvedById, AddSpans, RemoveSpans))
 			{
 				return NotFound();
 			}
diff --git a/Engine/Source/Programs/Horde/HordeServer/Issues/IIssue.cs b/Engine/Source/Programs/Horde/HordeServer/Issues/IIssue.cs
--- a/Engine/Source/Programs/Horde/HordeServer/Issues/IIssue.cs
+++ b/Engine/Source/Programs/Horde/HordeServer/Issues/IIssue.cs
@@ -147,6 +147,11 @@ namespace HordeServer.Models
 		/// </summary>
 		public IReadOnlyList<IIssueStream> Streams { get; }
 
+		/// <summary>
+		/// Spans which should be excluded from this issue
+		/// </summary>
+		public List<ObjectId>? ExcludeSpans { get; }
+
 		/// <summary>
 		/// Update index for this instance
 		/// </summary>
diff --git a/Engine/Source/Programs/Horde/HordeServer/Issues/IIssueCollection.cs b/Engine/Source/Programs/Horde/HordeServer/Issues/IIssueCollection.cs
--- a/Engine/Source/Programs/Horde/HordeServer/Issues/IIssueCollection.cs
+++ b/Engine/Source/Programs/Horde/HordeServer/Issues/IIssueCollection.cs
@@ -437,9 +437,10 @@ namespace HordeServer.Collections
 		/// <param name="NewDeclinedById">Name of a user that has declined the issue</param>
 		/// <param name="NewFixChange">Fix changelist for the issue. Pass 0 to clear the fix changelist, -1 for systemic issue.</param>
 		/// <param name="NewResolvedById">User that resolved the issue (may be ObjectId.Empty to clear)</param>
+		/// <param name="NewExcludeSpanIds">List of span ids to exclude from this issue</param>
 		/// <param name="NewLastSeenAt"></param>
 		/// <returns>True if the issue was updated</returns>
-		Task<IIssue?> TryUpdateIssueAsync(IIssue Issue, IssueSeverity? NewSeverity = null, string? NewSummary = null, string? NewUserSummary = null, string? NewDescription = null, bool? NewPromoted = null, UserId? NewOwnerId = null, UserId? NewNominatedById = null, bool? NewAcknowledged = null, UserId? NewDeclinedById = null, int? NewFixChange = null, UserId? NewResolvedById = null, DateTime? NewLastSeenAt = null);
+		Task<IIssue?> TryUpdateIssueAsync(IIssue Issue, IssueSeverity? NewSeverity = null, string? NewSummary = null, string? NewUserSummary = null, string? NewDescription = null, bool? NewPromoted = null, UserId? NewOwnerId = null, UserId? NewNominatedById = null, bool? NewAcknowledged = null, UserId? NewDeclinedById = null, int? NewFixChange = null, UserId? NewResolvedById = null, List<ObjectId>? NewExcludeSpanIds = null, DateTime? NewLastSeenAt = null);
 
 		/// <summary>
 		/// Updates derived data for an issue (ie. data computed from the spans attached to it). Also clears the issue's 'modified' state.
diff --git a/Engine/Source/Programs/Horde/HordeServer/Issues/Impl/IssueCollection.cs b/Engine/Source/Programs/Horde/HordeServer/Issues/Impl/IssueCollection.cs
--- a/Engine/Source/Programs/Horde/HordeServer/Issues/Impl/IssueCollection.cs
+++ b/Engine/Source/Programs/Horde/HordeServer/Issues/Impl/IssueCollection.cs
@@ -99,6 +99,9 @@ namespace HordeServer.Collections.Impl
 			[BsonElement("Suspects"), BsonIgnoreIfNull]
 			public List<IssueSuspect>? Suspects_DEPRECATED { get; set; }
 
+			[BsonIgnoreIfNull]
+			public List<ObjectId>? ExcludeSpans { get; set; }
+
 			public int UpdateIndex { get; set; }
 
 			bool IIssue.Promoted => Promoted ?? NotifySuspects_DEPRECATED;
@@ -775,7 +778,7 @@ namespace HordeServer.Collections.Impl
 		}
 
 		/// <inheritdoc/>
-		public async Task<IIssue?> TryUpdateIssueAsync(IIssue Issue, IssueSeverity? NewSeverity = null, string? NewSummary = null, string? NewUserSummary = null, string? NewDescription = null, bool? NewPromoted = null, UserId? NewOwnerId = null, UserId? NewNominatedById = null, bool? NewAcknowledged = null, UserId? NewDeclinedById = null, int? NewFixChange = null, UserId? NewResolvedById = null, DateTime? NewLastSeenAt = null)
+		public async Task<IIssue?> TryUpdateIssueAsync(IIssue Issue, IssueSeverity? NewSeverity = null, string? NewSummary = null, string? NewUserSummary = null, string? NewDescription = null, bool? NewPromoted = null, UserId? NewOwnerId = null, UserId? NewNominatedById = null, bool? NewAcknowledged = null, UserId? NewDeclinedById = null, int? NewFixChange = null, UserId? NewResolvedById = null, List<ObjectId>? NewExcludeSpanIds = null, DateTime? NewLastSeenAt = null)
 		{
 			Issue IssueDocument = (Issue)Issue;
 
@@ -890,6 +893,15 @@ namespace HordeServer.Collections.Impl
 					}
 				}
 			}
+			if (NewExcludeSpanIds != null)
+			{
+				List<ObjectId> NewCombinedExcludeSpanIds = NewExcludeSpanIds;
+				if (Issue.ExcludeSpans != null)
+				{
+					NewCombinedExcludeSpanIds = NewCombinedExcludeSpanIds.Union(Issue.ExcludeSpans).ToList();
+				}
+				Updates.Add(Builders<Issue>.Update.Set(x => x.ExcludeSpans, NewCombinedExcludeSpanIds));
+			}
 			if (NewLastSeenAt != null)
 			{
 				Updates.Add(Builders<Issue>.Update.Set(x => x.LastSeenAt, NewLastSeenAt.Value));
