Horde: Add log id back into task properties.

diff --git a/Engine/Source/Programs/Horde/HordeServer/Agents/Impl/Tasks/ConformTaskSource.cs b/Engine/Source/Programs/Horde/HordeServer/Agents/Impl/Tasks/ConformTaskSource.cs
--- a/Engine/Source/Programs/Horde/HordeServer/Agents/Impl/Tasks/ConformTaskSource.cs
+++ b/Engine/Source/Programs/Horde/HordeServer/Agents/Impl/Tasks/ConformTaskSource.cs
@@ -23,6 +23,7 @@ namespace HordeServer.Tasks.Impl
 {
 	using JobId = ObjectId<IJob>;
 	using LeaseId = ObjectId<ILease>;
+	using LogId = ObjectId<ILogFile>;
 
 	/// <summary>
 	/// Generates tasks telling agents to sync their workspaces
@@ -60,6 +61,8 @@ namespace HordeServer.Tasks.Impl
 			this.LogService = LogService;
 			this.Logger = Logger;
 			this.TickConformList = new ElectedTick(DatabaseService, new ObjectId("60afc5cf555a9a76aff0a50c"), CleanConformListAsync, TimeSpan.FromMinutes(1.0), Logger);
+
+			this.OnLeaseStartedProperties.Add(nameof(ConformTask.LogId), x => new LogId(x.LogId));
 		}
 
 		/// <inheritdoc/>
diff --git a/Engine/Source/Programs/Horde/HordeServer/Agents/Impl/Tasks/RestartTaskSource.cs b/Engine/Source/Programs/Horde/HordeServer/Agents/Impl/Tasks/RestartTaskSource.cs
--- a/Engine/Source/Programs/Horde/HordeServer/Agents/Impl/Tasks/RestartTaskSource.cs
+++ b/Engine/Source/Programs/Horde/HordeServer/Agents/Impl/Tasks/RestartTaskSource.cs
@@ -20,6 +20,7 @@ namespace HordeServer.Tasks.Impl
 {
 	using JobId = ObjectId<IJob>;
 	using LeaseId = ObjectId<ILease>;
+	using LogId = ObjectId<ILogFile>;
 
 	class RestartTaskSource : TaskSourceBase<RestartTask>
 	{
@@ -30,6 +31,8 @@ namespace HordeServer.Tasks.Impl
 		public RestartTaskSource(ILogFileService LogService)
 		{
 			this.LogService = LogService;
+
+			this.OnLeaseStartedProperties.Add(nameof(RestartTask.LogId), x => new LogId(x.LogId));
 		}
 
 		public override async Task<AgentLease?> AssignLeaseAsync(IAgent Agent, CancellationToken CancellationToken)
diff --git a/Engine/Source/Programs/Horde/HordeServer/Agents/Impl/Tasks/ShutdownTaskSource.cs b/Engine/Source/Programs/Horde/HordeServer/Agents/Impl/Tasks/ShutdownTaskSource.cs
--- a/Engine/Source/Programs/Horde/HordeServer/Agents/Impl/Tasks/ShutdownTaskSource.cs
+++ b/Engine/Source/Programs/Horde/HordeServer/Agents/Impl/Tasks/ShutdownTaskSource.cs
@@ -21,6 +21,7 @@ namespace HordeServer.Tasks.Impl
 {
 	using JobId = ObjectId<IJob>;
 	using LeaseId = ObjectId<ILease>;
+	using LogId = ObjectId<ILogFile>;
 
 	class ShutdownTaskSource : TaskSourceBase<ShutdownTask>
 	{
@@ -31,6 +32,7 @@ namespace HordeServer.Tasks.Impl
 		public ShutdownTaskSource(ILogFileService LogService)
 		{
 			this.LogService = LogService;
+			this.OnLeaseStartedProperties.Add(nameof(ShutdownTask.LogId), x => new LogId(x.LogId));
 		}
 
 		public override async Task<AgentLease?> AssignLeaseAsync(IAgent Agent, CancellationToken CancellationToken)
diff --git a/Engine/Source/Programs/Horde/HordeServer/Agents/Impl/Tasks/UpgradeTaskSource.cs b/Engine/Source/Programs/Horde/HordeServer/Agents/Impl/Tasks/UpgradeTaskSource.cs
--- a/Engine/Source/Programs/Horde/HordeServer/Agents/Impl/Tasks/UpgradeTaskSource.cs
+++ b/Engine/Source/Programs/Horde/HordeServer/Agents/Impl/Tasks/UpgradeTaskSource.cs
@@ -21,6 +21,7 @@ namespace HordeServer.Tasks.Impl
 	using AgentSoftwareChannelName = StringId<AgentSoftwareChannels>;
 	using JobId = ObjectId<IJob>;
 	using LeaseId = ObjectId<ILease>;
+	using LogId = ObjectId<ILogFile>;
 
 	class UpgradeTaskSource : TaskSourceBase<UpgradeTask>
 	{
@@ -35,6 +36,8 @@ namespace HordeServer.Tasks.Impl
 			this.AgentSoftwareService = AgentSoftwareService;
 			this.LogService = LogService;
 			this.Clock = Clock;
+
+			this.OnLeaseStartedProperties.Add(nameof(UpgradeTask.LogId), x => new LogId(x.LogId));
 		}
 
 		public async override Task<AgentLease?> AssignLeaseAsync(IAgent Agent, CancellationToken CancellationToken)
