Horde: Missing files.

diff --git a/Engine/Source/Programs/Horde/HordeServer/Issues/IIssueService.cs b/Engine/Source/Programs/Horde/HordeServer/Issues/IIssueService.cs
--- a/Engine/Source/Programs/Horde/HordeServer/Issues/IIssueService.cs
+++ b/Engine/Source/Programs/Horde/HordeServer/Issues/IIssueService.cs
@@ -141,8 +141,10 @@ namespace HordeServer.Services
 		/// <param name="DeclinedById">Name of the user that has declined this issue</param>
 		/// <param name="FixChange">Fix changelist for the issue</param>
 		/// <param name="ResolvedById">Whether the issue has been resolved</param>
+		/// <param name="AddSpanIds">Add spans to this issue</param>
+		/// <param name="RemoveSpanIds">Remove spans from this issue</param>
 		/// <returns>True if the issue was updated</returns>
-		Task<bool> UpdateIssueAsync(int Id, string? Summary = null, string? Description = null, bool? Promoted = null, UserId? OwnerId = null, UserId? NominatedById = null, bool? Acknowledged = null, UserId? DeclinedById = null, int? FixChange = null, UserId? ResolvedById = null);
+		Task<bool> UpdateIssueAsync(int Id, string? Summary = null, string? Description = null, bool? Promoted = null, UserId? OwnerId = null, UserId? NominatedById = null, bool? Acknowledged = null, UserId? DeclinedById = null, int? FixChange = null, UserId? ResolvedById = null, List<ObjectId>? AddSpanIds = null, List<ObjectId>? RemoveSpanIds = null);
 
 		/// <summary>
 		/// Marks a step as complete
diff --git a/Engine/Source/Programs/Horde/HordeServer/Issues/Impl/IssueService.cs b/Engine/Source/Programs/Horde/HordeServer/Issues/Impl/IssueService.cs
--- a/Engine/Source/Programs/Horde/HordeServer/Issues/Impl/IssueService.cs
+++ b/Engine/Source/Programs/Horde/HordeServer/Issues/Impl/IssueService.cs
@@ -384,23 +384,63 @@ namespace HordeServer.Services.Impl
 		}
 
 		/// <inheritdoc/>
-		public async Task<bool> UpdateIssueAsync(int Id, string? UserSummary = null, string? Description = null, bool? Promoted = null, UserId? OwnerId = null, UserId? NominatedById = null, bool? Acknowledged = null, UserId? DeclinedById = null, int? FixChange = null, UserId? ResolvedById = null)
+		public async Task<bool> UpdateIssueAsync(int Id, string? UserSummary = null, string? Description = null, bool? Promoted = null, UserId? OwnerId = null, UserId? NominatedById = null, bool? Acknowledged = null, UserId? DeclinedById = null, int? FixChange = null, UserId? ResolvedById = null, List<ObjectId>? AddSpanIds = null, List<ObjectId>? RemoveSpanIds = null)
 		{
+			IIssue? Issue;
 			for (; ; )
 			{
-				IIssue? Issue = await IssueCollection.GetIssueAsync(Id);
-				if(Issue == null)
+				Issue = await IssueCollection.GetIssueAsync(Id);
+				if (Issue == null)
 				{
 					return false;
 				}
 
-				Issue = await IssueCollection.TryUpdateIssueAsync(Issue, NewUserSummary: UserSummary, NewDescription: Description, NewPromoted: Promoted, NewOwnerId: OwnerId ?? ResolvedById, NewNominatedById: NominatedById, NewDeclinedById: DeclinedById, NewAcknowledged: Acknowledged, NewFixChange: FixChange, NewResolvedById: ResolvedById);
-				if(Issue != null)
+				Issue = await IssueCollection.TryUpdateIssueAsync(Issue, NewUserSummary: UserSummary, NewDescription: Description, NewPromoted: Promoted, NewOwnerId: OwnerId ?? ResolvedById, NewNominatedById: NominatedById, NewDeclinedById: DeclinedById, NewAcknowledged: Acknowledged, NewFixChange: FixChange, NewResolvedById: ResolvedById, NewExcludeSpanIds: RemoveSpanIds);
+				if (Issue != null)
 				{
-					await UpdateIssueDerivedDataAsync(Issue);
-					return true;
+					break;
 				}
 			}
+
+			await using IAsyncDisposable Lock = await IssueCollection.EnterCriticalSectionAsync();
+
+			List<IIssue> UpdateIssues = new List<IIssue>();
+			UpdateIssues.Add(Issue);
+
+			if (AddSpanIds != null)
+			{
+				foreach (ObjectId AddSpanId in AddSpanIds)
+				{
+					IIssueSpan? Span = await IssueCollection.GetSpanAsync(AddSpanId);
+					if (Span != null)
+					{
+						IIssue? OldIssue = await IssueCollection.GetIssueAsync(Span.IssueId);
+						UpdateIssues.Add(OldIssue ?? Issue);
+						await IssueCollection.TryUpdateSpanAsync(Span, NewIssueId: Issue.Id);
+					}
+				}
+			}
+
+			if (RemoveSpanIds != null)
+			{
+				foreach (ObjectId RemoveSpanId in RemoveSpanIds)
+				{
+					IIssueSpan? Span = await IssueCollection.GetSpanAsync(RemoveSpanId);
+					if (Span != null)
+					{
+						IIssue NewIssue = await FindOrAddIssueForSpanAsync(Span);
+						UpdateIssues.Add(NewIssue);
+						await IssueCollection.TryUpdateSpanAsync(Span, NewIssueId: NewIssue.Id);
+					}
+				}
+			}
+
+			foreach (IIssue UpdateIssue in UpdateIssues.GroupBy(x => x.Id).Select(x => x.First()))
+			{
+				await UpdateIssueDerivedDataAsync(UpdateIssue);
+			}
+
+			return true;
 		}
 
 		/// <summary>
@@ -906,7 +946,11 @@ namespace HordeServer.Services.Impl
 				}
 
 				// Get the new resolved timestamp
-				DateTime? NewResolvedAt = Issue.ResolvedAt ?? NewVerifiedAt;
+				DateTime? NewResolvedAt = NewVerifiedAt;
+				if (Issue.ResolvedById != null && Issue.ResolvedAt != null)
+				{
+					NewResolvedAt = Issue.ResolvedAt;
+				}
 
 				// Get the new list of streams
 				List<NewIssueStream> NewStreams = Spans.Select(x => x.StreamId).OrderBy(x => x.ToString(), StringComparer.Ordinal).Distinct().Select(x => new NewIssueStream(x)).ToList();
@@ -1051,6 +1095,32 @@ namespace HordeServer.Services.Impl
 			return Suspects;
 		}
 
+		/// <summary>
+		/// Find an existing issue that may match one of the given suspect changes
+		/// </summary>
+		/// <param name="Span">The span to find an issue for</param>
+		/// <returns>The matching issue</returns>
+		async Task<IIssue> FindOrAddIssueForSpanAsync(IIssueSpan Span)
+		{
+			List<IIssue> ExistingIssues;
+			if (Span.LastSuccess == null)
+			{
+				ExistingIssues = await IssueCollection.FindIssuesAsync(StreamId: Span.StreamId, MinChange: Span.FirstFailure.Change, MaxChange: Span.FirstFailure.Change, Resolved: false);
+			}
+			else
+			{
+				ExistingIssues = await IssueCollection.FindIssuesForChangesAsync(Span.Suspects.ConvertAll(x => x.OriginatingChange ?? x.Change));
+			}
+
+			IIssue? Issue = ExistingIssues.FirstOrDefault(x => x.Fingerprints.Any(y => y.IsMatch(Span.Fingerprint)));
+			if (Issue == null)
+			{
+				string Summary = GetSummary(Span.Fingerprint, Span.FirstFailure.Severity);
+				Issue = await IssueCollection.AddIssueAsync(Summary);
+			}
+			return Issue;
+		}
+
 		/// <summary>
 		/// Find an existing issue that may match one of the given suspect changes
 		/// </summary>
diff --git a/Engine/Source/Programs/Horde/HordeServer/Jobs/Impl/JobTaskSource.cs b/Engine/Source/Programs/Horde/HordeServer/Jobs/Impl/JobTaskSource.cs
--- a/Engine/Source/Programs/Horde/HordeServer/Jobs/Impl/JobTaskSource.cs
+++ b/Engine/Source/Programs/Horde/HordeServer/Jobs/Impl/JobTaskSource.cs
@@ -323,7 +323,7 @@ namespace HordeServer.Tasks.Impl
 			this.Settings = Settings;
 			this.Logger = Logger;
 
-			OnLeaseStartedProperties.Add(nameof(ExecuteJobTask.JobId), x => new JobId(x.JobId)).Add(nameof(ExecuteJobTask.BatchId), x => SubResourceId.Parse(x.BatchId));
+			OnLeaseStartedProperties.Add(nameof(ExecuteJobTask.JobId), x => new JobId(x.JobId)).Add(nameof(ExecuteJobTask.BatchId), x => SubResourceId.Parse(x.BatchId)).Add(nameof(ExecuteJobTask.LogId), x => new LogId(x.LogId));
 		}
 
 		/// <inheritdoc/>
