Horde: Rename the HordeServer project to Horde.Build.

diff --git a/Engine/Source/Programs/Horde/Horde.Build.Tests/Horde.Build.Tests.csproj b/Engine/Source/Programs/Horde/Horde.Build.Tests/Horde.Build.Tests.csproj
new file mode 100644
--- /dev/null
+++ b/Engine/Source/Programs/Horde/Horde.Build.Tests/Horde.Build.Tests.csproj
@@ -0,0 +1,55 @@
+<Project Sdk="Microsoft.NET.Sdk">
+
+  <PropertyGroup>
+    <TargetFramework>net5.0</TargetFramework>
+    <Nullable>enable</Nullable>
+    <IsPackable>false</IsPackable>
+  </PropertyGroup>
+
+  <ItemGroup>
+    <Compile Remove="Stubs\Collections\JobCollectionStub.cs" />
+    <Compile Remove="Stubs\Collections\StreamCollectionStub.cs" />
+  </ItemGroup>
+
+  <ItemGroup>
+    <None Include="Stubs\Collections\JobCollectionStub.cs" />
+    <None Include="Stubs\Collections\StreamCollectionStub.cs" />
+  </ItemGroup>
+
+  <ItemGroup>
+    <PackageReference Include="Microsoft.AspNetCore.Mvc.Testing" Version="3.1.6" />
+    <PackageReference Include="Microsoft.NET.Test.Sdk" Version="16.5.0" />
+    <PackageReference Include="Moq" Version="4.14.5" />
+    <PackageReference Include="MSTest.TestAdapter" Version="2.1.0" />
+    <PackageReference Include="MSTest.TestFramework" Version="2.1.0" />
+    <PackageReference Include="coverlet.collector" Version="1.2.0" />
+  </ItemGroup>
+
+  <ItemGroup>
+    <ProjectReference Include="..\..\Shared\EpicGames.Serialization\EpicGames.Serialization.csproj" />
+    <ProjectReference Include="..\HordeAgent\HordeAgent.csproj" Label="HordeAgent">
+      <Aliases>HordeAgent</Aliases>
+    </ProjectReference>
+    <ProjectReference Include="..\HordeServer\HordeServer.csproj" />
+  </ItemGroup>
+
+  <ItemGroup>
+    <Compile Update="Properties\Resources.Designer.cs">
+      <DesignTime>True</DesignTime>
+      <AutoGen>True</AutoGen>
+      <DependentUpon>Resources.resx</DependentUpon>
+    </Compile>
+  </ItemGroup>
+
+  <ItemGroup>
+    <EmbeddedResource Update="Properties\Resources.resx">
+      <Generator>ResXFileCodeGenerator</Generator>
+      <LastGenOutput>Resources.Designer.cs</LastGenOutput>
+    </EmbeddedResource>
+  </ItemGroup>
+
+  <ItemGroup>
+    <Folder Include="Resources\" />
+  </ItemGroup>
+
+</Project>
diff --git a/Engine/Source/Programs/Horde/Horde.Build/Agents/IAgentSoftwareCollection.cs b/Engine/Source/Programs/Horde/Horde.Build/Agents/IAgentSoftwareCollection.cs
new file mode 100644
--- /dev/null
+++ b/Engine/Source/Programs/Horde/Horde.Build/Agents/IAgentSoftwareCollection.cs
@@ -0,0 +1,54 @@
+// Copyright Epic Games, Inc. All Rights Reserved.
+
+using EpicGames.Core;
+using HordeCommon;
+using HordeServer.Models;
+using HordeServer.Services;
+using MongoDB.Bson;
+using MongoDB.Bson.Serialization.Attributes;
+using MongoDB.Driver;
+using System;
+using System.Collections.Generic;
+using System.Globalization;
+using System.IO;
+using System.IO.Compression;
+using System.Linq;
+using System.Text.Json;
+using System.Threading.Tasks;
+
+namespace HordeServer.Collections
+{
+	/// <summary>
+	/// Interface for a collection of template documents
+	/// </summary>
+	public interface IAgentSoftwareCollection
+	{
+		/// <summary>
+		/// Adds a new software revision
+		/// </summary>
+		/// <param name="Version">The version number</param>
+		/// <param name="Data">Zip file containing the new software</param>
+		/// <returns>New software instance</returns>
+		Task<bool> AddAsync(string Version, byte[] Data);
+
+		/// <summary>
+		/// Tests whether a given version exists
+		/// </summary>
+		/// <param name="Version">Version of the software</param>
+		/// <returns>True if it exists, false otherwise</returns>
+		Task<bool> ExistsAsync(string Version);
+
+		/// <summary>
+		/// Removes a software archive 
+		/// </summary>
+		/// <param name="Version">Version of the software to delete</param>
+		Task<bool> RemoveAsync(string Version);
+
+		/// <summary>
+		/// Downloads software of a given revision
+		/// </summary>
+		/// <param name="Version">Version of the software</param>
+		/// <returns>Data for the given software</returns>
+		Task<byte[]?> GetAsync(string Version);
+	}
+}
diff --git a/Engine/Source/Programs/Horde/Horde.Build/Agents/ILeaseCollection.cs b/Engine/Source/Programs/Horde/Horde.Build/Agents/ILeaseCollection.cs
new file mode 100644
--- /dev/null
+++ b/Engine/Source/Programs/Horde/Horde.Build/Agents/ILeaseCollection.cs
@@ -0,0 +1,86 @@
+// Copyright Epic Games, Inc. All Rights Reserved.
+
+using HordeCommon;
+using HordeServer.Models;
+using HordeServer.Services;
+using HordeServer.Utilities;
+using MongoDB.Bson;
+using MongoDB.Bson.Serialization.Attributes;
+using MongoDB.Driver;
+using System;
+using System.Collections.Generic;
+using System.Linq;
+using System.Threading.Tasks;
+
+namespace HordeServer.Collections
+{
+	using LeaseId = ObjectId<ILease>;
+	using LogId = ObjectId<ILogFile>;
+	using PoolId = StringId<IPool>;
+	using StreamId = StringId<IStream>;
+
+	/// <summary>
+	/// Interface for a collection of lease documents
+	/// </summary>
+	public interface ILeaseCollection
+	{
+		/// <summary>
+		/// Adds a lease to the collection
+		/// </summary>
+		/// <param name="Id">The lease id</param>
+		/// <param name="Name">Name of the lease</param>
+		/// <param name="AgentId">The agent id</param>
+		/// <param name="SessionId">The agent session handling the lease</param>
+		/// <param name="StreamId">Stream for the payload</param>
+		/// <param name="PoolId">The pool for the lease</param>
+		/// <param name="LogId">Log id for the lease</param>
+		/// <param name="StartTime">Start time of the lease</param>
+		/// <param name="Payload">Payload for the lease</param>
+		/// <returns>Async task</returns>
+		Task<ILease> AddAsync(LeaseId Id, string Name, AgentId AgentId, ObjectId SessionId, StreamId? StreamId, PoolId? PoolId, LogId? LogId, DateTime StartTime, byte[] Payload);
+
+		/// <summary>
+		/// Deletes a lease from the collection
+		/// </summary>
+		/// <param name="LeaseId">Unique id of the lease</param>
+		/// <returns>Async task</returns>
+		Task DeleteAsync(LeaseId LeaseId);
+
+		/// <summary>
+		/// Gets a specific lease
+		/// </summary>
+		/// <param name="LeaseId">Unique id of the lease</param>
+		/// <returns>The lease that was found, or null if it does not exist</returns>
+		Task<ILease?> GetAsync(LeaseId LeaseId);
+
+		/// <summary>
+		/// Finds all leases matching a set of criteria
+		/// </summary>
+		/// <param name="AgentId">Unqiue id of the agent executing this lease</param>
+		/// <param name="SessionId">Unique id of the agent session</param>
+		/// <param name="MinTime">Start of the window to include leases</param>
+		/// <param name="MaxTime">End of the window to include leases</param>
+		/// <param name="Index">Index of the first result to return</param>
+		/// <param name="Count">Number of results to return</param>
+		/// <returns>List of leases matching the given criteria</returns>
+		Task<List<ILease>> FindLeasesAsync(AgentId? AgentId = null, ObjectId? SessionId = null, DateTime? MinTime = null, DateTime? MaxTime = null, int? Index = null, int? Count = null);
+
+		/// <summary>
+		/// Finds all active leases
+		/// </summary>
+		/// <param name="Index">Index of the first result to return</param>
+		/// <param name="Count">Number of results to return</param>
+		/// <returns>List of leases</returns>
+		Task<List<ILease>> FindActiveLeasesAsync(int? Index = null, int? Count = null);
+
+		/// <summary>
+		/// Sets the outcome of a lease
+		/// </summary>
+		/// <param name="LeaseId">The lease to update</param>
+		/// <param name="FinishTime">Time at which the lease finished</param>
+		/// <param name="Outcome">Outcome of the lease</param>
+		/// <param name="Output">Output data from the task</param>
+		/// <returns>True if the lease was updated, false otherwise</returns>
+		Task<bool> TrySetOutcomeAsync(LeaseId LeaseId, DateTime FinishTime, LeaseOutcome Outcome, byte[]? Output);
+	}
+}
diff --git a/Engine/Source/Programs/Horde/Horde.Build/Agents/ISessionCollection.cs b/Engine/Source/Programs/Horde/Horde.Build/Agents/ISessionCollection.cs
new file mode 100644
--- /dev/null
+++ b/Engine/Source/Programs/Horde/Horde.Build/Agents/ISessionCollection.cs
@@ -0,0 +1,76 @@
+// Copyright Epic Games, Inc. All Rights Reserved.
+
+using HordeServer.Models;
+using HordeServer.Services;
+using HordeServer.Utilities;
+using Microsoft.Extensions.DependencyInjection;
+using MongoDB.Bson;
+using MongoDB.Bson.Serialization.Attributes;
+using MongoDB.Driver;
+using System;
+using System.Collections.Generic;
+using System.Linq;
+using System.Threading.Tasks;
+
+namespace HordeServer.Collections
+{
+	/// <summary>
+	/// Interface for a collection of session documents
+	/// </summary>
+	public interface ISessionCollection
+	{
+		/// <summary>
+		/// Adds a new session to the collection
+		/// </summary>
+		/// <param name="Id">The session id</param>
+		/// <param name="AgentId">The agent this session is for</param>
+		/// <param name="StartTime">Start time of this session</param>
+		/// <param name="Properties">Properties of this agent at the time the session started</param>
+		/// <param name="Resources">Resources which the agent has</param>
+		/// <param name="Version">Version of the agent software</param>
+		Task<ISession> AddAsync(ObjectId Id, AgentId AgentId, DateTime StartTime, IReadOnlyList<string>? Properties, IReadOnlyDictionary<string, int>? Resources, string? Version);
+
+		/// <summary>
+		/// Gets information about a particular session
+		/// </summary>
+		/// <param name="SessionId">The unique session id</param>
+		/// <returns>The session information</returns>
+		Task<ISession?> GetAsync(ObjectId SessionId);
+
+		/// <summary>
+		/// Find sessions for the given agent
+		/// </summary>
+		/// <param name="AgentId">The unique agent id</param>
+		/// <param name="StartTime">Start time to include in the search</param>
+		/// <param name="FinishTime">Finish time to include in the search</param>
+		/// <param name="Index">Index of the first result to return</param>
+		/// <param name="Count">Number of results to return</param>
+		/// <returns>List of sessions matching the given criteria</returns>
+		Task<List<ISession>> FindAsync(AgentId AgentId, DateTime? StartTime, DateTime? FinishTime, int Index, int Count);
+
+		/// <summary>
+		/// Finds any active sessions
+		/// </summary>
+		/// <param name="Index">Index of the first result to return</param>
+		/// <param name="Count">Number of results to return</param>
+		/// <returns>List of sessions</returns>
+		Task<List<ISession>> FindActiveSessionsAsync(int? Index = null, int? Count = null);
+
+		/// <summary>
+		/// Update a session from the collection
+		/// </summary>
+		/// <param name="SessionId">The session to update</param>
+		/// <param name="FinishTime">Time at which the session finished</param>
+		/// <param name="Properties">The agent properties</param>
+		/// <param name="Resources">Resources which the agent has</param>
+		/// <returns>Async task</returns>
+		Task UpdateAsync(ObjectId SessionId, DateTime FinishTime, IReadOnlyList<string> Properties, IReadOnlyDictionary<string, int> Resources);
+
+		/// <summary>
+		/// Delete a session from the collection
+		/// </summary>
+		/// <param name="SessionId">The session id</param>
+		/// <returns>Async task</returns>
+		Task DeleteAsync(ObjectId SessionId);
+	}
+}
diff --git a/Engine/Source/Programs/Horde/Horde.Build/Agents/Impl/AgentSoftwareCollection.cs b/Engine/Source/Programs/Horde/Horde.Build/Agents/Impl/AgentSoftwareCollection.cs
new file mode 100644
--- /dev/null
+++ b/Engine/Source/Programs/Horde/Horde.Build/Agents/Impl/AgentSoftwareCollection.cs
@@ -0,0 +1,210 @@
+// Copyright Epic Games, Inc. All Rights Reserved.
+
+using EpicGames.Core;
+using HordeCommon;
+using HordeServer.Models;
+using HordeServer.Services;
+using HordeServer.Utilities;
+using MongoDB.Bson;
+using MongoDB.Bson.Serialization.Attributes;
+using MongoDB.Driver;
+using System;
+using System.Collections.Generic;
+using System.Globalization;
+using System.IO;
+using System.IO.Compression;
+using System.Linq;
+using System.Text.Json;
+using System.Threading.Tasks;
+
+namespace HordeServer.Collections.Impl
+{
+	using AgentSoftwareVersion = StringId<IAgentSoftwareCollection>;
+
+	/// <summary>
+	/// Collection of agent software documents
+	/// </summary>
+	public class AgentSoftwareCollection : IAgentSoftwareCollection
+	{
+		/// <summary>
+		/// Information about agent software
+		/// </summary>
+		class AgentSoftwareDocument
+		{
+			[BsonId]
+			public string Id { get; set; }
+
+			[BsonIgnoreIfNull]
+			public byte[]? Data { get; set; }
+
+			[BsonIgnoreIfNull]
+            public List<string>? ChunkIds { get; set; }
+
+            [BsonConstructor]
+			private AgentSoftwareDocument()
+			{
+				this.Id = null!;
+				this.Data = null;
+                this.ChunkIds = null;			
+            }
+
+			public AgentSoftwareDocument(string Id, byte[]? Data, List<string>? ChunkIds = null)
+			{
+				this.Id = Id;
+				this.Data = Data;
+                this.ChunkIds = ChunkIds;
+            }
+		}
+
+		class ChunkDocument
+		{
+			[BsonId]
+			public string Id { get; set; }
+
+			public byte[] Data { get; set; } = null!;
+
+			public string Version { get; set; }
+
+
+            [BsonConstructor]
+			private ChunkDocument()
+			{
+				this.Id = null!;
+				this.Data = null!;
+                this.Version = null!;
+            }
+
+			public ChunkDocument(string Id, string Version, byte[] Data)
+			{
+				this.Id = Id;
+				this.Data = Data;
+                this.Version = Version;
+            }
+		}
+
+		// MongoDB document size limit is 16 megs, we chunk at this with a 1k buffer for anything else in document
+		readonly int ChunkSize = 16 * 1024 * 1024 - 1024;
+
+
+
+		/// <summary>
+		/// Template documents
+		/// </summary>
+		IMongoCollection<AgentSoftwareDocument> Collection;
+
+		/// <summary>
+		/// Agent template chunks
+		/// </summary>
+		IMongoCollection<ChunkDocument> AgentChunks;
+
+
+		/// <summary>
+		/// Constructor
+		/// </summary>
+		/// <param name="DatabaseService">The database service singleton</param>
+		public AgentSoftwareCollection(DatabaseService DatabaseService)
+		{
+			Collection = DatabaseService.GetCollection<AgentSoftwareDocument>("AgentSoftware");
+			AgentChunks = DatabaseService.GetCollection<ChunkDocument>("AgentSoftwareChunks");
+		}
+
+		/// <inheritdoc/>
+		public async Task<bool> AddAsync(string Version, byte[] Data)
+		{
+
+			if (await ExistsAsync(Version))
+			{
+                return false;
+            }
+
+            AgentSoftwareDocument Document;
+
+            // Check if agent will fit in document, otherwise we need to chunk
+            if (Data.Length <= ChunkSize)
+			{
+				Document = new AgentSoftwareDocument(Version, Data);
+				await Collection.InsertOneAsync(Document);
+                return true;
+            }
+
+			int TotalRead = 0;            
+            byte[] Buffer = new byte[ChunkSize];            
+			List<ChunkDocument> Chunks = new List<ChunkDocument>();
+            using (MemoryStream ChunkStream = new MemoryStream(Data))
+            {
+                for (; ; )
+                {
+					int BytesRead = ChunkStream.Read(Buffer, 0, ChunkSize);
+					Chunks.Add(new ChunkDocument(ObjectId.GenerateNewId().ToString(), Version, Buffer.Take(BytesRead).ToArray()));
+					TotalRead += BytesRead;
+					if (TotalRead == Data.Length)
+					{
+						break;
+					}
+                }
+            }
+
+            // insert the chunks
+            await AgentChunks.InsertManyAsync(Chunks);
+
+            Document = new AgentSoftwareDocument(Version, null, Chunks.Select( Chunk => Chunk.Id).ToList());
+			await Collection.InsertOneAsync(Document);
+            return true;
+        }
+
+		/// <inheritdoc/>
+		public async Task<bool> ExistsAsync(string Version)
+		{
+			return await Collection.Find(x => x.Id == Version).AnyAsync();
+		}
+
+		/// <inheritdoc/>
+		public async Task<bool> RemoveAsync(string Version)
+		{
+			// Delete any chunks
+			await AgentChunks.DeleteManyAsync(x => x.Version == Version);
+
+			DeleteResult Result = await Collection.DeleteOneAsync(x => x.Id == Version);			
+			return Result.DeletedCount > 0;
+		}
+
+		/// <inheritdoc/>
+		public async Task<byte[]?> GetAsync(string Version)
+		{			
+			IFindFluent<AgentSoftwareDocument, AgentSoftwareDocument> Query = Collection.Find(x => x.Id == Version);
+			AgentSoftwareDocument? Software = await Query.FirstOrDefaultAsync();
+
+			if (Software == null)
+			{
+                return null;
+            }
+
+			if (Software.Data != null && Software.Data.Length > 0)
+			{
+                return Software.Data;
+            }
+
+			if (Software.ChunkIds != null)
+			{
+				FilterDefinitionBuilder<ChunkDocument> FilterBuilder = Builders<ChunkDocument>.Filter;
+				FilterDefinition<ChunkDocument> Filter = FilterBuilder.Empty;
+				Filter &= FilterBuilder.In(x => x.Id, Software.ChunkIds);
+
+				IFindFluent<ChunkDocument, ChunkDocument> ChunkQuery = AgentChunks.Find(x => Software.ChunkIds.Contains(x.Id));				
+				List<ChunkDocument> Chunks = await ChunkQuery.ToListAsync();
+
+				using (MemoryStream DataStream = new MemoryStream())
+				{
+					foreach (string ChunkId in Software.ChunkIds)
+					{
+                        DataStream.Write(Chunks.First(x => x.Id == ChunkId).Data);
+                    }
+
+                    return DataStream.ToArray();
+                }
+			}
+
+            return null;
+        }
+	}
+}
diff --git a/Engine/Source/Programs/Horde/Horde.Build/Agents/Impl/AgentSoftwareService.cs b/Engine/Source/Programs/Horde/Horde.Build/Agents/Impl/AgentSoftwareService.cs
new file mode 100644
--- /dev/null
+++ b/Engine/Source/Programs/Horde/Horde.Build/Agents/Impl/AgentSoftwareService.cs
@@ -0,0 +1,348 @@
+// Copyright Epic Games, Inc. All Rights Reserved.
+
+using HordeCommon;
+using HordeServer.Collections;
+using HordeServer.Models;
+using HordeServer.Utilities;
+using Microsoft.IdentityModel.Tokens;
+using MongoDB.Bson;
+using MongoDB.Bson.Serialization.Attributes;
+using MongoDB.Driver;
+using System;
+using System.Collections.Generic;
+using System.Diagnostics;
+using System.IO;
+using System.Linq;
+using System.Threading.Channels;
+using System.Threading.Tasks;
+using EpicGames.Core;
+using Microsoft.Extensions.Logging;
+using System.Runtime.InteropServices;
+using Microsoft.Extensions.Options;
+using System.IO.Compression;
+using System.Security.Cryptography.X509Certificates;
+
+namespace HordeServer.Services
+{
+	using AgentSoftwareChannelName = StringId<AgentSoftwareChannels>;
+
+	/// <summary>
+	/// Information about a channel
+	/// </summary>
+	public class AgentSoftwareChannel : IAgentSoftwareChannel
+	{
+		/// <summary>
+		/// The channel id
+		/// </summary>
+		public AgentSoftwareChannelName Name { get; set; }
+
+		/// <summary>
+		/// Name of the user that made the last modification
+		/// </summary>
+		public string? ModifiedBy { get; set; }
+
+		/// <summary>
+		/// Last modification time
+		/// </summary>
+		public DateTime ModifiedTime { get; set; }
+
+		/// <summary>
+		/// The software revision number
+		/// </summary>
+		public string Version { get; set; } = String.Empty;
+
+		/// <summary>
+		/// Private constructor for serialization
+		/// </summary>
+		[BsonConstructor]
+		private AgentSoftwareChannel()
+		{
+		}
+
+		/// <summary>
+		/// Constructor
+		/// </summary>
+		/// <param name="Name">Name of the channel</param>
+		public AgentSoftwareChannel(AgentSoftwareChannelName Name)
+		{
+			this.Name = Name;
+		}
+	}
+
+	/// <summary>
+	/// Singleton document used to track different versions of the agent software
+	/// </summary>
+	[SingletonDocument("5f455039d97900f2b6c735a9")]
+	public class AgentSoftwareChannels : SingletonBase
+	{
+		/// <summary>
+		/// The next version number
+		/// </summary>
+		public int NextVersion { get; set; } = 1;
+
+		/// <summary>
+		/// List of channels
+		/// </summary>
+		public List<AgentSoftwareChannel> Channels { get; set; } = new List<AgentSoftwareChannel>();
+
+		/// <summary>
+		/// Finds an existing channel by the given name, or adds a new one
+		/// </summary>
+		/// <param name="Name"></param>
+		/// <returns></returns>
+		public AgentSoftwareChannel FindOrAddChannel(AgentSoftwareChannelName Name)
+		{
+			AgentSoftwareChannel? Channel = Channels.FirstOrDefault(x => x.Name == Name);
+			if (Channel == null)
+			{
+				Channel = new AgentSoftwareChannel(Name);
+				Channels.Add(Channel);
+			}
+			return Channel;
+		}
+	}
+
+	/// <summary>
+	/// Wrapper for a collection of software revisions. 
+	/// </summary>
+	public class AgentSoftwareService
+	{
+		/// <summary>
+		/// Name of the default channel
+		/// </summary>
+		public static AgentSoftwareChannelName DefaultChannelName { get; } = new AgentSoftwareChannelName("default");
+
+		/// <summary>
+		/// Collection of software documents
+		/// </summary>
+		IAgentSoftwareCollection Collection;
+
+		/// <summary>
+		/// Channels singleton
+		/// </summary>
+		ISingletonDocument<AgentSoftwareChannels> Singleton;
+
+		/// <summary>
+		/// Cached copy of the channels singleton
+		/// </summary>
+		LazyCachedValue<Task<AgentSoftwareChannels>> ChannelsDocument;
+
+		/// <summary>
+		/// The server settings
+		/// </summary>
+		IOptionsMonitor<ServerSettings> Settings;
+
+		/// <summary>
+		///  Logger for controller
+		/// </summary>
+		private readonly ILogger<AgentSoftwareService> Logger;
+
+
+		/// <summary>
+		/// Constructor
+		/// </summary>
+		/// <param name="Collection">The software collection</param>
+		/// <param name="Singleton">The channels singleton</param>		
+		/// <param name="Settings">The settings monitor</param>		
+		/// <param name="Logger">The logger instance</param>
+		public AgentSoftwareService(IAgentSoftwareCollection Collection, ISingletonDocument<AgentSoftwareChannels> Singleton, IOptionsMonitor<ServerSettings> Settings, ILogger<AgentSoftwareService> Logger)
+		{
+			this.Collection = Collection;
+			this.Singleton = Singleton;
+			this.Settings = Settings;
+			this.Logger = Logger;
+
+			ChannelsDocument = new LazyCachedValue<Task<AgentSoftwareChannels>>(() => Singleton.GetAsync(), TimeSpan.FromSeconds(20.0));
+
+			Task.Run(() => RegisterDefaultAgent(5000));
+		}
+
+		/// <summary>
+		/// Finds all software archives matching a set of criteria
+		/// </summary>
+		public async Task<List<IAgentSoftwareChannel>> FindChannelsAsync()
+		{
+			AgentSoftwareChannels Current = await ChannelsDocument.GetLatest();
+			return Current.Channels.ConvertAll<IAgentSoftwareChannel>(x => x);
+		}
+
+		/// <summary>
+		/// Gets a single software archive
+		/// </summary>
+		/// <param name="Name">Unique id of the channel to find software for</param>
+		public async Task<IAgentSoftwareChannel?> GetChannelAsync(AgentSoftwareChannelName Name)
+		{
+			AgentSoftwareChannels Current = await ChannelsDocument.GetLatest();
+			return Current.Channels.FirstOrDefault(x => x.Name == Name);
+		}
+
+		/// <summary>
+		/// Gets a cached channel name
+		/// </summary>
+		/// <param name="Name">The cached channel name</param>
+		/// <returns></returns>
+		public async Task<IAgentSoftwareChannel?> GetCachedChannelAsync(AgentSoftwareChannelName Name)
+		{
+			AgentSoftwareChannels Document = await ChannelsDocument.GetCached();
+			return Document.Channels.FirstOrDefault(x => x.Name == Name);
+		}
+
+		/// <summary>
+		/// Removes a channel
+		/// </summary>
+		/// <param name="Name">The channel id</param>
+		/// <returns>Async task</returns>
+		public async Task DeleteChannelAsync(AgentSoftwareChannelName Name)
+		{
+			for (; ; )
+			{
+				AgentSoftwareChannels Current = await Singleton.GetAsync();
+
+				int ChannelIdx = Current.Channels.FindIndex(x => x.Name == Name);
+				if (ChannelIdx == -1)
+				{
+					break;
+				}
+
+				string Version = Current.Channels[ChannelIdx].Version;
+				Current.Channels.RemoveAt(ChannelIdx);
+
+				if (await Singleton.TryUpdateAsync(Current))
+				{
+					if (!Current.Channels.Any(x => x.Version == Version))
+					{
+						await Collection.RemoveAsync(Version);
+					}
+					break;
+				}
+			}
+		}
+
+		/// <summary>
+		/// Updates a new software revision
+		/// </summary>
+		/// <param name="Name">Name of the channel</param>
+		/// <param name="Author">Name of the user uploading this file</param>
+		/// <param name="Data">The input data stream. This should be a zip archive containing the HordeAgent executable.</param>
+		/// <returns>Unique id for the file</returns>
+		public async Task<string> SetArchiveAsync(AgentSoftwareChannelName Name, string? Author, byte[] Data)
+		{
+			// Upload the software
+			string Version = AgentUtilities.ReadVersion(Data);
+			await Collection.AddAsync(Version, Data);
+
+			// Update the channel
+			for(; ;)
+			{
+				AgentSoftwareChannels Instance = await Singleton.GetAsync();
+
+				AgentSoftwareChannel Channel = Instance.FindOrAddChannel(Name);
+				Channel.ModifiedBy = Author;
+				Channel.ModifiedTime = DateTime.UtcNow;
+				Channel.Version = Version;
+
+				if (await Singleton.TryUpdateAsync(Instance))
+				{
+					break;
+				}
+			}
+			return Version;
+		}
+
+		/// <summary>
+		/// Gets the zip file for a given channel
+		/// </summary>
+		/// <param name="Name">The channel name</param>
+		/// <returns>Data for the given archive</returns>
+		public async Task<byte[]?> GetArchiveAsync(AgentSoftwareChannelName Name)
+		{
+			AgentSoftwareChannels Instance = await Singleton.GetAsync();
+
+			AgentSoftwareChannel? Channel = Instance.Channels.FirstOrDefault(x => x.Name == Name);
+			if (Channel == null)
+			{
+				return null;
+			}
+			else
+			{
+				return await Collection.GetAsync(Channel.Version);
+			}
+		}
+
+		/// <summary>
+		/// Gets the zip file with a given version number
+		/// </summary>
+		/// <param name="Version">The version</param>
+		/// <returns>Data for the given archive</returns>
+		public Task<byte[]?> GetArchiveAsync(string Version)
+		{
+			return Collection.GetAsync(Version);
+		}
+		async Task RegisterDefaultAgent(int DelayMs)
+		{
+			await Task.Delay(DelayMs);
+
+			// Check whether we have an installed agent zip
+			FileReference AgentZip = FileReference.Combine(Program.AppDir, "DefaultAgent/Agent.zip");
+			if (!Settings.CurrentValue.SingleInstance || !FileReference.Exists(AgentZip))
+			{
+				return;
+			}
+
+			Logger.LogInformation("Checking for default agent software update");
+
+			string AgentHash = ContentHash.MD5(AgentZip).ToString();
+
+			FileReference AgentHashFile = FileReference.Combine(Program.DataDir, "Agent/DefaultAgentHash");
+
+			try
+			{
+
+				if (FileReference.Exists(AgentHashFile) && FileReference.ReadAllText(AgentHashFile).Trim() == AgentHash)
+				{
+					Logger.LogInformation("Default agent software is up to date");
+					return;
+				}
+			}
+			catch (Exception Ex)
+			{
+				Logger.LogError(Ex, "Error checking default agent software, {Message}", Ex.Message);
+			}
+
+            byte[] Bytes = await File.ReadAllBytesAsync(AgentZip.ToString());
+
+			using X509Certificate2? GrpcCertificate = Program.ReadGrpcCertificate(Settings.CurrentValue);
+
+			if (GrpcCertificate == null)
+			{
+				throw new Exception("Unable to register default agent without valid grpc certicate");
+			}
+
+			// construct agent app settings
+			
+			string ProfileName = "Default";
+
+			string ServerUrl = $"https://{System.Net.Dns.GetHostName()}:{Settings.CurrentValue.HttpsPort}";			
+
+			Dictionary<string, object> DefaultProfile = new Dictionary<string, object>() { { "Name", ProfileName }, { "Environment", "prod" }, { "Thumbprint", GrpcCertificate.Thumbprint }, { "Url", ServerUrl } };
+
+			List<object> ServerProfiles = new List<object>() { DefaultProfile };
+
+			Dictionary<string, object> AgentSettings = new Dictionary<string, object>() { { "ServerProfiles", ServerProfiles }, { "Server", ProfileName }};
+
+			Bytes = AgentUtilities.UpdateAppSettings(Bytes, AgentSettings);
+
+			string Version = await SetArchiveAsync(new AgentSoftwareChannelName("default"), null, Bytes);
+
+			if (!DirectoryReference.Exists(AgentHashFile.Directory))
+			{
+				DirectoryReference.CreateDirectory(AgentHashFile.Directory);
+			}
+
+			FileReference.WriteAllText(AgentHashFile, AgentHash);
+
+			Logger.LogInformation($"Updated default agent software to {AgentHash}");
+
+		}
+	}
+}
diff --git a/Engine/Source/Programs/Horde/Horde.Build/Agents/Impl/LeaseCollection.cs b/Engine/Source/Programs/Horde/Horde.Build/Agents/Impl/LeaseCollection.cs
new file mode 100644
--- /dev/null
+++ b/Engine/Source/Programs/Horde/Horde.Build/Agents/Impl/LeaseCollection.cs
@@ -0,0 +1,165 @@
+// Copyright Epic Games, Inc. All Rights Reserved.
+
+using HordeCommon;
+using HordeServer.Models;
+using HordeServer.Services;
+using HordeServer.Utilities;
+using MongoDB.Bson;
+using MongoDB.Bson.Serialization.Attributes;
+using MongoDB.Driver;
+using System;
+using System.Collections.Generic;
+using System.Linq;
+using System.Threading.Tasks;
+
+namespace HordeServer.Collections.Impl
+{
+	using LeaseId = ObjectId<ILease>;
+	using LogId = ObjectId<ILogFile>;
+	using PoolId = StringId<IPool>;
+	using StreamId = StringId<IStream>;
+
+	/// <summary>
+	/// Collection of lease documents
+	/// </summary>
+	public class LeaseCollection : ILeaseCollection
+	{
+		/// <summary>
+		/// Concrete implementation of a lease document
+		/// </summary>
+		class LeaseDocument : ILease
+		{
+			public LeaseId Id { get; set; }
+			public string Name { get; set; }
+			public AgentId AgentId { get; set; }
+			public ObjectId SessionId { get; set; }
+			public StreamId? StreamId { get; set; }
+			public PoolId? PoolId { get; set; }
+			public LogId? LogId { get; set; }
+			public DateTime StartTime { get; set; }
+			public DateTime? FinishTime { get; set; }
+			public byte[] Payload { get; set; }
+			public LeaseOutcome Outcome { get; set; }
+
+			[BsonIgnoreIfNull]
+			public byte[]? Output { get; set; }
+
+			ReadOnlyMemory<byte> ILease.Payload => Payload;
+			ReadOnlyMemory<byte> ILease.Output => Output ?? Array.Empty<byte>();
+
+			[BsonConstructor]
+			private LeaseDocument()
+			{
+				Name = String.Empty;
+				Payload = null!;
+			}
+
+			public LeaseDocument(LeaseId Id, string Name, AgentId AgentId, ObjectId SessionId, StreamId? StreamId, PoolId? PoolId, LogId? LogId, DateTime StartTime, byte[] Payload)
+			{
+				this.Id = Id;
+				this.Name = Name;
+				this.AgentId = AgentId;
+				this.SessionId = SessionId;
+				this.StreamId = StreamId;
+				this.PoolId = PoolId;
+				this.LogId = LogId;
+				this.StartTime = StartTime;
+				this.Payload = Payload;
+			}
+		}
+
+		/// <summary>
+		/// Collection of lease documents
+		/// </summary>
+		readonly IMongoCollection<LeaseDocument> Leases;
+
+		/// <summary>
+		/// Constructor
+		/// </summary>
+		/// <param name="DatabaseService">The database service instance</param>
+		public LeaseCollection(DatabaseService DatabaseService)
+		{
+			Leases = DatabaseService.GetCollection<LeaseDocument>("Leases");
+
+			if (!DatabaseService.ReadOnlyMode)
+			{
+				Leases.Indexes.CreateOne(new CreateIndexModel<LeaseDocument>(Builders<LeaseDocument>.IndexKeys.Ascending(x => x.AgentId)));
+				Leases.Indexes.CreateOne(new CreateIndexModel<LeaseDocument>(Builders<LeaseDocument>.IndexKeys.Ascending(x => x.SessionId)));
+				Leases.Indexes.CreateOne(new CreateIndexModel<LeaseDocument>(Builders<LeaseDocument>.IndexKeys.Ascending(x => x.StartTime)));
+				Leases.Indexes.CreateOne(new CreateIndexModel<LeaseDocument>(Builders<LeaseDocument>.IndexKeys.Ascending(x => x.FinishTime)));
+			}
+		}
+
+		/// <inheritdoc/>
+		public async Task<ILease> AddAsync(LeaseId Id, string Name, AgentId AgentId, ObjectId SessionId, StreamId? StreamId, PoolId? PoolId, LogId? LogId, DateTime StartTime, byte[] Payload)
+		{
+			LeaseDocument Lease = new LeaseDocument(Id, Name, AgentId, SessionId, StreamId, PoolId, LogId, StartTime, Payload);
+			await Leases.ReplaceOneAsync(x => x.Id == Id, Lease, new ReplaceOptions { IsUpsert = true });
+			return Lease;
+		}
+
+		/// <inheritdoc/>
+		public async Task DeleteAsync(LeaseId LeaseId)
+		{
+			await Leases.DeleteOneAsync(x => x.Id == LeaseId);
+		}
+
+		/// <inheritdoc/>
+		public async Task<ILease?> GetAsync(LeaseId LeaseId)
+		{
+			return await Leases.Find(x => x.Id == LeaseId).FirstOrDefaultAsync();
+		}
+
+		/// <inheritdoc/>
+		public async Task<List<ILease>> FindLeasesAsync(AgentId? AgentId, ObjectId? SessionId, DateTime? MinTime, DateTime? MaxTime, int? Index, int? Count)
+		{
+			FilterDefinitionBuilder<LeaseDocument> FilterBuilder = Builders<LeaseDocument>.Filter;
+
+			FilterDefinition<LeaseDocument> Filter = FilterDefinition<LeaseDocument>.Empty;
+			if (AgentId != null)
+			{
+				Filter &= FilterBuilder.Eq(x => x.AgentId, AgentId.Value);
+			}
+			if (SessionId != null)
+			{
+				Filter &= FilterBuilder.Eq(x => x.SessionId, SessionId.Value);
+			}
+			if (MinTime != null)
+			{
+				Filter &= FilterBuilder.Or(FilterBuilder.Eq(x => x.FinishTime!, null), FilterBuilder.Gt(x => x.FinishTime!, MinTime.Value));
+			}
+			if (MaxTime != null)
+			{
+				Filter &= FilterBuilder.Lt(x => x.StartTime, MaxTime.Value);
+			}
+
+			List<LeaseDocument> Results = await Leases.Find(Filter).SortByDescending(x => x.StartTime).Range(Index, Count).ToListAsync();
+			return Results.ConvertAll<ILease>(x => x);
+		}
+
+		/// <inheritdoc/>
+		public async Task<List<ILease>> FindActiveLeasesAsync(int? Index = null, int? Count = null)
+		{
+			List<LeaseDocument> Results = await Leases.Find(x => x.FinishTime == null).Range(Index, Count).ToListAsync();
+			return Results.ConvertAll<ILease>(x => x);
+		}
+
+		/// <inheritdoc/>
+		public async Task<bool> TrySetOutcomeAsync(LeaseId LeaseId, DateTime FinishTime, LeaseOutcome Outcome, byte[]? Output)
+		{
+			FilterDefinitionBuilder<LeaseDocument> FilterBuilder = Builders<LeaseDocument>.Filter;
+			FilterDefinition<LeaseDocument> Filter = FilterBuilder.Eq(x => x.Id, LeaseId) & FilterBuilder.Eq(x => x.FinishTime, null);
+
+			UpdateDefinitionBuilder<LeaseDocument> UpdateBuilder = Builders<LeaseDocument>.Update;
+			UpdateDefinition<LeaseDocument> Update = UpdateBuilder.Set(x => x.FinishTime, FinishTime).Set(x => x.Outcome, Outcome);
+
+			if(Output != null && Output.Length > 0)
+			{
+				Update = Update.Set(x => x.Output, Output);
+			}
+
+			UpdateResult Result = await Leases.UpdateOneAsync(Filter, Update);
+			return Result.ModifiedCount > 0;
+		}
+	}
+}
diff --git a/Engine/Source/Programs/Horde/Horde.Build/Agents/Impl/SessionCollection.cs b/Engine/Source/Programs/Horde/Horde.Build/Agents/Impl/SessionCollection.cs
new file mode 100644
--- /dev/null
+++ b/Engine/Source/Programs/Horde/Horde.Build/Agents/Impl/SessionCollection.cs
@@ -0,0 +1,141 @@
+// Copyright Epic Games, Inc. All Rights Reserved.
+
+using HordeServer.Models;
+using HordeServer.Services;
+using HordeServer.Utilities;
+using Microsoft.Extensions.DependencyInjection;
+using MongoDB.Bson;
+using MongoDB.Bson.Serialization.Attributes;
+using MongoDB.Driver;
+using System;
+using System.Collections.Generic;
+using System.Linq;
+using System.Threading.Tasks;
+
+namespace HordeServer.Collections.Impl
+{
+	/// <summary>
+	/// Collection of session documents
+	/// </summary>
+	public class SessionCollection : ISessionCollection
+	{
+		/// <summary>
+		/// Concrete implementation of ISession
+		/// </summary>
+		class SessionDocument : ISession
+		{
+			[BsonRequired, BsonId]
+			public ObjectId Id { get; set; }
+
+			[BsonRequired]
+			public AgentId AgentId { get; set; }
+
+			public DateTime StartTime { get; set; }
+			public DateTime? FinishTime { get; set; }
+			public List<string>? Properties { get; set; }
+			public Dictionary<string, int>? Resources { get; set; }
+			public string? Version { get; set; }
+
+			IReadOnlyList<string>? ISession.Properties => Properties;
+
+			[BsonConstructor]
+			private SessionDocument()
+			{
+			}
+
+			public SessionDocument(ObjectId Id, AgentId AgentId, DateTime StartTime, IReadOnlyList<string>? Properties, IReadOnlyDictionary<string, int>? Resources, string? Version)
+			{
+				this.Id = Id;
+				this.AgentId = AgentId;
+				this.StartTime = StartTime;
+				if (Properties != null)
+				{
+					this.Properties = new List<string>(Properties);
+				}
+				if (Resources != null)
+				{
+					this.Resources = new Dictionary<string, int>(Resources);
+				}
+				this.Version = Version;
+			}
+		}
+
+		/// <summary>
+		/// Collection of session documents
+		/// </summary>
+		readonly IMongoCollection<SessionDocument> Sessions;
+
+		/// <summary>
+		/// Constructor
+		/// </summary>
+		/// <param name="DatabaseService">The database service</param>
+		public SessionCollection(DatabaseService DatabaseService)
+		{
+			Sessions = DatabaseService.GetCollection<SessionDocument>("Sessions");
+
+			if (!DatabaseService.ReadOnlyMode)
+			{
+				Sessions.Indexes.CreateOne(new CreateIndexModel<SessionDocument>(Builders<SessionDocument>.IndexKeys.Ascending(x => x.AgentId)));
+				Sessions.Indexes.CreateOne(new CreateIndexModel<SessionDocument>(Builders<SessionDocument>.IndexKeys.Ascending(x => x.StartTime)));
+				Sessions.Indexes.CreateOne(new CreateIndexModel<SessionDocument>(Builders<SessionDocument>.IndexKeys.Ascending(x => x.FinishTime)));
+			}
+		}
+
+		/// <inheritdoc/>
+		public async Task<ISession> AddAsync(ObjectId Id, AgentId AgentId, DateTime StartTime, IReadOnlyList<string>? Properties, IReadOnlyDictionary<string, int>? Resources, string? Version)
+		{
+			SessionDocument NewSession = new SessionDocument(Id, AgentId, StartTime, Properties, Resources, Version);
+			await Sessions.InsertOneAsync(NewSession);
+			return NewSession;
+		}
+
+		/// <inheritdoc/>
+		public async Task<ISession?> GetAsync(ObjectId SessionId)
+		{
+			return await Sessions.Find(x => x.Id == SessionId).FirstOrDefaultAsync();
+		}
+
+		/// <inheritdoc/>
+		public async Task<List<ISession>> FindAsync(AgentId AgentId, DateTime? StartTime, DateTime? FinishTime, int Index, int Count)
+		{
+			FilterDefinitionBuilder<SessionDocument> FilterBuilder = Builders<SessionDocument>.Filter;
+
+			FilterDefinition<SessionDocument> Filter = FilterBuilder.Eq(x => x.AgentId, AgentId);
+			if (StartTime != null)
+			{
+				Filter &= FilterBuilder.Gte(x => x.StartTime, StartTime.Value);
+			}
+			if (FinishTime != null)
+			{
+				Filter &= FilterBuilder.Or(FilterBuilder.Eq(x => x.FinishTime, null), FilterBuilder.Lte(x => x.FinishTime, FinishTime.Value));
+			}
+
+			List<SessionDocument> Results = await Sessions.Find(Filter).SortByDescending(x => x.StartTime).Skip(Index).Limit(Count).ToListAsync();
+			return Results.ConvertAll<ISession>(x => x);
+		}
+
+		/// <inheritdoc/>
+		public async Task<List<ISession>> FindActiveSessionsAsync(int? Index, int? Count)
+		{
+			List<SessionDocument> Results = await Sessions.Find(x => x.FinishTime == null).Range(Index, Count).ToListAsync();
+			return Results.ConvertAll<ISession>(x => x);
+		}
+
+		/// <inheritdoc/>
+		public Task UpdateAsync(ObjectId SessionId, DateTime FinishTime, IReadOnlyList<string> Properties, IReadOnlyDictionary<string, int> Resources)
+		{
+			UpdateDefinition<SessionDocument> Update = Builders<SessionDocument>.Update
+				.Set(x => x.FinishTime, FinishTime)
+				.Set(x => x.Properties, new List<string>(Properties))
+				.Set(x => x.Resources, new Dictionary<string, int>(Resources));
+
+			return Sessions.FindOneAndUpdateAsync(x => x.Id == SessionId, Update);
+		}
+
+		/// <inheritdoc/>
+		public Task DeleteAsync(ObjectId SessionId)
+		{
+			return Sessions.DeleteOneAsync(x => x.Id == SessionId);
+		}
+	}
+}
diff --git a/Engine/Source/Programs/Horde/Horde.Build/Agents/Impl/Tasks/ConformTaskSource.cs b/Engine/Source/Programs/Horde/Horde.Build/Agents/Impl/Tasks/ConformTaskSource.cs
new file mode 100644
--- /dev/null
+++ b/Engine/Source/Programs/Horde/Horde.Build/Agents/Impl/Tasks/ConformTaskSource.cs
@@ -0,0 +1,390 @@
+// Copyright Epic Games, Inc. All Rights Reserved.
+
+using Google.Protobuf;
+using Google.Protobuf.Reflection;
+using Google.Protobuf.WellKnownTypes;
+using HordeCommon;
+using HordeCommon.Rpc.Tasks;
+using HordeServer.Api;
+using HordeServer.Collections;
+using HordeServer.Models;
+using HordeServer.Services;
+using HordeServer.Utilities;
+using Microsoft.Extensions.Hosting;
+using Microsoft.Extensions.Logging;
+using MongoDB.Bson;
+using System;
+using System.Collections.Generic;
+using System.Linq;
+using System.Threading;
+using System.Threading.Tasks;
+
+namespace HordeServer.Tasks.Impl
+{
+	using JobId = ObjectId<IJob>;
+	using LeaseId = ObjectId<ILease>;
+	using LogId = ObjectId<ILogFile>;
+
+	/// <summary>
+	/// Generates tasks telling agents to sync their workspaces
+	/// </summary>
+	public sealed class ConformTaskSource : TaskSourceBase<ConformTask>, IHostedService, IDisposable
+	{
+		/// <inheritdoc/>
+		public override string Type => "Conform";
+
+		DatabaseService DatabaseService;
+		IAgentCollection AgentCollection;
+		PoolService PoolService;
+		SingletonDocument<ConformList> ConformList;
+		PerforceLoadBalancer PerforceLoadBalancer;
+		ILogFileService LogService;
+		ILogger Logger;
+		ElectedTick TickConformList;
+
+		/// <summary>
+		/// Constructor
+		/// </summary>
+		/// <param name="DatabaseService"></param>
+		/// <param name="AgentCollection"></param>
+		/// <param name="PoolService"></param>
+		/// <param name="LogService"></param>
+		/// <param name="PerforceLoadBalancer"></param>
+		/// <param name="Logger"></param>
+		public ConformTaskSource(DatabaseService DatabaseService, IAgentCollection AgentCollection, PoolService PoolService, ILogFileService LogService, PerforceLoadBalancer PerforceLoadBalancer, ILogger<ConformTaskSource> Logger)
+		{
+			this.DatabaseService = DatabaseService;
+			this.AgentCollection = AgentCollection;
+			this.PoolService = PoolService;
+			this.ConformList = new SingletonDocument<ConformList>(DatabaseService);
+			this.PerforceLoadBalancer = PerforceLoadBalancer;
+			this.LogService = LogService;
+			this.Logger = Logger;
+			this.TickConformList = new ElectedTick(DatabaseService, new ObjectId("60afc5cf555a9a76aff0a50c"), CleanConformListAsync, TimeSpan.FromMinutes(1.0), Logger);
+
+			this.OnLeaseStartedProperties.Add(nameof(ConformTask.LogId), x => new LogId(x.LogId));
+		}
+
+		/// <inheritdoc/>
+		public void Dispose()
+		{
+			TickConformList.Dispose();
+		}
+
+		/// <inheritdoc/>
+		public Task StartAsync(CancellationToken CancellationToken)
+		{
+			TickConformList.Start();
+			return Task.CompletedTask;
+		}
+
+		/// <inheritdoc/>
+		public async Task StopAsync(CancellationToken CancellationToken)
+		{
+			await TickConformList.StopAsync();
+		}
+
+		/// <summary>
+		/// Clean up the conform list of any outdated entries
+		/// </summary>
+		/// <param name="CancellationToken">Cancellation token</param>
+		/// <returns>Async task</returns>
+		async Task CleanConformListAsync(CancellationToken CancellationToken)
+		{
+			DateTime UtcNow = DateTime.UtcNow;
+			DateTime LastCheckTimeUtc = UtcNow - TimeSpan.FromMinutes(30.0);
+
+			// Get the current state of the conform list
+			ConformList List = await ConformList.GetAsync();
+
+			// Update any leases that are older than LastCheckTimeUtc
+			Dictionary<LeaseId, bool> RemoveLeases = new Dictionary<LeaseId, bool>();
+			foreach (ConformListEntry Entry in Enumerable.Concat(List.Entries, List.Servers.SelectMany(x => x.Entries)))
+			{
+				if (Entry.LastCheckTimeUtc < LastCheckTimeUtc)
+				{
+					IAgent? Agent = await AgentCollection.GetAsync(Entry.AgentId);
+
+					bool Remove = false;
+					if (Agent == null || !Agent.Leases.Any(x => x.Id == Entry.LeaseId))
+					{
+						Logger.LogWarning("Removing invalid lease from conform list: {AgentId} lease {LeaseId}", Entry.AgentId, Entry.LeaseId);
+						Remove = true;
+					}
+
+					RemoveLeases[Entry.LeaseId] = Remove;
+				}
+			}
+
+			// If there's anything to change, update the list
+			if (RemoveLeases.Count > 0)
+			{
+				await ConformList.UpdateAsync(List => UpdateConformList(List, UtcNow, RemoveLeases));
+			}
+		}
+
+		/// <summary>
+		/// Remove items from the conform list, and update timestamps for them
+		/// </summary>
+		/// <param name="List">The list to update</param>
+		/// <param name="UtcNow">Current time</param>
+		/// <param name="RemoveLeases">List of leases to update. Entries with values set to true will be removed, entries with values set to false will have their timestamp updated.</param>
+		static void UpdateConformList(ConformList List, DateTime UtcNow, Dictionary<LeaseId, bool> RemoveLeases)
+		{
+			UpdateConformList(List.Entries, UtcNow, RemoveLeases);
+			foreach (ConformListServer Server in List.Servers)
+			{
+				UpdateConformList(Server.Entries, UtcNow, RemoveLeases);
+			}
+		}
+
+		/// <summary>
+		/// Remove items from the conform list, and update timestamps for them
+		/// </summary>
+		/// <param name="Entries">The list to update</param>
+		/// <param name="UtcNow">Current time</param>
+		/// <param name="RemoveLeases">List of leases to update. Entries with values set to true will be removed, entries with values set to false will have their timestamp updated.</param>
+		static void UpdateConformList(List<ConformListEntry> Entries, DateTime UtcNow, Dictionary<LeaseId, bool> RemoveLeases)
+		{
+			for (int Idx = 0; Idx < Entries.Count; Idx++)
+			{
+				bool Remove;
+				if (RemoveLeases.TryGetValue(Entries[Idx].LeaseId, out Remove))
+				{
+					if (Remove)
+					{
+						Entries.RemoveAt(Idx--);
+					}
+					else if (Entries[Idx].LastCheckTimeUtc < UtcNow)
+					{
+						Entries[Idx].LastCheckTimeUtc = UtcNow;
+					}
+				}
+			}
+		}
+
+		/// <inheritdoc/>
+		public override async Task<AgentLease?> AssignLeaseAsync(IAgent Agent, CancellationToken CancellationToken)
+		{
+			DateTime UtcNow = DateTime.UtcNow;
+			if (!await IsConformPendingAsync(Agent, UtcNow))
+			{
+				return null;
+			}
+
+			if (Agent.Leases.Count == 0)
+			{
+				ConformTask Task = new ConformTask();
+				if (await GetWorkspacesAsync(Agent, Task.Workspaces))
+				{
+					LeaseId LeaseId = LeaseId.GenerateNewId();
+					if (!await AllocateConformLeaseAsync(Agent.Id, Task.Workspaces, LeaseId))
+					{
+						return null;
+					}
+
+					ILogFile Log = await LogService.CreateLogFileAsync(JobId.Empty, Agent.SessionId, LogType.Json);
+					Task.LogId = Log.Id.ToString();
+
+					byte[] Payload = Any.Pack(Task).ToByteArray();
+
+					return new AgentLease(LeaseId, "Updating workspaces", null, null, Log.Id, LeaseState.Pending, null, true, Payload);
+				}
+			}
+
+			if (Agent.RequestConform)
+			{
+				return AgentLease.Drain;
+			}
+
+			return null;
+		}
+
+		/// <inheritdoc/>
+		public async Task<bool> GetWorkspacesAsync(IAgent Agent, IList<HordeCommon.Rpc.Messages.AgentWorkspace> Workspaces)
+		{
+			Globals Globals = await DatabaseService.GetGlobalsAsync();
+
+			HashSet<AgentWorkspace> ConformWorkspaces = await PoolService.GetWorkspacesAsync(Agent, DateTime.UtcNow);
+			foreach (AgentWorkspace ConformWorkspace in ConformWorkspaces)
+			{
+				PerforceCluster? Cluster = Globals.FindPerforceCluster(ConformWorkspace.Cluster);
+				if (Cluster == null || !await Agent.TryAddWorkspaceMessage(ConformWorkspace, Cluster, PerforceLoadBalancer, Workspaces))
+				{
+					return false;
+				}
+			}
+
+			return true;
+		}
+
+		/// <inheritdoc/>
+		public override Task OnLeaseFinishedAsync(IAgent Agent, LeaseId LeaseId, ConformTask Payload, LeaseOutcome Outcome, ReadOnlyMemory<byte> Output, ILogger Logger)
+		{
+			return ReleaseConformLeaseAsync(LeaseId);
+		}
+
+		/// <summary>
+		/// Atempt to allocate a conform resource for the given lease
+		/// </summary>
+		/// <param name="AgentId">The agent id</param>
+		/// <param name="Workspaces">List of workspaces that are required</param>
+		/// <param name="LeaseId">The lease id</param>
+		/// <returns>True if the resource was allocated, false otherwise</returns>
+		private async Task<bool> AllocateConformLeaseAsync(AgentId AgentId, IEnumerable<HordeCommon.Rpc.Messages.AgentWorkspace> Workspaces, LeaseId LeaseId)
+		{
+			Globals Globals = await DatabaseService.GetGlobalsAsync();
+			for (; ; )
+			{
+				ConformList CurrentValue = await ConformList.GetAsync();
+				if (Globals.MaxConformCount != 0 && CurrentValue.Entries.Count + CurrentValue.Servers.Sum(x => x.Entries.Count) >= Globals.MaxConformCount)
+				{
+					return false;
+				}
+
+				HashSet<string> Servers = new HashSet<string>(StringComparer.OrdinalIgnoreCase);
+				foreach (HordeCommon.Rpc.Messages.AgentWorkspace Workspace in Workspaces)
+				{
+					if (Servers.Add(Workspace.ServerAndPort))
+					{
+						PerforceCluster? Cluster = Globals.FindPerforceCluster(Workspace.Cluster);
+						if (Cluster == null)
+						{
+							Logger.LogWarning("Unable to find perforce cluster '{Cluster}' for conform", Workspace.Cluster);
+							return false;
+						}
+
+						PerforceServer? Server = Cluster.Servers.FirstOrDefault(x => x.ServerAndPort.Equals(Workspace.BaseServerAndPort, StringComparison.Ordinal));
+						if (Server == null)
+						{
+							Logger.LogWarning("Unable to find perforce server '{Server}' in cluster '{Cluster}' for conform", Workspace.BaseServerAndPort, Workspace.Cluster);
+							return false;
+						}
+
+						ConformListServer? ConformServer = CurrentValue.Servers.FirstOrDefault(x => x.Cluster.Equals(Workspace.Cluster, StringComparison.Ordinal) && x.ServerAndPort.Equals(Workspace.ServerAndPort, StringComparison.Ordinal));
+						if (ConformServer == null)
+						{
+							ConformServer = new ConformListServer { Cluster = Workspace.Cluster, ServerAndPort = Workspace.ServerAndPort };
+							CurrentValue.Servers.Add(ConformServer);
+						}
+
+						if (Server.MaxConformCount != 0 && ConformServer.Entries.Count > Server.MaxConformCount)
+						{
+							return false;
+						}
+
+						ConformListEntry Entry = new ConformListEntry();
+						Entry.AgentId = AgentId;
+						Entry.LeaseId = LeaseId;
+						Entry.LastCheckTimeUtc = DateTime.UtcNow;
+						ConformServer.Entries.Add(Entry);
+					}
+				}
+
+				if (await ConformList.TryUpdateAsync(CurrentValue))
+				{
+					Logger.LogInformation("Added conform lease {LeaseId}", LeaseId);
+					return true;
+				}
+			}
+		}
+
+		/// <summary>
+		/// Terminate a conform lease
+		/// </summary>
+		/// <param name="LeaseId">The lease id</param>
+		/// <returns>Async task</returns>
+		public async Task ReleaseConformLeaseAsync(LeaseId LeaseId)
+		{
+			for (; ; )
+			{
+				ConformList CurrentValue = await ConformList.GetAsync();
+				if (CurrentValue.Entries.RemoveAll(x => x.LeaseId == LeaseId) + CurrentValue.Servers.Sum(x => x.Entries.RemoveAll(x => x.LeaseId == LeaseId)) == 0)
+				{
+					Logger.LogInformation("Conform lease {LeaseId} is not in singelton", LeaseId);
+					break;
+				}
+				if (await ConformList.TryUpdateAsync(CurrentValue))
+				{
+					Logger.LogInformation("Removed conform lease {LeaseId}", LeaseId);
+					break;
+				}
+			}
+		}
+
+		/// <summary>
+		/// Determine if an agent should be conformed
+		/// </summary>
+		/// <param name="Agent">The agent to test</param>
+		/// <param name="UtcNow">Current time</param>
+		/// <returns>True if the agent should be conformed</returns>
+		private async Task<bool> IsConformPendingAsync(IAgent Agent, DateTime UtcNow)
+		{
+			// If a conform was manually requested, allow it to run even if the agent is disabled
+			if (Agent.RequestConform)
+			{
+				return !IsConformCoolDownPeriod(Agent, UtcNow);
+			}
+
+			// Otherwise only run if the agent is enabled
+			if (Agent.Enabled)
+			{
+				// If we've attempted (and failed) a conform, run again after a certain time
+				if (Agent.ConformAttemptCount.HasValue)
+				{
+					return !IsConformCoolDownPeriod(Agent, UtcNow);
+				}
+
+				// Always run a conform every 24h
+				if (UtcNow > Agent.LastConformTime + TimeSpan.FromDays(1.0))
+				{
+					return true;
+				}
+
+				// Check if the workspaces have changed (first check against a cached list of workspaces, then an accurate one)
+				HashSet<AgentWorkspace> Workspaces = await PoolService.GetWorkspacesAsync(Agent, UtcNow - TimeSpan.FromSeconds(30.0));
+				if (!Workspaces.SetEquals(Agent.Workspaces))
+				{
+					Workspaces = await PoolService.GetWorkspacesAsync(Agent, UtcNow);
+					if (!Workspaces.SetEquals(Agent.Workspaces))
+					{
+						return true;
+					}
+				}
+			}
+			return false;
+		}
+
+		/// <summary>
+		/// Determines if the given agent is in a cooldown period, waiting to retry after a failed conform
+		/// </summary>
+		/// <param name="Agent">The agent to test</param>
+		/// <param name="UtcNow">Current time</param>
+		/// <returns></returns>
+		private static bool IsConformCoolDownPeriod(IAgent Agent, DateTime UtcNow)
+		{
+			if (!Agent.ConformAttemptCount.HasValue)
+			{
+				return false;
+			}
+
+			TimeSpan RetryTime;
+			switch (Agent.ConformAttemptCount.Value)
+			{
+				case 1:
+					RetryTime = TimeSpan.FromMinutes(5.0);
+					break;
+				case 2:
+					RetryTime = TimeSpan.FromMinutes(20.0);
+					break;
+				case 3:
+					RetryTime = TimeSpan.FromHours(1.0);
+					break;
+				default:
+					RetryTime = TimeSpan.FromHours(6.0);
+					break;
+			}
+			return UtcNow < Agent.LastConformTime + RetryTime;
+		}
+	}
+}
diff --git a/Engine/Source/Programs/Horde/Horde.Build/Agents/Impl/Tasks/RestartTaskSource.cs b/Engine/Source/Programs/Horde/Horde.Build/Agents/Impl/Tasks/RestartTaskSource.cs
new file mode 100644
--- /dev/null
+++ b/Engine/Source/Programs/Horde/Horde.Build/Agents/Impl/Tasks/RestartTaskSource.cs
@@ -0,0 +1,59 @@
+// Copyright Epic Games, Inc. All Rights Reserved.
+
+using Google.Protobuf;
+using Google.Protobuf.Reflection;
+using Google.Protobuf.WellKnownTypes;
+using HordeCommon;
+using HordeCommon.Rpc.Tasks;
+using HordeServer.Api;
+using HordeServer.Models;
+using HordeServer.Services;
+using HordeServer.Utilities;
+using MongoDB.Bson;
+using System;
+using System.Collections.Generic;
+using System.Linq;
+using System.Threading;
+using System.Threading.Tasks;
+
+namespace HordeServer.Tasks.Impl
+{
+	using JobId = ObjectId<IJob>;
+	using LeaseId = ObjectId<ILease>;
+	using LogId = ObjectId<ILogFile>;
+
+	class RestartTaskSource : TaskSourceBase<RestartTask>
+	{
+		public override string Type => "Restart";
+
+		ILogFileService LogService;
+
+		public RestartTaskSource(ILogFileService LogService)
+		{
+			this.LogService = LogService;
+
+			this.OnLeaseStartedProperties.Add(nameof(RestartTask.LogId), x => new LogId(x.LogId));
+		}
+
+		public override async Task<AgentLease?> AssignLeaseAsync(IAgent Agent, CancellationToken CancellationToken)
+		{
+			if (!Agent.RequestRestart)
+			{
+				return null;
+			}
+			if (Agent.Leases.Count > 0)
+			{
+				return AgentLease.Drain;
+			}
+
+			ILogFile Log = await LogService.CreateLogFileAsync(JobId.Empty, Agent.SessionId, LogType.Json);
+
+			RestartTask Task = new RestartTask();
+			Task.LogId = Log.Id.ToString();
+
+			byte[] Payload = Any.Pack(Task).ToByteArray();
+
+			return new AgentLease(LeaseId.GenerateNewId(), "Restart", null, null, Log.Id, LeaseState.Pending, null, true, Payload);
+		}
+	}
+}
diff --git a/Engine/Source/Programs/Horde/Horde.Build/Agents/Impl/Tasks/ShutdownTaskSource.cs b/Engine/Source/Programs/Horde/Horde.Build/Agents/Impl/Tasks/ShutdownTaskSource.cs
new file mode 100644
--- /dev/null
+++ b/Engine/Source/Programs/Horde/Horde.Build/Agents/Impl/Tasks/ShutdownTaskSource.cs
@@ -0,0 +1,59 @@
+// Copyright Epic Games, Inc. All Rights Reserved.
+
+using EpicGames.Core;
+using Google.Protobuf;
+using Google.Protobuf.Reflection;
+using Google.Protobuf.WellKnownTypes;
+using HordeCommon;
+using HordeCommon.Rpc.Tasks;
+using HordeServer.Api;
+using HordeServer.Models;
+using HordeServer.Services;
+using HordeServer.Utilities;
+using MongoDB.Bson;
+using System;
+using System.Collections.Generic;
+using System.Linq;
+using System.Threading;
+using System.Threading.Tasks;
+
+namespace HordeServer.Tasks.Impl
+{
+	using JobId = ObjectId<IJob>;
+	using LeaseId = ObjectId<ILease>;
+	using LogId = ObjectId<ILogFile>;
+
+	class ShutdownTaskSource : TaskSourceBase<ShutdownTask>
+	{
+		public override string Type => "Shutdown";
+
+		ILogFileService LogService;
+
+		public ShutdownTaskSource(ILogFileService LogService)
+		{
+			this.LogService = LogService;
+			this.OnLeaseStartedProperties.Add(nameof(ShutdownTask.LogId), x => new LogId(x.LogId));
+		}
+
+		public override async Task<AgentLease?> AssignLeaseAsync(IAgent Agent, CancellationToken CancellationToken)
+		{
+			if (!Agent.RequestShutdown)
+			{
+				return null;
+			}
+			if (Agent.Leases.Count > 0)
+			{
+				return AgentLease.Drain;
+			}
+
+			ILogFile Log = await LogService.CreateLogFileAsync(JobId.Empty, Agent.SessionId, LogType.Json);
+
+			ShutdownTask Task = new ShutdownTask();
+			Task.LogId = Log.Id.ToString();
+
+			byte[] Payload = Any.Pack(Task).ToByteArray();
+
+			return new AgentLease(LeaseId.GenerateNewId(), "Shutdown", null, null, Log.Id, LeaseState.Pending, null, true, Payload);
+		}
+	}
+}
diff --git a/Engine/Source/Programs/Horde/Horde.Build/Agents/Impl/Tasks/UpgradeTaskSource.cs b/Engine/Source/Programs/Horde/Horde.Build/Agents/Impl/Tasks/UpgradeTaskSource.cs
new file mode 100644
--- /dev/null
+++ b/Engine/Source/Programs/Horde/Horde.Build/Agents/Impl/Tasks/UpgradeTaskSource.cs
@@ -0,0 +1,89 @@
+// Copyright Epic Games, Inc. All Rights Reserved.
+
+using Google.Protobuf;
+using Google.Protobuf.Reflection;
+using Google.Protobuf.WellKnownTypes;
+using HordeCommon;
+using HordeCommon.Rpc.Tasks;
+using HordeServer.Api;
+using HordeServer.Models;
+using HordeServer.Services;
+using HordeServer.Utilities;
+using MongoDB.Bson;
+using System;
+using System.Collections.Generic;
+using System.Linq;
+using System.Threading;
+using System.Threading.Tasks;
+
+namespace HordeServer.Tasks.Impl
+{
+	using AgentSoftwareChannelName = StringId<AgentSoftwareChannels>;
+	using JobId = ObjectId<IJob>;
+	using LeaseId = ObjectId<ILease>;
+	using LogId = ObjectId<ILogFile>;
+
+	class UpgradeTaskSource : TaskSourceBase<UpgradeTask>
+	{
+		public override string Type => "Upgrade";
+
+		AgentSoftwareService AgentSoftwareService;
+		ILogFileService LogService;
+		IClock Clock;
+
+		public UpgradeTaskSource(AgentSoftwareService AgentSoftwareService, ILogFileService LogService, IClock Clock)
+		{
+			this.AgentSoftwareService = AgentSoftwareService;
+			this.LogService = LogService;
+			this.Clock = Clock;
+
+			this.OnLeaseStartedProperties.Add(nameof(UpgradeTask.LogId), x => new LogId(x.LogId));
+		}
+
+		public async override Task<AgentLease?> AssignLeaseAsync(IAgent Agent, CancellationToken CancellationToken)
+		{
+			string? RequiredVersion = await GetRequiredSoftwareVersion(Agent);
+			if (RequiredVersion != null && Agent.Version != RequiredVersion)
+			{
+				AgentLease? Lease = null;
+				if (Agent.Leases.Count == 0 && (Agent.LastUpgradeTime == null || Agent.LastUpgradeTime.Value + TimeSpan.FromMinutes(5.0) < Clock.UtcNow || Agent.LastUpgradeVersion != RequiredVersion.ToString()))
+				{
+					ILogFile LogFile = await LogService.CreateLogFileAsync(JobId.Empty, Agent.SessionId, LogType.Json);
+
+					UpgradeTask Task = new UpgradeTask();
+					Task.SoftwareId = RequiredVersion.ToString();
+					Task.LogId = LogFile.Id.ToString();
+
+					byte[] Payload;
+					if (Agent.Version == "5.0.0-17425336" || Agent.Version == "5.0.0-17448746")
+					{
+						Any Any = new Any();
+						Any.TypeUrl = "type.googleapis.com/Horde.UpgradeTask";
+						Any.Value = Task.ToByteString();
+						Payload = Any.ToByteArray();
+					}
+					else
+					{
+						Payload = Any.Pack(Task).ToByteArray();
+					}
+
+					Lease = new AgentLease(LeaseId.GenerateNewId(), $"Upgrade to {RequiredVersion}", null, null, LogFile.Id, LeaseState.Pending, null, true, Payload);
+				}
+				return Lease;
+			}
+			return null;
+		}
+
+		/// <summary>
+		/// Determines the client software version that should be installed on an agent
+		/// </summary>
+		/// <param name="Agent">The agent instance</param>
+		/// <returns>Unique id of the client version this agent should be running</returns>
+		public async Task<string?> GetRequiredSoftwareVersion(IAgent Agent)
+		{
+			AgentSoftwareChannelName ChannelName = Agent.Channel ?? AgentSoftwareService.DefaultChannelName;
+			IAgentSoftwareChannel? Channel = await AgentSoftwareService.GetCachedChannelAsync(ChannelName);
+			return Channel?.Version;
+		}
+	}
+}
diff --git a/Engine/Source/Programs/Horde/Horde.Build/Agents/Impl/UpgradeService.cs b/Engine/Source/Programs/Horde/Horde.Build/Agents/Impl/UpgradeService.cs
new file mode 100644
--- /dev/null
+++ b/Engine/Source/Programs/Horde/Horde.Build/Agents/Impl/UpgradeService.cs
@@ -0,0 +1,110 @@
+// Copyright Epic Games, Inc. All Rights Reserved.
+
+using HordeServer.Collections;
+using HordeServer.Collections.Impl;
+using HordeServer.Models;
+using Microsoft.Extensions.DependencyInjection;
+using Microsoft.Extensions.Logging;
+using MongoDB.Driver;
+using Serilog.Events;
+using System;
+using System.Collections.Generic;
+using System.Linq;
+using System.Threading.Tasks;
+
+namespace HordeServer.Services
+{
+	/// <summary>
+	/// Service to perform schema upgrades on the database
+	/// </summary>
+	public class UpgradeService
+	{
+		/// <summary>
+		/// The current schema version
+		/// </summary>
+		public const int LatestSchemaVersion = 5;
+
+		/// <summary>
+		/// The database service instance
+		/// </summary>
+		DatabaseService DatabaseService { get; set; }
+
+		/// <summary>
+		/// The DI service provider
+		/// </summary>
+		IServiceProvider ServiceProvider;
+
+		/// <summary>
+		/// Logger instance
+		/// </summary>
+		ILogger<UpgradeService> Logger;
+
+		/// <summary>
+		/// Constructor
+		/// </summary>
+		/// <param name="DatabaseService">The database service singleton</param>
+		/// <param name="ServiceProvider">The DI service provider</param>
+		/// <param name="Logger">Logger instance</param>
+		public UpgradeService(DatabaseService DatabaseService, IServiceProvider ServiceProvider, ILogger<UpgradeService> Logger)
+		{
+			this.DatabaseService = DatabaseService;
+			this.ServiceProvider = ServiceProvider;
+			this.Logger = Logger;
+		}
+
+		/// <summary>
+		/// Update the database schema
+		/// </summary>
+		/// <param name="FromVersion">The version to upgrade from. Automatically detected if this is null.</param>
+		/// <returns>Async task</returns>
+		public async Task UpgradeSchemaAsync(int? FromVersion)
+		{
+			Globals Globals = await DatabaseService.GetGlobalsAsync();
+			int SchemaVersion = FromVersion ?? Globals.SchemaVersion ?? 0;
+
+			while (SchemaVersion < LatestSchemaVersion)
+			{
+				Logger.LogInformation("Upgrading from schema version {Version}", SchemaVersion);
+
+				// Handle the current version
+				//				if (SchemaVersion == 0)
+				//				{
+				//					IssueCollection IssueCollection = new IssueCollection(DatabaseService);
+				//					IJobCollection JobCollection = ServiceProvider.GetService<IJobCollection>();
+				//					await IssueCollection.UpdateSchemaLogIdFields(JobCollection, Logger);
+				//				}
+				//				if(SchemaVersion == 2)
+				//				{
+				//					IssueCollection IssueCollection = new IssueCollection(DatabaseService);
+				//					await IssueCollection.FixResolvedIssuesAsync(Logger);
+				//				}
+
+				if (SchemaVersion == 3)
+				{
+					IJobCollection JobCollection = ServiceProvider.GetRequiredService<IJobCollection>();
+					await JobCollection.UpgradeDocumentsAsync();
+				}
+				if (SchemaVersion == 4)
+				{
+					UserCollectionV1 UserCollectionV1 = new UserCollectionV1(ServiceProvider.GetRequiredService<DatabaseService>());
+					using UserCollectionV2 UserCollectionV2 = new UserCollectionV2(ServiceProvider.GetRequiredService<DatabaseService>(), ServiceProvider.GetRequiredService<ILogger<UserCollectionV2>>());
+					await UserCollectionV2.ResaveDocumentsAsync(UserCollectionV1);
+				}
+
+				// Increment the version number
+				SchemaVersion++;
+
+				// Try to update the current schema version number
+				while (Globals.SchemaVersion == null || Globals.SchemaVersion < SchemaVersion)
+				{
+					Globals.SchemaVersion = SchemaVersion;
+					if (await DatabaseService.TryUpdateSingletonAsync(Globals))
+					{
+						break;
+					}
+					Globals = await DatabaseService.GetGlobalsAsync();
+				}
+			}
+		}
+	}
+}
diff --git a/Engine/Source/Programs/Horde/Horde.Build/Api/Acls.cs b/Engine/Source/Programs/Horde/Horde.Build/Api/Acls.cs
new file mode 100644
--- /dev/null
+++ b/Engine/Source/Programs/Horde/Horde.Build/Api/Acls.cs
@@ -0,0 +1,164 @@
+// Copyright Epic Games, Inc. All Rights Reserved.
+
+using HordeServer.Models;
+using System;
+using System.Collections.Generic;
+using System.ComponentModel.DataAnnotations;
+using System.Linq;
+using System.Security.Claims;
+using System.Threading.Tasks;
+
+namespace HordeServer.Api
+{
+	/// <summary>
+	/// New claim to create
+	/// </summary>
+	public class CreateAclClaimRequest
+	{
+		/// <summary>
+		/// The claim type
+		/// </summary>
+		[Required]
+		public string Type { get; set; } = null!;
+
+		/// <summary>
+		/// The claim value
+		/// </summary>
+		[Required]
+		public string Value { get; set; } = null!;
+	}
+
+	/// <summary>
+	/// Individual entry in an ACL
+	/// </summary>
+	public class CreateAclEntryRequest
+	{
+		/// <summary>
+		/// Name of the user or group
+		/// </summary>
+		[Required]
+		public CreateAclClaimRequest Claim { get; set; } = null!;
+
+		/// <summary>
+		/// Array of actions to allow
+		/// </summary>
+		[Required]
+		public string[] Actions { get; set; } = null!;
+	}
+
+	/// <summary>
+	/// Parameters to update an ACL
+	/// </summary>
+	public class UpdateAclRequest
+	{
+		/// <summary>
+		/// Entries to replace the existing ACL
+		/// </summary>
+		public List<CreateAclEntryRequest>? Entries { get; set; }
+
+		/// <summary>
+		/// Whether to inherit permissions from the parent ACL
+		/// </summary>
+		public bool? Inherit { get; set; }
+
+		/// <summary>
+		/// List of exceptions to the inherited setting
+		/// </summary>
+		public List<string>? Exceptions { get; set; }
+	}
+
+	/// <summary>
+	/// New claim to update
+	/// </summary>
+	public class GetAclClaimResponse
+	{
+		/// <summary>
+		/// The claim type
+		/// </summary>
+		public string Type { get; set; }
+
+		/// <summary>
+		/// The claim value
+		/// </summary>
+		public string Value { get; set; }
+
+		/// <summary>
+		/// Constructor
+		/// </summary>
+		/// <param name="Claim">The claim to construct from</param>
+		public GetAclClaimResponse(AclClaim Claim)
+			: this(Claim.Type, Claim.Value)
+		{
+		}
+
+		/// <summary>
+		/// Constructor
+		/// </summary>
+		/// <param name="Type">The claim type</param>
+		/// <param name="Value">The claim value</param>
+		public GetAclClaimResponse(string Type, string Value)
+		{
+			this.Type = Type;
+			this.Value = Value;
+		}
+	}
+
+	/// <summary>
+	/// Individual entry in an ACL
+	/// </summary>
+	public class GetAclEntryResponse
+	{
+		/// <summary>
+		/// Names of the user or group
+		/// </summary>
+		public GetAclClaimResponse Claim { get; set; }
+
+		/// <summary>
+		/// Array of actions to allow
+		/// </summary>
+		public List<string> Actions { get; set; }
+
+		/// <summary>
+		/// Constructor
+		/// </summary>
+		/// <param name="AclEntry">The acl entry to construct from</param>
+		public GetAclEntryResponse(AclEntry AclEntry)
+		{
+			this.Claim = new GetAclClaimResponse(AclEntry.Claim.Type, AclEntry.Claim.Value);
+			this.Actions = AclEntry.GetActionNames(AclEntry.Actions);
+		}
+	}
+
+	/// <summary>
+	/// Information about an ACL
+	/// </summary>
+	public class GetAclResponse
+	{
+		/// <summary>
+		/// Entries to replace the existing ACL
+		/// </summary>
+		[Required]
+		public List<GetAclEntryResponse> Entries { get; set; }
+
+		/// <summary>
+		/// Whether to inherit permissions from the parent entity
+		/// </summary>
+		public bool Inherit { get; set; }
+
+		/// <summary>
+		/// Exceptions from permission inheritance setting
+		/// </summary>
+		public List<string>? Exceptions { get; set; }
+
+		/// <summary>
+		/// Constructor
+		/// </summary>
+		/// <param name="Acl">The ACL to construct from</param>
+		public GetAclResponse(Acl Acl)
+		{
+			this.Entries = Acl.Entries.ConvertAll(x => new GetAclEntryResponse(x));
+			this.Inherit = Acl.Inherit;
+			this.Exceptions = Acl.Exceptions?.ConvertAll(x => x.ToString());
+		}
+	}
+}
diff --git a/Engine/Source/Programs/Horde/Horde.Build/Api/AgentSoftwareChannels.cs b/Engine/Source/Programs/Horde/Horde.Build/Api/AgentSoftwareChannels.cs
new file mode 100644
--- /dev/null
+++ b/Engine/Source/Programs/Horde/Horde.Build/Api/AgentSoftwareChannels.cs
@@ -0,0 +1,48 @@
+// Copyright Epic Games, Inc. All Rights Reserved.
+
+using HordeServer.Models;
+using System;
+using System.Collections.Generic;
+using System.Linq;
+using System.Threading.Tasks;
+
+namespace HordeServer.Api
+{
+	/// <summary>
+	/// Information about an agent software channel
+	/// </summary>
+	public class GetAgentSoftwareChannelResponse
+	{
+		/// <summary>
+		/// Name of the channel
+		/// </summary>
+		public string? Name { get; set; }
+
+		/// <summary>
+		/// Name of the user that last modified this channel
+		/// </summary>
+		public string? ModifiedBy { get; set; }
+
+		/// <summary>
+		/// The modified timestamp
+		/// </summary>
+		public DateTime ModifiedTime { get; set; }
+
+		/// <summary>
+		/// Version number of this software
+		/// </summary>
+		public string? Version { get; set; }
+
+		/// <summary>
+		/// Constructor
+		/// </summary>
+		/// <param name="Channel">The channel information</param>
+		public GetAgentSoftwareChannelResponse(IAgentSoftwareChannel Channel)
+		{
+			this.Name = Channel.Name.ToString();
+			this.ModifiedBy = Channel.ModifiedBy;
+			this.ModifiedTime = Channel.ModifiedTime;
+			this.Version = Channel.Version.ToString();
+		}
+	}
+}
diff --git a/Engine/Source/Programs/Horde/Horde.Build/Api/Agents.cs b/Engine/Source/Programs/Horde/Horde.Build/Api/Agents.cs
new file mode 100644
--- /dev/null
+++ b/Engine/Source/Programs/Horde/Horde.Build/Api/Agents.cs
@@ -0,0 +1,467 @@
+// Copyright Epic Games, Inc. All Rights Reserved.
+
+using EpicGames.Core;
+using HordeCommon;
+using HordeServer.Models;
+using HordeServer.Utilities;
+using System;
+using System.Collections.Generic;
+using System.ComponentModel.DataAnnotations;
+using System.Linq;
+using System.Threading.Tasks;
+
+namespace HordeServer.Api
+{
+	using PoolId = StringId<IPool>;
+
+	/// <summary>
+	/// Parameters to register a new agent
+	/// </summary>
+	public class CreateAgentRequest
+	{
+		/// <summary>
+		/// Friendly name for the agent
+		/// </summary>
+		[Required]
+		public string Name { get; set; } = null!; // Enforced by [required] attribute
+
+		/// <summary>
+		/// Whether the agent is currently enabled
+		/// </summary>
+		public bool Enabled { get; set; } = true;
+
+		/// <summary>
+		/// Whether the agent is ephemeral (ie. should not be shown when inactive)
+		/// </summary>
+		public bool Ephemeral { get; set; }
+
+		/// <summary>
+		/// Per-agent override for the desired agent software channel
+		/// </summary>
+		public string? Channel { get; set; }
+
+		/// <summary>
+		/// Pools for this agent
+		/// </summary>
+		public List<string>? Pools { get; set; }
+	}
+
+	/// <summary>
+	/// Response from creating an agent
+	/// </summary>
+	public class CreateAgentResponse
+	{
+		/// <summary>
+		/// Unique id for the new agent
+		/// </summary>
+		public string Id { get; set; }
+
+		/// <summary>
+		/// Constructor
+		/// </summary>
+		/// <param name="Id">Unique id for this agent</param>
+		public CreateAgentResponse(string Id)
+		{
+			this.Id = Id;
+		}
+	}
+
+	/// <summary>
+	/// Parameters to update an agent
+	/// </summary>
+	public class UpdateAgentRequest
+	{
+		/// <summary>
+		/// Whether the agent is currently enabled
+		/// </summary>
+		public bool? Enabled { get; set; }
+
+		/// <summary>
+		/// Whether the agent is ephemeral
+		/// </summary>
+		public bool? Ephemeral { get; set; }
+
+		/// <summary>
+		/// Request a conform be performed using the current agent set
+		/// </summary>
+		public bool? RequestConform { get; set; }
+
+		/// <summary>
+		/// Request the machine be restarted
+		/// </summary>
+		public bool? RequestRestart { get; set; }
+
+		/// <summary>
+		/// Request the machine be shut down
+		/// </summary>
+		public bool? RequestShutdown { get; set; }
+
+		/// <summary>
+		/// Per-agent override for the desired agent software channel
+		/// </summary>
+		public string? Channel { get; set; }
+
+		/// <summary>
+		/// Pools for this agent
+		/// </summary>
+		public List<string>? Pools { get; set; }
+
+		/// <summary>
+		/// New ACL for this agent
+		/// </summary>
+		public UpdateAclRequest? Acl { get; set; }
+		
+		/// <summary>
+		/// New comment
+		/// </summary>
+		public string? Comment { get; set; }
+	}
+
+	/// <summary>
+	/// Response for queries to find a particular lease within an agent
+	/// </summary>
+	public class GetAgentLeaseResponse
+	{
+		/// <summary>
+		/// Identifier for the lease
+		/// </summary>
+		public string Id { get; set; }
+
+		/// <summary>
+		/// The agent id
+		/// </summary>
+		public string? AgentId { get; set; }
+
+		/// <summary>
+		/// Name of the lease
+		/// </summary>
+		public string Name { get; set; }
+
+		/// <summary>
+		/// Log id for this lease
+		/// </summary>
+		public string? LogId { get; set; }
+
+		/// <summary>
+		/// Time at which the lease started (UTC)
+		/// </summary>
+		public DateTime StartTime { get; set; }
+
+		/// <summary>
+		/// Time at which the lease started (UTC)
+		/// </summary>
+		public DateTime? FinishTime { get; set; }
+
+		/// <summary>
+		/// Whether this lease has started executing on the agent yet
+		/// </summary>
+		public bool Executing { get; set; }
+
+		/// <summary>
+		/// Details of the payload being executed
+		/// </summary>
+		public Dictionary<string, string>? Details { get; set; }
+
+		/// <summary>
+		/// Outcome of the lease
+		/// </summary>
+		public LeaseOutcome? Outcome { get; set; }
+
+		/// <summary>
+		/// State of the lease (for AgentLeases)
+		/// </summary>
+		public LeaseState? State { get; set; }
+
+		/// <summary>
+		/// Constructor
+		/// </summary>
+		/// <param name="Lease">The lease to initialize from</param>
+		public GetAgentLeaseResponse(AgentLease Lease)
+		{
+			this.Id = Lease.Id.ToString();
+			this.Name = Lease.Name;
+			this.LogId = Lease.LogId?.ToString();
+			this.State = Lease.State;
+			this.StartTime = Lease.StartTime;
+			this.Executing = Lease.Active;
+			this.FinishTime = Lease.ExpiryTime;
+			this.Details = AgentLease.GetPayloadDetails(Lease.Payload);			
+		}
+
+		/// <summary>
+		/// Converts this lease to a response object
+		/// </summary>
+		/// <param name="Lease">The lease to initialize from</param>
+		public GetAgentLeaseResponse(ILease Lease)
+		{
+			this.Id = Lease.Id.ToString();
+			this.AgentId = Lease.AgentId.ToString();
+			this.Name = Lease.Name;
+			this.LogId = Lease.LogId?.ToString();
+			this.StartTime = Lease.StartTime;
+			this.Executing = (Lease.FinishTime == null);
+			this.FinishTime = Lease.FinishTime;
+			this.Details = AgentLease.GetPayloadDetails(Lease.Payload);
+			this.Outcome = Lease.Outcome;
+		}
+	}
+
+	/// <summary>
+	/// Information about an agent session
+	/// </summary>
+	public class GetAgentSessionResponse
+	{
+		/// <summary>
+		/// Unique id for this session
+		/// </summary>
+		[Required]
+		public string Id { get; set; }
+
+		/// <summary>
+		/// Start time for this session
+		/// </summary>
+		[Required]
+		public DateTime StartTime { get; set; }
+
+		/// <summary>
+		/// Finishing time for this session
+		/// </summary>
+		public DateTime? FinishTime { get; set; }
+
+		/// <summary>
+		/// Properties of this agent
+		/// </summary>
+		public List<string>? Properties { get; set; }
+
+		/// <summary>
+		/// Version of the software running during this session
+		/// </summary>
+		public string? Version { get; set; }
+
+		/// <summary>
+		/// Constructor
+		/// </summary>
+		/// <param name="Session">The session to construct from</param>
+		public GetAgentSessionResponse(ISession Session)
+		{
+			this.Id = Session.Id.ToString();
+			this.StartTime = Session.StartTime;
+			this.FinishTime = Session.FinishTime;
+			this.Properties = (Session.Properties != null) ? new List<string>(Session.Properties) : null;
+			this.Version = Session.Version?.ToString();
+		}
+	}
+
+	/// <summary>
+	/// Information about a workspace synced on an agent
+	/// </summary>
+	public class GetAgentWorkspaceResponse
+	{
+		/// <summary>
+		/// The Perforce server and port to connect to
+		/// </summary>
+		public string? Cluster { get; set; }
+
+		/// <summary>
+		/// User to log into Perforce with (eg. buildmachine)
+		/// </summary>
+		public string? UserName { get; set; }
+
+		/// <summary>
+		/// Identifier to distinguish this workspace from other workspaces
+		/// </summary>
+		public string Identifier { get; set; }
+
+		/// <summary>
+		/// The stream to sync
+		/// </summary>
+		public string Stream { get; set; }
+
+		/// <summary>
+		/// Custom view for the workspace
+		/// </summary>
+		public List<string>? View { get; set; }
+
+		/// <summary>
+		/// Whether to use an incremental workspace
+		/// </summary>
+		public bool bIncremental { get; set; }
+
+		/// <summary>
+		/// Constructor
+		/// </summary>
+		/// <param name="Workspace">The workspace to construct from</param>
+		public GetAgentWorkspaceResponse(AgentWorkspace Workspace)
+		{
+			this.Cluster = Workspace.Cluster;
+			this.UserName = Workspace.UserName;
+			this.Identifier = Workspace.Identifier;
+			this.Stream = Workspace.Stream;
+			this.View = Workspace.View;
+			this.bIncremental = Workspace.bIncremental;
+		}
+	}
+
+	/// <summary>
+	/// Information about an agent
+	/// </summary>
+	public class GetAgentResponse
+	{
+		/// <summary>
+		/// The agent's unique ID
+		/// </summary>
+		public string Id { get; set; }
+
+		/// <summary>
+		/// Friendly name of the agent
+		/// </summary>
+		public string Name { get; set; }
+
+		/// <summary>
+		/// Whether the agent is currently enabled
+		/// </summary>
+		public bool Enabled { get; set; }
+
+		/// <summary>
+		/// The current session id
+		/// </summary>
+		public string? SessionId { get; set; }
+
+		/// <summary>
+		/// Whether the agent is ephemeral
+		/// </summary>
+		public bool Ephemeral { get; set; }
+
+		/// <summary>
+		/// Whether the agent is currently online
+		/// </summary>
+		public bool Online { get; set; }
+
+		/// <summary>
+		/// Whether this agent has expired
+		/// </summary>
+		public bool Deleted { get; set; }
+
+		/// <summary>
+		/// Whether a conform job is pending
+		/// </summary>
+		public bool PendingConform { get; set; }
+
+		/// <summary>
+		/// Whether a restart is pending
+		/// </summary>
+		public bool PendingRestart { get; set; }
+		
+		/// <summary>
+		/// Whether a shutdown is pending
+		/// </summary>
+		public bool PendingShutdown { get; set; }
+
+		/// <summary>
+		/// Last time a conform was attempted
+		/// </summary>
+		public DateTime LastConformTime { get; set; }
+
+		/// <summary>
+		/// Number of times a conform has been attempted
+		/// </summary>
+		public int? ConformAttemptCount { get; set; }
+
+		/// <summary>
+		/// Last time a conform was attempted
+		/// </summary>
+		public DateTime? NextConformTime { get; set; }
+
+		/// <summary>
+		/// The current client version
+		/// </summary>
+		public string? Version { get; set; }
+
+		/// <summary>
+		/// Per-agent override for the desired agent software channel
+		/// </summary>
+		public string? Channel { get; set; }
+
+		/// <summary>
+		/// Properties for the agent
+		/// </summary>
+		public List<string> Properties { get; set; }
+
+		/// <summary>
+		/// Resources for the agent
+		/// </summary>
+		public Dictionary<string, int> Resources { get; set; }
+
+		/// <summary>
+		/// Last update time of this agent
+		/// </summary>
+		public DateTime? UpdateTime { get; set; }
+
+		/// <summary>
+		/// Pools for this agent
+		/// </summary>
+		public List<string>? Pools { get; set; }
+
+		/// <summary>
+		/// List of workspaces currently synced to this machine
+		/// </summary>
+		public List<GetAgentWorkspaceResponse> Workspaces { get; set; } = new List<GetAgentWorkspaceResponse>();
+
+		/// <summary>
+		/// Capabilities of this agent
+		/// </summary>
+		public object? Capabilities { get; }
+
+		/// <summary>
+		/// Array of active leases.
+		/// </summary>
+		public List<GetAgentLeaseResponse> Leases { get; }
+
+		/// <summary>
+		/// Per-object permissions
+		/// </summary>
+		public GetAclResponse? Acl { get; set; }
+		
+		/// <summary>
+		/// Comment for this agent
+		/// </summary>
+		public string? Comment { get; set; }
+
+		/// <summary>
+		/// Constructor
+		/// </summary>
+		/// <param name="Agent">The agent to construct from</param>
+		/// <param name="bIncludeAcl">Whether to include the ACL in the response</param>
+		public GetAgentResponse(IAgent Agent, bool bIncludeAcl)
+		{
+			this.Id = Agent.Id.ToString();
+			this.Name = Agent.Id.ToString();
+			this.Enabled = Agent.Enabled;
+			this.Properties = new List<string>(Agent.Properties);
+			this.Resources = new Dictionary<string, int>(Agent.Resources);
+			this.SessionId = Agent.SessionId?.ToString();
+			this.Online = Agent.IsSessionValid(DateTime.UtcNow);
+			this.Deleted = Agent.Deleted;
+			this.PendingConform = Agent.RequestConform;
+			this.PendingRestart = Agent.RequestRestart;
+			this.PendingShutdown = Agent.RequestShutdown;
+			this.LastConformTime = Agent.LastConformTime;
+			this.NextConformTime =   Agent.LastConformTime;
+			this.ConformAttemptCount = Agent.ConformAttemptCount;
+			this.Version = Agent.Version?.ToString() ?? "Unknown";
+			if(Agent.Channel != null)
+			{
+				this.Version += $" ({Agent.Channel})";
+			}
+			this.Version = Agent.Version?.ToString();
+			this.Channel = Agent.Channel?.ToString();
+			this.UpdateTime = Agent.UpdateTime;
+			this.Pools = Agent.GetPools().Select(x => x.ToString()).ToList();
+			this.Workspaces = Agent.Workspaces.ConvertAll(x => new GetAgentWorkspaceResponse(x));
+			this.Capabilities = new { Devices = new[] { new { Properties = Agent.Properties, Resources = Agent.Resources } } };
+			this.Leases = Agent.Leases.ConvertAll(x => new GetAgentLeaseResponse(x));
+			this.Acl = (bIncludeAcl && Agent.Acl != null) ? new GetAclResponse(Agent.Acl) : null;
+			this.Comment = Agent.Comment;
+		}
+	}
+}
diff --git a/Engine/Source/Programs/Horde/Horde.Build/Api/Artifacts.cs b/Engine/Source/Programs/Horde/Horde.Build/Api/Artifacts.cs
new file mode 100644
--- /dev/null
+++ b/Engine/Source/Programs/Horde/Horde.Build/Api/Artifacts.cs
@@ -0,0 +1,119 @@
+// Copyright Epic Games, Inc. All Rights Reserved.
+
+using HordeServer.Models;
+using System.Collections.Generic;
+using System.ComponentModel.DataAnnotations;
+
+namespace HordeServer.Api
+{
+	/// <summary>
+	/// Response from creating an artifact
+	/// </summary>
+	public class CreateArtifactResponse
+	{
+		/// <summary>
+		/// Unique id for this artifact
+		/// </summary>
+		public string Id { get; set; }
+
+		/// <summary>
+		/// Constructor
+		/// </summary>
+		/// <param name="Id">The artifact file id</param>
+		public CreateArtifactResponse(string Id)
+		{
+			this.Id = Id;
+		}
+	}
+
+	/// <summary>
+	/// Request to get a zip of artifacts
+	/// </summary>
+	public class GetArtifactZipRequest
+	{
+		/// <summary>
+		/// Job id to get all artifacts for
+		/// </summary>
+		public string? JobId { get; set; }
+
+		/// <summary>
+		/// Step id to filter by
+		/// </summary>
+		public string? StepId { get; set; }
+
+		/// <summary>
+		/// Further filter by a list of artifact ids
+		/// </summary>
+		public List<string>? ArtifactIds { get; set; }
+
+		/// <summary>
+		/// Constructor
+		/// </summary>
+		/// <param name="JobId">Job id to get all artifacts for</param>
+		/// <param name="StepId">step to filter by</param>
+		/// <param name="ArtifactIds">The artifact ids.  Returns all artifacts for a job </param>
+		public GetArtifactZipRequest(string? JobId, string? StepId, List<string>? ArtifactIds)
+		{
+			this.JobId = JobId;
+			this.StepId = StepId;
+			this.ArtifactIds = ArtifactIds;
+		}
+	}
+
+	/// <summary>
+	/// Response describing an artifact
+	/// </summary>
+	public class GetArtifactResponse
+	{
+		/// <summary>
+		/// Unique id of the artifact
+		/// </summary>
+		public string Id { get; set; }
+
+		/// <summary>
+		/// Name of the artifact
+		/// </summary>
+		public string Name { get; set; }
+
+		/// <summary>
+		/// Unique id of the job for this artifact
+		/// </summary>
+		public string JobId { get; set; }
+
+		/// <summary>
+		/// Optional unique id of the step for this artifact
+		/// </summary>
+		public string? StepId { get; set; }
+
+		/// <summary>
+		/// Type of artifact
+		/// </summary>
+		public string MimeType { get; set; }
+
+		/// <summary>
+		/// Length of artifact
+		/// </summary>
+		public long Length { get; set; }
+
+		/// <summary>
+		/// Short-lived code that can be used to download this artifact through direct download links in the browser
+		/// </summary>
+		public string? Code { get; set; }
+
+		/// <summary>
+		/// Constructor
+		/// </summary>
+		/// <param name="Artifact">The artifact to construct from</param>
+		/// <param name="Code">The direct download code</param>
+		public GetArtifactResponse(IArtifact Artifact, string? Code)
+		{
+			this.Id = Artifact.Id.ToString();
+			this.Name = Artifact.Name;
+			this.JobId = Artifact.JobId.ToString();
+			this.StepId = Artifact.StepId.ToString();
+			this.MimeType = Artifact.MimeType;
+			this.Length = Artifact.Length;
+			this.Code = Code;
+		}
+	}
+}
diff --git a/Engine/Source/Programs/Horde/Horde.Build/Api/Changes.cs b/Engine/Source/Programs/Horde/Horde.Build/Api/Changes.cs
new file mode 100644
--- /dev/null
+++ b/Engine/Source/Programs/Horde/Horde.Build/Api/Changes.cs
@@ -0,0 +1,92 @@
+// Copyright Epic Games, Inc. All Rights Reserved.
+
+using HordeServer.Models;
+using System;
+using System.Collections.Generic;
+using System.Linq;
+using System.Threading.Tasks;
+
+namespace HordeServer.Api
+{
+	/// <summary>
+	/// Information about a submitted changelist
+	/// </summary>
+	public class GetChangeSummaryResponse
+	{
+		/// <summary>
+		/// The source changelist number
+		/// </summary>
+		public int Number { get; set; }
+
+		/// <summary>
+		/// Name of the user that authored this change [DEPRECATED]
+		/// </summary>
+		public string Author { get; set; }
+
+		/// <summary>
+		/// Information about the change author
+		/// </summary>
+		public GetThinUserInfoResponse AuthorInfo { get; set; }
+
+		/// <summary>
+		/// The description text
+		/// </summary>
+		public string Description { get; set; }
+
+		/// <summary>
+		/// Constructor
+		/// </summary>
+		/// <param name="ChangeSummary">The commit to construct from</param>
+		public GetChangeSummaryResponse(ChangeSummary ChangeSummary)
+		{
+			this.Number = ChangeSummary.Number;
+			this.Author = ChangeSummary.Author.Name;
+			this.AuthorInfo = new GetThinUserInfoResponse(ChangeSummary.Author);
+			this.Description = ChangeSummary.Description;
+		}
+	}
+
+	/// <summary>
+	/// Information about a submitted changelist
+	/// </summary>
+	public class GetChangeDetailsResponse
+	{
+		/// <summary>
+		/// The source changelist number
+		/// </summary>
+		public int Number { get; set; }
+
+		/// <summary>
+		/// Name of the user that authored this change [DEPRECATED]
+		/// </summary>
+		public string Author { get; set; }
+
+		/// <summary>
+		/// Information about the user that authored this change
+		/// </summary>
+		public GetThinUserInfoResponse AuthorInfo { get; set; }
+
+		/// <summary>
+		/// The description text
+		/// </summary>
+		public string Description { get; set; }
+
+		/// <summary>
+		/// List of files that were modified, relative to the stream base
+		/// </summary>
+		public List<string> Files { get; set; }
+
+		/// <summary>
+		/// Constructor
+		/// </summary>
+		/// <param name="ChangeDetails">The commit to construct from</param>
+		public GetChangeDetailsResponse(ChangeDetails ChangeDetails)
+		{
+			this.Number = ChangeDetails.Number;
+			this.Author = ChangeDetails.Author.Name;
+			this.AuthorInfo = new GetThinUserInfoResponse(ChangeDetails.Author);
+			this.Description = ChangeDetails.Description;
+			this.Files = ChangeDetails.Files.ConvertAll(x => x.Path);
+		}
+	}
+}
diff --git a/Engine/Source/Programs/Horde/Horde.Build/Api/Credentials.cs b/Engine/Source/Programs/Horde/Horde.Build/Api/Credentials.cs
new file mode 100644
--- /dev/null
+++ b/Engine/Source/Programs/Horde/Horde.Build/Api/Credentials.cs
@@ -0,0 +1,108 @@
+// Copyright Epic Games, Inc. All Rights Reserved.
+
+using HordeServer.Models;
+using System;
+using System.Collections.Generic;
+using System.ComponentModel.DataAnnotations;
+using System.Linq;
+using System.Threading.Tasks;
+
+namespace HordeServer.Api
+{
+	/// <summary>
+	/// Parameters to create a new credential
+	/// </summary>
+	public class CreateCredentialRequest
+	{
+		/// <summary>
+		/// Name for the new credential
+		/// </summary>
+		[Required]
+		public string Name { get; set; } = null!;
+
+		/// <summary>
+		/// Properties for the new credential
+		/// </summary>
+		public Dictionary<string, string>? Properties { get; set; }
+	}
+
+	/// <summary>
+	/// Response from creating a new credential
+	/// </summary>
+	public class CreateCredentialResponse
+	{
+		/// <summary>
+		/// Unique id for the new credential
+		/// </summary>
+		public string Id { get; set; }
+
+		/// <summary>
+		/// Constructor
+		/// </summary>
+		/// <param name="Id">Unique id for the new credential</param>
+		public CreateCredentialResponse(string Id)
+		{
+			this.Id = Id;
+		}
+	}
+
+	/// <summary>
+	/// Parameters to update a credential
+	/// </summary>
+	public class UpdateCredentialRequest
+	{
+		/// <summary>
+		/// Optional new name for the credential
+		/// </summary>
+		public string? Name { get; set; }
+
+		/// <summary>
+		/// Properties to update for the credential. Properties set to null will be removed.
+		/// </summary>
+		public Dictionary<string, string?>? Properties { get; set; }
+
+		/// <summary>
+		/// Custom permissions for this object
+		/// </summary>
+		public UpdateAclRequest? Acl { get; set; }
+	}
+
+	/// <summary>
+	/// Response describing a credential
+	/// </summary>
+	public class GetCredentialResponse
+	{
+		/// <summary>
+		/// Unique id of the credential
+		/// </summary>
+		public string Id { get; set; }
+
+		/// <summary>
+		/// Name of the credential
+		/// </summary>
+		public string Name { get; set; }
+
+		/// <summary>
+		/// Properties for the credential
+		/// </summary>
+		public Dictionary<string, string> Properties { get; set; }
+
+		/// <summary>
+		/// Custom permissions for this object
+		/// </summary>
+		public GetAclResponse? Acl { get; set; }
+
+		/// <summary>
+		/// Constructor
+		/// </summary>
+		/// <param name="Credential">The credential to construct from</param>
+		/// <param name="bIncludeAcl">Whether to include the ACL in the response</param>
+		public GetCredentialResponse(Credential Credential, bool bIncludeAcl)
+		{
+			this.Id = Credential.Id.ToString();
+			this.Name = Credential.Name;
+			this.Properties = Credential.Properties;
+			this.Acl = (bIncludeAcl && Credential.Acl != null)? new GetAclResponse(Credential.Acl) : null;
+		}
+	}
+}
diff --git a/Engine/Source/Programs/Horde/Horde.Build/Api/Events.cs b/Engine/Source/Programs/Horde/Horde.Build/Api/Events.cs
new file mode 100644
--- /dev/null
+++ b/Engine/Source/Programs/Horde/Horde.Build/Api/Events.cs
@@ -0,0 +1,189 @@
+// Copyright Epic Games, Inc. All Rights Reserved.
+
+using EpicGames.Core;
+using HordeCommon;
+using HordeServer.Models;
+using HordeServer.Utilities;
+using Microsoft.AspNetCore.Mvc;
+using System;
+using System.Collections.Generic;
+using System.ComponentModel.DataAnnotations;
+using System.Linq;
+using System.Threading.Tasks;
+
+namespace HordeServer.Api
+{
+	using StreamId = StringId<IStream>;
+	using TemplateRefId = StringId<TemplateRef>;
+
+	/// <summary>
+	/// Base class for event records
+	/// </summary>
+	[JsonKnownTypes(typeof(JobCompleteEventRecord), typeof(LabelCompleteEventRecord), typeof(StepCompleteEventRecord))]
+	public class EventRecord
+	{
+	}
+
+	/// <summary>
+	/// Event that occurs when a job completes
+	/// </summary>
+	[JsonDiscriminator("Job")]
+	public class JobCompleteEventRecord : EventRecord
+	{
+		/// <summary>
+		/// The stream id
+		/// </summary>
+		[Required]
+		public string StreamId { get; set; }
+
+		/// <summary>
+		/// The template id
+		/// </summary>
+		[Required]
+		public string TemplateId { get; set; }
+
+		/// <summary>
+		/// Outcome of the job
+		/// </summary>
+		public LabelOutcome Outcome { get; set; } = LabelOutcome.Success;
+
+		/// <summary>
+		/// Default constructor
+		/// </summary>
+		private JobCompleteEventRecord()
+		{
+			StreamId = String.Empty;
+			TemplateId = String.Empty;
+		}
+
+		/// <summary>
+		/// Constructor
+		/// </summary>
+		/// <param name="StreamId">The stream id</param>
+		/// <param name="TemplateId">The template id</param>
+		/// <param name="Outcome">Outcome of the job</param>
+		public JobCompleteEventRecord(StreamId StreamId, TemplateRefId TemplateId, LabelOutcome Outcome)
+		{
+			this.StreamId = StreamId.ToString();
+			this.TemplateId = TemplateId.ToString();
+			this.Outcome = Outcome;
+		}
+	}
+
+	/// <summary>
+	/// Event that occurs when a label completes
+	/// </summary>
+	[JsonDiscriminator("Label")]
+	public class LabelCompleteEventRecord : EventRecord
+	{
+		/// <summary>
+		/// The stream id
+		/// </summary>
+		[Required]
+		public string StreamId { get; set; }
+
+		/// <summary>
+		/// The template id
+		/// </summary>
+		[Required]
+		public string TemplateId { get; set; }
+
+		/// <summary>
+		/// Name of the category for this label
+		/// </summary>
+		public string? CategoryName { get; set; }
+
+		/// <summary>
+		/// Name of the label to monitor
+		/// </summary>
+		[Required]
+		public string LabelName { get; set; }
+
+		/// <summary>
+		/// Outcome of the job
+		/// </summary>
+		public LabelOutcome Outcome { get; set; } = LabelOutcome.Success;
+
+		/// <summary>
+		/// Default constructor
+		/// </summary>
+		private LabelCompleteEventRecord()
+		{
+			this.StreamId = String.Empty;
+			this.TemplateId = String.Empty;
+			this.LabelName = String.Empty;
+		}
+
+		/// <summary>
+		/// Constructor
+		/// </summary>
+		/// <param name="StreamId">The stream id</param>
+		/// <param name="TemplateId">The template id</param>
+		/// <param name="CategoryName">Name of the category</param>
+		/// <param name="LabelName">The label name</param>
+		/// <param name="Outcome">Outcome of the label</param>
+		public LabelCompleteEventRecord(StreamId StreamId, TemplateRefId TemplateId, string? CategoryName, string LabelName, LabelOutcome Outcome)
+		{
+			this.StreamId = StreamId.ToString();
+			this.TemplateId = TemplateId.ToString();
+			this.CategoryName = CategoryName;
+			this.LabelName = LabelName;
+			this.Outcome = Outcome;
+		}
+	}
+
+	/// <summary>
+	/// Event that occurs when a step completes
+	/// </summary>
+	[JsonDiscriminator("Step")]
+	public class StepCompleteEventRecord : EventRecord
+	{
+		/// <summary>
+		/// The stream id
+		/// </summary>
+		[Required]
+		public string StreamId { get; set; }
+
+		/// <summary>
+		/// The template id
+		/// </summary>
+		[Required]
+		public string TemplateId { get; set; }
+
+		/// <summary>
+		/// Name of the step to monitor
+		/// </summary>
+		[Required]
+		public string StepName { get; set; }
+
+		/// <summary>
+		/// Outcome of the job step
+		/// </summary>
+		public JobStepOutcome Outcome { get; set; } = JobStepOutcome.Success;
+
+		/// <summary>
+		/// Default constructor
+		/// </summary>
+		private StepCompleteEventRecord()
+		{
+			StreamId = String.Empty;
+			TemplateId = String.Empty;
+			StepName = String.Empty;
+		}
+
+		/// <summary>
+		/// Constructor
+		/// </summary>
+		/// <param name="StreamId">The stream id</param>
+		/// <param name="TemplateId">The template id</param>
+		/// <param name="StepName">The label name</param>
+		/// <param name="Outcome">Outcome of the step</param>
+		public StepCompleteEventRecord(StreamId StreamId, TemplateRefId TemplateId, string StepName, JobStepOutcome Outcome)
+		{
+			this.StreamId = StreamId.ToString();
+			this.TemplateId = TemplateId.ToString();
+			this.StepName = StepName;
+			this.Outcome = Outcome;
+		}
+	}
+}
diff --git a/Engine/Source/Programs/Horde/Horde.Build/Api/HelixSwarm.cs b/Engine/Source/Programs/Horde/Horde.Build/Api/HelixSwarm.cs
new file mode 100644
--- /dev/null
+++ b/Engine/Source/Programs/Horde/Horde.Build/Api/HelixSwarm.cs
@@ -0,0 +1,94 @@
+﻿// Copyright Epic Games, Inc. All Rights Reserved.
+
+using System.ComponentModel.DataAnnotations;
+using System.Diagnostics.CodeAnalysis;
+
+namespace HordeServer.Api
+{
+	/// <summary>
+	/// Request sent by the Helix Swarm project
+	/// The format of this request can be modified through the settings web UI in Swarm.
+	/// Please ensure it matches the request class below.
+	/// </summary>
+	[SuppressMessage("Compiler", "CA1056:URI properties should not be strings")]
+	public class SwarmAutoTestRequest
+	{
+		/// <summary>
+		/// Name of the test
+		/// </summary>
+		public string? Test { get; set; } = null;
+		
+		/// <summary>
+		/// ID of the test run
+		/// </summary>
+		public string? TestRunId { get; set; } = null;
+		
+		/// <summary>
+		/// Perforce changelist number to run the test on
+		/// </summary>
+		[Required]
+		public string? Changelist { get; set; } = null;
+		
+		/// <summary>
+		/// Status of the changelist, for example "submitted"
+		/// </summary>
+		public string? Status { get; set; } = null;
+		
+		/// <summary>
+		/// Version of the review
+		/// </summary>
+		public string? Version { get; set; } = null;
+		
+		/// <summary>
+		/// Description of the review, usually the changelist message itself
+		/// </summary>
+		public string? Description { get; set; } = null;
+		
+		/// <summary>
+		/// Swarm project ID
+		/// </summary>
+		public string? Project { get; set; } = null;
+		
+		/// <summary>
+		/// Swarm project name
+		/// </summary>
+		public string? ProjectName { get; set; } = null;
+		
+		/// <summary>
+		/// Perforce branch ID in Swarm
+		/// </summary>
+		public string? Branch { get; set; } = null;
+		
+		/// <summary>
+		/// Perforce branch name in Swarm
+		/// </summary>
+		public string? BranchName { get; set; } = null;
+		
+		/// <summary>
+		/// Callback URL to Swarm server to be called when "auto test" is finished.
+		/// In our case, that would be when the job is completed.
+		/// </summary>
+		[Required]
+		public string? UpdateUrl { get; set; } = null;
+	}
+
+	/// <summary>
+	/// Response object describing the created job
+	/// </summary>
+	public class SwarmAutoTestResponse
+	{
+		/// <summary>
+		/// The ID of the job that got started
+		/// </summary>
+		public string Id { get; set; }
+
+		/// <summary>
+		/// Constructor
+		/// </summary>
+		/// <param name="Id">Id of the new document</param>
+		public SwarmAutoTestResponse(string Id)
+		{
+			this.Id = Id;
+		}
+	}
+}
