Horde: Rename the HordeServer project to Horde.Build.

diff --git a/Engine/Source/Programs/Horde/Horde.Build/Controllers/AgentsController.cs b/Engine/Source/Programs/Horde/Horde.Build/Controllers/AgentsController.cs
new file mode 100644
--- /dev/null
+++ b/Engine/Source/Programs/Horde/Horde.Build/Controllers/AgentsController.cs
@@ -0,0 +1,362 @@
+// Copyright Epic Games, Inc. All Rights Reserved.
+
+using HordeServer.Api;
+using HordeServer.Collections;
+using HordeServer.Models;
+using HordeServer.Services;
+using HordeServer.Utilities;
+using Microsoft.AspNetCore.Authorization;
+using Microsoft.AspNetCore.Mvc;
+using Microsoft.Extensions.Logging;
+using MongoDB.Bson;
+using System;
+using System.Collections.Generic;
+using System.Linq;
+using System.Text;
+using System.Threading.Tasks;
+
+namespace HordeServer.Controllers
+{
+	using AgentSoftwareChannelName = StringId<AgentSoftwareChannels>;
+	using LeaseId = ObjectId<ILease>;
+	using PoolId = StringId<IPool>;
+
+	/// <summary>
+	/// Controller for the /api/v1/agents endpoint
+	/// </summary>
+	[ApiController]
+	[Authorize]
+	[Route("[controller]")]
+	public class AgentsController : ControllerBase
+	{
+		/// <summary>
+		/// Singleton instance of the ACL service
+		/// </summary>
+		AclService AclService;
+
+		/// <summary>
+		/// Singleton instance of the agent service
+		/// </summary>
+		AgentService AgentService;
+
+		/// <summary>
+		/// The pool service
+		/// </summary>
+		PoolService PoolService;
+
+		/// <summary>
+		/// Constructor
+		/// </summary>
+		/// <param name="AclService">The ACL service singleton</param>
+		/// <param name="AgentService">The agent service</param>
+		/// <param name="PoolService">The pool service</param>
+		public AgentsController(AclService AclService, AgentService AgentService, PoolService PoolService)
+		{
+			this.AclService = AclService;
+			this.AgentService = AgentService;
+			this.PoolService = PoolService;
+		}
+
+		/// <summary>
+		/// Register an agent to perform remote work.
+		/// </summary>
+		/// <param name="Request">Request parameters</param>
+		/// <returns>Information about the registered agent</returns>
+		[HttpPost]
+		[Route("/api/v1/agents")]
+		public async Task<ActionResult<CreateAgentResponse>> CreateAgentAsync([FromBody] CreateAgentRequest Request)
+		{
+			if(!await AclService.AuthorizeAsync(AclAction.CreateAgent, User))
+			{
+				return Forbid();
+			}
+
+			AgentSoftwareChannelName? Channel = String.IsNullOrEmpty(Request.Channel) ? (AgentSoftwareChannelName?)null : new AgentSoftwareChannelName(Request.Channel);
+			IAgent Agent = await AgentService.CreateAgentAsync(Request.Name, Request.Enabled, Channel, Request.Pools?.ConvertAll(x => new PoolId(x)));
+
+			return new CreateAgentResponse(Agent.Id.ToString());
+		}
+
+		/// <summary>
+		/// Finds the agents matching specified criteria.
+		/// </summary>
+		/// <param name="Pool">The pool containing the agent</param>
+		/// <param name="Index">First result to return</param>
+		/// <param name="Count">Number of results to return</param>
+		/// <param name="ModifiedAfter">If set, only returns agents modified after this time</param>
+		/// <param name="Filter">Filter for the properties to return</param>
+		/// <returns>List of matching agents</returns>
+		[HttpGet]
+		[Route("/api/v1/agents")]
+		[ProducesResponseType(typeof(List<GetAgentResponse>), 200)]
+		public async Task<ActionResult<List<object>>> FindAgentsAsync([FromQuery] string? Pool = null, [FromQuery] int? Index = null, [FromQuery] int? Count = null, [FromQuery] DateTimeOffset? ModifiedAfter = null, [FromQuery] PropertyFilter? Filter = null)
+		{
+			GlobalPermissionsCache PermissionsCache = new GlobalPermissionsCache();
+			List<IAgent> Agents = await AgentService.FindAgentsAsync(Pool?.ToObjectId(), ModifiedAfter?.UtcDateTime, Index, Count);
+
+			List<object> Responses = new List<object>();
+			foreach (IAgent Agent in Agents)
+			{
+				bool bIncludeAcl = await AgentService.AuthorizeAsync(Agent, AclAction.ViewPermissions, User, PermissionsCache);
+				Responses.Add(new GetAgentResponse(Agent, bIncludeAcl).ApplyFilter(Filter));
+			}
+
+			return Responses;
+		}
+
+		/// <summary>
+		/// Retrieve information about a specific agent
+		/// </summary>
+		/// <param name="AgentId">Id of the agent to get information about</param>
+		/// <param name="Filter">Filter for the properties to return</param>
+		/// <returns>Information about the requested agent</returns>
+		[HttpGet]
+		[Route("/api/v1/agents/{AgentId}")]
+		[ProducesResponseType(typeof(GetAgentResponse), 200)]
+		public async Task<ActionResult<object>> GetAgentAsync(string AgentId, [FromQuery] PropertyFilter? Filter = null)
+		{
+			IAgent? Agent = await AgentService.GetAgentAsync(new AgentId(AgentId));
+			if (Agent == null)
+			{
+				return NotFound();
+			}
+
+			GlobalPermissionsCache Cache = new GlobalPermissionsCache();
+			bool bIncludeAcl = await AgentService.AuthorizeAsync(Agent, AclAction.ViewPermissions, User, Cache);
+			return new GetAgentResponse(Agent, bIncludeAcl).ApplyFilter(Filter);
+		}
+
+		/// <summary>
+		/// Update an agent's properties.
+		/// </summary>
+		/// <param name="AgentId">Id of the agent to update.</param>
+		/// <param name="Update">Properties on the agent to update.</param>
+		/// <returns>Http result code</returns>
+		[HttpPut]
+		[Route("/api/v1/agents/{AgentId}")]
+		public async Task<ActionResult> UpdateAgentAsync(string AgentId, [FromBody] UpdateAgentRequest Update)
+		{
+			List<PoolId>? UpdatePools = Update.Pools?.ConvertAll(x => new PoolId(x));
+			AgentSoftwareChannelName? Channel = String.IsNullOrEmpty(Update.Channel) ? (AgentSoftwareChannelName?)null : new AgentSoftwareChannelName(Update.Channel);
+
+			string UserName = User.GetUserName() ?? "Unknown";
+
+			GlobalPermissionsCache Cache = new GlobalPermissionsCache();
+			for (; ; )
+			{
+				IAgent? Agent = await AgentService.GetAgentAsync(new AgentId(AgentId));
+				if (Agent == null)
+				{
+					return NotFound();
+				}
+				if (!await AgentService.AuthorizeAsync(Agent, AclAction.UpdateAgent, User, Cache))
+				{
+					return Forbid();
+				}
+				if (Update.Acl != null && !await AgentService.AuthorizeAsync(Agent, AclAction.ChangePermissions, User, Cache))
+				{
+					return Forbid();
+				}
+
+				IAgent? NewAgent = await AgentService.Agents.TryUpdateSettingsAsync(Agent, Update.Enabled, Update.RequestConform, Update.RequestRestart, Update.RequestShutdown, Channel, Update.Pools?.ConvertAll(x => new PoolId(x)), Acl.Merge(Agent.Acl, Update.Acl), Update.Comment);
+				if (NewAgent == null)
+				{
+					continue;
+				}
+
+				IAuditLogChannel<AgentId> Logger = AgentService.Agents.GetLogger(Agent.Id);
+				if (Agent.Enabled != NewAgent.Enabled)
+				{
+					Logger.LogInformation("Setting changed: Enabled = {State} ({UserName})", NewAgent.Enabled, UserName);
+				}
+				if (Agent.RequestConform != NewAgent.RequestConform)
+				{
+					Logger.LogInformation("Setting changed: RequestConfig = {State} ({UserName})", NewAgent.RequestConform, UserName);
+				}
+				if (Agent.RequestRestart != NewAgent.RequestRestart)
+				{
+					Logger.LogInformation("Setting changed: RequestRestart = {State} ({UserName})", NewAgent.RequestRestart, UserName);
+				}
+				if (Agent.RequestShutdown != NewAgent.RequestShutdown)
+				{
+					Logger.LogInformation("Setting changed: RequestShutdown = {State} ({UserName})", NewAgent.RequestShutdown, UserName);
+				}
+				if (Agent.Comment != NewAgent.Comment)
+				{
+					Logger.LogInformation("Setting changed: Comment = \"{Comment}\" ({UserName})", Update.Comment, UserName);
+				}
+				foreach (PoolId AddedPool in NewAgent.ExplicitPools.Except(Agent.ExplicitPools))
+				{
+					Logger.LogInformation("Added to pool {PoolId} ({UserName})", AddedPool, UserName);
+				}
+				foreach (PoolId RemovedPool in Agent.ExplicitPools.Except(NewAgent.ExplicitPools))
+				{
+					Logger.LogInformation("Removed from pool {PoolId} ({UserName})", RemovedPool, UserName);
+				}
+				break;
+			}
+			return Ok();
+		}
+
+		/// <summary>
+		/// Remove a registered agent.
+		/// </summary>
+		/// <param name="AgentId">Id of the agent to delete.</param>
+		/// <returns>Http result code</returns>
+		[HttpDelete]
+		[Route("/api/v1/agents/{AgentId}")]
+		public async Task<ActionResult> DeleteAgentAsync(string AgentId)
+		{
+			IAgent? Agent = await AgentService.GetAgentAsync(new AgentId(AgentId));
+			if (Agent == null)
+			{
+				return NotFound();
+			}
+			if (!await AgentService.AuthorizeAsync(Agent, AclAction.DeleteAgent, User, null))
+			{
+				return Forbid();
+			}
+
+			await AgentService.DeleteAgentAsync(Agent);
+			return new OkResult();
+		}
+
+		/// <summary>
+		/// Retrieve historical information about a specific agent
+		/// </summary>
+		/// <param name="AgentId">Id of the agent to get information about</param>
+		/// <param name="MinTime">Minimum time for records to return</param>
+		/// <param name="MaxTime">Maximum time for records to return</param>
+		/// <param name="Index">Offset of the first result</param>
+		/// <param name="Count">Number of records to return</param>
+		/// <returns>Information about the requested agent</returns>
+		[HttpGet]
+		[Route("/api/v1/agents/{AgentId}/history")]
+		public async Task GetAgentHistoryAsync(string AgentId, [FromQuery] DateTime? MinTime = null, [FromQuery] DateTime? MaxTime = null, [FromQuery] int Index = 0, [FromQuery] int Count = 50)
+		{
+			Response.ContentType = "application/json";
+			Response.StatusCode = 200;
+			await Response.StartAsync();
+			await AgentService.Agents.GetLogger(AgentId.ToAgentId()).FindAsync(Response.BodyWriter, MinTime, MaxTime, Index, Count);
+		}
+
+		/// <summary>
+		/// Find all the sessions of a particular agent
+		/// </summary>
+		/// <param name="AgentId">Unique id of the agent to find</param>
+		/// <param name="StartTime">Start time to include in the search</param>
+		/// <param name="FinishTime">Finish time to include in the search</param>
+		/// <param name="Index">Index of the first result to return</param>
+		/// <param name="Count">Number of results to return</param>
+		/// <returns>Sessions </returns>
+		[HttpGet]
+		[Route("/api/v1/agents/{AgentId}/sessions")]
+		public async Task<ActionResult<List<GetAgentSessionResponse>>> FindSessionsAsync(string AgentId, [FromQuery] DateTimeOffset? StartTime, [FromQuery] DateTimeOffset? FinishTime, [FromQuery] int Index = 0, [FromQuery] int Count = 50)
+		{
+			AgentId AgentIdValue = new AgentId(AgentId);
+
+			IAgent? Agent = await AgentService.GetAgentAsync(AgentIdValue);
+			if (Agent == null)
+			{
+				return NotFound();
+			}
+			if (!await AgentService.AuthorizeAsync(Agent, AclAction.ViewSession, User, null))
+			{
+				return Forbid();
+			}
+
+			List<ISession> Sessions = await AgentService.FindSessionsAsync(AgentIdValue, StartTime?.UtcDateTime, FinishTime?.UtcDateTime, Index, Count);
+			return Sessions.ConvertAll(x => new GetAgentSessionResponse(x));
+		}
+
+		/// <summary>
+		/// Find all the sessions of a particular agent
+		/// </summary>
+		/// <param name="AgentId">Unique id of the agent to find</param>
+		/// <param name="SessionId">Unique id of the session</param>
+		/// <returns>Sessions </returns>
+		[HttpGet]
+		[Route("/api/v1/agents/{AgentId}/sessions/{SessionId}")]
+		public async Task<ActionResult<GetAgentSessionResponse>> GetSessionAsync(string AgentId, string SessionId)
+		{
+			AgentId AgentIdValue = new AgentId(AgentId);
+
+			IAgent? Agent = await AgentService.GetAgentAsync(AgentIdValue);
+			if (Agent == null)
+			{
+				return NotFound();
+			}
+			if (!await AgentService.AuthorizeAsync(Agent, AclAction.ViewSession, User, null))
+			{
+				return Forbid();
+			}
+
+			ISession? Session = await AgentService.GetSessionAsync(SessionId.ToObjectId());
+			if(Session == null || Session.AgentId != AgentIdValue)
+			{
+				return NotFound();
+			}
+
+			return new GetAgentSessionResponse(Session);
+		}
+
+		/// <summary>
+		/// Find all the leases for a particular agent
+		/// </summary>
+		/// <param name="AgentId">Unique id of the agent to find</param>
+		/// <param name="SessionId">The session to query</param>
+		/// <param name="StartTime">Start of the time window to consider</param>
+		/// <param name="FinishTime">End of the time window to consider</param>
+		/// <param name="Index">Index of the first result to return</param>
+		/// <param name="Count">Number of results to return</param>
+		/// <returns>Sessions </returns>
+		[HttpGet]
+		[Route("/api/v1/agents/{AgentId}/leases")]
+		public async Task<ActionResult<List<GetAgentLeaseResponse>>> FindLeasesAsync(string AgentId, [FromQuery] string? SessionId, [FromQuery] DateTimeOffset? StartTime, [FromQuery] DateTimeOffset? FinishTime, [FromQuery] int Index = 0, [FromQuery] int Count = 1000)
+		{
+			AgentId AgentIdValue = new AgentId(AgentId);
+
+			IAgent? Agent = await AgentService.GetAgentAsync(AgentIdValue);
+			if (Agent == null)
+			{
+				return NotFound();
+			}
+			if (!await AgentService.AuthorizeAsync(Agent, AclAction.ViewLeases, User, null))
+			{
+				return Forbid();
+			}
+
+			List<ILease> Leases = await AgentService.FindLeasesAsync(AgentIdValue, SessionId?.ToObjectId(), StartTime?.UtcDateTime, FinishTime?.UtcDateTime, Index, Count);
+			return Leases.ConvertAll(x => new GetAgentLeaseResponse(x));
+		}
+
+		/// <summary>
+		/// Get info about a particular lease
+		/// </summary>
+		/// <param name="AgentId">Unique id of the agent to find</param>
+		/// <param name="LeaseId">Unique id of the particular lease</param>
+		/// <returns>Lease matching the given id</returns>
+		[HttpGet]
+		[Route("/api/v1/agents/{AgentId}/leases/{LeaseId}")]
+		public async Task<ActionResult<GetAgentLeaseResponse>> GetLeaseAsync(AgentId AgentId, LeaseId LeaseId)
+		{
+			IAgent? Agent = await AgentService.GetAgentAsync(AgentId);
+			if(Agent == null)
+			{
+				return NotFound();
+			}
+			if (!await AclService.AuthorizeAsync(AclAction.ViewLeases, User))
+			{
+				return Forbid();
+			}
+
+			ILease? Lease = await AgentService.GetLeaseAsync(LeaseId);
+			if (Lease == null || Lease.AgentId != AgentId)
+			{
+				return NotFound();
+			}
+
+			return new GetAgentLeaseResponse(Lease);
+		}
+	}
+}
diff --git a/Engine/Source/Programs/Horde/Horde.Build/Controllers/ArtifactsController.cs b/Engine/Source/Programs/Horde/Horde.Build/Controllers/ArtifactsController.cs
new file mode 100644
--- /dev/null
+++ b/Engine/Source/Programs/Horde/Horde.Build/Controllers/ArtifactsController.cs
@@ -0,0 +1,447 @@
+// Copyright Epic Games, Inc. All Rights Reserved.
+
+using Amazon.S3.Transfer;
+using HordeServer.Api;
+using HordeServer.Jobs;
+using HordeServer.Models;
+using HordeServer.Services;
+using HordeServer.Utilities;
+using Microsoft.AspNetCore.Authorization;
+using Microsoft.AspNetCore.Http;
+using Microsoft.AspNetCore.Mvc;
+using Microsoft.AspNetCore.StaticFiles;
+using Microsoft.IdentityModel.Tokens;
+using MongoDB.Bson;
+using System;
+using System.Collections.Generic;
+using System.Globalization;
+using System.IdentityModel.Tokens.Jwt;
+using System.IO;
+using System.IO.Compression;
+using System.Linq;
+using System.Net.Http.Headers;
+using System.Net.Mime;
+using System.Security.Claims;
+using System.Text;
+using System.Threading.Tasks;
+
+namespace HordeServer.Controllers
+{
+	using JobId = ObjectId<IJob>;
+
+	/// <summary>
+	/// Controller for the /api/artifacts endpoint
+	/// </summary>
+	[ApiController]
+	[Route("[controller]")]
+	public class ArtifactsController : ControllerBase
+	{
+		/// <summary>
+		/// Instance of the database service
+		/// </summary>
+		private readonly DatabaseService DatabaseService;
+
+		/// <summary>
+		/// Instance of the artifact collection
+		/// </summary>
+		private readonly IArtifactCollection ArtifactCollection;
+
+		/// <summary>
+		/// Instance of the ACL service
+		/// </summary>
+		private readonly AclService AclService;
+
+		/// <summary>
+		/// Instance of the Job service
+		/// </summary>
+		private readonly JobService JobService;
+
+		/// <summary>
+		/// Constructor
+		/// </summary>
+		public ArtifactsController(DatabaseService DatabaseSerivce, IArtifactCollection ArtifactCollection, AclService AclService, JobService JobService)
+		{
+			this.DatabaseService = DatabaseSerivce;
+			this.ArtifactCollection = ArtifactCollection;
+			this.AclService = AclService;
+			this.JobService = JobService;
+		}
+
+		/// <summary>
+		/// Creates an artifact
+		/// </summary>
+		/// <param name="JobId">BatchId</param>
+		/// <param name="StepId">StepId</param>
+		/// <param name="File">The file contents</param>
+		/// <returns>Http result code</returns>
+		[HttpPost]
+		[Authorize]
+		[Route("/api/v1/artifacts")]
+		public async Task<ActionResult<CreateArtifactResponse>> CreateArtifact([FromQuery] JobId JobId, [FromQuery]string? StepId, IFormFile File)
+		{
+			IJob? Job = await JobService.GetJobAsync(JobId);
+			if(Job == null)
+			{
+				return NotFound();
+			}
+
+			if (!await JobService.AuthorizeAsync(Job, AclAction.UploadArtifact, User, null))
+			{
+				return Forbid();
+			}
+
+			IJobStep? Step = null;
+			if(StepId != null)
+			{
+				foreach(IJobStepBatch Batch in Job.Batches)
+				{
+					if(Batch.TryGetStep(StepId.ToSubResourceId(), out Step))
+					{
+						break;
+					}
+				}
+				if(Step == null)
+				{
+					// if the step doesn't exist in any of the batches, not found
+					return NotFound();
+				}
+			}
+
+			IArtifact NewArtifact = await ArtifactCollection.CreateArtifactAsync(Job.Id, Step?.Id, File.FileName, File.ContentType ?? "horde-mime/unknown", File.OpenReadStream());
+			return new CreateArtifactResponse(NewArtifact.Id.ToString());
+		}
+
+		/// <summary>
+		/// Updates an artifact
+		/// </summary>
+		/// <param name="ArtifactId">JobId</param>
+		/// <param name="File">The file contents</param>
+		/// <returns>Http result code</returns>
+		[HttpPut]
+		[Authorize]
+		[Route("/api/v1/artifacts/{ArtifactId}")]
+		public async Task<ActionResult<CreateArtifactResponse>> UpdateArtifact(string ArtifactId, IFormFile File)
+		{
+			IArtifact? Artifact = await ArtifactCollection.GetArtifactAsync(ArtifactId.ToObjectId());
+			if (Artifact == null)
+			{
+				return NotFound();
+			}
+			if (!await JobService.AuthorizeAsync(Artifact.JobId, AclAction.UploadArtifact, User, null))
+			{
+				return Forbid();
+			}
+
+			await ArtifactCollection.UpdateArtifactAsync(Artifact, File.ContentType ?? "horde-mime/unknown", File.OpenReadStream());
+			return Ok();
+		}
+
+		/// <summary>
+		/// Query artifacts for a job step
+		/// </summary>
+		/// <param name="JobId">Optional JobId to filter by</param>
+		/// <param name="StepId">Optional StepId to filter by</param>
+		/// <param name="Code">Whether to generate a direct download code</param>
+		/// <param name="Filter">Filter for the properties to return</param>
+		/// <returns>Information about all the artifacts</returns>
+		[HttpGet]
+		[Authorize]
+		[Route("/api/v1/artifacts")]
+		[ProducesResponseType(typeof(List<GetArtifactResponse>), 200)]
+		public async Task<ActionResult<List<object>>> GetArtifacts([FromQuery] JobId JobId, [FromQuery] string? StepId = null, [FromQuery] bool Code = false, [FromQuery] PropertyFilter? Filter = null)
+		{
+			if (!await JobService.AuthorizeAsync(JobId, AclAction.DownloadArtifact, User, null))
+			{
+				return Forbid();
+			}
+
+			string? DownloadCode = Code ? (string?)GetDirectDownloadCodeForJob(JobId) : null;
+
+			List<IArtifact> Artifacts = await ArtifactCollection.GetArtifactsAsync(JobId, StepId?.ToSubResourceId(), null);
+			return Artifacts.ConvertAll(x => new GetArtifactResponse(x, DownloadCode).ApplyFilter(Filter));
+		}
+
+		/// <summary>
+		/// Gets the claim required to download artifacts for a particular job
+		/// </summary>
+		/// <param name="JobId">The job id</param>
+		/// <returns>The required claim</returns>
+		static Claim GetDirectDownloadClaim(JobId JobId)
+		{
+			return new Claim(HordeClaimTypes.JobArtifacts, JobId.ToString());
+		}
+
+		/// <summary>
+		/// Get a download code for the artifacts of a job
+		/// </summary>
+		/// <param name="JobId">The job id</param>
+		/// <returns>The download code</returns>
+		string GetDirectDownloadCodeForJob(JobId JobId)
+		{
+			Claim DownloadClaim = GetDirectDownloadClaim(JobId);
+			return AclService.IssueBearerToken(new[] { DownloadClaim }, TimeSpan.FromHours(4.0));
+		}
+
+		/// <summary>
+		/// Retrieve metadata about a specific artifact
+		/// </summary>
+		/// <param name="ArtifactId">Id of the artifact to get information about</param>
+		/// <param name="Code">Whether to generate a direct download code</param>
+		/// <param name="Filter">Filter for the properties to return</param>
+		/// <returns>Information about the requested project</returns>
+		[HttpGet]
+		[Authorize]
+		[Route("/api/v1/artifacts/{ArtifactId}")]
+		[ProducesResponseType(typeof(GetArtifactResponse), 200)]
+		public async Task<ActionResult<object>> GetArtifact(string ArtifactId, bool Code = false, [FromQuery] PropertyFilter? Filter = null)
+		{
+			IArtifact? Artifact = await ArtifactCollection.GetArtifactAsync(ArtifactId.ToObjectId());
+			if (Artifact == null)
+			{
+				return NotFound();
+			}
+			if (!await JobService.AuthorizeAsync(Artifact.JobId, AclAction.DownloadArtifact, User, null))
+			{
+				return Forbid();
+			}
+
+			string? DownloadCode = Code? (string?)GetDirectDownloadCodeForJob(Artifact.JobId) : null;
+			return new GetArtifactResponse(Artifact, DownloadCode).ApplyFilter(Filter);
+		}
+
+		/// <summary>
+		/// Retrieve raw data for an artifact
+		/// </summary>
+		/// <param name="ArtifactId">Id of the artifact to get information about</param>
+		/// <returns>Raw artifact data</returns>
+		[HttpGet]
+		[Authorize]
+		[Route("/api/v1/artifacts/{ArtifactId}/data")]
+		public async Task<ActionResult> GetArtifactData(string ArtifactId)
+		{
+			IArtifact? Artifact = await ArtifactCollection.GetArtifactAsync(ArtifactId.ToObjectId());
+			if (Artifact == null)
+			{
+				return NotFound();
+			}
+			if (!await JobService.AuthorizeAsync(Artifact.JobId, AclAction.DownloadArtifact, User, null))
+			{
+				return Forbid();
+			}
+
+			// Fun, filestream result automatically closes the stream!
+			return new FileStreamResult(await ArtifactCollection.OpenArtifactReadStreamAsync(Artifact), Artifact.MimeType);
+		}
+		
+		/// <summary>
+		/// Retrieve raw data for an artifact by filename
+		/// </summary>
+		/// <param name="JobId">Unique id for the job</param>
+		/// <param name="StepId">Unique id for the step</param>
+		/// <param name="Filename">Filename of artifact from step</param>
+		/// <returns>Raw artifact data</returns>
+		[HttpGet]
+		[Route("/api/v1/jobs/{JobId}/steps/{StepId}/artifacts/{FileName}/data")]
+		public async Task<ActionResult<object>> GetArtifactDataByFilename(JobId JobId, string StepId, string Filename)
+		{
+			SubResourceId StepIdValue = StepId.ToSubResourceId();
+
+			if (!await JobService.AuthorizeAsync(JobId, AclAction.DownloadArtifact, User, null))
+			{
+				return Forbid();
+			}
+			
+			List<IArtifact> Artifacts = await ArtifactCollection.GetArtifactsAsync(JobId, StepIdValue, Filename);
+			if (Artifacts.Count == 0)
+			{
+				return NotFound();
+			}
+
+			IArtifact Artifact = Artifacts[0];
+			return new FileStreamResult(await ArtifactCollection.OpenArtifactReadStreamAsync(Artifact), Artifact.MimeType);
+		}
+
+		/// <summary>
+		/// Class to return a file stream without the "content-disposition: attachment" header
+		/// </summary>
+		class InlineFileStreamResult : FileStreamResult
+		{
+			/// <summary>
+			/// The suggested download filename
+			/// </summary>
+			string FileName;
+
+			/// <summary>
+			/// Constructor
+			/// </summary>
+			/// <param name="Stream"></param>
+			/// <param name="MimeType"></param>
+			/// <param name="FileName"></param>
+			public InlineFileStreamResult(System.IO.Stream Stream, string MimeType, string FileName)
+				: base(Stream, MimeType)
+			{
+				this.FileName = FileName;
+			}
+
+			/// <inheritdoc/>
+			public override Task ExecuteResultAsync(ActionContext Context)
+			{
+				ContentDisposition ContentDisposition = new ContentDisposition();
+				ContentDisposition.Inline = true;
+				ContentDisposition.FileName = FileName;
+				Context.HttpContext.Response.Headers.Add("Content-Disposition", ContentDisposition.ToString());
+				
+				return base.ExecuteResultAsync(Context);
+			}
+		}
+
+		/// <summary>
+		/// Retrieve raw data for an artifact
+		/// </summary>
+		/// <param name="ArtifactId">Id of the artifact to get information about</param>
+		/// <param name="Code">The authorization code for this resource</param>
+		/// <returns>Raw artifact data</returns>
+		[HttpGet]
+		[AllowAnonymous]
+		[Route("/api/v1/artifacts/{ArtifactId}/download")]
+		public async Task<ActionResult> DownloadArtifact(string ArtifactId, [FromQuery] string Code)
+		{
+			TokenValidationParameters Parameters = new TokenValidationParameters();
+			Parameters.ValidateAudience = false;
+			Parameters.RequireExpirationTime = true;
+			Parameters.ValidateLifetime = true;
+			Parameters.ValidIssuer = DatabaseService.JwtIssuer;
+			Parameters.ValidateIssuer = true;
+			Parameters.ValidateIssuerSigningKey = true;
+			Parameters.IssuerSigningKey = DatabaseService.JwtSigningKey;
+
+			SecurityToken Token;
+			JwtSecurityTokenHandler Handler = new JwtSecurityTokenHandler();
+			ClaimsPrincipal Principal = Handler.ValidateToken(Code, Parameters, out Token);
+
+			IArtifact? Artifact = await ArtifactCollection.GetArtifactAsync(ArtifactId.ToObjectId());
+			if (Artifact == null)
+			{
+				return NotFound();
+			}
+
+			Claim DirectDownloadClaim = GetDirectDownloadClaim(Artifact.JobId);
+			if (!Principal.HasClaim(DirectDownloadClaim.Type, DirectDownloadClaim.Value))
+			{
+				return Forbid();
+			}
+
+			return new InlineFileStreamResult(await ArtifactCollection.OpenArtifactReadStreamAsync(Artifact), Artifact.MimeType, Path.GetFileName(Artifact.Name));
+		}
+
+		/// <summary>
+		/// Returns a zip archive of many artifacts
+		/// </summary>
+		/// <param name="ArtifactZipRequest">Artifact request params</param>
+		/// <returns>Zip of many artifacts</returns>
+		[HttpPost]
+		[Authorize]
+		[Route("/api/v1/artifacts/zip")]
+		public async Task<ActionResult> ZipArtifacts(GetArtifactZipRequest ArtifactZipRequest)
+		{
+			if (ArtifactZipRequest.JobId == null)
+			{
+				return BadRequest("Must specify a JobId");
+			}
+
+			IJob? Job = await JobService.GetJobAsync(new JobId(ArtifactZipRequest.JobId!));
+			if (Job == null)
+			{
+				return NotFound();
+			}
+			if (!await JobService.AuthorizeAsync(Job, AclAction.DownloadArtifact, User, null))
+			{
+				return Forbid();
+			}
+
+			List<IArtifact> Artifacts = await ArtifactCollection.GetArtifactsAsync(Job.Id, ArtifactZipRequest.StepId?.ToSubResourceId(), null);
+
+			Dictionary<ObjectId, IArtifact> IdToArtifact = Artifacts.ToDictionary(x => x.Id, x => x);
+
+			List<IArtifact> ZipArtifacts;
+			if (ArtifactZipRequest.ArtifactIds == null)
+			{
+				ZipArtifacts = Artifacts;
+			}
+			else
+			{
+				ZipArtifacts = new List<IArtifact>();
+				foreach (string ArtifactId in ArtifactZipRequest.ArtifactIds)
+				{
+					IArtifact? Artifact;
+					if (IdToArtifact.TryGetValue(ArtifactId.ToObjectId(), out Artifact))
+					{
+						ZipArtifacts.Add(Artifact);
+					}
+					else
+					{
+						return NotFound();
+					}
+				}
+			}
+
+			IGraph Graph = await JobService.GetGraphAsync(Job);
+
+			return new CustomFileCallbackResult("Artifacts.zip", "application/octet-stream", false, async (OutputStream, Context) =>
+			{
+				// Make an unseekable MemoryStream for the ZipArchive. We have to do this because the ZipEntry stream falls back to a synchronous write to it's own stream wrappers.
+				using (CustomBufferStream ZipOutputStream = new CustomBufferStream())
+				{
+					// Keep the stream open after dispose so we can write the EOF bits.
+					using (ZipArchive ZipArchive = new ZipArchive(ZipOutputStream, ZipArchiveMode.Create, true))
+					{
+						foreach (IArtifact Artifact in ZipArtifacts)
+						{
+							await using (System.IO.Stream ArtifactStream = await ArtifactCollection.OpenArtifactReadStreamAsync(Artifact))
+							{
+								// tack on the step name into the directory if it exists
+								string StepName = string.Empty;
+								if (Artifact.StepId.HasValue)
+								{
+									foreach (IJobStepBatch Batch in Job.Batches)
+									{
+										IJobStep? Step;
+										if (Batch.TryGetStep(Artifact.StepId.Value, out Step))
+										{
+											StepName = Graph.Groups[Batch.GroupIdx].Nodes[Step.NodeIdx].Name;
+											break;
+										}
+									}
+								}
+
+								ZipArchiveEntry ZipEntry = ZipArchive.CreateEntry(Artifact.Name);
+								using (System.IO.Stream EntryStream = ZipEntry.Open())
+								{
+									byte[] Buffer = new byte[4096];
+									int TotalBytesRead = 0;
+									while (TotalBytesRead < Artifact.Length)
+									{
+										int BytesRead = await ArtifactStream.ReadAsync(Buffer, 0, Buffer.Length);
+
+										// Write bytes to the entry stream.  Also advances the MemStream pos.
+										await EntryStream.WriteAsync(Buffer, 0, BytesRead);
+										// Dump what we have to the output stream
+										await OutputStream.WriteAsync(ZipOutputStream.GetBuffer(), 0, (int)ZipOutputStream.Position);
+
+										// Reset everything.
+										ZipOutputStream.Position = 0;
+										ZipOutputStream.SetLength(0);
+										TotalBytesRead += BytesRead;
+									}
+								}
+							}
+						}
+					}
+					// Write out the EOF stuff
+					ZipOutputStream.Position = 0;
+					await ZipOutputStream.CopyToAsync(OutputStream);
+				}
+			});
+		}
+	}
+}
diff --git a/Engine/Source/Programs/Horde/Horde.Build/Controllers/CredentialsController.cs b/Engine/Source/Programs/Horde/Horde.Build/Controllers/CredentialsController.cs
new file mode 100644
--- /dev/null
+++ b/Engine/Source/Programs/Horde/Horde.Build/Controllers/CredentialsController.cs
@@ -0,0 +1,184 @@
+// Copyright Epic Games, Inc. All Rights Reserved.
+
+using HordeServer.Api;
+using HordeServer.Models;
+using HordeServer.Services;
+using HordeServer.Utilities;
+using Microsoft.AspNetCore.Authorization;
+using Microsoft.AspNetCore.Mvc;
+using MongoDB.Bson;
+using System;
+using System.Collections.Generic;
+using System.ComponentModel.DataAnnotations;
+using System.Linq;
+using System.Threading.Tasks;
+
+namespace HordeServer.Controllers
+{
+	/// <summary>
+	/// Controller for the /api/v1/credentials endpoint
+	/// </summary>
+	[ApiController]
+	[Authorize]
+	[Route("[controller]")]
+	public class CredentialsController : ControllerBase
+	{
+		/// <summary>
+		/// Singleton instance of the ACL service
+		/// </summary>
+		private readonly AclService AclService;
+
+		/// <summary>
+		/// Singleton instance of the credential service
+		/// </summary>
+		private readonly CredentialService CredentialService;
+
+		/// <summary>
+		/// Constructor
+		/// </summary>
+		/// <param name="AclService">The ACL service</param>
+		/// <param name="CredentialService">The credential service</param>
+		public CredentialsController(AclService AclService, CredentialService CredentialService)
+		{
+			this.AclService = AclService;
+			this.CredentialService = CredentialService;
+		}
+
+		/// <summary>
+		/// Creates a new credential
+		/// </summary>
+		/// <param name="Create">Parameters for the new credential.</param>
+		/// <returns>Http result code</returns>
+		[HttpPost]
+		[Route("/api/v1/credentials")]
+		public async Task<ActionResult<CreateCredentialResponse>> CreateCredentialAsync([FromBody] CreateCredentialRequest Create)
+		{
+			if(!await AclService.AuthorizeAsync(AclAction.CreateCredential, User))
+			{
+				return Forbid();
+			}
+
+			Credential NewCredential = await CredentialService.CreateCredentialAsync(Create.Name, Create.Properties);
+			return new CreateCredentialResponse(NewCredential.Id.ToString());
+		}
+
+		/// <summary>
+		/// Query all the credentials
+		/// </summary>
+		/// <param name="Name">Id of the credential to get information about</param>
+		/// <param name="Filter">Filter for the properties to return</param>
+		/// <returns>Information about all the credentials</returns>
+		[HttpGet]
+		[Route("/api/v1/credentials")]
+		[ProducesResponseType(typeof(List<GetCredentialResponse>), 200)]
+		public async Task<ActionResult<object>> FindCredentialAsync([FromQuery] string? Name = null, [FromQuery] PropertyFilter? Filter = null)
+		{
+			if (!await AclService.AuthorizeAsync(AclAction.ListCredentials, User))
+			{
+				return Forbid();
+			}
+
+			List<Credential> Credentials = await CredentialService.FindCredentialsAsync(Name);
+			GlobalPermissionsCache Cache = new GlobalPermissionsCache();
+
+			List<object> Responses = new List<object>();
+			foreach (Credential Credential in Credentials)
+			{
+				if (await CredentialService.AuthorizeAsync(Credential, AclAction.ViewCredential, User, Cache))
+				{
+					bool bIncludeAcl = await CredentialService.AuthorizeAsync(Credential, AclAction.ViewPermissions, User, Cache);
+					Responses.Add(new GetCredentialResponse(Credential, bIncludeAcl).ApplyFilter(Filter));
+				}
+			}
+			return Responses;
+		}
+
+		/// <summary>
+		/// Retrieve information about a specific credential
+		/// </summary>
+		/// <param name="CredentialId">Id of the credential to get information about</param>
+		/// <param name="Filter">Filter for properties to return</param>
+		/// <returns>Information about the requested credential</returns>
+		[HttpGet]
+		[Route("/api/v1/credentials/{CredentialId}")]
+		[ProducesResponseType(typeof(GetCredentialResponse), 200)]
+		public async Task<ActionResult<object>> GetCredentialAsync(string CredentialId, [FromQuery] PropertyFilter? Filter = null)
+		{
+			ObjectId ProjectIdValue = CredentialId.ToObjectId();
+
+			Credential? Credential = await CredentialService.GetCredentialAsync(ProjectIdValue);
+			if (Credential == null)
+			{
+				return NotFound();
+			}
+
+			GlobalPermissionsCache Cache = new GlobalPermissionsCache();
+			if (!await CredentialService.AuthorizeAsync(Credential, AclAction.ViewCredential, User, Cache))
+			{
+				return Forbid();
+			}
+
+			bool bIncludeAcl = await CredentialService.AuthorizeAsync(Credential, AclAction.ViewPermissions, User, Cache);
+			return new GetCredentialResponse(Credential, bIncludeAcl).ApplyFilter(Filter);
+		}
+
+		/// <summary>
+		/// Update a credential's properties.
+		/// </summary>
+		/// <param name="CredentialId">Id of the credential to update</param>
+		/// <param name="Update">Items on the credential to update</param>
+		/// <returns>Http result code</returns>
+		[HttpPut]
+		[Route("/api/v1/credentials/{CredentialId}")]
+		public async Task<ActionResult> UpdateCredentialAsync(string CredentialId, [FromBody] UpdateCredentialRequest Update)
+		{
+			ObjectId CredentialIdValue = CredentialId.ToObjectId();
+
+			Credential? Credential = await CredentialService.GetCredentialAsync(CredentialIdValue);
+			if(Credential == null)
+			{
+				return NotFound();
+			}
+
+			GlobalPermissionsCache Cache = new GlobalPermissionsCache();
+			if (!await CredentialService.AuthorizeAsync(Credential, AclAction.UpdateCredential, User, Cache))
+			{
+				return Forbid();
+			}
+			if (Update.Acl != null && !await CredentialService.AuthorizeAsync(Credential, AclAction.ChangePermissions, User, Cache))
+			{
+				return Forbid();
+			}
+
+			await CredentialService.UpdateCredentialAsync(CredentialIdValue, Update.Name, Update.Properties, Acl.Merge(Credential.Acl, Update.Acl));
+			return new OkResult();
+		}
+
+		/// <summary>
+		/// Delete a credential
+		/// </summary>
+		/// <param name="CredentialId">Id of the credential to delete</param>
+		/// <returns>Http result code</returns>
+		[HttpDelete]
+		[Route("/api/v1/credentials/{CredentialId}")]
+		public async Task<ActionResult> DeleteCredentialAsync(string CredentialId)
+		{
+			ObjectId CredentialIdValue = CredentialId.ToObjectId();
+
+			Credential? Credential = await CredentialService.GetCredentialAsync(CredentialIdValue);
+			if (Credential == null)
+			{
+				return NotFound();
+			}
+			if (!await CredentialService.AuthorizeAsync(Credential, AclAction.DeleteCredential, User, null))
+			{
+				return Forbid();
+			}
+			if (!await CredentialService.DeleteCredentialAsync(CredentialIdValue))
+			{
+				return NotFound();
+			}
+			return new OkResult();
+		}
+	}
+}
diff --git a/Engine/Source/Programs/Horde/Horde.Build/Controllers/DebugController.cs b/Engine/Source/Programs/Horde/Horde.Build/Controllers/DebugController.cs
new file mode 100644
--- /dev/null
+++ b/Engine/Source/Programs/Horde/Horde.Build/Controllers/DebugController.cs
@@ -0,0 +1,931 @@
+// Copyright Epic Games, Inc. All Rights Reserved.
+
+using EpicGames.Core;
+using HordeServer.Api;
+using HordeCommon;
+using HordeServer.Collections;
+using HordeServer.Models;
+using HordeServer.Services;
+using HordeServer.Tasks.Impl;
+using HordeServer.Utilities;
+using Microsoft.AspNetCore.Authorization;
+using Microsoft.AspNetCore.Mvc;
+using Microsoft.Extensions.Primitives;
+using MongoDB.Bson;
+using System;
+using System.Collections.Generic;
+using System.Diagnostics;
+using System.Globalization;
+using System.IO;
+using System.Linq;
+using System.Net;
+using System.Runtime;
+using System.Runtime.InteropServices;
+using System.Text;
+using System.Threading.Tasks;
+using System.Web;
+
+using MongoDB.Driver;
+using System.Text.Json;
+using System.Threading;
+using HordeServer.Services.Impl;
+using Microsoft.Net.Http.Headers;
+using HordeServer.Storage;
+using Microsoft.Extensions.Hosting;
+using Microsoft.Extensions.Options;
+using StatsdClient;
+
+using Amazon.S3;
+using JetBrains.Profiler.SelfApi;
+
+namespace HordeServer.Controllers
+{
+	using LogId = ObjectId<ILogFile>;
+	using PoolId = StringId<IPool>;
+	using ProjectId = StringId<IProject>;
+	using TemplateRefId = StringId<TemplateRef>;
+	using AgentSoftwareVersion = StringId<IAgentSoftwareCollection>;
+	using IStream = HordeServer.Models.IStream;
+	using P4Debugging = Perforce.P4.P4Debugging;
+
+	/// <summary>
+	/// Public endpoints for the debug controller
+	/// </summary>
+	[ApiController]
+	[Route("[controller]")]
+	public class PublicDebugController : ControllerBase
+	{
+		/// <summary>
+		/// The connection tracker service singleton
+		/// </summary>
+		RequestTrackerService RequestTrackerService;
+
+		IHostApplicationLifetime ApplicationLifetime;
+
+		IDogStatsd DogStatsd;
+
+		/// <summary>
+		/// Constructor
+		/// </summary>
+		/// <param name="RequestTrackerService"></param>
+		/// <param name="ApplicationLifetime"></param>
+		/// <param name="DogStatsd"></param>
+		public PublicDebugController(RequestTrackerService RequestTrackerService, IHostApplicationLifetime ApplicationLifetime, IDogStatsd DogStatsd)
+		{
+			this.RequestTrackerService = RequestTrackerService;
+			this.ApplicationLifetime = ApplicationLifetime;
+			this.DogStatsd = DogStatsd;
+		}
+
+		/// <summary>
+		/// Prints all the headers for the incoming request
+		/// </summary>
+		/// <returns>Http result</returns>
+		[HttpGet]
+		[Route("/api/v1/debug/headers")]
+		public IActionResult GetRequestHeaders()
+		{
+			StringBuilder Content = new StringBuilder();
+			Content.AppendLine("<html><body><pre>");
+			foreach (KeyValuePair<string, StringValues> Pair in HttpContext.Request.Headers)
+			{
+				foreach (string Value in Pair.Value)
+				{
+					Content.AppendLine(HttpUtility.HtmlEncode($"{Pair.Key}: {Value}"));
+				}
+			}
+			Content.Append("</pre></body></html>");
+			return new ContentResult { ContentType = "text/html", StatusCode = (int)HttpStatusCode.OK, Content = Content.ToString() };
+		}
+
+		/// <summary>
+		/// Waits specified number of milliseconds and then returns a response
+		/// Used for testing timeouts proxy settings.
+		/// </summary>
+		/// <returns>Http result</returns>
+		[HttpGet]
+		[Route("/api/v1/debug/wait")]
+		public async Task<ActionResult> GetAndWait([FromQuery] int WaitTimeMs = 1000)
+		{
+			await Task.Delay(WaitTimeMs);
+			string Content = $"Waited {WaitTimeMs} ms. " + new Random().Next(0, 10000000);
+			return new ContentResult { ContentType = "text/plain", StatusCode = (int)HttpStatusCode.OK, Content = Content };
+		}
+
+		/// <summary>
+		/// Waits specified number of milliseconds and then throws an exception
+		/// Used for testing graceful shutdown and interruption of outstanding requests.
+		/// </summary>
+		/// <returns>Http result</returns>
+		[HttpGet]
+		[Route("/api/v1/debug/exception")]
+		public async Task<ActionResult> ThrowException([FromQuery] int WaitTimeMs = 0)
+		{
+			await Task.Delay(WaitTimeMs);
+			throw new Exception("Test exception triggered by debug controller!");
+		}
+
+		/// <summary>
+		/// Trigger an increment of a DogStatsd metric
+		/// </summary>
+		/// <returns>Http result</returns>
+		[HttpGet]
+		[Route("/api/v1/debug/metric")]
+		public ActionResult TriggerMetric([FromQuery] int Value = 10)
+		{
+			DogStatsd.Increment("hordeMetricTest", Value);
+			return Ok("Incremented metric 'hordeMetricTest' Type: " + DogStatsd.GetType());
+		}
+
+		/// <summary>
+		/// Display metrics related to the .NET runtime
+		/// </summary>
+		/// <returns>Http result</returns>
+		[HttpGet]
+		[Route("/api/v1/debug/dotnet-metrics")]
+		public ActionResult DotNetMetrics()
+		{
+			ThreadPool.GetMaxThreads(out int MaxWorkerThreads, out int MaxIoThreads);
+			ThreadPool.GetAvailableThreads(out int FreeWorkerThreads, out int FreeIoThreads);
+			ThreadPool.GetMinThreads(out int MinWorkerThreads, out int MinIoThreads);
+
+			int BusyIoThreads = MaxIoThreads - FreeIoThreads;
+			int BusyWorkerThreads = MaxWorkerThreads - FreeWorkerThreads;
+
+			StringBuilder Content = new StringBuilder();
+			Content.AppendLine("Threads:");
+			Content.AppendLine("-------------------------------------------------------------");
+			Content.AppendLine("Worker busy={0,-5} free={1,-5} min={2,-5} max={3,-5}", BusyWorkerThreads, FreeWorkerThreads, MinWorkerThreads, MaxWorkerThreads);
+			Content.AppendLine("  IOCP busy={0,-5} free={1,-5} min={2,-5} max={3,-5}", BusyIoThreads, FreeIoThreads, MinIoThreads, MaxWorkerThreads);
+
+
+			NumberFormatInfo Nfi = (NumberFormatInfo)CultureInfo.InvariantCulture.NumberFormat.Clone();
+			Nfi.NumberGroupSeparator = " ";
+
+			string FormatBytes(long Number)
+			{
+				return (Number / 1024 / 1024).ToString("#,0", Nfi) + " MB";
+			}
+
+			GCMemoryInfo GcMemoryInfo = GC.GetGCMemoryInfo();
+			Content.AppendLine("");
+			Content.AppendLine("");
+			Content.AppendLine("Garbage collection (GC):");
+			Content.AppendLine("-------------------------------------------------------------");
+			Content.AppendLine("              Latency mode: " + GCSettings.LatencyMode);
+			Content.AppendLine("              Is server GC: " + GCSettings.IsServerGC);
+			Content.AppendLine("              Total memory: " + FormatBytes(GC.GetTotalMemory(false)));
+			Content.AppendLine("           Total allocated: " + FormatBytes(GC.GetTotalAllocatedBytes(false)));
+			Content.AppendLine("                 Heap size: " + FormatBytes(GcMemoryInfo.HeapSizeBytes));
+			Content.AppendLine("                Fragmented: " + FormatBytes(GcMemoryInfo.FragmentedBytes));
+			Content.AppendLine("               Memory Load: " + FormatBytes(GcMemoryInfo.MemoryLoadBytes));
+			Content.AppendLine("    Total available memory: " + FormatBytes(GcMemoryInfo.TotalAvailableMemoryBytes));
+			Content.AppendLine("High memory load threshold: " + FormatBytes(GcMemoryInfo.HighMemoryLoadThresholdBytes));
+
+			return Ok(Content.ToString());
+		}
+
+		/// <summary>
+		/// Force a full GC of all generations
+		/// </summary>
+		/// <returns>Prints time taken in ms</returns>
+		[HttpGet]
+		[Route("/api/v1/debug/force-gc")]
+		public ActionResult ForceTriggerGc()
+		{
+			Stopwatch Timer = new Stopwatch();
+			Timer.Start();
+			GC.Collect();
+			Timer.Stop();
+			return Ok($"Time taken: {Timer.Elapsed.TotalMilliseconds} ms");
+		}
+
+		/// <summary>
+		/// Lists requests in progress
+		/// </summary>
+		/// <returns>HTML result</returns>
+		[HttpGet]
+		[Route("/api/v1/debug/requests-in-progress")]
+		public ActionResult GetRequestsInProgress()
+		{
+			StringBuilder Content = new StringBuilder();
+			Content.AppendLine("<html><body>");
+			Content.AppendLine("<h1>Requests in progress</h1>");
+			Content.AppendLine("<table border=\"1\">");
+			Content.AppendLine("<tr>");
+			Content.AppendLine("<th>Request Trace ID</th>");
+			Content.AppendLine("<th>Path</th>");
+			Content.AppendLine("<th>Started At</th>");
+			Content.AppendLine("<th>Age</th>");
+			Content.AppendLine("</tr>");
+
+			List<KeyValuePair<string, TrackedRequest>> Requests = RequestTrackerService.GetRequestsInProgress().ToList();
+			Requests.Sort((A, B) => A.Value.StartedAt.CompareTo(B.Value.StartedAt));
+
+			foreach (KeyValuePair<string, TrackedRequest> Entry in Requests)
+			{
+				Content.Append("<tr>");
+				Content.AppendLine($"<td>{Entry.Key}</td>");
+				Content.AppendLine($"<td>{Entry.Value.Request.Path}</td>");
+				Content.AppendLine($"<td>{Entry.Value.StartedAt}</td>");
+				Content.AppendLine($"<td>{Entry.Value.GetTimeSinceStartInMs()} ms</td>");
+				Content.Append("</tr>");
+			}
+			Content.Append("</table>\n</body>\n</html>");
+
+			return new ContentResult { ContentType = "text/html", StatusCode = (int)HttpStatusCode.OK, Content = Content.ToString() };
+		}
+
+		/*
+		// Used during development only
+		[HttpGet]
+		[Route("/api/v1/debug/stop")]
+		public ActionResult StopApp()
+		{
+			Task.Run(async () =>
+			{
+				await Task.Delay(100);
+				ApplicationLifetime.StopApplication();
+			});
+			
+			return new ContentResult { ContentType = "text/plain", StatusCode = (int)HttpStatusCode.OK, Content = "App stopping..." };
+		}
+		/**/
+	}
+
+	/// <summary>
+	/// Controller managing account status
+	/// </summary>
+	[ApiController]
+	[Authorize]
+	[Route("[controller]")]
+	public class SecureDebugController : ControllerBase
+	{
+		/// <summary>
+		/// The ACL service singleton
+		/// </summary>
+		AclService AclService;
+
+		/// <summary>
+		/// The database service instance
+		/// </summary>
+		DatabaseService DatabaseService;
+
+		/// <summary>
+		/// The job task source singelton
+		/// </summary>
+		JobTaskSource JobTaskSource;
+
+		/// <summary>
+		/// The globals singleton
+		/// </summary>
+		ISingletonDocument<Globals> GlobalsSingleton;
+
+		/// <summary>
+		/// Collection of pool documents
+		/// </summary>
+		IPoolCollection PoolCollection;
+
+		/// <summary>
+		/// Collection of project documents
+		/// </summary>
+		IProjectCollection ProjectCollection;
+
+		/// <summary>
+		/// Collection of agent documents
+		/// </summary>
+		IAgentCollection AgentCollection;
+
+		/// <summary>
+		/// Collection of session documents
+		/// </summary>
+		ISessionCollection SessionCollection;
+
+		/// <summary>
+		/// Collection of lease documents
+		/// </summary>
+		ILeaseCollection LeaseCollection;
+
+		/// <summary>
+		/// Collection of template documents
+		/// </summary>
+		ITemplateCollection TemplateCollection;
+
+		/// <summary>
+		/// Collection of stream documents
+		/// </summary>
+		IStreamCollection StreamCollection;
+
+		/// <summary>
+		/// The graph collection singleton
+		/// </summary>
+		IGraphCollection GraphCollection;
+
+		/// <summary>
+		/// The log file collection singleton
+		/// </summary>
+		ILogFileCollection LogFileCollection;
+
+		/// <summary>
+		/// Perforce client
+		/// </summary>
+		IPerforceService Perforce;
+
+		/// <summary>
+		/// Fleet manager
+		/// </summary>
+		IFleetManager FleetManager;
+
+		/// <summary>
+		/// Settings
+		/// </summary>
+		private IOptionsMonitor<ServerSettings> Settings;
+
+		/// <summary>
+		/// Constructor
+		/// </summary>
+		/// <param name="AclService">The ACL service singleton</param>
+		/// <param name="DatabaseService">The database service instance</param>
+		/// <param name="JobTaskSource">The dispatch service singleton</param>
+		/// <param name="GlobalsSingleton">The globals singleton</param>
+		/// <param name="PoolCollection">Collection of pool documents</param>
+		/// <param name="ProjectCollection">Collection of project documents</param>
+		/// <param name="AgentCollection">Collection of agent documents</param>
+		/// <param name="SessionCollection">Collection of session documents</param>
+		/// <param name="LeaseCollection">Collection of lease documents</param>
+		/// <param name="TemplateCollection">Collection of template documents</param>
+		/// <param name="StreamCollection">Collection of stream documents</param>
+		/// <param name="GraphCollection">The graph collection</param>
+		/// <param name="LogFileCollection">The log file collection</param>
+		/// <param name="Perforce">Perforce client</param>
+		/// <param name="FleetManager">The default fleet manager</param>
+		/// <param name="Settings">Settings</param>
+		public SecureDebugController(AclService AclService, DatabaseService DatabaseService, JobTaskSource JobTaskSource, ISingletonDocument<Globals> GlobalsSingleton, IPoolCollection PoolCollection, IProjectCollection ProjectCollection, IAgentCollection AgentCollection, ISessionCollection SessionCollection, ILeaseCollection LeaseCollection, ITemplateCollection TemplateCollection, IStreamCollection StreamCollection, IGraphCollection GraphCollection, ILogFileCollection LogFileCollection, IPerforceService Perforce, IFleetManager FleetManager, IOptionsMonitor<ServerSettings> Settings)
+		{
+			this.AclService = AclService;
+			this.DatabaseService = DatabaseService;
+			this.JobTaskSource = JobTaskSource;
+			this.GlobalsSingleton = GlobalsSingleton;
+			this.PoolCollection = PoolCollection;
+			this.ProjectCollection = ProjectCollection;
+			this.AgentCollection = AgentCollection;
+			this.SessionCollection = SessionCollection;
+			this.LeaseCollection = LeaseCollection;
+			this.TemplateCollection = TemplateCollection;
+			this.StreamCollection = StreamCollection;
+			this.GraphCollection = GraphCollection;
+			this.LogFileCollection = LogFileCollection;
+			this.Perforce = Perforce;
+			this.FleetManager = FleetManager;
+			this.Settings = Settings;
+		}
+
+		/// <summary>
+		/// Prints all the environment variables
+		/// </summary>
+		/// <returns>Http result</returns>
+		[HttpGet]
+		[Route("/api/v1/debug/environment")]
+		public async Task<ActionResult> GetServerEnvVars()
+		{
+			if (!await AclService.AuthorizeAsync(AclAction.AdminRead, User))
+			{
+				return Forbid();
+			}
+
+			StringBuilder Content = new StringBuilder();
+			Content.AppendLine("<html><body><pre>");
+			foreach (System.Collections.DictionaryEntry? Pair in System.Environment.GetEnvironmentVariables())
+			{
+				if (Pair != null)
+				{
+					Content.AppendLine(HttpUtility.HtmlEncode($"{Pair.Value.Key}={Pair.Value.Value}"));
+				}
+			}
+			Content.Append("</pre></body></html>");
+			return new ContentResult { ContentType = "text/html", StatusCode = (int)HttpStatusCode.OK, Content = Content.ToString() };
+		}
+
+		/// <summary>
+		/// Returns diagnostic information about the current state of the queue
+		/// </summary>
+		/// <returns>Information about the queue</returns>
+		[HttpGet]
+		[Route("/api/v1/debug/queue")]
+		public async Task<ActionResult<object>> GetQueueStatusAsync()
+		{
+			if (!await AclService.AuthorizeAsync(AclAction.AdminRead, User))
+			{
+				return Forbid();
+			}
+
+			return JobTaskSource.GetStatus();
+		}
+
+		/// <summary>
+		/// Returns the complete config Horde uses
+		/// </summary>
+		/// <returns>Information about the config</returns>
+		[HttpGet]
+		[Route("/api/v1/debug/config")]
+		public async Task<ActionResult> GetConfig()
+		{
+			if (!await AclService.AuthorizeAsync(AclAction.AdminRead, User))
+			{
+				return Forbid();
+			}
+
+			JsonSerializerOptions Options = new JsonSerializerOptions
+			{
+				WriteIndented = true
+			};
+
+			return Ok(JsonSerializer.Serialize(Settings.CurrentValue, Options));
+		}
+
+		/// <summary>
+		/// Queries for all graphs
+		/// </summary>
+		/// <returns>The graph definitions</returns>
+		[HttpGet]
+		[Route("/api/v1/debug/graphs")]
+		[ProducesResponseType(200, Type = typeof(GetGraphResponse))]
+		public async Task<ActionResult<List<object>>> GetGraphsAsync([FromQuery] int? Index = null, [FromQuery] int? Count = null, [FromQuery] PropertyFilter? Filter = null)
+		{
+			if (!await AclService.AuthorizeAsync(AclAction.AdminRead, User))
+			{
+				return Forbid();
+			}
+
+			List<IGraph> Graphs = await GraphCollection.FindAllAsync(null, Index, Count);
+			return Graphs.ConvertAll(x => new GetGraphResponse(x).ApplyFilter(Filter));
+		}
+
+		/// <summary>
+		/// Queries for a particular graph by hash
+		/// </summary>
+		/// <returns>The graph definition</returns>
+		[HttpGet]
+		[Route("/api/v1/debug/graphs/{GraphId}")]
+		[ProducesResponseType(200, Type = typeof(GetGraphResponse))]
+		public async Task<ActionResult<object>> GetGraphAsync(string GraphId, [FromQuery] PropertyFilter? Filter = null)
+		{
+			if (!await AclService.AuthorizeAsync(AclAction.AdminRead, User))
+			{
+				return Forbid();
+			}
+
+			IGraph Graph = await GraphCollection.GetAsync(ContentHash.Parse(GraphId));
+			return new GetGraphResponse(Graph).ApplyFilter(Filter);
+		}
+
+		/// <summary>
+		/// Query all the job templates.
+		/// </summary>
+		/// <param name="Filter">Filter for properties to return</param>
+		/// <returns>Information about all the job templates</returns>
+		[HttpGet]
+		[Route("/api/v1/debug/templates")]
+		[ProducesResponseType(typeof(List<GetTemplateResponse>), 200)]
+		public async Task<ActionResult<object>> GetTemplatesAsync([FromQuery] PropertyFilter? Filter = null)
+		{
+			if (!await AclService.AuthorizeAsync(AclAction.AdminRead, User))
+			{
+				return Forbid();
+			}
+
+			List<ITemplate> Templates = await TemplateCollection.FindAllAsync();
+			return Templates.ConvertAll(x => new GetTemplateResponse(x).ApplyFilter(Filter));
+		}
+
+		/// <summary>
+		/// Retrieve information about a specific job template.
+		/// </summary>
+		/// <param name="TemplateHash">Id of the template to get information about</param>
+		/// <param name="Filter">List of properties to return</param>
+		/// <returns>Information about the requested template</returns>
+		[HttpGet]
+		[Route("/api/v1/debug/templates/{TemplateHash}")]
+		[ProducesResponseType(typeof(GetTemplateResponse), 200)]
+		public async Task<ActionResult<object>> GetTemplateAsync(string TemplateHash, [FromQuery] PropertyFilter? Filter = null)
+		{
+			ContentHash TemplateHashValue = ContentHash.Parse(TemplateHash);
+			if (!await AclService.AuthorizeAsync(AclAction.AdminRead, User))
+			{
+				return Forbid();
+			}
+
+			ITemplate? Template = await TemplateCollection.GetAsync(TemplateHashValue);
+			if (Template == null)
+			{
+				return NotFound();
+			}
+			return Template.ApplyFilter(Filter);
+		}
+
+		/// <summary>
+		/// Retrieve metadata about a specific log file
+		/// </summary>
+		/// <param name="LogFileId">Id of the log file to get information about</param>
+		/// <param name="Filter">Filter for the properties to return</param>
+		/// <returns>Information about the requested project</returns>
+		[HttpGet]
+		[Route("/api/v1/debug/logs/{LogFileId}")]
+		public async Task<ActionResult<object>> GetLogAsync(LogId LogFileId, [FromQuery] PropertyFilter? Filter = null)
+		{
+			if (!await AclService.AuthorizeAsync(AclAction.AdminRead, User))
+			{
+				return Forbid();
+			}
+
+			ILogFile? LogFile = await LogFileCollection.GetLogFileAsync(LogFileId);
+			if (LogFile == null)
+			{
+				return NotFound();
+			}
+
+			return LogFile.ApplyFilter(Filter);
+		}
+
+		/// <summary>
+		/// Populate the database with test data
+		/// </summary>
+		/// <returns>Async task</returns>
+		[HttpGet]
+		[Route("/api/v1/debug/collections/{Name}")]
+		public async Task<ActionResult<object>> GetDocumentsAsync(string Name, [FromQuery] string? Filter = null, [FromQuery] int Index = 0, [FromQuery] int Count = 10)
+		{
+			if (!await AclService.AuthorizeAsync(AclAction.AdminRead, User))
+			{
+				return Forbid();
+			}
+
+			IMongoCollection<Dictionary<string, object>> Collection = DatabaseService.GetCollection<Dictionary<string, object>>(Name);
+			List<Dictionary<string, object>> Documents = await Collection.Find(Filter ?? "{}").Skip(Index).Limit(Count).ToListAsync();
+			return Documents;
+		}
+
+		/// <summary>
+		/// Create P4 login ticket for the specified username
+		/// </summary>
+		/// <returns>Perforce ticket</returns>
+		[HttpGet]
+		[Route("/api/v1/debug/p4ticket")]
+		public async Task<ActionResult<string>> CreateTicket([FromQuery] string? Username = null)
+		{
+			if (!await AclService.AuthorizeAsync(AclAction.AdminRead, User))
+			{
+				return Forbid();
+			}
+
+			if (Username == null)
+			{
+				return BadRequest();
+			}
+
+			return await Perforce.CreateTicket(PerforceCluster.DefaultName, Username);
+		}
+
+		/// <summary>
+		/// Debugs perforce service commands
+		/// </summary>
+		/// <returns>Perforce ticket</returns>
+		[HttpGet]
+		[Route("/api/v1/debug/perforce/stress")]
+		public async Task<ActionResult<string>> GetPerforceStress()
+		{
+			if (!await AclService.AuthorizeAsync(AclAction.AdminWrite, User))
+			{
+				return Forbid();
+			}
+
+			Dictionary<string, string> Results = new Dictionary<string, string>();
+			const string DevBuild = "//UE4/Dev-Build";
+			string? PerforceUser = User.GetPerforceUser();
+
+			object UpdateLock = new object();
+
+			// ParallelTest.ForAsync doesn't seem to handle async lambdas
+			IEnumerable<Task> Tasks = Enumerable.Range(0, 16).Select(async (Idx) => {
+
+				int LatestChange = await Perforce.GetLatestChangeAsync(PerforceCluster.DefaultName, DevBuild, PerforceUser);
+				int CodeChange = await Perforce.GetCodeChangeAsync(PerforceCluster.DefaultName, DevBuild, LatestChange);				
+
+				lock (UpdateLock)
+				{
+					Results.Add($"Parallel Perforce Result {Idx}", $"LatestChange: {LatestChange} - CodeChange: {CodeChange}");
+				}
+			});
+
+			await Task.WhenAll(Tasks);			
+
+			Tasks = Enumerable.Range(0, 16).Select(async (Idx) => {
+
+				int LatestChange = await Perforce.GetLatestChangeAsync(PerforceCluster.DefaultName, DevBuild, PerforceUser);
+				int CodeChange = await Perforce.GetCodeChangeAsync(PerforceCluster.DefaultName, DevBuild, LatestChange);				
+
+				lock (UpdateLock)
+				{
+					Results.Add($"Parallel Perforce Result {Idx + 16}", $"LatestChange: {LatestChange} - CodeChange: {CodeChange}");
+				}
+			});
+
+
+			await Task.WhenAll(Tasks);
+
+			StringBuilder Content = new StringBuilder();
+			Content.AppendLine("<html><body><pre>");
+			foreach (KeyValuePair<string, string> Pair in Results)
+			{
+				Content.AppendLine(HttpUtility.HtmlEncode($"{Pair.Key}={Pair.Value}"));
+			}
+			Content.Append("</pre></body></html>");
+			return new ContentResult { ContentType = "text/html", StatusCode = (int)HttpStatusCode.OK, Content = Content.ToString() };
+
+		}
+
+		/// <summary>
+		/// Debugs perforce service commands
+		/// </summary>
+		/// <returns>Perforce ticket</returns>
+		[HttpGet]
+		[Route("/api/v1/debug/perforce/commands")]
+		public async Task<ActionResult<string>> GetPerforceCommands()
+		{
+			if (!await AclService.AuthorizeAsync(AclAction.AdminWrite, User))
+			{
+				return Forbid();
+			}
+
+			StringBuilder Content = new StringBuilder();
+
+			const string DevBuild = "//UE4/Dev-Build";
+			// a test CL I have setup in Dev-Build
+			const int TestCL = 16687685;
+
+			string? PerforceUser = User.GetPerforceUser();
+
+			if (string.IsNullOrEmpty(PerforceUser))
+			{
+				throw new Exception("Perforce user required");
+			}
+
+			Dictionary<string, string> Results = new Dictionary<string, string>();
+
+			// NOTE: This is in Private-Build not DevBuild
+			int NewChangeId = await Perforce.CreateNewChangeAsync(PerforceCluster.DefaultName, "//UE4/Private-Build", "Counter.txt");
+
+			List<ChangeDetails> ChangeDetails = await Perforce.GetChangeDetailsAsync(PerforceCluster.DefaultName, "//UE4/Private-Build", new int[] { NewChangeId }, PerforceUser);
+
+			Results.Add("CreateNewChangeAsync", $"{NewChangeId} : {ChangeDetails[0].Description}");
+
+			await Perforce.UpdateChangelistDescription(PerforceCluster.DefaultName, TestCL, $"Updated Description: New Change CL was {NewChangeId}");
+
+			ChangeDetails = await Perforce.GetChangeDetailsAsync(PerforceCluster.DefaultName, DevBuild, new int[] { TestCL }, PerforceUser);
+
+			Results.Add("UpdateChangelistDescription", $"{TestCL} : {ChangeDetails[0].Description}");
+
+			int LatestChange = await Perforce.GetLatestChangeAsync(PerforceCluster.DefaultName, DevBuild, PerforceUser);
+			Results.Add("GetLatestChangeAsync", LatestChange.ToString(CultureInfo.InvariantCulture));
+
+			int CodeChange = await Perforce.GetCodeChangeAsync(PerforceCluster.DefaultName, DevBuild, LatestChange);
+			Results.Add("GetCodeChangeAsync", CodeChange.ToString(CultureInfo.InvariantCulture));
+
+			PerforceUserInfo? UserInfo = await Perforce.GetUserInfoAsync(PerforceCluster.DefaultName, PerforceUser);
+
+			if (UserInfo == null)
+			{
+				Results.Add("GetUserInfoAsync", "null");
+			}
+			else
+			{
+				Results.Add("GetUserInfoAsync", $"{UserInfo.Login} : {UserInfo.Email}");
+			}
+
+			// changes
+			List<ChangeSummary> Changes = await Perforce.GetChangesAsync(PerforceCluster.DefaultName, DevBuild, null, LatestChange, 10, PerforceUser);
+			int Count = 0;
+			foreach (ChangeSummary Summary in Changes)
+			{
+				Results.Add($"GetChangesAsync Result {Count++}", $"{Summary.Number} : {Summary.Author} - {Summary.Description}");
+			}
+
+			ChangeDetails = await Perforce.GetChangeDetailsAsync(PerforceCluster.DefaultName, DevBuild, Changes.Select(Change => Change.Number).ToArray(), PerforceUser);
+			Count = 0;
+			foreach (ChangeDetails Details in ChangeDetails)
+			{
+				Results.Add($"GetChangeDetailsAsync Result {Count++}", $"{Details.Number} : {Details.Author} - {Details.Files.ToJson()}");
+			}
+
+			String Ticket = "Failed (Expected if not running service account)";
+			try
+			{
+				await Perforce.CreateTicket(PerforceCluster.DefaultName, PerforceUser);
+				Ticket = "Success";
+			}
+			catch
+			{
+
+			}
+
+			Results.Add($"CreateTicket", $"{Ticket}");
+			List<FileSummary> FileSummaries = await Perforce.FindFilesAsync(PerforceCluster.DefaultName, ChangeDetails.SelectMany(Details => Details.Files).Select(File => $"{DevBuild}{File}"));
+			Count = 0;
+			foreach (FileSummary File in FileSummaries)
+			{
+				Results.Add($"FindFilesAsync Result {Count++}", $"{File.Change} : {File.DepotPath} : {File.Exists}");
+			}
+
+			// print async
+
+			byte[] Bytes = await Perforce.PrintAsync(PerforceCluster.DefaultName, $"{DevBuild}/RunUAT.bat");
+			Results.Add($"PrintAsync: {DevBuild}/RunUAT.bat ", $"{System.Text.Encoding.Default.GetString(Bytes)}");
+
+			Bytes = await Perforce.PrintAsync(PerforceCluster.DefaultName, $"{DevBuild}/RunUAT.bat@13916380");
+			Results.Add($"PrintAsync: {DevBuild}/RunUAT.bat@13916380", $"{System.Text.Encoding.Default.GetString(Bytes)}");
+
+			// need to be running as service account for these
+
+			int? DuplicateCL = null;
+
+			int ShelvedChangeOnDevBuild = 16687685;
+
+			try 
+			{
+				DuplicateCL = await Perforce.DuplicateShelvedChangeAsync(PerforceCluster.DefaultName, ShelvedChangeOnDevBuild);
+				Results.Add("DuplicateShelvedChangeAsync", $"Duplicated CL {DuplicateCL} from Shelved CL {ShelvedChangeOnDevBuild}");
+			}
+			catch 
+			{
+				Results.Add($"DuplicateShelvedChangeAsync",  "Failed - expected unless using service account");
+			}
+
+			if (DuplicateCL.HasValue)
+			{
+				// Note change client much exist for the update
+				ChangeDetails = await Perforce.GetChangeDetailsAsync(PerforceCluster.DefaultName, DevBuild, new int[]{ DuplicateCL.Value}, PerforceUser);
+
+				Results.Add("UpdateChangelistDescription - PreUpdate", $"{DuplicateCL} : {ChangeDetails[0].Description}"); 
+
+				try 
+				{
+					await Perforce.UpdateChangelistDescription(PerforceCluster.DefaultName, DuplicateCL.Value, "Updated Description from Horde");
+
+					ChangeDetails = await Perforce.GetChangeDetailsAsync(PerforceCluster.DefaultName, DevBuild, new int[]{ DuplicateCL.Value}, PerforceUser);
+
+					Results.Add("UpdateChangelistDescription - PostUpdate", $"{DuplicateCL} : {ChangeDetails[0].Description}"); 
+				}
+				catch 
+				{
+					Results.Add("UpdateChangelistDescription - Faile", $"Client from change {DuplicateCL} must exist"); 
+				}
+
+				await Perforce.DeleteShelvedChangeAsync(PerforceCluster.DefaultName, DuplicateCL.Value);
+
+				ChangeDetails = await Perforce.GetChangeDetailsAsync(PerforceCluster.DefaultName, DevBuild, new int[] {DuplicateCL.Value} , PerforceUser);
+
+				Results.Add("GetChangeDetailsAsync - After Delete", $"ChangeDetails.Count: {ChangeDetails.Count}"); 
+
+			}
+
+			byte[] ImageBytes = await Perforce.PrintAsync(PerforceCluster.DefaultName, $"{DevBuild}/Samples/Games/ShooterGame/ShooterGame.png");
+
+			Content.AppendLine("<html><body><pre>");
+			foreach (KeyValuePair<string, string> Pair in Results)
+			{
+				Content.AppendLine(HttpUtility.HtmlEncode($"{Pair.Key}={Pair.Value}"));
+			}
+
+			Content.Append($"PrintAsync Binary: {DevBuild}/Samples/Games/ShooterGame/ShooterGame.png\n</pre>");
+			Content.Append($"<img src=\"data:image/png;base64,{Convert.ToBase64String(ImageBytes, Base64FormattingOptions.None)}\"/>");
+
+			Content.Append("</body></html>");
+			return new ContentResult { ContentType = "text/html", StatusCode = (int)HttpStatusCode.OK, Content = Content.ToString() };
+
+		}
+
+		/// <summary>
+		/// Simulates a crash of the native p4 bridge
+		/// </summary>		
+		[HttpGet]		
+		[Route("/api/v1/debug/perforce/debugcrash")]		
+		public async Task<ActionResult> GetPerforceDebugCrash()
+		{
+			if (!await AclService.AuthorizeAsync(AclAction.AdminWrite, User))
+			{
+				return Forbid();
+			}
+
+			P4Debugging.DebugCrash();
+
+			return Ok();
+		}
+
+		/// <summary>
+		/// Debugging fleet managers
+		/// </summary>
+		/// <returns>Nothing</returns>
+		[HttpGet]
+		[Route("/api/v1/debug/fleetmanager")]
+		public async Task<ActionResult<string>> FleetManage([FromQuery] string? PoolId = null, [FromQuery] string? ChangeType = null, [FromQuery] int? Count = null)
+		{
+			if (!await AclService.AuthorizeAsync(AclAction.AdminRead, User))
+			{
+				return Forbid();
+			}
+
+			if (PoolId == null)
+			{
+				return BadRequest("Missing 'PoolId' query parameter");
+			}
+
+			IPool? Pool = await PoolCollection.GetAsync(new PoolId(PoolId));
+			if (Pool == null)
+			{
+				return BadRequest($"Pool with ID '{PoolId}' not found");
+			}
+
+			if (!(ChangeType == "expand" || ChangeType == "shrink"))
+			{
+				return BadRequest("Missing 'ChangeType' query parameter and/or must be set to 'expand' or 'shrink'");
+			}
+
+			if (Count == null || Count.Value <= 0)
+			{
+				return BadRequest("Missing 'Count' query parameter and/or must be greater than 0");
+			}
+
+			List<IAgent> Agents = new List<IAgent>();
+			string Message = "No operation";
+
+			if (ChangeType == "expand")
+			{
+				await FleetManager.ExpandPool(Pool, Agents, Count.Value);
+				Message = $"Expanded pool {Pool.Name} by {Count}";
+			}
+			else if (ChangeType == "shrink")
+			{
+				Agents = await AgentCollection.FindAsync();
+				Agents = Agents.FindAll(a => a.IsInPool(Pool.Id)).ToList();
+				await FleetManager.ShrinkPool(Pool, Agents, Count.Value);
+				Message = $"Shrank pool {Pool.Name} by {Count}";
+			}
+
+			return new ContentResult { ContentType = "text/plain", StatusCode = (int)HttpStatusCode.OK, Content = Message };
+		}
+		
+		/// <summary>
+		/// Starts the profiler session
+		/// </summary>
+		/// <returns>Text message</returns>
+		[HttpGet]
+		[Route("/api/v1/debug/profiler/start")]
+		public async Task<ActionResult> StartProfiler()
+		{
+			await DotTrace.EnsurePrerequisiteAsync();
+
+			string SnapshotDir = Path.Join(Path.GetTempPath(), "horde-profiler-snapshots");
+			if (!Directory.Exists(SnapshotDir))
+			{
+				Directory.CreateDirectory(SnapshotDir);
+			}
+
+			var Config = new DotTrace.Config();
+			Config.SaveToDir(SnapshotDir);
+			DotTrace.Attach(Config);
+			DotTrace.StartCollectingData();
+			
+			return new ContentResult { ContentType = "text/plain", StatusCode = (int)HttpStatusCode.OK, Content = "Profiling session started. Using dir " + SnapshotDir };
+		}
+		
+		/// <summary>
+		/// Stops the profiler session
+		/// </summary>
+		/// <returns>Text message</returns>
+		[HttpGet]
+		[Route("/api/v1/debug/profiler/stop")]
+		public ActionResult StopProfiler()
+		{
+			DotTrace.SaveData();
+			DotTrace.Detach();
+			return new ContentResult { ContentType = "text/plain", StatusCode = (int)HttpStatusCode.OK, Content = "Profiling session stopped" };
+		}
+		
+		/// <summary>
+		/// Downloads the captured profiling snapshots
+		/// </summary>
+		/// <returns>A .zip file containing the profiling snapshots</returns>
+		[HttpGet]
+		[Route("/api/v1/debug/profiler/download")]
+		public ActionResult DownloadProfilingData()
+		{
+			string SnapshotZipFile = DotTrace.GetCollectedSnapshotFilesArchive(false);
+			if (!System.IO.File.Exists(SnapshotZipFile))
+			{
+				return NotFound("The generated snapshot .zip file was not found");
+			}
+			
+			return PhysicalFile(SnapshotZipFile, "application/zip", Path.GetFileName(SnapshotZipFile));
+		}
+	}
+}
diff --git a/Engine/Source/Programs/Horde/Horde.Build/Controllers/IssuesController.cs b/Engine/Source/Programs/Horde/Horde.Build/Controllers/IssuesController.cs
new file mode 100644
--- /dev/null
+++ b/Engine/Source/Programs/Horde/Horde.Build/Controllers/IssuesController.cs
@@ -0,0 +1,436 @@
+// Copyright Epic Games, Inc. All Rights Reserved.
+
+using EpicGames.Core;
+using HordeServer.Api;
+using HordeServer.Collections;
+using HordeCommon;
+using HordeServer.Models;
+using HordeServer.Services;
+using HordeServer.Utilities;
+using Microsoft.AspNetCore.Authorization;
+using Microsoft.AspNetCore.Mvc;
+using Microsoft.AspNetCore.Server.Kestrel.Core.Features;
+using Microsoft.IdentityModel.Tokens;
+using MongoDB.Bson;
+using System;
+using System.Collections.Generic;
+using System.ComponentModel.DataAnnotations;
+using System.Linq;
+using System.Text;
+using System.Threading.Tasks;
+
+namespace HordeServer.Controllers
+{
+	using JobId = ObjectId<IJob>;
+	using LogId = ObjectId<ILogFile>;
+	using StreamId = StringId<IStream>;
+	using UserId = ObjectId<IUser>;
+
+	/// <summary>
+	/// Controller for the /api/v1/issues endpoint
+	/// </summary>
+	[Authorize]
+	[ApiController]
+	[Route("[controller]")]
+	public class IssuesController : ControllerBase
+	{
+		private readonly IIssueCollection IssueCollection;
+		private readonly IIssueService IssueService;
+		private readonly JobService JobService;
+		private readonly StreamService StreamService;
+		private readonly IUserCollection UserCollection;
+		private readonly ILogEventCollection LogEventCollection;
+		private readonly ILogFileService LogFileService;
+
+		/// <summary>
+		/// Constructor
+		/// </summary>
+		public IssuesController(IIssueCollection IssueCollection, IIssueService IssueService, JobService JobService, StreamService StreamService, IUserCollection UserCollection, ILogEventCollection LogEventCollection, ILogFileService LogFileService)
+		{
+			this.IssueCollection = IssueCollection;
+			this.IssueService = IssueService;
+			this.JobService = JobService;
+			this.StreamService = StreamService;
+			this.UserCollection = UserCollection;
+			this.LogEventCollection = LogEventCollection;
+			this.LogFileService = LogFileService;
+		}
+
+		/// <summary>
+		/// Retrieve information about a specific issue
+		/// </summary>
+		/// <param name="Ids">Set of issue ids to find</param>
+		/// <param name="StreamId">The stream to query for</param>
+		/// <param name="Change">The changelist to query</param>
+		/// <param name="MinChange">The minimum changelist range to query, inclusive</param>
+		/// <param name="MaxChange">The minimum changelist range to query, inclusive</param>
+		/// <param name="JobId">Job id to filter by</param>
+		/// <param name="BatchId">The batch to filter by</param>
+		/// <param name="StepId">The step to filter by</param>
+		/// <param name="LabelIdx">The label within the job to filter by</param>
+		/// <param name="UserId">User to filter issues for</param>
+		/// <param name="Resolved">Whether to include resolved issues</param>
+		/// <param name="Promoted">Whether to include promoted issues</param>
+		/// <param name="Index">Starting offset of the window of results to return</param>
+		/// <param name="Count">Number of results to return</param>
+		/// <param name="Filter">Filter for the properties to return</param>
+		/// <returns>List of matching agents</returns>
+		[HttpGet]
+		[Route("/api/v1/issues")]
+		[ProducesResponseType(typeof(List<GetIssueResponse>), 200)]
+		public async Task<ActionResult<object>> FindIssuesAsync([FromQuery(Name = "Id")] int[]? Ids = null, [FromQuery] string? StreamId = null, [FromQuery] int? Change = null, [FromQuery] int? MinChange = null, [FromQuery] int? MaxChange = null, [FromQuery] JobId? JobId = null, [FromQuery] string? BatchId = null, [FromQuery] string? StepId = null, [FromQuery(Name = "label")] int? LabelIdx = null, [FromQuery] string? UserId = null, [FromQuery] bool? Resolved = null, [FromQuery] bool? Promoted = null, [FromQuery] int Index = 0, [FromQuery] int Count = 10, [FromQuery] PropertyFilter? Filter = null)
+		{
+			if(Ids != null && Ids.Length == 0)
+			{
+				Ids = null;
+			}
+
+			UserId? UserIdValue = null;
+			if (UserId != null)
+			{
+				UserIdValue = new UserId(UserId);
+			}
+
+			List<IIssue> Issues;
+			if (JobId == null)
+			{
+				StreamId? StreamIdValue = null;
+				if (StreamId != null)
+				{
+					StreamIdValue = new StreamId(StreamId);
+				}
+
+				Issues = await IssueService.FindIssuesAsync(Ids, UserIdValue, StreamIdValue, MinChange ?? Change, MaxChange ?? Change, Resolved, Promoted, Index, Count);
+			}
+			else
+			{
+				IJob? Job = await JobService.GetJobAsync(JobId.Value);
+				if (Job == null)
+				{
+					return NotFound();
+				}
+				if(!await JobService.AuthorizeAsync(Job, AclAction.ViewJob, this.User, null))
+				{
+					return Forbid();
+				}
+
+				IGraph Graph = await JobService.GetGraphAsync(Job);
+				Issues = await IssueService.FindIssuesForJobAsync(Ids, Job, Graph, StepId?.ToSubResourceId(), BatchId?.ToSubResourceId(), LabelIdx, UserIdValue, Resolved, Promoted, Index, Count);
+			}
+
+			StreamPermissionsCache PermissionsCache = new StreamPermissionsCache();
+
+			List<object> Responses = new List<object>();
+			foreach (IIssue Issue in Issues)
+			{
+				IIssueDetails Details = await IssueService.GetIssueDetailsAsync(Issue);
+				if (await AuthorizeIssue(Details, PermissionsCache))
+				{
+					bool bShowDesktopAlerts = IssueService.ShowDesktopAlertsForIssue(Issue, Details.Spans);
+					GetIssueResponse Response = await CreateIssueResponseAsync(Details, bShowDesktopAlerts);
+					Responses.Add(PropertyFilter.Apply(Response, Filter));
+				}
+			}
+			return Responses;
+		}
+
+		/// <summary>
+		/// Retrieve information about a specific issue
+		/// </summary>
+		/// <param name="IssueId">Id of the issue to get information about</param>
+		/// <param name="Filter">Filter for the properties to return</param>
+		/// <returns>List of matching agents</returns>
+		[HttpGet]
+		[Route("/api/v1/issues/{IssueId}")]
+		[ProducesResponseType(typeof(GetIssueResponse), 200)]
+		public async Task<ActionResult<object>> GetIssueAsync(int IssueId, [FromQuery] PropertyFilter? Filter = null)
+		{
+			IIssueDetails? Details = await IssueService.GetIssueDetailsAsync(IssueId);
+			if (Details == null)
+			{
+				return NotFound();
+			}
+			if (!await AuthorizeIssue(Details, null))
+			{
+				return Forbid();
+			}
+
+			bool bShowDesktopAlerts = IssueService.ShowDesktopAlertsForIssue(Details.Issue, Details.Spans);
+			return PropertyFilter.Apply(await CreateIssueResponseAsync(Details, bShowDesktopAlerts), Filter);
+		}
+
+		/// <summary>
+		/// Retrieve historical information about a specific issue
+		/// </summary>
+		/// <param name="IssueId">Id of the agent to get information about</param>
+		/// <param name="MinTime">Minimum time for records to return</param>
+		/// <param name="MaxTime">Maximum time for records to return</param>
+		/// <param name="Index">Offset of the first result</param>
+		/// <param name="Count">Number of records to return</param>
+		/// <returns>Information about the requested agent</returns>
+		[HttpGet]
+		[Route("/api/v1/issues/{IssueId}/history")]
+		public async Task GetAgentHistoryAsync(int IssueId, [FromQuery] DateTime? MinTime = null, [FromQuery] DateTime? MaxTime = null, [FromQuery] int Index = 0, [FromQuery] int Count = 50)
+		{
+			Response.ContentType = "application/json";
+			Response.StatusCode = 200;
+			await Response.StartAsync();
+			await IssueCollection.GetLogger(IssueId).FindAsync(Response.BodyWriter, MinTime, MaxTime, Index, Count);
+		}
+
+		/// <summary>
+		/// Create an issue response object
+		/// </summary>
+		/// <param name="Details"></param>
+		/// <param name="ShowDesktopAlerts"></param>
+		/// <returns></returns>
+		async Task<GetIssueResponse> CreateIssueResponseAsync(IIssueDetails Details, bool ShowDesktopAlerts)
+		{
+			List<GetIssueAffectedStreamResponse> AffectedStreams = new List<GetIssueAffectedStreamResponse>();
+			foreach (IGrouping<StreamId, IIssueSpan> StreamSpans in Details.Spans.GroupBy(x => x.StreamId))
+			{
+				IStream? Stream = await StreamService.GetCachedStream(StreamSpans.Key);
+				AffectedStreams.Add(new GetIssueAffectedStreamResponse(Stream, StreamSpans));
+			}
+			return new GetIssueResponse(Details, AffectedStreams, ShowDesktopAlerts);
+		}
+
+		/// <summary>
+		/// Retrieve events for a specific issue
+		/// </summary>
+		/// <param name="IssueId">Id of the issue to get information about</param>
+		/// <param name="Filter">Filter for the properties to return</param>
+		/// <returns>List of matching agents</returns>
+		[HttpGet]
+		[Route("/api/v1/issues/{IssueId}/streams")]
+		[ProducesResponseType(typeof(List<GetIssueStreamResponse>), 200)]
+		public async Task<ActionResult<List<object>>> GetIssueStreamsAsync(int IssueId, [FromQuery] PropertyFilter? Filter = null)
+		{
+			IIssueDetails? Issue = await IssueService.GetIssueDetailsAsync(IssueId);
+			if (Issue == null)
+			{
+				return NotFound();
+			}
+
+			StreamPermissionsCache Cache = new StreamPermissionsCache();
+			if (!await AuthorizeIssue(Issue, Cache))
+			{
+				return Forbid();
+			}
+
+			List<object> Responses = new List<object>();
+			foreach (IGrouping<StreamId, IIssueSpan> SpanGroup in Issue.Spans.GroupBy(x => x.StreamId))
+			{
+				if (await StreamService.AuthorizeAsync(SpanGroup.Key, AclAction.ViewStream, User, Cache))
+				{
+					HashSet<ObjectId> SpanIds = new HashSet<ObjectId>(SpanGroup.Select(x => x.Id));
+					List<IIssueStep> Steps = Issue.Steps.Where(x => SpanIds.Contains(x.SpanId)).ToList();
+					Responses.Add(PropertyFilter.Apply(new GetIssueStreamResponse(SpanGroup.Key, SpanGroup.ToList(), Steps), Filter));
+				}
+			}
+			return Responses;
+		}
+
+		/// <summary>
+		/// Retrieve events for a specific issue
+		/// </summary>
+		/// <param name="IssueId">Id of the issue to get information about</param>
+		/// <param name="StreamId">The stream id</param>
+		/// <param name="Filter">Filter for the properties to return</param>
+		/// <returns>List of matching agents</returns>
+		[HttpGet]
+		[Route("/api/v1/issues/{IssueId}/streams/{StreamId}")]
+		[ProducesResponseType(typeof(List<GetIssueStreamResponse>), 200)]
+		public async Task<ActionResult<object>> GetIssueStreamAsync(int IssueId, string StreamId, [FromQuery] PropertyFilter? Filter = null)
+		{
+			IIssueDetails? Details = await IssueService.GetIssueDetailsAsync(IssueId);
+			if (Details == null)
+			{
+				return NotFound();
+			}
+
+			StreamId StreamIdValue = new StreamId(StreamId);
+			if (!await StreamService.AuthorizeAsync(StreamIdValue, AclAction.ViewStream, User, null))
+			{
+				return Forbid();
+			}
+
+			List<IIssueSpan> Spans = Details.Spans.Where(x => x.StreamId == StreamIdValue).ToList();
+			if(Spans.Count == 0)
+			{
+				return NotFound();
+			}
+
+			HashSet<ObjectId> SpanIds = new HashSet<ObjectId>(Spans.Select(x => x.Id));
+			List<IIssueStep> Steps = Details.Steps.Where(x => SpanIds.Contains(x.SpanId)).ToList();
+
+			return PropertyFilter.Apply(new GetIssueStreamResponse(StreamIdValue, Spans, Steps), Filter);
+		}
+
+		/// <summary>
+		/// Retrieve events for a specific issue
+		/// </summary>
+		/// <param name="IssueId">Id of the issue to get information about</param>
+		/// <param name="JobId">The job id to filter for</param>
+		/// <param name="BatchId">The batch to filter by</param>
+		/// <param name="StepId">The step to filter by</param>
+		/// <param name="LabelIdx">The label within the job to filter by</param>
+		/// <param name="LogIds">List of log ids to return issues for</param>
+		/// <param name="Index">Index of the first event</param>
+		/// <param name="Count">Number of events to return</param>
+		/// <param name="Filter">Filter for the properties to return</param>
+		/// <returns>List of matching agents</returns>
+		[HttpGet]
+		[Route("/api/v1/issues/{IssueId}/events")]
+		[ProducesResponseType(typeof(List<GetLogEventResponse>), 200)]
+		public async Task<ActionResult<List<object>>> GetIssueEventsAsync(int IssueId, [FromQuery] JobId? JobId = null, [FromQuery] string? BatchId = null, [FromQuery] string? StepId = null, [FromQuery(Name = "label")] int? LabelIdx = null, [FromQuery] string[]? LogIds = null, [FromQuery] int Index = 0, [FromQuery] int Count = 10, [FromQuery] PropertyFilter? Filter = null)
+		{
+			HashSet<LogId> LogIdValues = new HashSet<LogId>();
+			if(JobId != null)
+			{
+				IJob? Job = await JobService.GetJobAsync(JobId.Value);
+				if(Job == null)
+				{
+					return NotFound();
+				}
+
+				if (StepId != null)
+				{
+					IJobStep? Step;
+					if (Job.TryGetStep(StepId.ToSubResourceId(), out Step) && Step.Outcome != JobStepOutcome.Success && Step.LogId != null)
+					{
+						LogIdValues.Add(Step.LogId.Value);
+					}
+				}
+				else if (BatchId != null)
+				{
+					IJobStepBatch? Batch;
+					if (Job.TryGetBatch(BatchId.ToSubResourceId(), out Batch))
+					{
+						LogIdValues.UnionWith(Batch.Steps.Where(x => x.Outcome != JobStepOutcome.Success && x.LogId != null).Select(x => x.LogId!.Value));
+					}
+				}
+				else if (LabelIdx != null)
+				{
+					IGraph Graph = await JobService.GetGraphAsync(Job);
+
+					HashSet<NodeRef> IncludedNodes = new HashSet<NodeRef>(Graph.Labels[LabelIdx.Value].IncludedNodes);
+
+					foreach (IJobStepBatch Batch in Job.Batches)
+					{
+						foreach (IJobStep Step in Batch.Steps)
+						{
+							NodeRef NodeRef = new NodeRef(Batch.GroupIdx, Step.NodeIdx);
+							if (Step.Outcome != JobStepOutcome.Success && Step.LogId != null && IncludedNodes.Contains(NodeRef))
+							{
+								LogIdValues.Add(Step.LogId.Value);
+							}
+						}
+					}
+				}
+				else
+				{
+					LogIdValues.UnionWith(Job.Batches.SelectMany(x => x.Steps).Where(x => x.Outcome != JobStepOutcome.Success && x.LogId != null).Select(x => x.LogId!.Value));
+				}
+			}
+			if(LogIds != null)
+			{
+				LogIdValues.UnionWith(LogIds.Select(x => new LogId(x)));
+			}
+
+			List<ILogEvent> Events = await IssueService.FindEventsForIssueAsync(IssueId, LogIdValues.ToArray(), Index, Count);
+
+			JobPermissionsCache PermissionsCache = new JobPermissionsCache();
+			Dictionary<LogId, ILogFile?> LogFiles = new Dictionary<LogId, ILogFile?>();
+
+			List<object> Responses = new List<object>();
+			foreach (ILogEvent Event in Events)
+			{
+				ILogFile? LogFile;
+				if (!LogFiles.TryGetValue(Event.LogId, out LogFile))
+				{
+					LogFile = await LogFileService.GetLogFileAsync(Event.LogId);
+					LogFiles[Event.LogId] = LogFile;
+				}
+				if (LogFile != null && await JobService.AuthorizeAsync(LogFile.JobId, AclAction.ViewLog, User, PermissionsCache))
+				{
+					ILogEventData Data = await LogFileService.GetEventDataAsync(LogFile, Event.LineIndex, Event.LineCount);
+					GetLogEventResponse Response = new GetLogEventResponse(Event, Data, IssueId);
+					Responses.Add(PropertyFilter.Apply(Response, Filter));
+				}
+			}
+			return Responses;
+		}
+
+		/// <summary>
+		/// Authorize the current user to see an issue
+		/// </summary>
+		/// <param name="Issue">The issue to authorize</param>
+		/// <param name="PermissionsCache">Cache of permissions</param>
+		/// <returns>True if the user is authorized to see the issue</returns>
+		private async Task<bool> AuthorizeIssue(IIssueDetails Issue, StreamPermissionsCache? PermissionsCache)
+		{
+			foreach (StreamId StreamId in Issue.Spans.Select(x => x.StreamId).Distinct())
+			{
+				if (await StreamService.AuthorizeAsync(StreamId, AclAction.ViewStream, User, PermissionsCache))
+				{
+					return true;
+				}
+			}
+			return false;
+		}
+
+		/// <summary>
+		/// Update an issue
+		/// </summary>
+		/// <param name="IssueId">Id of the issue to get information about</param>
+		/// <param name="Request">The update information</param>
+		/// <returns>List of matching agents</returns>
+		[HttpPut]
+		[Route("/api/v1/issues/{IssueId}")]
+		public async Task<ActionResult> UpdateIssueAsync(int IssueId, [FromBody] UpdateIssueRequest Request)
+		{
+			UserId? NewOwnerId = null;
+			if (Request.OwnerId != null)
+			{
+				NewOwnerId = Request.OwnerId.Length == 0 ? UserId.Empty : new UserId(Request.OwnerId);
+			}
+
+			UserId? NewNominatedById = null;
+			if (Request.NominatedById != null)
+			{
+				NewNominatedById = new UserId(Request.NominatedById);
+			}
+
+			UserId? NewDeclinedById = null;
+			if (Request.Declined ?? false)
+			{
+				NewDeclinedById = User.GetUserId();
+			}
+
+			UserId? NewResolvedById = null;
+			if (Request.Resolved.HasValue)
+			{
+				NewResolvedById = Request.Resolved.Value ? User.GetUserId() : UserId.Empty;
+			}
+
+			List<ObjectId>? AddSpans = null;
+			if (Request.AddSpans != null && Request.AddSpans.Count > 0)
+			{
+				AddSpans = Request.AddSpans.ConvertAll(x => ObjectId.Parse(x));
+			}
+
+			List<ObjectId>? RemoveSpans = null;
+			if (Request.RemoveSpans != null && Request.RemoveSpans.Count > 0)
+			{
+				RemoveSpans = Request.RemoveSpans.ConvertAll(x => ObjectId.Parse(x));
+			}
+
+			if (!await IssueService.UpdateIssueAsync(IssueId, Request.Summary, Request.Description, Request.Promoted, NewOwnerId, NewNominatedById, Request.Acknowledged, NewDeclinedById, Request.FixChange, NewResolvedById, AddSpans, RemoveSpans))
+			{
+				return NotFound();
+			}
+			return Ok();
+		}
+	}
+}
