Horde: Rename the HordeServer project to Horde.Build.

diff --git a/Engine/Source/Programs/Horde/Horde.Build/Controllers/StreamsController.cs b/Engine/Source/Programs/Horde/Horde.Build/Controllers/StreamsController.cs
new file mode 100644
--- /dev/null
+++ b/Engine/Source/Programs/Horde/Horde.Build/Controllers/StreamsController.cs
@@ -0,0 +1,401 @@
+// Copyright Epic Games, Inc. All Rights Reserved.
+
+using HordeServer.Api;
+using HordeServer.Collections;
+using HordeServer.Models;
+using HordeServer.Services;
+using HordeServer.Utilities;
+using Microsoft.AspNetCore.Authorization;
+using Microsoft.AspNetCore.Mvc;
+using System;
+using System.Collections.Generic;
+using System.Diagnostics.CodeAnalysis;
+using System.IO;
+using System.Linq;
+using System.Threading.Tasks;
+
+namespace HordeServer.Controllers
+{
+	using ProjectId = StringId<IProject>;
+	using StreamId = StringId<IStream>;
+	using TemplateRefId = StringId<TemplateRef>;
+
+	/// <summary>
+	/// Controller for the /api/v1/streams endpoint
+	/// </summary>
+	[ApiController]
+	[Authorize]
+	[Route("[controller]")]
+	public class StreamsController : ControllerBase
+	{
+		private readonly AclService AclService;
+		private readonly ProjectService ProjectService;
+		private readonly StreamService StreamService;
+		private readonly ITemplateCollection TemplateCollection;
+		private readonly JobService JobService;
+		private readonly IJobStepRefCollection JobStepRefCollection;
+		private readonly IPerforceService PerforceService;
+		private readonly IUserCollection UserCollection;
+
+		/// <summary>
+		/// Constructor
+		/// </summary>
+		public StreamsController(AclService AclService, ProjectService ProjectService, StreamService StreamService, ITemplateCollection TemplateCollection, JobService JobService, IJobStepRefCollection JobStepRefCollection, IPerforceService PerforceService, IUserCollection UserCollection)
+		{
+			this.AclService = AclService;
+			this.ProjectService = ProjectService;
+			this.StreamService = StreamService;
+			this.TemplateCollection = TemplateCollection;
+			this.JobService = JobService;
+			this.JobStepRefCollection = JobStepRefCollection;
+			this.PerforceService = PerforceService;
+			this.UserCollection = UserCollection;
+		}
+
+		/// <summary>
+		/// Query all the streams for a particular project.
+		/// </summary>
+		/// <param name="ProjectIds">Unique id of the project to query</param>
+		/// <param name="Filter">Filter for the properties to return</param>
+		/// <returns>Information about all the projects</returns>
+		[HttpGet]
+		[Route("/api/v1/streams")]
+		[ProducesResponseType(typeof(List<GetStreamResponse>), 200)]
+		public async Task<ActionResult<List<object>>> GetStreamsAsync([FromQuery(Name = "ProjectId")] string[] ProjectIds, [FromQuery] PropertyFilter? Filter = null)
+		{
+			ProjectId[] ProjectIdValues = Array.ConvertAll(ProjectIds, x => new ProjectId(x));
+
+			List<IStream> Streams = await StreamService.GetStreamsAsync(ProjectIdValues);
+			ProjectPermissionsCache PermissionsCache = new ProjectPermissionsCache();
+
+			List<GetStreamResponse> Responses = new List<GetStreamResponse>();
+			foreach (IStream Stream in Streams)
+			{
+				if (await StreamService.AuthorizeAsync(Stream, AclAction.ViewStream, User, PermissionsCache))
+				{
+					GetStreamResponse Response = await CreateGetStreamResponse(Stream, PermissionsCache);
+					Responses.Add(Response);
+				}
+			}
+			return Responses.OrderBy(x => x.Order).ThenBy(x => x.Name).Select(x => PropertyFilter.Apply(x, Filter)).ToList();
+		}
+
+		/// <summary>
+		/// Retrieve information about a specific stream.
+		/// </summary>
+		/// <param name="StreamId">Id of the stream to get information about</param>
+		/// <param name="Filter">Filter for the properties to return</param>
+		/// <returns>Information about the requested project</returns>
+		[HttpGet]
+		[Route("/api/v1/streams/{StreamId}")]
+		[ProducesResponseType(typeof(GetStreamResponse), 200)]
+		public async Task<ActionResult<object>> GetStreamAsync(string StreamId, [FromQuery] PropertyFilter? Filter = null)
+		{
+			StreamId StreamIdValue = new StreamId(StreamId);
+
+			IStream? Stream = await StreamService.GetStreamAsync(StreamIdValue);
+			if (Stream == null)
+			{
+				return NotFound();
+			}
+
+			ProjectPermissionsCache PermissionsCache = new ProjectPermissionsCache();
+			if (!await StreamService.AuthorizeAsync(Stream, AclAction.ViewStream, User, PermissionsCache))
+			{
+				return Forbid();
+			}
+
+			bool bIncludeAcl = await StreamService.AuthorizeAsync(Stream, AclAction.ViewPermissions, User, PermissionsCache);
+			return PropertyFilter.Apply(await CreateGetStreamResponse(Stream, PermissionsCache), Filter);
+		}
+
+		/// <summary>
+		/// Create a stream response object, including all the templates
+		/// </summary>
+		/// <param name="Stream">Stream to create response for</param>
+		/// <param name="Cache">Permissions cache</param>
+		/// <returns>Response object</returns>
+		async Task<GetStreamResponse> CreateGetStreamResponse(IStream Stream, ProjectPermissionsCache Cache)
+		{
+			bool bIncludeAcl = Stream.Acl != null && await StreamService.AuthorizeAsync(Stream, AclAction.ViewPermissions, User, Cache);
+
+			List<GetTemplateRefResponse> ApiTemplateRefs = new List<GetTemplateRefResponse>();
+			foreach (KeyValuePair<TemplateRefId, TemplateRef> Pair in Stream.Templates)
+			{
+				if (await StreamService.AuthorizeAsync(Stream, Pair.Value, AclAction.ViewTemplate, User, Cache))
+				{
+					ITemplate? Template = await TemplateCollection.GetAsync(Pair.Value.Hash);
+					if (Template != null)
+					{
+						bool bIncludeTemplateAcl = Pair.Value.Acl != null && await StreamService.AuthorizeAsync(Stream, Pair.Value, AclAction.ViewPermissions, User, Cache);
+						ApiTemplateRefs.Add(new GetTemplateRefResponse(Pair.Key, Pair.Value, Template, bIncludeTemplateAcl));
+					}
+				}
+			}
+
+			return Stream.ToApiResponse(bIncludeAcl, ApiTemplateRefs);
+		}
+
+		/// <summary>
+		/// Gets a list of changes for a stream
+		/// </summary>
+		/// <param name="StreamId">The stream id</param>
+		/// <param name="Min">The starting changelist number</param>
+		/// <param name="Max">The ending changelist number</param>
+		/// <param name="Results">Number of results to return</param>
+		/// <param name="Filter">The filter to apply to the results</param>
+		/// <returns>Http result code</returns>
+		[HttpGet]
+		[Route("/api/v1/streams/{StreamId}/changes")]
+		[ProducesResponseType(typeof(List<GetChangeSummaryResponse>), 200)]
+		public async Task<ActionResult<List<object>>> GetChangesAsync(string StreamId, [FromQuery] int? Min = null, [FromQuery] int? Max = null, [FromQuery] int Results = 50, PropertyFilter? Filter = null)
+		{
+			StreamId StreamIdValue = new StreamId(StreamId);
+
+			IStream? Stream = await StreamService.GetStreamAsync(StreamIdValue);
+			if (Stream == null)
+			{
+				return NotFound();
+			}
+			if (!await StreamService.AuthorizeAsync(Stream, AclAction.ViewChanges, User, null))
+			{
+				return Forbid();
+			}
+
+			string? PerforceUser = User.GetPerforceUser();
+			if(PerforceUser == null)
+			{
+				return Forbid();
+			}
+
+			List<ChangeSummary> Commits = await PerforceService.GetChangesAsync(Stream.ClusterName, Stream.Name, Min, Max, Results, PerforceUser);
+
+			List<GetChangeSummaryResponse> Responses = new List<GetChangeSummaryResponse>();
+			foreach (ChangeSummary Commit in Commits)
+			{
+				Responses.Add(new GetChangeSummaryResponse(Commit));
+			}
+			return Responses.ConvertAll(x => PropertyFilter.Apply(x, Filter));
+		}
+
+		/// <summary>
+		/// Gets a list of changes for a stream
+		/// </summary>
+		/// <param name="StreamId">The stream id</param>
+		/// <param name="Number">The changelist number</param>
+		/// <param name="Filter">The filter to apply to the results</param>
+		/// <returns>Http result code</returns>
+		[HttpGet]
+		[Route("/api/v1/streams/{StreamId}/changes/{Number}")]
+		[ProducesResponseType(typeof(GetChangeDetailsResponse), 200)]
+		public async Task<ActionResult<object>> GetChangeDetailsAsync(string StreamId, int Number, PropertyFilter? Filter = null)
+		{
+			StreamId StreamIdValue = new StreamId(StreamId);
+
+			IStream? Stream = await StreamService.GetStreamAsync(StreamIdValue);
+			if (Stream == null)
+			{
+				return NotFound();
+			}
+			if (!await StreamService.AuthorizeAsync(Stream, AclAction.ViewChanges, User, null))
+			{
+				return Forbid();
+			}
+
+			string? PerforceUser = User.GetPerforceUser();
+			if(PerforceUser == null)
+			{
+				return Forbid();
+			}
+
+			ChangeDetails? ChangeDetails = await PerforceService.GetChangeDetailsAsync(Stream.ClusterName, Stream.Name, Number, PerforceUser);
+			if(ChangeDetails == null)
+			{
+				return NotFound();
+			}
+
+			return PropertyFilter.Apply(new GetChangeDetailsResponse(ChangeDetails), Filter);
+		}
+
+		/// <summary>
+		/// Gets the history of a step in the stream
+		/// </summary>
+		/// <param name="StreamId">The stream id</param>
+		/// <param name="TemplateId"></param>
+		/// <param name="Step">Name of the step to search for</param>
+		/// <param name="Change">Maximum changelist number to return</param>
+		/// <param name="Count">Number of results to return</param>
+		/// <param name="Filter">The filter to apply to the results</param>
+		/// <returns>Http result code</returns>
+		[HttpGet]
+		[Route("/api/v1/streams/{StreamId}/history")]
+		[ProducesResponseType(typeof(List<GetJobStepRefResponse>), 200)]
+		public async Task<ActionResult<List<object>>> GetStepHistoryAsync(string StreamId, [FromQuery] string TemplateId, [FromQuery] string Step, [FromQuery] int? Change = null, [FromQuery] int Count = 10, [FromQuery] PropertyFilter? Filter = null)
+		{
+			StreamId StreamIdValue = new StreamId(StreamId);
+
+			IStream? Stream = await StreamService.GetStreamAsync(StreamIdValue);
+			if (Stream == null)
+			{
+				return NotFound();
+			}
+			if (!await StreamService.AuthorizeAsync(Stream, AclAction.ViewJob, User, null))
+			{
+				return Forbid();
+			}
+
+			TemplateRefId TemplateIdValue = new TemplateRefId(TemplateId);
+
+			List<IJobStepRef> Steps = await JobStepRefCollection.GetStepsForNodeAsync(StreamIdValue, TemplateIdValue, Step, Change, true, Count);
+			return Steps.ConvertAll(x => PropertyFilter.Apply(new GetJobStepRefResponse(x), Filter));
+		}
+
+		/// <summary>
+		/// Deletes a stream
+		/// </summary>
+		/// <param name="StreamId">Id of the stream to update.</param>
+		/// <returns>Http result code</returns>
+		[HttpDelete]
+		[Route("/api/v1/streams/{StreamId}")]
+		public async Task<ActionResult> DeleteStreamAsync(string StreamId)
+		{
+			StreamId StreamIdValue = new StreamId(StreamId);
+
+			IStream? Stream = await StreamService.GetStreamAsync(StreamIdValue);
+			if (Stream == null)
+			{
+				return NotFound();
+			}
+			if (!await StreamService.AuthorizeAsync(Stream, AclAction.DeleteStream, User, null))
+			{
+				return Forbid();
+			}
+
+			await StreamService.DeleteStreamAsync(StreamIdValue);
+			return new OkResult();
+		}
+
+		/// <summary>
+		/// Validates that the user is allowed to specify the given template refs
+		/// </summary>
+		/// <param name="Requests">List of template refs</param>
+		/// <param name="OutResult">On failure, receives information about the missing claim</param>
+		/// <returns>True if the refs are valid, false otherwise</returns>
+		bool TryAuthorizeTemplateClaims(List<CreateTemplateRefRequest>? Requests, [NotNullWhen(false)] out ActionResult? OutResult)
+		{
+			if (Requests != null)
+			{
+				foreach (CreateTemplateRefRequest Request in Requests)
+				{
+					if (Request.Schedule != null && Request.Schedule.Claims != null)
+					{
+						foreach (CreateAclClaimRequest Claim in Request.Schedule.Claims)
+						{
+							if (!User.HasClaim(Claim.Type, Claim.Value))
+							{
+								OutResult = BadRequest("User does not have the {Claim.Type}: {Claim.Value} claim");
+								return false;
+							}
+						}
+					}
+				}
+			}
+
+			OutResult = null;
+			return true;
+		}
+
+		/// <summary>
+		/// Creates a list of template refs from a set of request objects
+		/// </summary>
+		/// <param name="Requests">Request objects</param>
+		/// <param name="Stream">The current stream state</param>
+		/// <param name="TemplateCollection">The template service</param>
+		/// <returns>List of new template references</returns>
+		static async Task<Dictionary<TemplateRefId, TemplateRef>> CreateTemplateRefs(List<CreateTemplateRefRequest> Requests, IStream? Stream, ITemplateCollection TemplateCollection)
+		{
+			Dictionary<TemplateRefId, TemplateRef> NewTemplateRefs = new Dictionary<TemplateRefId, TemplateRef>();
+			foreach (CreateTemplateRefRequest Request in Requests)
+			{
+				// Create the template
+				ITemplate NewTemplate = await TemplateCollection.AddAsync(Request.Name, Request.Priority, Request.AllowPreflights, Request.InitialAgentType, Request.SubmitNewChange, Request.Arguments, Request.Parameters.ConvertAll(x => x.ToModel()));
+
+				// Get an identifier for the new template ref
+				TemplateRefId NewTemplateRefId;
+				if (Request.Id != null)
+				{
+					NewTemplateRefId = new TemplateRefId(Request.Id);
+				}
+				else
+				{
+					NewTemplateRefId = TemplateRefId.Sanitize(Request.Name);
+				}
+
+				// Add it to the list
+				TemplateRef NewTemplateRef = new TemplateRef(NewTemplate, Request.ShowUgsBadges, Request.ShowUgsAlerts, Request.NotificationChannel, Request.NotificationChannelFilter, Request.TriageChannel, Request.Schedule?.ToModel(), Request.ChainedJobs?.ConvertAll(x => new ChainedJobTemplate(x)), Acl.Merge(null, Request.Acl));
+				if (Stream != null && Stream.Templates.TryGetValue(NewTemplateRefId, out TemplateRef? OldTemplateRef))
+				{
+					if (OldTemplateRef.Schedule != null && NewTemplateRef.Schedule != null)
+					{
+						NewTemplateRef.Schedule.CopyState(OldTemplateRef.Schedule);
+					}
+				}
+				NewTemplateRefs.Add(NewTemplateRefId, NewTemplateRef);
+			}
+			foreach (TemplateRef TemplateRef in NewTemplateRefs.Values)
+			{
+				if (TemplateRef.ChainedJobs != null)
+				{
+					foreach (ChainedJobTemplate ChainedJob in TemplateRef.ChainedJobs)
+					{
+						if (!NewTemplateRefs.ContainsKey(ChainedJob.TemplateRefId))
+						{
+							throw new InvalidDataException($"Invalid template ref id '{ChainedJob.TemplateRefId}");
+						}
+					}
+				}
+			}
+			return NewTemplateRefs;
+		}
+
+		/// <summary>
+		/// Gets all the templates for a stream
+		/// </summary>
+		/// <param name="StreamId">Unique id of the stream to query</param>
+		/// <param name="TemplateRefId">Unique id of the template to query</param>
+		/// <param name="Filter">Filter for properties to return</param>
+		/// <returns>Information about all the templates</returns>
+		[HttpGet]
+		[Route("/api/v1/streams/{StreamId}/templates/{TemplateRefId}")]
+		[ProducesResponseType(typeof(List<GetTemplateResponse>), 200)]
+		public async Task<ActionResult<object>> GetTemplateAsync(string StreamId, string TemplateRefId, [FromQuery] PropertyFilter? Filter = null)
+		{
+			StreamId StreamIdValue = new StreamId(StreamId);
+			TemplateRefId TemplateRefIdValue = new TemplateRefId(TemplateRefId);
+
+			IStream? Stream = await StreamService.GetStreamAsync(StreamIdValue);
+			if (Stream == null)
+			{
+				return NotFound();
+			}
+
+			TemplateRef? TemplateRef;
+			if (!Stream.Templates.TryGetValue(TemplateRefIdValue, out TemplateRef))
+			{
+				return NotFound();
+			}
+			if (!await StreamService.AuthorizeAsync(Stream, TemplateRef, AclAction.ViewTemplate, User, null))
+			{
+				return Forbid();
+			}
+
+			ITemplate? Template = await TemplateCollection.GetAsync(TemplateRef.Hash);
+			if(Template == null)
+			{
+				return NotFound();
+			}
+
+			return new GetTemplateResponse(Template).ApplyFilter(Filter);
+		}
+	}
+}
diff --git a/Engine/Source/Programs/Horde/Horde.Build/Controllers/SubscriptionsController.cs b/Engine/Source/Programs/Horde/Horde.Build/Controllers/SubscriptionsController.cs
new file mode 100644
--- /dev/null
+++ b/Engine/Source/Programs/Horde/Horde.Build/Controllers/SubscriptionsController.cs
@@ -0,0 +1,192 @@
+// Copyright Epic Games, Inc. All Rights Reserved.
+
+using HordeServer.Api;
+using HordeServer.Collections;
+using HordeServer.Models;
+using HordeServer.Services;
+using HordeServer.Utilities;
+using Microsoft.AspNetCore.Authorization;
+using Microsoft.AspNetCore.Mvc;
+using MongoDB.Bson;
+using MongoDB.Driver;
+using System;
+using System.Collections.Generic;
+using System.Linq;
+using System.Threading.Tasks;
+
+namespace HordeServer.Controllers
+{
+	using UserId = ObjectId<IUser>;
+
+	/// <summary>
+	/// Controller for the /api/v1/agents endpoint
+	/// </summary>
+	[ApiController]
+	[Authorize]
+	[Route("[controller]")]
+	public class SubscriptionsController : ControllerBase
+	{
+		/// <summary>
+		/// The ACL service singleton
+		/// </summary>
+		AclService AclService;
+
+		/// <summary>
+		/// Collection of subscription documents
+		/// </summary>
+		ISubscriptionCollection SubscriptionCollection;
+
+		/// <summary>
+		/// Constructor
+		/// </summary>
+		/// <param name="AclService">The acl service singleton</param>
+		/// <param name="SubscriptionCollection">The collection of subscription documents</param>
+		public SubscriptionsController(AclService AclService, ISubscriptionCollection SubscriptionCollection)
+		{
+			this.AclService = AclService;
+			this.SubscriptionCollection = SubscriptionCollection;
+		}
+
+		/// <summary>
+		/// Find subscriptions matching a criteria
+		/// </summary>
+		/// <param name="UserId">Name of the user</param>
+		/// <param name="Filter">Filter for properties to return</param>
+		/// <returns>List of subscriptions</returns>
+		[HttpGet]
+		[Route("/api/v1/subscriptions")]
+		[ProducesResponseType(typeof(List<GetSubscriptionResponse>), 200)]
+		public async Task<ActionResult<List<object>>> GetSubscriptionsAsync([FromQuery] string UserId, [FromQuery] PropertyFilter? Filter = null)
+		{
+			UserId UserIdValue;
+			if (!TryParseUserId(UserId, out UserIdValue))
+			{
+				return BadRequest("Invalid user id");
+			}
+			if (!await AclService.AuthorizeAsUserAsync(User, UserIdValue))
+			{
+				return Forbid();
+			}
+
+			List<ISubscription> Results = await SubscriptionCollection.FindSubscriptionsAsync(UserIdValue);
+			return Results.ConvertAll(x => PropertyFilter.Apply(new GetSubscriptionResponse(x), Filter));
+		}
+
+		/// <summary>
+		/// Find subscriptions matching a criteria
+		/// </summary>
+		/// <param name="SubscriptionId">The subscription id</param>
+		/// <param name="Filter">Filter for properties to return</param>
+		/// <returns>List of subscriptions</returns>
+		[HttpGet]
+		[Route("/api/v1/subscriptions/{SubscriptionId}")]
+		[ProducesResponseType(typeof(GetSubscriptionResponse), 200)]
+		public async Task<ActionResult<object>> GetSubscriptionAsync(string SubscriptionId, [FromQuery] PropertyFilter? Filter = null)
+		{
+			ISubscription? Subscription = await SubscriptionCollection.GetAsync(SubscriptionId);
+			if (Subscription == null)
+			{
+				return NotFound();
+			}
+			if (!await AclService.AuthorizeAsUserAsync(User, Subscription.UserId))
+			{
+				return Forbid();
+			}
+
+			return PropertyFilter.Apply(new GetSubscriptionResponse(Subscription), Filter);
+		}
+
+
+		/// <summary>
+		/// Remove a subscription
+		/// </summary>
+		/// <param name="SubscriptionId">The subscription id</param>
+		/// <returns>Async task</returns>
+		[HttpDelete]
+		[Route("/api/v1/subscriptions/{SubscriptionId}")]
+		[ProducesResponseType(typeof(List<GetSubscriptionResponse>), 200)]
+		public async Task<ActionResult> DeleteSubscriptionAsync(string SubscriptionId)
+		{
+			ISubscription? Subscription = await SubscriptionCollection.GetAsync(SubscriptionId);
+			if (Subscription == null)
+			{
+				return NotFound();
+			}
+			if (!await AclService.AuthorizeAsUserAsync(User, Subscription.UserId))
+			{
+				return Forbid();
+			}
+
+			await SubscriptionCollection.RemoveAsync(new[] { Subscription });
+			return Ok();
+		}
+
+		/// <summary>
+		/// Find subscriptions matching a criteria
+		/// </summary>
+		/// <param name="Subscriptions">The new subscriptions to create</param>
+		/// <returns>List of subscriptions</returns>
+		[HttpPost]
+		[Route("/api/v1/subscriptions")]
+		public async Task<ActionResult<List<CreateSubscriptionResponse>>> CreateSubscriptionsAsync(List<CreateSubscriptionRequest> Subscriptions)
+		{
+			HashSet<UserId> AuthorizedUsers = new HashSet<UserId>();
+
+			UserId? CurrentUserId = User.GetUserId();
+			if(CurrentUserId != null)
+			{
+				AuthorizedUsers.Add(CurrentUserId.Value);
+			}
+
+			GlobalPermissionsCache Cache = new GlobalPermissionsCache();
+
+			List<NewSubscription> NewSubscriptions = new List<NewSubscription>();
+			foreach (CreateSubscriptionRequest Subscription in Subscriptions)
+			{
+				UserId NewUserId;
+				if (!TryParseUserId(Subscription.UserId, out NewUserId))
+				{
+					return BadRequest($"Invalid user id: '{Subscription.UserId}'.");
+				}
+				if (AuthorizedUsers.Add(NewUserId) && !await AclService.AuthorizeAsync(AclAction.Impersonate, User, Cache))
+				{
+					return Forbid();
+				}
+				NewSubscriptions.Add(new NewSubscription(Subscription.Event, NewUserId, Subscription.NotificationType));
+			}
+
+			List<ISubscription> Results = await SubscriptionCollection.AddAsync(NewSubscriptions);
+			return Results.ConvertAll(x => new CreateSubscriptionResponse(x));
+		}
+
+		/// <summary>
+		/// Parse a user id from a string. Allows passing the user's name as well as their objectid value.
+		/// </summary>
+		/// <param name="UserName"></param>
+		/// <param name="ObjectId"></param>
+		/// <returns></returns>
+		bool TryParseUserId(string UserName, out UserId ObjectId)
+		{
+			UserId NewObjectId;
+			if (UserId.TryParse(UserName, out NewObjectId))
+			{
+				ObjectId = NewObjectId;
+				return true;
+			}
+
+			string? CurrentUserName = User.GetUserName();
+			if (CurrentUserName != null && String.Equals(UserName, CurrentUserName, StringComparison.OrdinalIgnoreCase))
+			{
+				UserId? CurrentUserId = User.GetUserId();
+				if (CurrentUserId != null)
+				{
+					ObjectId = CurrentUserId.Value;
+					return true;
+				}
+			}
+
+			ObjectId = default;
+			return false;
+		}
+	}
+}
\ No newline at end of file
diff --git a/Engine/Source/Programs/Horde/Horde.Build/Controllers/SwarmController.cs b/Engine/Source/Programs/Horde/Horde.Build/Controllers/SwarmController.cs
new file mode 100644
--- /dev/null
+++ b/Engine/Source/Programs/Horde/Horde.Build/Controllers/SwarmController.cs
@@ -0,0 +1,45 @@
+// Copyright Epic Games, Inc. All Rights Reserved.
+
+using Amazon.Runtime.Internal.Util;
+using Microsoft.AspNetCore.Mvc;
+using Microsoft.Extensions.Logging;
+using System;
+using System.Collections.Generic;
+using System.Linq;
+using System.Threading.Tasks;
+
+namespace HordeServer.Controllers
+{
+	/// <summary>
+	/// Controller for the /api/v1/swarm endpoint
+	/// </summary>
+	[ApiController]
+	[Route("[controller]")]
+	public sealed class SwarmController : ControllerBase
+	{
+		/// <summary>
+		/// Logger instance
+		/// </summary>
+		ILogger<SwarmController> Logger;
+
+		/// <summary>
+		/// Constructor
+		/// </summary>
+		public SwarmController(ILogger<SwarmController> Logger)
+		{
+			this.Logger = Logger;
+		}
+
+		/// <summary>
+		/// Gets the latest version info
+		/// </summary>
+		/// <returns>Result code</returns>
+		[HttpPost]
+		[Route("/api/v1/swarm/reviews")]
+		public ActionResult AddReview(object Request)
+		{
+			Logger.LogInformation("Added Swarm review: {Request}", Request.ToString());
+			return Ok();
+		}
+	}
+}
diff --git a/Engine/Source/Programs/Horde/Horde.Build/Controllers/TestDataController.cs b/Engine/Source/Programs/Horde/Horde.Build/Controllers/TestDataController.cs
new file mode 100644
--- /dev/null
+++ b/Engine/Source/Programs/Horde/Horde.Build/Controllers/TestDataController.cs
@@ -0,0 +1,147 @@
+// Copyright Epic Games, Inc. All Rights Reserved.
+
+using HordeServer.Api;
+using HordeServer.Collections;
+using HordeCommon;
+using HordeServer.Models;
+using HordeServer.Services;
+using HordeServer.Utilities;
+using Microsoft.AspNetCore.Authorization;
+using Microsoft.AspNetCore.Mvc;
+using Microsoft.AspNetCore.Server.Kestrel.Core.Features;
+using MongoDB.Bson;
+using MongoDB.Bson.Serialization;
+using System;
+using System.Collections.Generic;
+using System.ComponentModel.DataAnnotations;
+using System.IO;
+using System.Linq;
+using System.Threading.Tasks;
+
+namespace HordeServer.Controllers
+{
+	using JobId = ObjectId<IJob>;
+	using StreamId = StringId<IStream>;
+
+	/// <summary>
+	/// Controller for the /api/v1/testdata endpoint
+	/// </summary>
+	[ApiController]
+	[Authorize]
+	[Route("[controller]")]
+	public class TestDataController : ControllerBase
+	{
+		/// <summary>
+		/// Collection of job documents
+		/// </summary>
+		private JobService JobService;
+
+		/// <summary>
+		/// Collection of test data documents
+		/// </summary>
+		private readonly ITestDataCollection TestDataCollection;
+
+		/// <summary>
+		/// Constructor
+		/// </summary>
+		/// <param name="JobService">The job service singleton</param>
+		/// <param name="TestDataCollection">Collection of test data documents</param>
+		public TestDataController(JobService JobService, ITestDataCollection TestDataCollection)
+		{
+			this.JobService = JobService;
+			this.TestDataCollection = TestDataCollection;
+		}
+
+		/// <summary>
+		/// Creates a new TestData document
+		/// </summary>
+		/// <returns>The stream document</returns>
+		[HttpPost]
+		[Route("/api/v1/testdata")]
+		public async Task<ActionResult<CreateTestDataResponse>> CreateAsync(CreateTestDataRequest Request)
+		{
+			IJob? Job = await JobService.GetJobAsync(new JobId(Request.JobId));
+			if (Job == null)
+			{
+				return NotFound();
+			}
+			if (!await JobService.AuthorizeAsync(Job, AclAction.UpdateJob, User, null))
+			{
+				return Forbid();
+			}
+
+			IJobStep? JobStep;
+			if (!Job.TryGetStep(Request.StepId.ToSubResourceId(), out JobStep))
+			{
+				return NotFound();
+			}
+
+			ITestData TestData = await TestDataCollection.AddAsync(Job, JobStep, Request.Key, new BsonDocument(Request.Data));
+			return new CreateTestDataResponse(TestData.Id.ToString());
+		}
+
+		/// <summary>
+		/// Searches for test data that matches a set of criteria
+		/// </summary>
+		/// <param name="StreamId">The stream id</param>
+		/// <param name="MinChange">The minimum changelist number to return (inclusive)</param>
+		/// <param name="MaxChange">The maximum changelist number to return (inclusive)</param>
+		/// <param name="JobId">The job id</param>
+		/// <param name="JobStepId">The unique step id</param>
+		/// <param name="Key">Key identifying the result to return</param>
+		/// <param name="Index">Offset within the results to return</param>
+		/// <param name="Count">Number of results to return</param>
+		/// <param name="Filter">Filter for properties to return</param>
+		/// <returns>The stream document</returns>
+		[HttpGet]
+		[Route("/api/v1/testdata")]
+		[ProducesResponseType(typeof(List<GetTestDataResponse>), 200)]
+		public async Task<ActionResult<List<object>>> FindTestDataAsync([FromQuery] string? StreamId = null, [FromQuery] int? MinChange = null, [FromQuery] int? MaxChange = null, string? JobId = null, string? JobStepId = null, string? Key = null, int Index = 0, int Count = 10, PropertyFilter? Filter = null)
+		{
+			StreamId? StreamIdValue = null;
+			if(StreamId != null)
+			{
+				StreamIdValue = new StreamId(StreamId);
+			}
+
+			JobPermissionsCache Cache = new JobPermissionsCache();
+
+			List<object> Results = new List<object>();
+
+			List<ITestData> Documents = await TestDataCollection.FindAsync(StreamIdValue, MinChange, MaxChange, JobId?.ToObjectId<IJob>(), JobStepId?.ToSubResourceId(), Key, Index, Count);
+			foreach (ITestData Document in Documents)
+			{
+				if (await JobService.AuthorizeAsync(Document.JobId, AclAction.ViewJob, User, Cache))
+				{
+					Results.Add(PropertyFilter.Apply(new GetTestDataResponse(Document), Filter));
+				}
+			}
+
+			return Results;
+		}
+
+		/// <summary>
+		/// Retrieve information about a specific issue
+		/// </summary>
+		/// <param name="TestDataId">Id of the document to get information about</param>
+		/// <param name="Filter">Filter for the properties to return</param>
+		/// <returns>List of matching agents</returns>
+		[HttpGet]
+		[Route("/api/v1/testdata/{TestDataId}")]
+		[ProducesResponseType(typeof(GetTestDataResponse), 200)]
+		public async Task<ActionResult<object>> GetTestDataAsync(string TestDataId, [FromQuery] PropertyFilter? Filter = null)
+		{
+			ITestData? TestData = await TestDataCollection.GetAsync(TestDataId.ToObjectId());
+			if (TestData == null)
+			{
+				return NotFound();
+			}
+			if (!await JobService.AuthorizeAsync(TestData.JobId, AclAction.ViewJob, User, null))
+			{
+				return Forbid();
+			}
+
+			return PropertyFilter.Apply(new GetTestDataResponse(TestData), Filter);
+		}
+	}
+}
diff --git a/Engine/Source/Programs/Horde/Horde.Build/Controllers/UgsController.cs b/Engine/Source/Programs/Horde/Horde.Build/Controllers/UgsController.cs
new file mode 100644
--- /dev/null
+++ b/Engine/Source/Programs/Horde/Horde.Build/Controllers/UgsController.cs
@@ -0,0 +1,432 @@
+// Copyright Epic Games, Inc. All Rights Reserved.
+
+using HordeServer.Api;
+using HordeServer.Collections;
+using HordeCommon;
+using HordeServer.Models;
+using HordeServer.Services;
+using HordeServer.Utilities;
+using Microsoft.AspNetCore.Authorization;
+using Microsoft.AspNetCore.Mvc;
+using Microsoft.AspNetCore.Server.Kestrel.Core.Features;
+using Microsoft.Extensions.Caching.Memory;
+using Microsoft.Extensions.Options;
+using MongoDB.Bson;
+using Serilog.Events;
+using System;
+using System.Collections.Generic;
+using System.ComponentModel.DataAnnotations;
+using System.IO;
+using System.Linq;
+using System.Threading.Tasks;
+using Microsoft.Extensions.Logging;
+
+namespace HordeServer.Controllers
+{
+	using LogId = ObjectId<ILogFile>;
+	using StreamId = StringId<IStream>;
+	using UserId = ObjectId<IUser>;
+
+	/// <summary>
+	/// Controller for the /api/v1/issues endpoint
+	/// </summary>
+	[ApiController]
+	[Route("[controller]")]
+	public sealed class UgsController : ControllerBase
+	{
+		/// <summary>
+		/// Singleton instance of the issue service
+		/// </summary>
+		private readonly IIssueService IssueService;
+
+		/// <summary>
+		/// Singleton instance of the job service
+		/// </summary>
+		private readonly JobService JobService;
+
+		/// <summary>
+		/// Collection of stream documents
+		/// </summary>
+		private readonly StreamService StreamService;
+
+		/// <summary>
+		/// Collection of metadata documents
+		/// </summary>
+		private readonly IUgsMetadataCollection UgsMetadataCollection;
+
+		/// <summary>
+		/// Collection of log events
+		/// </summary>
+		private readonly ILogEventCollection LogEventCollection;
+
+		/// <summary>
+		/// Collection of users
+		/// </summary>
+		private readonly IUserCollection UserCollection;
+
+		/// <summary>
+		/// The log file service
+		/// </summary>
+		private readonly ILogFileService LogFileService;
+
+		/// <summary>
+		/// Server settings
+		/// </summary>
+		private readonly ServerSettings Settings;
+
+		/// <summary>
+		/// Logger 
+		/// </summary>
+		ILogger Logger;
+
+		/// <summary>
+		/// Constructor
+		/// </summary>
+		/// <param name="IssueService">The issue service</param>
+		/// <param name="JobService">The job service</param>
+		/// <param name="StreamService">Collection of stream documents</param>
+		/// <param name="UgsMetadataCollection">Collection of UGS metadata documents</param>
+		/// <param name="LogEventCollection">Collection of log event documents</param>
+		/// <param name="UserCollection">Collection of user documents</param>
+		/// <param name="LogFileService">The log file service</param>
+		/// <param name="OptionsMonitor">The server settings</param>
+		/// <param name="Logger">Logger</param>
+		public UgsController(IIssueService IssueService, JobService JobService, StreamService StreamService, IUgsMetadataCollection UgsMetadataCollection, ILogEventCollection LogEventCollection, IUserCollection UserCollection, ILogFileService LogFileService, IOptionsMonitor<ServerSettings> OptionsMonitor, ILogger<UgsController> Logger)
+		{
+			this.IssueService = IssueService;
+			this.JobService = JobService;
+			this.StreamService = StreamService;
+			this.UgsMetadataCollection = UgsMetadataCollection;
+			this.LogEventCollection = LogEventCollection;
+			this.UserCollection = UserCollection;
+			this.LogFileService = LogFileService;
+			this.Settings = OptionsMonitor.CurrentValue;
+			this.Logger = Logger;
+		}
+
+		/// <summary>
+		/// Gets the latest version info
+		/// </summary>
+		/// <returns>Result code</returns>
+		[HttpGet]
+		[Route("/ugs/api/latest")]
+		public ActionResult<object> GetLatest()
+		{
+			return new { Version = 2 };
+		}
+
+		/// <summary>
+		/// Adds new metadata to the database
+		/// </summary>
+		/// <param name="Request">Request object</param>
+		/// <returns>Result code</returns>
+		[HttpPost]
+		[Route("/ugs/api/metadata")]
+		public async Task<ActionResult> AddMetadataAsync(AddUgsMetadataRequest Request)
+		{
+			IUgsMetadata Metadata = await UgsMetadataCollection.FindOrAddAsync(Request.Stream, Request.Change, Request.Project);
+			if (Request.Synced != null || Request.Vote != null || Request.Investigating != null || Request.Starred != null || Request.Comment != null)
+			{
+				if (Request.UserName == null)
+				{
+					return BadRequest("Missing UserName field on request body");
+				}
+				Metadata = await UgsMetadataCollection.UpdateUserAsync(Metadata, Request.UserName, Request.Synced, Request.Vote, Request.Investigating, Request.Starred, Request.Comment);
+			}
+			if (Request.Badges != null)
+			{
+				foreach (AddUgsBadgeRequest Badge in Request.Badges)
+				{
+					Metadata = await UgsMetadataCollection.UpdateBadgeAsync(Metadata, Badge.Name, Badge.Url, Badge.State);
+				}
+			}
+			return Ok();
+		}
+
+		/// <summary>
+		/// Searches for metadata updates
+		/// </summary>
+		/// <param name="Stream">THe stream to search for</param>
+		/// <param name="MinChange">Minimum changelist number</param>
+		/// <param name="MaxChange">Maximum changelist number</param>
+		/// <param name="Project">The project identifiers to search for</param>
+		/// <param name="Sequence">Last sequence number</param>
+		/// <returns>List of metadata updates</returns>
+		[HttpGet]
+		[Route("/ugs/api/metadata")]
+		public async Task<GetUgsMetadataListResponse> FindMetadataAsync([FromQuery] string Stream, [FromQuery] int MinChange, [FromQuery] int? MaxChange = null, [FromQuery] string? Project = null, [FromQuery] long? Sequence = null)
+		{
+			List<IUgsMetadata> MetadataList = await UgsMetadataCollection.FindAsync(Stream, MinChange, MaxChange, Sequence);
+		
+			GetUgsMetadataListResponse Response = new GetUgsMetadataListResponse();
+			if(Sequence != null)
+			{
+				Response.SequenceNumber = Sequence.Value;
+			}
+
+			foreach (IUgsMetadata Metadata in MetadataList)
+			{
+				if (Metadata.UpdateTicks > Response.SequenceNumber)
+				{
+					Response.SequenceNumber = Metadata.UpdateTicks;
+				}
+				if (String.IsNullOrEmpty(Metadata.Project) || Metadata.Project.Equals(Project, StringComparison.OrdinalIgnoreCase))
+				{
+					Response.Items.Add(new GetUgsMetadataResponse(Metadata));
+				}
+			}
+			return Response;
+		}
+
+		/// <summary>
+		/// Retrieve information about open issues
+		/// </summary>
+		/// <param name="User"></param>
+		/// <param name="IncludeResolved">Whether to include resolved issues</param>
+		/// <param name="MaxResults">Maximum number of results to return</param>
+		/// <returns>List of matching agents</returns>
+		[HttpGet]
+		[Route("/ugs/api/issues")]
+		[ProducesResponseType(typeof(GetUgsIssueResponse), 200)]
+		[System.Diagnostics.CodeAnalysis.SuppressMessage("Design", "CA1801:Remove unused parameter", Justification = "<Pending>")]
+		[System.Diagnostics.CodeAnalysis.SuppressMessage("Style", "IDE0060:Remove unused parameter", Justification = "<Pending>")]
+		public async Task<ActionResult<List<GetUgsIssueResponse>>> GetIssuesAsync([FromQuery] string? User = null, [FromQuery] bool IncludeResolved = false, [FromQuery] int MaxResults = 100)
+		{
+			IUser? UserInfo = (User != null) ? await UserCollection.FindUserByLoginAsync(User) : null;
+
+			List<GetUgsIssueResponse> Responses = new List<GetUgsIssueResponse>();
+			if (IncludeResolved)
+			{
+				List<IIssue> Issues = await IssueService.FindIssuesAsync(null, Resolved: null, Count: MaxResults);
+				foreach(IIssue Issue in Issues)
+				{
+					IIssueDetails Details = await IssueService.GetIssueDetailsAsync(Issue);
+					bool bNotify = UserInfo != null && Details.Suspects.Any(x => x.AuthorId == UserInfo.Id);
+					Responses.Add(await CreateIssueResponseAsync(Details, bNotify));
+				}
+			}
+			else
+			{
+				foreach (IIssueDetails CachedOpenIssue in IssueService.CachedOpenIssues)
+				{
+					if (Responses.Count >= MaxResults)
+					{
+						break;
+					}
+
+					if (CachedOpenIssue.ShowNotifications())
+					{
+						bool bNotify = UserInfo != null && CachedOpenIssue.IncludeForUser(UserInfo.Id);
+						Responses.Add(await CreateIssueResponseAsync(CachedOpenIssue, bNotify));
+					}
+				}
+			}
+			return Responses;
+		}
+
+		/// <summary>
+		/// Retrieve information about a specific issue
+		/// </summary>
+		/// <param name="IssueId">Id of the issue to get information about</param>
+		/// <param name="Filter">Filter for the properties to return</param>
+		/// <returns>List of matching agents</returns>
+		[HttpGet]
+		[Route("/ugs/api/issues/{IssueId}")]
+		[ProducesResponseType(typeof(GetUgsIssueBuildResponse), 200)]
+		public async Task<ActionResult<object>> GetIssueAsync(int IssueId, [FromQuery] PropertyFilter? Filter = null)
+		{
+			IIssueDetails? Issue = await IssueService.GetIssueDetailsAsync(IssueId);
+			if (Issue == null)
+			{
+				return NotFound();
+			}
+
+			return PropertyFilter.Apply(await CreateIssueResponseAsync(Issue, false), Filter);
+		}
+
+		/// <summary>
+		/// Retrieve information about builds for a specific issue
+		/// </summary>
+		/// <param name="IssueId">Id of the issue to get information about</param>
+		/// <param name="Filter">Filter for the properties to return</param>
+		/// <returns>List of matching agents</returns>
+		[HttpGet]
+		[Route("/ugs/api/issues/{IssueId}/builds")]
+		[ProducesResponseType(typeof(List<GetUgsIssueBuildResponse>), 200)]
+		public async Task<ActionResult<List<object>>> GetIssueBuildsAsync(int IssueId, [FromQuery] PropertyFilter? Filter = null)
+		{
+			IIssueDetails? Issue = await IssueService.GetCachedIssueDetailsAsync(IssueId);
+			if (Issue == null)
+			{
+				return NotFound();
+			}
+
+			List<GetUgsIssueBuildResponse> Responses = new List<GetUgsIssueBuildResponse>();
+			foreach (IIssueSpan Span in Issue.Spans)
+			{
+				if (Span.LastSuccess != null)
+				{
+					Responses.Add(CreateBuildResponse(Span, Span.LastSuccess, IssueBuildOutcome.Success));
+				}
+				foreach (IIssueStep Step in Issue.Steps)
+				{
+					Responses.Add(CreateBuildResponse(Span, Step, IssueBuildOutcome.Error));
+				}
+				if (Span.NextSuccess != null)
+				{
+					Responses.Add(CreateBuildResponse(Span, Span.NextSuccess, IssueBuildOutcome.Success));
+				}
+			}
+
+			return Responses.ConvertAll(x => PropertyFilter.Apply(x, Filter));
+		}
+
+		/// <summary>
+		/// Retrieve information about builds for a specific issue
+		/// </summary>
+		/// <param name="IssueId">Id of the issue to get information about</param>
+		/// <returns>List of matching agents</returns>
+		[HttpGet]
+		[Route("/ugs/api/issues/{IssueId}/diagnostics")]
+		[ProducesResponseType(typeof(List<GetUgsIssueDiagnosticResponse>), 200)]
+		public async Task<ActionResult<List<GetUgsIssueDiagnosticResponse>>> GetIssueDiagnosticsAsync(int IssueId)
+		{
+			List<GetUgsIssueDiagnosticResponse> Diagnostics = new List<GetUgsIssueDiagnosticResponse>();
+
+			Dictionary<LogId, ILogFile?> LogFiles = new Dictionary<LogId, ILogFile?>();
+
+			JobPermissionsCache PermissionsCache = new JobPermissionsCache();
+
+			List<ILogEvent> Events = await IssueService.FindEventsForIssueAsync(IssueId, Count: 10);
+			foreach (ILogEvent Event in Events)
+			{
+				ILogFile? LogFile;
+				if(!LogFiles.TryGetValue(Event.LogId, out LogFile))
+				{
+					LogFile = await LogFileService.GetLogFileAsync(Event.LogId);
+					LogFiles.Add(Event.LogId, LogFile);
+				}
+				if (LogFile != null)
+				{
+					ILogEventData EventData = await LogFileService.GetEventDataAsync(LogFile, Event.LineIndex, Event.LineCount);
+					long BuildId = Event.LogId.GetHashCode();
+					Uri Url = new Uri(Settings.DashboardUrl, $"log/{Event.LogId}?lineindex={Event.LineIndex}");
+
+					GetUgsIssueDiagnosticResponse Diagnostic = new GetUgsIssueDiagnosticResponse(BuildId, EventData.Message, Url);
+					Diagnostics.Add(Diagnostic);
+				}
+			}
+
+			return Diagnostics;
+		}
+
+		/// <summary>
+		/// Gets the URL for a failing step in the
+		/// </summary>
+		/// <param name="Details">The issue to get a URL for</param>
+		/// <param name="bNotify">Whether to show notifications for this issue</param>
+		/// <returns>The issue response</returns>
+		async Task<GetUgsIssueResponse> CreateIssueResponseAsync(IIssueDetails Details, bool bNotify)
+		{
+			Uri? BuildUrl = GetIssueBuildUrl(Details);
+
+			IUser? Owner = Details.Issue.OwnerId.HasValue ? await UserCollection.GetCachedUserAsync(Details.Issue.OwnerId.Value) : null;
+			IUser? NominatedBy = Details.Issue.NominatedById.HasValue ? await UserCollection.GetCachedUserAsync(Details.Issue.NominatedById.Value) : null;
+
+			return new GetUgsIssueResponse(Details, Owner, NominatedBy, bNotify, BuildUrl);
+		}
+
+		/// <summary>
+		/// Gets the URL for a failing step in the given issue
+		/// </summary>
+		/// <param name="Issue">The issue details</param>
+		/// <returns>The build URL</returns>
+		Uri? GetIssueBuildUrl(IIssueDetails Issue)
+		{
+			HashSet<ObjectId> UnresolvedSpans = new HashSet<ObjectId>(Issue.Spans.Where(x => x.NextSuccess == null).Select(x => x.Id));
+
+			IIssueStep? Step = Issue.Steps.OrderByDescending(x => UnresolvedSpans.Contains(x.SpanId)).ThenByDescending(x => x.Change).FirstOrDefault();
+			if (Step == null)
+			{
+				return null;
+			}
+
+			return new Uri(Settings.DashboardUrl, $"job/{Step.JobId}?step={Step.StepId}");
+		}
+
+		/// <summary>
+		/// Creates the response for a particular build
+		/// </summary>
+		/// <param name="Span">Span containing the step</param>
+		/// <param name="Step">The step to describe</param>
+		/// <param name="Outcome">Outcome of this step</param>
+		/// <returns>Response object</returns>
+		GetUgsIssueBuildResponse CreateBuildResponse(IIssueSpan Span, IIssueStep Step, IssueBuildOutcome Outcome)
+		{
+			GetUgsIssueBuildResponse Response = new GetUgsIssueBuildResponse(Span.StreamName, Step.Change, Outcome);
+			Response.Id = Step.LogId.GetHashCode();
+			Response.JobName = $"{Step.JobName}: {Span.NodeName}";
+			Response.JobUrl = new Uri(Settings.DashboardUrl, $"job/{Step.JobId}");
+			Response.JobStepName = Span.NodeName;
+			Response.JobStepUrl = new Uri(Settings.DashboardUrl, $"job/{Step.JobId}?step={Step.StepId}");
+			Response.ErrorUrl = Response.JobStepUrl;
+			return Response;
+		}
+
+		/// <summary>
+		/// Update an issue
+		/// </summary>
+		/// <param name="IssueId">Id of the issue to get information about</param>
+		/// <param name="Request">The update information</param>
+		/// <returns>List of matching agents</returns>
+		[HttpPut]
+		[Route("/ugs/api/issues/{IssueId}")]
+		public async Task<ActionResult> UpdateIssueAsync(int IssueId, [FromBody] UpdateUgsIssueRequest Request)
+		{
+			UserId? NewOwnerId = null;
+			if (!String.IsNullOrEmpty(Request.Owner))
+			{
+				NewOwnerId = (await UserCollection.FindOrAddUserByLoginAsync(Request.Owner))?.Id;
+			}
+
+			UserId? NewNominatedById = null;
+			if (!String.IsNullOrEmpty(Request.NominatedBy))
+			{
+				NewNominatedById = (await UserCollection.FindOrAddUserByLoginAsync(Request.NominatedBy))?.Id;
+			}
+
+			UserId? NewDeclinedById = null;
+			if (!String.IsNullOrEmpty(Request.DeclinedBy))
+			{
+				NewDeclinedById = (await UserCollection.FindOrAddUserByLoginAsync(Request.DeclinedBy))?.Id;
+			}
+
+			UserId? NewResolvedById = null;
+			if (!String.IsNullOrEmpty(Request.ResolvedBy))
+			{
+				NewResolvedById = (await UserCollection.FindOrAddUserByLoginAsync(Request.ResolvedBy))?.Id;
+			}
+			if (NewResolvedById == null && Request.Resolved.HasValue)
+			{
+				NewResolvedById = Request.Resolved.Value ? IIssue.ResolvedByUnknownId : UserId.Empty;
+			}
+
+			if (!await IssueService.UpdateIssueAsync(IssueId, OwnerId: NewOwnerId, NominatedById: NewNominatedById, Acknowledged: Request.Acknowledged, DeclinedById: NewDeclinedById, FixChange: Request.FixChange, ResolvedById: NewResolvedById))
+			{
+				return NotFound();
+			}
+			return Ok();
+		}
+
+		/// <summary>
+		/// Post information about net core installation. 
+		/// </summary>
+		[HttpPost]
+		[Route("/ugs/api/netcore")]
+		public ActionResult<object> PostNetCoreInfo(string? User = null, string? Machine = null, bool NetCore = false)
+		{
+			Logger.LogInformation("NetCore: User={User}, Machine={Machine}, NetCore={NetCore}", User, Machine, NetCore);
+			return new object();
+		}
+	}
+}
diff --git a/Engine/Source/Programs/Horde/Horde.Build/Controllers/UserController.cs b/Engine/Source/Programs/Horde/Horde.Build/Controllers/UserController.cs
new file mode 100644
--- /dev/null
+++ b/Engine/Source/Programs/Horde/Horde.Build/Controllers/UserController.cs
@@ -0,0 +1,88 @@
+// Copyright Epic Games, Inc. All Rights Reserved.
+
+using HordeServer.Api;
+using HordeServer.Collections;
+using HordeServer.Models;
+using HordeServer.Services;
+using HordeServer.Utilities;
+using Microsoft.AspNetCore.Authorization;
+using Microsoft.AspNetCore.Mvc;
+using MongoDB.Bson;
+using System;
+using System.Collections.Generic;
+using System.Linq;
+using System.Threading.Tasks;
+
+namespace HordeServer.Controllers
+{
+	using JobId = ObjectId<IJob>;
+	using UserId = ObjectId<IUser>;
+
+	/// <summary>
+	/// Controller for the /api/v1/user endpoint
+	/// </summary>
+	[ApiController]
+	[Authorize]
+	[Route("[controller]")]
+	public class UserController : ControllerBase
+	{
+		/// <summary>
+		/// The user collection instance
+		/// </summary>
+		IUserCollection UserCollection { get; set; }
+
+		/// <summary>
+		/// The avatar service
+		/// </summary>
+		IAvatarService? AvatarService { get; set; }
+
+		/// <summary>
+		/// Constructor
+		/// </summary>
+		/// <param name="UserCollection"></param>
+		/// <param name="AvatarService"></param>
+		public UserController(IUserCollection UserCollection, IAvatarService? AvatarService)
+		{
+			this.UserCollection = UserCollection;
+			this.AvatarService = AvatarService;
+		}
+
+		/// <summary>
+		/// Gets information about the logged in user
+		/// </summary>
+		/// <returns>Http result code</returns>
+		[HttpGet]
+		[Route("/api/v1/user")]
+		[ProducesResponseType(typeof(List<GetUserResponse>), 200)]
+		public async Task<ActionResult<object>> GetUserAsync([FromQuery] PropertyFilter? Filter = null)
+		{
+			IUser? InternalUser = await UserCollection.GetUserAsync(User);
+			if (InternalUser == null)
+			{
+				return NotFound();
+			}
+
+			IAvatar? Avatar = (AvatarService == null)? (IAvatar?)null : await AvatarService.GetAvatarAsync(InternalUser);
+			IUserClaims Claims = await UserCollection.GetClaimsAsync(InternalUser.Id);
+			IUserSettings Settings = await UserCollection.GetSettingsAsync(InternalUser.Id);
+			return PropertyFilter.Apply(new GetUserResponse(InternalUser, Avatar, Claims, Settings), Filter);
+		}
+
+		/// <summary>
+		/// Updates the logged in user
+		/// </summary>
+		/// <returns>Http result code</returns>
+		[HttpPut]
+		[Route("/api/v1/user")]
+		public async Task<ActionResult> UpdateUserAsync(UpdateUserRequest Request)
+		{
+			UserId? UserId = User.GetUserId();
+			if(UserId == null)
+			{
+				return BadRequest("Current user does not have a registered profile");
+			}
+			await UserCollection.UpdateSettingsAsync(UserId.Value, Request.EnableExperimentalFeatures, Request.EnableIssueNotifications, Request.DashboardSettings?.ToBsonValue(), Request.AddPinnedJobIds?.Select(x => new JobId(x)), Request.RemovePinnedJobIds?.Select(x => new JobId(x)));
+			return Ok();
+		}
+	}
+}
diff --git a/Engine/Source/Programs/Horde/Horde.Build/Fleet/IPoolCollection.cs b/Engine/Source/Programs/Horde/Horde.Build/Fleet/IPoolCollection.cs
new file mode 100644
--- /dev/null
+++ b/Engine/Source/Programs/Horde/Horde.Build/Fleet/IPoolCollection.cs
@@ -0,0 +1,81 @@
+// Copyright Epic Games, Inc. All Rights Reserved.
+
+using EpicGames.Horde.Common;
+using HordeServer.Models;
+using HordeServer.Services;
+using HordeServer.Utilities;
+using MongoDB.Bson;
+using MongoDB.Bson.Serialization;
+using MongoDB.Bson.Serialization.Attributes;
+using MongoDB.Driver;
+using System;
+using System.Collections.Generic;
+using System.Linq;
+using System.Linq.Expressions;
+using System.Threading.Tasks;
+
+namespace HordeServer.Collections
+{
+	using PoolId = StringId<IPool>;
+
+	/// <summary>
+	/// Collection of pool documents
+	/// </summary>
+	public interface IPoolCollection
+	{
+		/// <summary>
+		/// Creates a new pool
+		/// </summary>
+		/// <param name="Id">Unique id for the new pool</param>
+		/// <param name="Name">Name of the new pool</param>
+		/// <param name="Condition">Condition for agents to be included in this pool</param>
+		/// <param name="EnableAutoscaling">Whether to enable autoscaling for this pool</param>
+		/// <param name="MinAgents">Minimum number of agents in the pool</param>
+		/// <param name="NumReserveAgents">Minimum number of idle agents to maintain</param>
+		/// <param name="Properties">Properties for the pool</param>
+		/// <returns>The new pool document</returns>
+		Task<IPool> AddAsync(PoolId Id, string Name, Condition? Condition = null, bool? EnableAutoscaling = null, int? MinAgents = null, int? NumReserveAgents = null, IEnumerable<KeyValuePair<string, string>>? Properties = null);
+
+		/// <summary>
+		/// Enumerates all the pools
+		/// </summary>
+		/// <returns>The pool documents</returns>
+		Task<List<IPool>> GetAsync();
+
+		/// <summary>
+		/// Gets a pool by ID
+		/// </summary>
+		/// <param name="Id">Unique id of the pool</param>
+		/// <returns>The pool document</returns>
+		Task<IPool?> GetAsync(PoolId Id);
+
+		/// <summary>
+		/// Gets a list of all valid pool ids
+		/// </summary>
+		/// <returns>List of pool ids</returns>
+		Task<List<PoolId>> GetPoolIdsAsync();
+
+		/// <summary>
+		/// Gets a pool by ID
+		/// </summary>
+		/// <param name="Id">Unique id of the pool</param>
+		/// <returns>The pool document</returns>
+		Task<bool> DeleteAsync(PoolId Id);
+
+		/// <summary>
+		/// Updates a pool
+		/// </summary>
+		/// <param name="Pool">The pool to update</param>
+		/// <param name="NewName">New name for the pool</param>
+		/// <param name="NewCondition">New condition for the pool</param>
+		/// <param name="NewEnableAutoscaling">New setting for whether to enable autoscaling</param>
+		/// <param name="NewMinAgents">Minimum number of agents in the pool</param>
+		/// <param name="NewNumReserveAgents">Minimum number of idle agents to maintain</param>
+		/// <param name="NewWorkspaces">New workspaces for the pool</param>
+		/// <param name="NewProperties">New properties for the pool</param>
+		/// <param name="LastScaleUpTime">New time for last (auto) scale up</param>
+		/// <param name="LastScaleDownTime">New time for last (auto) scale down</param>
+		/// <returns>Async task</returns>
+		Task<IPool?> TryUpdateAsync(IPool Pool, string? NewName = null, Condition? NewCondition = null, bool? NewEnableAutoscaling = null, int? NewMinAgents = null, int? NewNumReserveAgents = null, List<AgentWorkspace>? NewWorkspaces = null, Dictionary<string, string?>? NewProperties = null, DateTime? LastScaleUpTime = null, DateTime? LastScaleDownTime = null);
+	}
+}
diff --git a/Engine/Source/Programs/Horde/Horde.Build/Fleet/ITelemetryCollection.cs b/Engine/Source/Programs/Horde/Horde.Build/Fleet/ITelemetryCollection.cs
new file mode 100644
--- /dev/null
+++ b/Engine/Source/Programs/Horde/Horde.Build/Fleet/ITelemetryCollection.cs
@@ -0,0 +1,41 @@
+// Copyright Epic Games, Inc. All Rights Reserved.
+
+using HordeServer.Models;
+using HordeServer.Utilities;
+using System;
+using System.Collections.Generic;
+using System.Linq;
+using System.Threading.Tasks;
+
+namespace HordeServer.Collections
+{
+	using PoolId = StringId<IPool>;
+	using StreamId = StringId<IStream>;
+
+	/// <summary>
+	/// Collection of utilization collection
+	/// </summary>
+	public interface ITelemetryCollection
+	{
+		/// <summary>
+		/// Adds entries for the given utilization
+		/// </summary>
+		/// <param name="Telemetry">Telemetry data to add</param>
+		/// <returns>Async task</returns>
+		Task AddUtilizationTelemetryAsync(IUtilizationTelemetry Telemetry);
+
+		/// <summary>
+		/// Finds utilization data matching the given criteria
+		/// </summary>
+		/// <param name="StartTimeUtc">Start time to query utilization for</param>
+		/// <param name="FinishTimeUtc">Finish time to query utilization for</param>
+		/// <returns>The utilization data</returns>
+		Task<List<IUtilizationTelemetry>> GetUtilizationTelemetryAsync(DateTime StartTimeUtc, DateTime FinishTimeUtc);
+
+		/// <summary>
+		/// Finds the latest utilization data
+		/// </summary>
+		/// <returns>The utilization data</returns>
+		Task<IUtilizationTelemetry?> GetLatestUtilizationTelemetryAsync();
+	}
+}
diff --git a/Engine/Source/Programs/Horde/Horde.Build/Fleet/Impl/PoolCollection.cs b/Engine/Source/Programs/Horde/Horde.Build/Fleet/Impl/PoolCollection.cs
new file mode 100644
--- /dev/null
+++ b/Engine/Source/Programs/Horde/Horde.Build/Fleet/Impl/PoolCollection.cs
@@ -0,0 +1,244 @@
+// Copyright Epic Games, Inc. All Rights Reserved.
+
+using EpicGames.Horde.Common;
+using HordeServer.Models;
+using HordeServer.Services;
+using HordeServer.Utilities;
+using MongoDB.Bson;
+using MongoDB.Bson.Serialization;
+using MongoDB.Bson.Serialization.Attributes;
+using MongoDB.Driver;
+using System;
+using System.Collections.Generic;
+using System.Linq;
+using System.Linq.Expressions;
+using System.Threading.Tasks;
+
+namespace HordeServer.Collections.Impl
+{
+	using PoolId = StringId<IPool>;
+
+	/// <summary>
+	/// Collection of pool documents
+	/// </summary>
+	public class PoolCollection : IPoolCollection
+	{
+		class PoolDocument : IPool
+		{
+			// Properties
+			[BsonId]
+			public PoolId Id { get; set; }
+			[BsonRequired]
+			public string Name { get; set; } = null!;
+			[BsonIgnoreIfNull]
+			public Condition? Condition { get; set; }
+			public List<AgentWorkspace> Workspaces { get; set; } = new List<AgentWorkspace>();
+			public Dictionary<string, string> Properties { get; set; } = new Dictionary<string, string>();
+			[BsonIgnoreIfDefault(true)]
+			public bool EnableAutoscaling { get; set; } = true;
+			public int? MinAgents { get; set; }
+			public int? NumReserveAgents { get; set; }
+			public DateTime? LastScaleUpTime { get; set; }
+			public DateTime? LastScaleDownTime { get; set; }
+			public int UpdateIndex { get; set; }
+
+			// Read-only wrappers
+			IReadOnlyList<AgentWorkspace> IPool.Workspaces => Workspaces;
+			IReadOnlyDictionary<string, string> IPool.Properties => Properties;
+
+			public PoolDocument()
+			{
+			}
+
+			public PoolDocument(IPool Other)
+			{
+				Id = Other.Id;
+				Name = Other.Name;
+				Condition = Other.Condition;
+				Workspaces.AddRange(Other.Workspaces);
+				Properties = new Dictionary<string, string>(Other.Properties);
+				EnableAutoscaling = Other.EnableAutoscaling;
+				MinAgents = Other.MinAgents;
+				NumReserveAgents = Other.NumReserveAgents;
+				LastScaleUpTime = Other.LastScaleUpTime;
+				LastScaleDownTime = Other.LastScaleDownTime;
+				UpdateIndex = Other.UpdateIndex;
+			}
+		}
+
+		/// <summary>
+		/// Collection of pool documents
+		/// </summary>
+		IMongoCollection<PoolDocument> Pools;
+
+		/// <summary>
+		/// Constructor
+		/// </summary>
+		/// <param name="DatabaseService">The database service instance</param>
+		public PoolCollection(DatabaseService DatabaseService)
+		{
+			Pools = DatabaseService.GetCollection<PoolDocument>("Pools");
+		}
+
+		/// <inheritdoc/>
+		public async Task<IPool> AddAsync(PoolId Id, string Name, Condition? Condition, bool? EnableAutoscaling, int? MinAgents, int? NumReserveAgents, IEnumerable<KeyValuePair<string, string>>? Properties)
+		{
+			PoolDocument Pool = new PoolDocument();
+			Pool.Id = Id;
+			Pool.Name = Name;
+			Pool.Condition = Condition;
+			if (EnableAutoscaling != null)
+			{
+				Pool.EnableAutoscaling = EnableAutoscaling.Value;
+			}
+			Pool.MinAgents = MinAgents;
+			Pool.NumReserveAgents = NumReserveAgents;
+			if (Properties != null)
+			{
+				Pool.Properties = new Dictionary<string, string>(Properties);
+			}
+			await Pools.InsertOneAsync(Pool);
+			return Pool;
+		}
+
+		/// <inheritdoc/>
+		public async Task<List<IPool>> GetAsync()
+		{
+			List<PoolDocument> Results = await Pools.Find(FilterDefinition<PoolDocument>.Empty).ToListAsync();
+			return Results.ConvertAll<IPool>(x => x);
+		}
+
+		/// <inheritdoc/>
+		public async Task<IPool?> GetAsync(PoolId Id)
+		{
+			return await Pools.Find(x => x.Id == Id).FirstOrDefaultAsync();
+		}
+
+		/// <inheritdoc/>
+		public async Task<List<PoolId>> GetPoolIdsAsync()
+		{
+			ProjectionDefinition<PoolDocument, BsonDocument> Projection = Builders<PoolDocument>.Projection.Include(x => x.Id);
+			List<BsonDocument> Results = await Pools.Find(FilterDefinition<PoolDocument>.Empty).Project(Projection).ToListAsync();
+			return Results.ConvertAll(x => new PoolId(x.GetElement("_id").Value.AsString));
+		}
+
+		/// <inheritdoc/>
+		public async Task<bool> DeleteAsync(PoolId Id)
+		{
+			FilterDefinition<PoolDocument> Filter = Builders<PoolDocument>.Filter.Eq(x => x.Id, Id);
+			DeleteResult Result = await Pools.DeleteOneAsync(Filter);
+			return Result.DeletedCount > 0;
+		}
+
+		/// <summary>
+		/// Attempts to update a pool, upgrading it to the latest document schema if necessary
+		/// </summary>
+		/// <param name="PoolToUpdate">Interface for the document to update</param>
+		/// <param name="Transaction">The transaction to execute</param>
+		/// <returns>New pool definition, or null on failure.</returns>
+		async Task<IPool?> TryUpdateAsync(IPool PoolToUpdate, TransactionBuilder<PoolDocument> Transaction)
+		{
+			if (Transaction.IsEmpty)
+			{
+				return PoolToUpdate;
+			}
+
+			Transaction.Set(x => x.UpdateIndex, PoolToUpdate.UpdateIndex + 1);
+
+			PoolDocument? MutablePool = PoolToUpdate as PoolDocument;
+			if (MutablePool != null)
+			{
+				UpdateResult Result = await Pools.UpdateOneAsync(x => x.Id == PoolToUpdate.Id && x.UpdateIndex == PoolToUpdate.UpdateIndex, Transaction.ToUpdateDefinition());
+				if (Result.ModifiedCount == 0)
+				{
+					return null;
+				}
+
+				Transaction.ApplyTo(MutablePool);
+			}
+			else
+			{
+				MutablePool = new PoolDocument(PoolToUpdate);
+				Transaction.ApplyTo(MutablePool);
+
+				ReplaceOneResult Result = await Pools.ReplaceOneAsync<PoolDocument>(x => x.Id == PoolToUpdate.Id && x.UpdateIndex == PoolToUpdate.UpdateIndex, MutablePool);
+				if (Result.ModifiedCount == 0)
+				{
+					return null;
+				}
+			}
+			return MutablePool;
+		}
+
+		/// <inheritdoc/>
+		public Task<IPool?> TryUpdateAsync(IPool Pool, string? NewName, Condition? NewCondition, bool? NewEnableAutoscaling, int? NewMinAgents, int? NewNumReserveAgents, List<AgentWorkspace>? NewWorkspaces, Dictionary<string, string?>? NewProperties, DateTime? LastScaleUpTime, DateTime? LastScaleDownTime)
+		{
+			TransactionBuilder<PoolDocument> Transaction = new TransactionBuilder<PoolDocument>();
+			if (NewName != null)
+			{
+				Transaction.Set(x => x.Name, NewName);
+			}
+			if (NewCondition != null)
+			{
+				if (NewCondition.IsEmpty())
+				{
+					Transaction.Unset(x => x.Condition!);
+				}
+				else
+				{
+					Transaction.Set(x => x.Condition, NewCondition);
+				}
+			}
+			if (NewEnableAutoscaling != null)
+			{
+				if (NewEnableAutoscaling.Value)
+				{
+					Transaction.Unset(x => x.EnableAutoscaling);
+				}
+				else
+				{
+					Transaction.Set(x => x.EnableAutoscaling, NewEnableAutoscaling.Value);
+				}
+			}
+			if (NewMinAgents != null)
+			{
+				if (NewMinAgents.Value < 0)
+				{
+					Transaction.Unset(x => x.MinAgents!);
+				}
+				else
+				{
+					Transaction.Set(x => x.MinAgents, NewMinAgents.Value);
+				}
+			}
+			if (NewNumReserveAgents != null)
+			{
+				if (NewNumReserveAgents.Value < 0)
+				{
+					Transaction.Unset(x => x.NumReserveAgents!);
+				}
+				else
+				{
+					Transaction.Set(x => x.NumReserveAgents, NewNumReserveAgents.Value);
+				}
+			}
+			if (NewWorkspaces != null)
+			{
+				Transaction.Set(x => x.Workspaces, NewWorkspaces);
+			}
+			if (NewProperties != null)
+			{
+				Transaction.UpdateDictionary(x => x.Properties, NewProperties);
+			}
+			if (LastScaleUpTime != null)
+			{
+				Transaction.Set(x => x.LastScaleUpTime, LastScaleUpTime);
+			}
+			if (LastScaleDownTime != null)
+			{
+				Transaction.Set(x => x.LastScaleDownTime, LastScaleDownTime);
+			}
+			return TryUpdateAsync(Pool, Transaction);
+		}
+	}
+}
diff --git a/Engine/Source/Programs/Horde/Horde.Build/Fleet/Impl/TelemetryCollection.cs b/Engine/Source/Programs/Horde/Horde.Build/Fleet/Impl/TelemetryCollection.cs
new file mode 100644
--- /dev/null
+++ b/Engine/Source/Programs/Horde/Horde.Build/Fleet/Impl/TelemetryCollection.cs
@@ -0,0 +1,136 @@
+// Copyright Epic Games, Inc. All Rights Reserved.
+
+using EpicGames.Core;
+using HordeServer.Models;
+using HordeServer.Services;
+using HordeServer.Utilities;
+using Microsoft.Extensions.Hosting;
+using MongoDB.Bson;
+using MongoDB.Bson.Serialization.Attributes;
+using MongoDB.Driver;
+using System;
+using System.Collections.Generic;
+using System.Linq;
+using System.Threading;
+using System.Threading.Tasks;
+
+namespace HordeServer.Collections.Impl
+{
+	using PoolId = StringId<IPool>;
+	using StreamId = StringId<IStream>;
+
+	/// <summary>
+	/// Collection of utilization data
+	/// </summary>
+	public class TelemetryCollection : ITelemetryCollection
+	{
+		class UtilizationDocument : IUtilizationTelemetry
+		{
+			public DateTime StartTime { get; set; }
+			public DateTime FinishTime { get; set; }
+
+			public int NumAgents { get; set; }
+			public List<PoolUtilizationDocument> Pools { get; set; }
+			IReadOnlyList<IPoolUtilizationTelemetry> IUtilizationTelemetry.Pools => Pools;
+
+			public double AdminTime { get; set; }
+			public int UpdateIndex { get; set; }
+
+			[BsonConstructor]
+			private UtilizationDocument()
+			{
+				Pools = new List<PoolUtilizationDocument>();
+			}
+
+			public UtilizationDocument(IUtilizationTelemetry Other)
+			{
+				this.StartTime = Other.StartTime;
+				this.FinishTime = Other.FinishTime;
+				this.NumAgents = Other.NumAgents;
+				this.Pools = Other.Pools.ConvertAll(x => new PoolUtilizationDocument(x));
+				this.AdminTime = Other.AdminTime;
+			}
+		}
+
+		class PoolUtilizationDocument : IPoolUtilizationTelemetry
+		{
+			public PoolId PoolId { get; set; }
+			public int NumAgents { get; set; }
+
+			public List<StreamUtilizationDocument> Streams { get; set; } 
+			IReadOnlyList<IStreamUtilizationTelemetry> IPoolUtilizationTelemetry.Streams => Streams;
+
+			public double AdminTime { get; set; }
+			public double OtherTime { get; set; }
+
+			[BsonConstructor]
+			private PoolUtilizationDocument()
+			{
+				Streams = new List<StreamUtilizationDocument>();
+			}
+
+			public PoolUtilizationDocument(IPoolUtilizationTelemetry Other)
+			{
+				this.PoolId = Other.PoolId;
+				this.NumAgents = Other.NumAgents;
+				this.Streams = Other.Streams.ConvertAll(x => new StreamUtilizationDocument(x));
+				this.AdminTime = Other.AdminTime;
+				this.OtherTime = Other.OtherTime;
+			}
+		}
+
+		class StreamUtilizationDocument : IStreamUtilizationTelemetry
+		{
+			public StreamId StreamId { get; set; }
+			public double Time { get; set; }
+
+			[BsonConstructor]
+			private StreamUtilizationDocument()
+			{
+			}
+
+			public StreamUtilizationDocument(IStreamUtilizationTelemetry Other)
+			{
+				this.StreamId = Other.StreamId;
+				this.Time = Other.Time;
+			}
+		}
+
+		IMongoCollection<UtilizationDocument> Utilization;
+
+		/// <summary>
+		/// Constructor
+		/// </summary>
+		/// <param name="Database"></param>
+		public TelemetryCollection(DatabaseService Database)
+		{
+			Utilization = Database.GetCollection<UtilizationDocument>("Utilization");
+
+			if (!Database.ReadOnlyMode)
+			{
+				Utilization.Indexes.CreateOne(new CreateIndexModel<UtilizationDocument>(Builders<UtilizationDocument>.IndexKeys.Ascending(x => x.FinishTime).Ascending(x => x.StartTime)));
+			}
+		}
+
+		/// <inheritdoc/>
+		public async Task AddUtilizationTelemetryAsync(IUtilizationTelemetry NewTelemetry)
+		{
+			await Utilization.InsertOneAsync(new UtilizationDocument(NewTelemetry));
+		}
+
+		/// <inheritdoc/>
+		public async Task<List<IUtilizationTelemetry>> GetUtilizationTelemetryAsync(DateTime StartTimeUtc, DateTime FinishTimeUtc)
+		{
+			FilterDefinition<UtilizationDocument> Filter = Builders<UtilizationDocument>.Filter.Gte(x => x.FinishTime, StartTimeUtc) & Builders<UtilizationDocument>.Filter.Lte(x => x.StartTime, FinishTimeUtc);
+
+			List<UtilizationDocument> Documents = await Utilization.Find(Filter).ToListAsync();
+			return Documents.ConvertAll<IUtilizationTelemetry>(x => x);
+		}
+
+		/// <inheritdoc/>
+		public async Task<IUtilizationTelemetry?> GetLatestUtilizationTelemetryAsync()
+		{
+			return await Utilization.Find(FilterDefinition<UtilizationDocument>.Empty).SortByDescending(x => x.FinishTime).FirstOrDefaultAsync();
+		}
+	}
+}
diff --git a/Engine/Source/Programs/Horde/Horde.Build/Fleet/PoolService.cs b/Engine/Source/Programs/Horde/Horde.Build/Fleet/PoolService.cs
new file mode 100644
--- /dev/null
+++ b/Engine/Source/Programs/Horde/Horde.Build/Fleet/PoolService.cs
@@ -0,0 +1,185 @@
+// Copyright Epic Games, Inc. All Rights Reserved.
+
+using EpicGames.Horde.Common;
+using HordeCommon;
+using HordeServer.Api;
+using HordeServer.Collections;
+using HordeServer.Models;
+using HordeServer.Utilities;
+using Microsoft.Extensions.Logging;
+using Microsoft.Extensions.Options;
+using MongoDB.Bson;
+using MongoDB.Driver;
+using System;
+using System.Collections.Generic;
+using System.Linq;
+using System.Threading;
+using System.Threading.Tasks;
+
+namespace HordeServer.Services
+{
+	using PoolId = StringId<IPool>;
+
+	/// <summary>
+	/// Wraps functionality for manipulating pools
+	/// </summary>
+	public class PoolService
+	{
+		/// <summary>
+		/// The database service instance
+		/// </summary>
+		DatabaseService DatabaseService;
+
+		/// <summary>
+		/// Collection of pool documents
+		/// </summary>
+		IPoolCollection Pools;
+
+		/// <summary>
+		/// Returns the current time
+		/// </summary>
+		IClock Clock;
+
+		/// <summary>
+		/// Cached set of pools, along with the timestamp that it was obtained
+		/// </summary>
+		Tuple<DateTime, Dictionary<PoolId, IPool>>? CachedPoolLookup;
+
+		/// <summary>
+		/// Constructor
+		/// </summary>
+		/// <param name="DatabaseService"></param>
+		/// <param name="Pools">Collection of pool documents</param>
+		/// <param name="Clock"></param>
+		public PoolService(DatabaseService DatabaseService, IPoolCollection Pools, IClock Clock)
+		{
+			this.DatabaseService = DatabaseService;
+			this.Pools = Pools;
+			this.Clock = Clock;
+		}
+
+		/// <summary>
+		/// Creates a new pool
+		/// </summary>
+		/// <param name="Name">Name of the new pool</param>
+		/// <param name="Condition">Condition for agents to be automatically included in this pool</param>
+		/// <param name="EnableAutoscaling">Whether to enable autoscaling for this pool</param>
+		/// <param name="MinAgents">Minimum number of agents in the pool</param>
+		/// <param name="NumReserveAgents">Minimum number of idle agents to maintain</param>
+		/// <param name="Properties">Properties for the new pool</param>
+		/// <returns>The new pool document</returns>
+		public Task<IPool> CreatePoolAsync(string Name, Condition? Condition = null, bool? EnableAutoscaling = null, int? MinAgents = null, int? NumReserveAgents = null, Dictionary<string, string>? Properties = null)
+		{
+			return Pools.AddAsync(PoolId.Sanitize(Name), Name, Condition, EnableAutoscaling, MinAgents, NumReserveAgents, Properties);
+		}
+
+		/// <summary>
+		/// Deletes a pool
+		/// </summary>
+		/// <param name="PoolId">Unique id of the pool</param>
+		/// <returns>Async task object</returns>
+		public Task<bool> DeletePoolAsync(PoolId PoolId)
+		{
+			return Pools.DeleteAsync(PoolId);
+		}
+
+		/// <summary>
+		/// Updates an existing pool
+		/// </summary>
+		/// <param name="Pool">The pool to update</param>
+		/// <param name="NewName">The new name for the pool</param>
+		/// <param name="NewCondition">New requirements for the pool</param>
+		/// <param name="NewEnableAutoscaling">Whether to enable autoscaling</param>
+		/// <param name="NewMinAgents">Minimum number of agents in the pool</param>
+		/// <param name="NewNumReserveAgents">Minimum number of idle agents to maintain</param>
+		/// <param name="NewProperties">Properties on the pool to update. Any properties with a value of null will be removed.</param>
+		/// <returns>Async task object</returns>
+		public async Task<IPool?> UpdatePoolAsync(IPool? Pool, string? NewName = null, Condition? NewCondition = null, bool? NewEnableAutoscaling = null, int? NewMinAgents = null, int? NewNumReserveAgents = null, Dictionary<string, string?>? NewProperties = null)
+		{
+			for (; Pool != null; Pool = await Pools.GetAsync(Pool.Id))
+			{
+				IPool? NewPool = await Pools.TryUpdateAsync(Pool, NewName, NewCondition, NewEnableAutoscaling, NewMinAgents, NewNumReserveAgents, null, NewProperties, null, null);
+				if (NewPool != null)
+				{
+					return NewPool;
+				}
+			}
+			return Pool;
+		}
+
+		/// <summary>
+		/// Gets all the available pools
+		/// </summary>
+		/// <returns>List of pool documents</returns>
+		public Task<List<IPool>> GetPoolsAsync()
+		{
+			return Pools.GetAsync();
+		}
+
+		/// <summary>
+		/// Gets a pool by ID
+		/// </summary>
+		/// <param name="PoolId">Unique id of the pool</param>
+		/// <returns>The pool document</returns>
+		public Task<IPool?> GetPoolAsync(PoolId PoolId)
+		{
+			return Pools.GetAsync(PoolId);
+		}
+
+		/// <summary>
+		/// Get a list of workspaces for the given agent
+		/// </summary>
+		/// <param name="Agent">The agent to return workspaces for</param>
+		/// <param name="ValidAtTime">Absolute time at which we expect the results to be valid. Values may be cached as long as they are after this time.</param>
+		/// <returns>List of workspaces</returns>
+		public async Task<HashSet<AgentWorkspace>> GetWorkspacesAsync(IAgent Agent, DateTime ValidAtTime)
+		{
+			HashSet<AgentWorkspace> Workspaces = new HashSet<AgentWorkspace>();
+
+			Dictionary<PoolId, IPool> PoolMapping = await GetPoolLookupAsync(ValidAtTime);
+			foreach (PoolId PoolId in Agent.GetPools())
+			{
+				IPool? Pool;
+				if (PoolMapping.TryGetValue(PoolId, out Pool))
+				{
+					Workspaces.UnionWith(Pool.Workspaces);
+				}
+			}
+
+			Globals Globals = await DatabaseService.GetGlobalsAsync();
+			Workspaces.UnionWith(Agent.GetAutoSdkWorkspaces(Globals, Workspaces.ToList()));
+
+			return Workspaces;
+		}
+
+		/// <summary>
+		/// Gets a mapping from pool identifiers to definitions
+		/// </summary>
+		/// <param name="ValidAtTime">Absolute time at which we expect the results to be valid. Values may be cached as long as they are after this time.</param>
+		/// <returns>Map of pool ids to pool documents</returns>
+		private async Task<Dictionary<PoolId, IPool>> GetPoolLookupAsync(DateTime ValidAtTime)
+		{
+			Tuple<DateTime, Dictionary<PoolId, IPool>>? CachedPoolLookupCopy = CachedPoolLookup;
+			if (CachedPoolLookupCopy == null || CachedPoolLookupCopy.Item1 < ValidAtTime)
+			{
+				// Get a new list of cached pools
+				DateTime NewCacheTime = Clock.UtcNow;
+				List<IPool> NewPools = await Pools.GetAsync();
+				Tuple<DateTime, Dictionary<PoolId, IPool>> NewCachedPoolLookup = Tuple.Create(NewCacheTime, NewPools.ToDictionary(x => x.Id, x => x));
+
+				// Try to swap it with the current version
+				while (CachedPoolLookupCopy == null || CachedPoolLookupCopy.Item1 < NewCacheTime)
+				{
+					Tuple<DateTime, Dictionary<PoolId, IPool>>? OriginalValue = Interlocked.CompareExchange(ref CachedPoolLookup, NewCachedPoolLookup, CachedPoolLookupCopy);
+					if (OriginalValue == CachedPoolLookupCopy)
+					{
+						CachedPoolLookupCopy = NewCachedPoolLookup;
+						break;
+					}
+					CachedPoolLookupCopy = OriginalValue;
+				}
+			}
+			return CachedPoolLookupCopy.Item2;
+		}
+	}
+}
diff --git a/Engine/Source/Programs/Horde/Horde.Build/Fleet/PoolUpdateService.cs b/Engine/Source/Programs/Horde/Horde.Build/Fleet/PoolUpdateService.cs
new file mode 100644
--- /dev/null
+++ b/Engine/Source/Programs/Horde/Horde.Build/Fleet/PoolUpdateService.cs
@@ -0,0 +1,131 @@
+// Copyright Epic Games, Inc. All Rights Reserved.
+
+using HordeServer.Collections;
+using HordeServer.Models;
+using HordeServer.Utilities;
+using Microsoft.Extensions.Logging;
+using System;
+using System.Collections.Generic;
+using System.Linq;
+using System.Threading;
+using System.Threading.Tasks;
+
+namespace HordeServer.Services
+{
+	using PoolId = StringId<IPool>;
+
+	/// <summary>
+	/// Periodically updates pool documents to contain the correct workspaces
+	/// </summary>
+	public class PoolUpdateService : TickedBackgroundService
+	{
+		/// <summary>
+		/// Collection of pool documents
+		/// </summary>
+		IPoolCollection Pools;
+
+		/// <summary>
+		/// Collection of stream documents
+		/// </summary>
+		IStreamCollection Streams;
+
+		/// <summary>
+		/// The logger instance for this service
+		/// </summary>
+		ILogger<PoolService> Logger;
+
+		/// <summary>
+		/// Constructor
+		/// </summary>
+		/// <param name="Pools">Collection of pool documents</param>
+		/// <param name="Streams">Collection of stream documents</param>
+		/// <param name="Logger">Logger instance</param>
+		public PoolUpdateService(IPoolCollection Pools, IStreamCollection Streams, ILogger<PoolService> Logger)
+			: base(TimeSpan.FromSeconds(30.0), Logger)
+		{
+			this.Pools = Pools;
+			this.Streams = Streams;
+			this.Logger = Logger;
+		}
+
+		/// <summary>
+		/// Execute the background task
+		/// </summary>
+		/// <param name="StoppingToken">Cancellation token for the async task</param>
+		/// <returns>Async task</returns>
+		protected override async Task TickAsync(CancellationToken StoppingToken)
+		{
+			// Capture the start time for this operation. We use this to attempt to sequence updates to agents, and prevent overriding another server's updates.
+			DateTime StartTime = DateTime.UtcNow;
+
+			// Update the list
+			bool bRetryUpdate = true;
+			while (bRetryUpdate && !StoppingToken.IsCancellationRequested)
+			{
+				Logger.LogDebug("Updating pool->workspace map");
+
+				// Assume this will be the last iteration
+				bRetryUpdate = false;
+
+				// Capture the list of pools at the start of this update
+				List<IPool> CurrentPools = await Pools.GetAsync();
+
+				// Lookup table of pool id to workspaces
+				Dictionary<PoolId, List<AgentWorkspace>> PoolToAgentWorkspaces = new Dictionary<PoolId, List<AgentWorkspace>>();
+
+				// Populate the workspace list from the current stream
+				List<IStream> ActiveStreams = await Streams.FindAllAsync();
+				foreach (IStream ActiveStream in ActiveStreams)
+				{
+					foreach (KeyValuePair<string, AgentType> AgentTypePair in ActiveStream.AgentTypes)
+					{
+						// Create the new agent workspace
+						AgentWorkspace? AgentWorkspace;
+						if (ActiveStream.TryGetAgentWorkspace(AgentTypePair.Value, out AgentWorkspace))
+						{
+							AgentType AgentType = AgentTypePair.Value;
+
+							// Find or add a list of workspaces for this pool
+							List<AgentWorkspace>? AgentWorkspaces;
+							if (!PoolToAgentWorkspaces.TryGetValue(AgentType.Pool, out AgentWorkspaces))
+							{
+								AgentWorkspaces = new List<AgentWorkspace>();
+								PoolToAgentWorkspaces.Add(AgentType.Pool, AgentWorkspaces);
+							}
+
+							// Add it to the list
+							if (!AgentWorkspaces.Contains(AgentWorkspace))
+							{
+								AgentWorkspaces.Add(AgentWorkspace);
+							}
+						}
+					}
+				}
+
+				// Update the list of workspaces for each pool
+				foreach (IPool CurrentPool in CurrentPools)
+				{
+					// Get the new list of workspaces for this pool
+					List<AgentWorkspace>? NewWorkspaces;
+					if (!PoolToAgentWorkspaces.TryGetValue(CurrentPool.Id, out NewWorkspaces))
+					{
+						NewWorkspaces = new List<AgentWorkspace>();
+					}
+
+					// Update the pools document
+					if (!AgentWorkspace.SetEquals(CurrentPool.Workspaces, NewWorkspaces) || CurrentPool.Workspaces.Count != NewWorkspaces.Count)
+					{
+						Logger.LogInformation("New workspaces for pool {Pool}:{Workspaces}", CurrentPool.Id, String.Join("", NewWorkspaces.Select(x => $"\n  Identifier=\"{x.Identifier}\", Stream={x.Stream}")));
+
+						IPool? Result = await Pools.TryUpdateAsync(CurrentPool, NewWorkspaces: NewWorkspaces);
+						if (Result == null)
+						{
+							Logger.LogInformation("Pool modified; will retry");
+							bRetryUpdate = true;
+						}
+					}
+				}
+			}
+		}
+	}
+}
diff --git a/Engine/Source/Programs/Horde/Horde.Build/Fleet/TelemetryService.cs b/Engine/Source/Programs/Horde/Horde.Build/Fleet/TelemetryService.cs
new file mode 100644
--- /dev/null
+++ b/Engine/Source/Programs/Horde/Horde.Build/Fleet/TelemetryService.cs
@@ -0,0 +1,142 @@
+// Copyright Epic Games, Inc. All Rights Reserved.
+
+using Google.Protobuf.WellKnownTypes;
+using HordeCommon.Rpc.Tasks;
+using HordeServer.Collections;
+using HordeServer.Models;
+using HordeServer.Utilities;
+using Microsoft.Extensions.Logging;
+using MongoDB.Bson;
+using System;
+using System.Collections.Generic;
+using System.Linq;
+using System.Threading;
+using System.Threading.Tasks;
+
+namespace HordeServer.Services
+{
+	using PoolId = StringId<IPool>;
+
+	/// <summary>
+	/// Service which updates telemetry periodically
+	/// </summary>
+	public class TelemetryService : ElectedBackgroundService
+	{
+		/// <summary>
+		/// Collection of telemetry documents
+		/// </summary>
+		ITelemetryCollection TelemetryCollection;
+
+		/// <summary>
+		/// Collection of agent documents
+		/// </summary>
+		IAgentCollection AgentCollection;
+
+		/// <summary>
+		/// Collection of lease documents
+		/// </summary>
+		ILeaseCollection LeaseCollection;
+
+		/// <summary>
+		/// Collection of pool documents
+		/// </summary>
+		IPoolCollection PoolCollection;
+
+		/// <summary>
+		/// Logger for diagnostic messages
+		/// </summary>
+		ILogger Logger;
+
+		/// <summary>
+		/// Constructor
+		/// </summary>
+		/// <param name="DatabaseService"></param>
+		/// <param name="TelemetryCollection"></param>
+		/// <param name="AgentCollection"></param>
+		/// <param name="LeaseCollection"></param>
+		/// <param name="PoolCollection"></param>
+		/// <param name="Logger"></param>
+		public TelemetryService(DatabaseService DatabaseService, ITelemetryCollection TelemetryCollection, IAgentCollection AgentCollection, ILeaseCollection LeaseCollection, IPoolCollection PoolCollection, ILogger<TelemetryService> Logger)
+			: base(DatabaseService, new ObjectId("600082cb9b822d7afee6f0d1"), Logger)
+		{
+			this.TelemetryCollection = TelemetryCollection;
+			this.AgentCollection = AgentCollection;
+			this.LeaseCollection = LeaseCollection;
+			this.PoolCollection = PoolCollection;
+			this.Logger = Logger;
+		}
+
+		/// <inheritdoc/>
+		protected override async Task<DateTime> TickLeaderAsync(CancellationToken StoppingToken)
+		{
+			DateTime CurrentTime = DateTime.UtcNow;
+			DateTime NextTickTime = CurrentTime + TimeSpan.FromMinutes(10.0);
+
+			// Find the last time that we need telemetry for
+			DateTime MaxTime = CurrentTime.Date + TimeSpan.FromHours(CurrentTime.Hour);
+
+			// Get the latest telemetry data
+			IUtilizationTelemetry? Latest = await TelemetryCollection.GetLatestUtilizationTelemetryAsync();
+			TimeSpan Interval = TimeSpan.FromHours(1.0);
+			int Count = (Latest == null) ? (7 * 24) : (int)Math.Round((MaxTime - Latest.FinishTime) / Interval);
+			DateTime MinTime = MaxTime - Count * Interval;
+
+			// Query all the current data
+			List<IAgent> Agents = await AgentCollection.FindAsync();
+			List<IPool> Pools = await PoolCollection.GetAsync();
+			List<ILease> Leases = await LeaseCollection.FindLeasesAsync(MinTime: MinTime);
+
+			// Find all the agents
+			Dictionary<AgentId, List<PoolId>> AgentToPoolIds = Agents.ToDictionary(x => x.Id, x => x.GetPools().ToList());
+
+			// Generate all the telemetry data
+			DateTime BucketMinTime = MinTime;
+			for (int Idx = 0; Idx < Count; Idx++)
+			{
+				DateTime BucketMaxTime = BucketMinTime + Interval;
+				Logger.LogInformation("Creating telemetry for {MinTime} to {MaxTime}", BucketMinTime, BucketMaxTime);
+
+				NewUtilizationTelemetry Telemetry = new NewUtilizationTelemetry(BucketMinTime, BucketMaxTime);
+				Telemetry.NumAgents = Agents.Count;
+				foreach (PoolId PoolId in AgentToPoolIds.Values.SelectMany(x => x))
+				{
+					Telemetry.FindOrAddPool(PoolId).NumAgents++;
+				}
+				foreach (ILease Lease in Leases)
+				{
+					if (Lease.StartTime < BucketMaxTime && Lease.FinishTime >= BucketMinTime)
+					{
+						List<PoolId>? LeasePools;
+						if (AgentToPoolIds.TryGetValue(Lease.AgentId, out LeasePools))
+						{
+							DateTime FinishTime = Lease.FinishTime ?? BucketMaxTime;
+							double Time = new TimeSpan(Math.Min(FinishTime.Ticks, BucketMaxTime.Ticks) - Math.Max(Lease.StartTime.Ticks, BucketMinTime.Ticks)).TotalHours;
+
+							foreach (PoolId PoolId in LeasePools)
+							{
+								NewPoolUtilizationTelemetry PoolTelemetry = Telemetry.FindOrAddPool(PoolId);
+								if (Lease.PoolId == null || Lease.StreamId == null)
+								{
+									PoolTelemetry.AdminTime += Time;
+								}
+								else if (PoolId == Lease.PoolId)
+								{
+									PoolTelemetry.FindOrAddStream(Lease.StreamId.Value).Time += Time;
+								}
+								else
+								{
+									PoolTelemetry.OtherTime += Time;
+								}
+							}
+						}
+					}
+				}
+				await TelemetryCollection.AddUtilizationTelemetryAsync(Telemetry);
+
+				BucketMinTime = BucketMaxTime;
+			}
+
+			return NextTickTime;
+		}
+	}
+}
diff --git a/Engine/Source/Programs/Horde/Horde.Build/GlobalSuppressions.cs b/Engine/Source/Programs/Horde/Horde.Build/GlobalSuppressions.cs
new file mode 100644
--- /dev/null
+++ b/Engine/Source/Programs/Horde/Horde.Build/GlobalSuppressions.cs
@@ -0,0 +1,43 @@
+// Copyright Epic Games, Inc. All Rights Reserved.
+
+// This file is used by Code Analysis to maintain SuppressMessage
+// attributes that are applied to this project.
+// Project-level suppressions either have no target or are given
+// a specific target and scoped to a namespace, type, member, etc.
+
+using System.Diagnostics.CodeAnalysis;
+
+[assembly: System.Diagnostics.CodeAnalysis.SuppressMessage("Reliability", "CA2007:Consider calling ConfigureAwait on the awaited task", Justification = "Not necessary in NET Core; no SynchronizationContext.", Scope = "module")]
+[assembly: System.Diagnostics.CodeAnalysis.SuppressMessage("Usage", "CA2227:Collection properties should be read only", Justification = "<Pending>", Scope = "module")]
+[assembly: System.Diagnostics.CodeAnalysis.SuppressMessage("Design", "CA1062:Validate arguments of public methods", Justification = "<Pending>", Scope = "module")]
+[assembly: System.Diagnostics.CodeAnalysis.SuppressMessage("Globalization", "CA1303:Do not pass literals as localized parameters", Justification = "<Pending>", Scope = "module")]
+[assembly: System.Diagnostics.CodeAnalysis.SuppressMessage("Performance", "CA1819:Properties should not return arrays", Justification = "<Pending>", Scope = "module")]
+[assembly: System.Diagnostics.CodeAnalysis.SuppressMessage("Naming", "CA1716:Identifiers should not match keywords", Justification = "<Pending>", Scope = "type", Target = "~T:HordeServer.Models.Event")]
+[assembly: System.Diagnostics.CodeAnalysis.SuppressMessage("Design", "CA1032:Implement standard exception constructors", Justification = "<Pending>", Scope = "type", Target = "~T:HordeServer.Services.RetryNotAllowedException")]
+[assembly: System.Diagnostics.CodeAnalysis.SuppressMessage("Naming", "CA1714:Flags enums should have plural names", Justification = "<Pending>", Scope = "type", Target = "~T:HordeServer.Api.ScheduleDaysOfWeek")]
+[assembly: System.Diagnostics.CodeAnalysis.SuppressMessage("Naming", "CA1720:Identifier contains type name", Justification = "<Pending>", Scope = "member", Target = "~F:HordeServer.Api.TemplateParameterType.String")]
+
+[assembly: System.Diagnostics.CodeAnalysis.SuppressMessage("Build", "CA1822:Member Login does not access instance data and can be marked as static (Shared in VisualBasic)", Justification = "<Pending>", Scope = "member", Target = "~M:HordeServer.Controllers.AccountController.Login~Microsoft.AspNetCore.Mvc.IActionResult")]
+[assembly: System.Diagnostics.CodeAnalysis.SuppressMessage("Build", "CA1031:Modify 'Logout' to catch a more specific allowed exception type, or rethrow the exception.", Justification = "<Pending>", Scope = "member", Target = "~M:HordeServer.Controllers.AccountController.Logout~System.Threading.Tasks.Task{Microsoft.AspNetCore.Mvc.IActionResult}")]
+[assembly: System.Diagnostics.CodeAnalysis.SuppressMessage("Build", "CA1305:The behavior of 'string.ToString()' could vary based on the current user's locale settings. Replace this call in 'AgentLease.ToResponse()' with a call to 'string.ToString(IFormatProvider)'.", Justification = "<Pending>", Scope = "member", Target = "~M:HordeServer.Models.AgentLease.ToResponse~HordeServer.Api.GetAgentLeaseResponse")]
+[assembly: System.Diagnostics.CodeAnalysis.SuppressMessage("Build", "CA5380:Adding certificates to the operating system's trusted root certificates increases the risk of incorrectly authenticating an illegitimate certificate", Justification = "<Pending>", Scope = "member", Target = "~T:HordeServer.Services.DatabaseService")]
+[assembly: System.Diagnostics.CodeAnalysis.SuppressMessage("Build", "CA1724:The type name Template conflicts in whole or in part with the namespace name 'Microsoft.AspNetCore.Routing.Template'. Change either name to eliminate the conflict.", Justification = "<Pending>", Scope = "type", Target = "~T:HordeServer.Models.Template")]
+[assembly: System.Diagnostics.CodeAnalysis.SuppressMessage("Build", "CA1724:The type name Session conflicts in whole or in part with the namespace name 'Microsoft.AspNetCore.Session'. Change either name to eliminate the conflict.", Justification = "<Pending>", Scope = "type", Target = "~T:HordeServer.Models.Session")]
+[assembly: System.Diagnostics.CodeAnalysis.SuppressMessage("Build", "CA1812:BotsService is an internal class that is apparently never instantiated. If so, remove the code from the assembly. If this class is intended to contain only static members, make it static (Shared in Visual Basic).", Justification = "<Pending>", Scope = "type", Target = "~T:HordeServer.Services.BotsService")]
+[assembly: System.Diagnostics.CodeAnalysis.SuppressMessage("Build", "CA1812:RpcService is an internal class that is apparently never instantiated. If so, remove the code from the assembly. If this class is intended to contain only static members, make it static (Shared in Visual Basic).", Justification = "<Pending>", Scope = "type", Target = "~T:HordeServer.Services.RpcService")]
+[assembly: System.Diagnostics.CodeAnalysis.SuppressMessage("Build", "CA1812:Startup is an internal class that is apparently never instantiated. If so, remove the code from the assembly. If this class is intended to contain only static members, make it static (Shared in Visual Basic).", Justification = "<Pending>", Scope = "type", Target = "~T:HordeServer.Startup")]
+[assembly: System.Diagnostics.CodeAnalysis.SuppressMessage("Build", "CA1812:Startup.AgentJwtBearerHandler is an internal class that is apparently never instantiated. If so, remove the code from the assembly. If this class is intended to contain only static members, make it static (Shared in Visual Basic).", Justification = "<Pending>", Scope = "type", Target = "~T:HordeServer.Startup.AgentJwtBearerHandler")]
+[assembly: System.Diagnostics.CodeAnalysis.SuppressMessage("Build", "CA1812:Startup.AnonymousAuthenticationHandler is an internal class that is apparently never instantiated. If so, remove the code from the assembly. If this class is intended to contain only static members, make it static (Shared in Visual Basic).", Justification = "<Pending>", Scope = "type", Target = "~T:HordeServer.Startup.AnonymousAuthenticationHandler")]
+[assembly: System.Diagnostics.CodeAnalysis.SuppressMessage("Performance", "CA1822:Mark members as static", Justification = "<Pending>", Scope = "member", Target = "~M:HordeServer.Controllers.DebugController.Environment~Microsoft.AspNetCore.Mvc.IActionResult")]
+[assembly: System.Diagnostics.CodeAnalysis.SuppressMessage("Naming", "CA1716:Identifiers should not match keywords", Justification = "<Pending>", Scope = "module")]
+[assembly: SuppressMessage("Design", "CA1031:Do not catch general exception types", Justification = "<Pending>", Scope = "member", Target = "~M:HordeServer.Services.AgentService.ExecuteAsync(System.Threading.CancellationToken)~System.Threading.Tasks.Task")]
+[assembly: SuppressMessage("Design", "CA1031:Do not catch general exception types", Justification = "<Pending>", Scope = "member", Target = "~M:HordeServer.Utilities.PeriodicBackgroundService.ExecuteAsync(System.Threading.CancellationToken)~System.Threading.Tasks.Task")]
+[assembly: SuppressMessage("Performance", "CA1815:Override equals and operator equals on value types", Justification = "<Pending>", Scope = "type", Target = "~T:HordeServer.Models.JobStepRefId")]
+[assembly: SuppressMessage("Style", "IDE0054:Use compound assignment", Justification = "<Pending>", Scope = "member", Target = "~M:HordeServer.Models.Job.GetTimingInfo~System.Collections.Generic.Dictionary{HordeServer.Models.Node,HordeServer.Models.TimingInfo}")]
+[assembly: SuppressMessage("Design", "CA1031:Do not catch general exception types", Justification = "<Pending>", Scope = "member", Target = "~M:HordeServer.Startup.DashboardBearerHandler.HandleAuthenticateAsync~System.Threading.Tasks.Task{Microsoft.AspNetCore.Authentication.AuthenticateResult}")]
+[assembly: SuppressMessage("Design", "CA1031:Do not catch general exception types", Justification = "<Pending>", Scope = "member", Target = "~M:HordeServer.Services.ConformService.ExecuteAsync(System.Threading.CancellationToken)~System.Threading.Tasks.Task")]
+[assembly: SuppressMessage("Design", "CA1031:Do not catch general exception types", Justification = "<Pending>", Scope = "member", Target = "~M:HordeServer.Services.ScheduleService.ExecuteAsync(System.Threading.CancellationToken)~System.Threading.Tasks.Task")]
+[assembly: SuppressMessage("Design", "CA1031:Do not catch general exception types", Justification = "<Pending>", Scope = "member", Target = "~M:HordeServer.Services.ScheduleService.UpdateSchedulesAsync~System.Threading.Tasks.Task{System.DateTimeOffset}")]
+[assembly: SuppressMessage("Design", "CA1030:Use events where appropriate", Justification = "<Pending>", Scope = "module")]
+[assembly: SuppressMessage("Design", "CA1031:Do not catch general exception types", Justification = "<Pending>", Scope = "module")]
+[assembly: SuppressMessage("Design", "CA1032:Implement standard exception constructors", Justification = "<Pending>", Scope = "module")]
diff --git a/Engine/Source/Programs/Horde/Horde.Build/Horde.Build.csproj b/Engine/Source/Programs/Horde/Horde.Build/Horde.Build.csproj
new file mode 100644
--- /dev/null
+++ b/Engine/Source/Programs/Horde/Horde.Build/Horde.Build.csproj
@@ -0,0 +1,106 @@
+﻿<Project Sdk="Microsoft.NET.Sdk.Web">
+
+  <PropertyGroup>
+    <TargetFramework>net5.0</TargetFramework>
+    <Nullable>enable</Nullable>
+    <Deterministic>true</Deterministic>
+    <AssemblyName>HordeServer</AssemblyName>
+    <DocumentationFile>HordeServer.xml</DocumentationFile>
+    <ApplicationIcon>HordeServer.ico</ApplicationIcon>
+  </PropertyGroup>
+
+  <ItemGroup>
+    <PackageReference Include="AWSSDK.EC2" Version="3.7.0" />
+    <PackageReference Include="AWSSDK.Extensions.NETCore.Setup" Version="3.7.0.1" />
+    <PackageReference Include="AWSSDK.S3" Version="3.7.0.17" />
+    <PackageReference Include="AWSSDK.SecurityToken" Version="3.7.1.69" />
+    <PackageReference Include="Datadog.Trace.OpenTracing" Version="1.18.0" />
+    <PackageReference Include="DogStatsD-CSharp-Client" Version="6.0.0" />
+    <PackageReference Include="Grpc.AspNetCore" Version="2.37.0" />
+    <PackageReference Include="Grpc.AspNetCore.Server.Reflection" Version="2.37.0" />
+    <PackageReference Include="Grpc.Net.Client" Version="2.37.0" />
+    <PackageReference Include="Grpc.Tools" Version="2.37.0">
+      <PrivateAssets>all</PrivateAssets>
+      <IncludeAssets>runtime; build; native; contentfiles; analyzers; buildtransitive</IncludeAssets>
+    </PackageReference>
+    <PackageReference Include="JetBrains.Profiler.SelfApi" Version="2.1.2" />
+    <PackageReference Include="Microsoft.AspNetCore.Authentication.Cookies" Version="2.2.0" />
+    <PackageReference Include="Microsoft.AspNetCore.Authentication.JwtBearer" Version="3.1.5" />
+    <PackageReference Include="Microsoft.AspNetCore.Authentication.OpenIdConnect" Version="3.1.5" />
+    <PackageReference Include="Microsoft.AspNetCore.DataProtection.StackExchangeRedis" Version="3.1.7" />
+    <PackageReference Include="Microsoft.AspNetCore.SpaServices.Extensions" Version="3.1.17" />
+    <PackageReference Include="Microsoft.CodeAnalysis.FxCopAnalyzers" Version="3.0.0">
+      <PrivateAssets>all</PrivateAssets>
+      <IncludeAssets>runtime; build; native; contentfiles; analyzers; buildtransitive</IncludeAssets>
+    </PackageReference>
+    <PackageReference Include="Microsoft.Extensions.Caching.Memory" Version="5.0.0" />
+    <PackageReference Include="MongoDB.Driver" Version="2.11.6" />
+    <PackageReference Include="OpenTracing" Version="0.12.1" />
+    <PackageReference Include="OpenTracing.Contrib.Grpc" Version="0.4.0" />
+    <PackageReference Include="Polly" Version="7.2.1" />
+    <PackageReference Include="Polly.Extensions.Http" Version="3.0.0" />
+    <PackageReference Include="Serilog.AspNetCore" Version="3.2.0" />
+    <PackageReference Include="Serilog.Enrichers.OpenTracing" Version="0.0.2" />
+    <PackageReference Include="Serilog.Extensions.Hosting" Version="3.1.0" />
+    <PackageReference Include="Serilog.Sinks.Console" Version="3.1.1" />
+    <PackageReference Include="Serilog.Sinks.File" Version="4.1.0" />
+    <PackageReference Include="SerilogAnalyzer" Version="0.15.0" />
+    <PackageReference Include="SharpZipLib" Version="1.2.0" />
+    <PackageReference Include="StackExchange.Redis" Version="2.2.4" />
+    <PackageReference Include="Swashbuckle.AspNetCore" Version="5.5.0" />
+    <PackageReference Include="System.Linq.Async" Version="5.0.0" />
+    <PackageReference Include="TimeZoneConverter" Version="3.2.0" />
+  </ItemGroup>
+
+  <ItemGroup>
+    <Content Include="..\..\Shared\EpicGames.Horde.Common\Protos\**\*.*">
+      <Link>Protos\%(RecursiveDir)%(FileName)%(Extension)</Link>
+    </Content>
+  </ItemGroup>
+
+  <ItemGroup>
+    <None Remove="HordeServer.xml" />
+  </ItemGroup>
+
+  <ItemGroup>
+    <Protobuf Include="..\..\Shared\EpicGames.Horde\Protos\horde\storage\ref_table.proto" ProtoRoot="..\..\Shared\EpicGames.Horde\Protos" GrpcServices="Server" />
+    <Protobuf Include="..\..\Shared\EpicGames.Horde\Protos\horde\storage\blob_store.proto" ProtoRoot="..\..\Shared\EpicGames.Horde\Protos" GrpcServices="Server" />
+    <Protobuf Include="..\..\Shared\EpicGames.Horde\Protos\horde\compute\compute_rpc.proto" ProtoRoot="..\..\Shared\EpicGames.Horde\Protos" GrpcServices="Server" />
+    <Protobuf Include="..\HordeCommon\Protos\horde\horde_rpc.proto" ProtoRoot="..\HordeCommon\Protos" GrpcServices="Server" />
+    <Protobuf Include="..\HordeCommon\Protos\grpc\health\v1\health.proto" ProtoRoot="..\HordeCommon\Protos" GrpcServices="Server" />
+  </ItemGroup>
+
+  <ItemGroup>
+    <ProjectReference Include="..\..\Shared\EpicGames.Core\EpicGames.Core.csproj" />
+    <ProjectReference Include="..\..\Shared\EpicGames.Horde\EpicGames.Horde.csproj" />
+    <ProjectReference Include="..\..\Shared\EpicGames.IoHash\EpicGames.IoHash.csproj" />
+    <ProjectReference Include="..\..\Shared\EpicGames.Perforce.Managed\EpicGames.Perforce.Managed.csproj" />
+    <ProjectReference Include="..\..\Shared\EpicGames.Perforce\EpicGames.Perforce.csproj" />
+    <ProjectReference Include="..\..\Shared\EpicGames.Redis\EpicGames.Redis.csproj" />
+    <ProjectReference Include="..\..\Shared\EpicGames.Serialization\EpicGames.Serialization.csproj" />
+    <ProjectReference Include="..\HordeCommon\HordeCommon.csproj" />
+    <ProjectReference Include="..\..\..\ThirdParty\Perforce\P4Api.Net\p4api.net\p4api.net.csproj" />
+    <ProjectReference Include="..\HordeServer.Common\HordeServer.Common.csproj" />
+    <ProjectReference Include="..\HordeServer.Compute\HordeServer.Compute.csproj" />
+  </ItemGroup>
+
+  <ItemGroup>
+    <Content Include="..\..\..\ThirdParty\Perforce\P4Api.Net\native\windows\release\p4bridge.*" Condition=" '$(OS)' == 'Windows_NT' ">
+      <CopyToOutputDirectory>PreserveNewest</CopyToOutputDirectory>
+    </Content>
+    <Content Include="..\..\..\ThirdParty\Perforce\P4Api.Net\native\linux\release\p4bridge.*" Condition="$([MSBuild]::IsOSPlatform('Linux'))">
+      <CopyToOutputDirectory>PreserveNewest</CopyToOutputDirectory>
+    </Content>
+    <Content Include="..\ThirdParty\**" Condition="$([MSBuild]::IsOSPlatform('Windows'))">
+      <CopyToOutputDirectory>PreserveNewest</CopyToOutputDirectory>
+      <Link>ThirdParty/%(RecursiveDir)/%(Filename)%(Extension)</Link>
+    </Content>
+  </ItemGroup>
+
+  <ItemGroup>
+    <Content Include="Defaults\ue5.project.png">
+      <CopyToOutputDirectory>PreserveNewest</CopyToOutputDirectory>
+    </Content>
+  </ItemGroup>
+
+</Project>
