Horde: Rename the HordeServer project to Horde.Build.

diff --git a/Engine/Source/Programs/Horde/Horde.Build/Jobs/IJobCollection.cs b/Engine/Source/Programs/Horde/Horde.Build/Jobs/IJobCollection.cs
new file mode 100644
--- /dev/null
+++ b/Engine/Source/Programs/Horde/Horde.Build/Jobs/IJobCollection.cs
@@ -0,0 +1,244 @@
+// Copyright Epic Games, Inc. All Rights Reserved.
+
+using EpicGames.Core;
+using HordeServer.Api;
+using HordeCommon;
+using HordeServer.Models;
+using HordeServer.Services;
+using HordeServer.Utilities;
+using Microsoft.Extensions.Logging;
+using MongoDB.Bson;
+using MongoDB.Bson.Serialization.Attributes;
+using MongoDB.Driver;
+using System;
+using System.Collections.Generic;
+using System.Diagnostics.CodeAnalysis;
+using System.IO;
+using System.Linq;
+using System.Threading.Tasks;
+
+namespace HordeServer.Collections
+{
+	using JobId = ObjectId<IJob>;
+	using LeaseId = ObjectId<ILease>;
+	using LogId = ObjectId<ILogFile>;
+	using PoolId = StringId<IPool>;
+	using StreamId = StringId<IStream>;
+	using TemplateRefId = StringId<TemplateRef>;
+	using UserId = ObjectId<IUser>;
+
+	/// <summary>
+	/// Interface for a collection of job documents
+	/// </summary>
+	[SuppressMessage("Compiler", "CA1054:URI properties should not be strings")]
+	public interface IJobCollection
+	{
+		/// <summary>
+		/// Creates a new job
+		/// </summary>
+		/// <param name="JobId">A requested job id</param>
+		/// <param name="StreamId">Unique id of the stream that this job belongs to</param>
+		/// <param name="TemplateRefId">Name of the template ref</param>
+		/// <param name="TemplateHash">Template for this job</param>
+		/// <param name="Graph">The graph for the new job</param>
+		/// <param name="Name">Name of the job</param>
+		/// <param name="Change">The change to build</param>
+		/// <param name="CodeChange">The corresponding code changelist number</param>
+		/// <param name="PreflightChange">Optional changelist to preflight</param>
+		/// <param name="ClonedPreflightChange">Optional cloned preflight changelist</param>
+		/// <param name="StartedByUserId">User that started the job</param>
+		/// <param name="Priority">Priority of the job</param>
+		/// <param name="AutoSubmit">Whether to automatically submit the preflighted change on completion</param>
+		/// <param name="UpdateIssues">Whether to update issues based on the outcome of this job</param>
+		/// <param name="JobTriggers">List of downstream job triggers</param>
+		/// <param name="ShowUgsBadges">Whether to show badges in UGS for this job</param>
+		/// <param name="ShowUgsAlerts">Whether to show alerts in UGS for this job</param>
+		/// <param name="NotificationChannel">Notification channel for this job</param>
+		/// <param name="NotificationChannelFilter">Notification channel filter for this job</param>
+		/// <param name="Arguments">Arguments for the job</param>
+		/// <returns>The new job document</returns>
+		Task<IJob> AddAsync(JobId JobId, StreamId StreamId, TemplateRefId TemplateRefId, ContentHash TemplateHash, IGraph Graph, string Name, int Change, int CodeChange, int? PreflightChange, int? ClonedPreflightChange, UserId? StartedByUserId, Priority? Priority, bool? AutoSubmit, bool? UpdateIssues, List<ChainedJobTemplate>? JobTriggers, bool ShowUgsBadges, bool ShowUgsAlerts, string? NotificationChannel, string? NotificationChannelFilter, List<string>? Arguments);
+
+		/// <summary>
+		/// Gets a job with the given unique id
+		/// </summary>
+		/// <param name="JobId">Job id to search for</param>
+		/// <returns>Information about the given job</returns>
+		Task<IJob?> GetAsync(JobId JobId);
+
+		/// <summary>
+		/// Deletes a job
+		/// </summary>
+		/// <param name="Job">The job to remove</param>
+		Task<bool> RemoveAsync(IJob Job);
+
+		/// <summary>
+		/// Delete all the jobs for a stream
+		/// </summary>
+		/// <param name="StreamId">Unique id of the stream</param>
+		/// <returns>Async task</returns>
+		Task RemoveStreamAsync(StreamId StreamId);
+
+		/// <summary>
+		/// Gets a job's permissions info by ID
+		/// </summary>
+		/// <param name="JobId">Unique id of the job</param>
+		/// <returns>The job document</returns>
+		Task<IJobPermissions?> GetPermissionsAsync(JobId JobId);
+
+		/// <summary>
+		/// Searches for jobs matching the given criteria
+		/// </summary>
+		/// <param name="JobIds">List of job ids to return</param>
+		/// <param name="StreamId">The stream containing the job</param>
+		/// <param name="Name">Name of the job</param>
+		/// <param name="Templates">Templates to look for</param>
+		/// <param name="MinChange">The minimum changelist number</param>
+		/// <param name="MaxChange">The maximum changelist number</param>
+		/// <param name="PreflightChange">Preflight change to find</param>
+		/// <param name="StartedByUser">User id for which to include jobs</param>
+		/// <param name="PreflightStartedByUser">User for which to include preflight jobs</param>
+		/// <param name="MinCreateTime">The minimum creation time</param>
+		/// <param name="MaxCreateTime">The maximum creation time</param>
+		/// <param name="ModifiedBefore">Filter the results by modified time</param>
+		/// <param name="ModifiedAfter">Filter the results by modified time</param>
+		/// <param name="Index">Index of the first result to return</param>
+		/// <param name="Count">Number of results to return</param>
+		/// <returns>List of jobs matching the given criteria</returns>
+		Task<List<IJob>> FindAsync(JobId[]? JobIds = null, StreamId? StreamId = null, string? Name = null, TemplateRefId[]? Templates = null, int? MinChange = null, int? MaxChange = null, int? PreflightChange = null, UserId? PreflightStartedByUser = null, UserId? StartedByUser = null, DateTimeOffset? MinCreateTime = null, DateTimeOffset? MaxCreateTime = null, DateTimeOffset? ModifiedBefore = null, DateTimeOffset? ModifiedAfter = null, int? Index = null, int? Count = null);
+
+		/// <summary>
+		/// Updates a new job
+		/// </summary>
+		/// <param name="Job">The job document to update</param>
+		/// <param name="Graph">Graph for the job</param>
+		/// <param name="Name">Name of the job</param>
+		/// <param name="Priority">Priority of the job</param>
+		/// <param name="AutoSubmit">Automatically submit the job on completion</param>
+		/// <param name="AutoSubmitChange">Changelist that was automatically submitted</param>
+		/// <param name="AutoSubmitMessage"></param>
+		/// <param name="AbortedByUserId">Name of the user that aborted the job</param>
+		/// <param name="NotificationTriggerId">Id for a notification trigger</param>
+		/// <param name="Reports">New reports</param>
+		/// <param name="Arguments">New arguments for the job</param>
+		/// <param name="LabelIdxToTriggerId">New trigger ID for a label in the job</param>
+		/// <param name="JobTrigger">New downstream job id</param>
+		Task<bool> TryUpdateJobAsync(IJob Job, IGraph Graph, string? Name = null, Priority? Priority = null, bool? AutoSubmit = null, int? AutoSubmitChange = null, string? AutoSubmitMessage = null, UserId? AbortedByUserId = null, ObjectId? NotificationTriggerId = null, List<Report>? Reports = null, List<string>? Arguments = null, KeyValuePair<int, ObjectId>? LabelIdxToTriggerId = null, KeyValuePair<TemplateRefId, JobId>? JobTrigger = null);
+
+		/// <summary>
+		/// Updates the state of a batch
+		/// </summary>
+		/// <param name="Job">Job to update</param>
+		/// <param name="Graph">The graph for this job</param>
+		/// <param name="BatchId">Unique id of the batch to update</param>
+		/// <param name="NewLogId">The new log file id</param>
+		/// <param name="NewState">New state of the jobstep</param>
+		/// <param name="NewError">Error code for the batch</param>
+		/// <returns>True if the job was updated, false if it was deleted</returns>
+		Task<bool> TryUpdateBatchAsync(IJob Job, IGraph Graph, SubResourceId BatchId, LogId? NewLogId, JobStepBatchState? NewState, JobStepBatchError? NewError);
+
+		/// <summary>
+		/// Update a jobstep state
+		/// </summary>
+		/// <param name="Job">Job to update</param>
+		/// <param name="Graph">The graph for this job</param>
+		/// <param name="BatchId">Unique id of the batch containing the step</param>
+		/// <param name="StepId">Unique id of the step to update</param>
+		/// <param name="NewState">New state of the jobstep</param>
+		/// <param name="NewOutcome">New outcome of the jobstep</param>
+		/// <param name="NewAbortRequested">New state of request abort</param>
+		/// <param name="NewAbortByUserId">New name of user that requested the abort</param>
+		/// <param name="NewLogId">New log id for the jobstep</param>
+		/// <param name="NewNotificationTriggerId">New id for a notification trigger</param>
+		/// <param name="NewRetryByUserId">Whether the step should be retried</param>
+		/// <param name="NewPriority">New priority for this step</param>
+		/// <param name="NewReports">New report documents</param>
+		/// <param name="NewProperties">Property changes. Any properties with a null value will be removed.</param>
+		/// <returns>True if the job was updated, false if it was deleted in the meantime</returns>
+		Task<bool> TryUpdateStepAsync(IJob Job, IGraph Graph, SubResourceId BatchId, SubResourceId StepId, JobStepState NewState = default, JobStepOutcome NewOutcome = default, bool? NewAbortRequested = null, UserId? NewAbortByUserId = null, LogId? NewLogId = null, ObjectId? NewNotificationTriggerId = null, UserId? NewRetryByUserId = null, Priority? NewPriority = null, List<Report>? NewReports = null, Dictionary<string, string?>? NewProperties = null);
+
+		/// <summary>
+		/// Attempts to update the node groups to be executed for a job. Fails if another write happens in the meantime.
+		/// </summary>
+		/// <param name="Job">The job to update</param>
+		/// <param name="NewGraph">New graph for this job</param>
+		/// <returns>True if the groups were updated to the given list. False if another write happened first.</returns>
+		Task<bool> TryUpdateGraphAsync(IJob Job, IGraph NewGraph);
+
+		/// <summary>
+		/// Removes a job from the dispatch queue. Ignores the state of any batches still remaining to execute. Should only be used to correct for inconsistent state.
+		/// </summary>
+		/// <param name="Job">Job</param>
+		/// <returns></returns>
+		Task<bool> TryRemoveFromDispatchQueueAsync(IJob Job);
+
+		/// <summary>
+		/// Adds an issue to a job
+		/// </summary>
+		/// <param name="JobId">The job id</param>
+		/// <param name="IssueId">The issue to add</param>
+		/// <returns>Async task</returns>
+		Task AddIssueToJobAsync(JobId JobId, int IssueId);
+
+		/// <summary>
+		/// Gets a queue of jobs to consider for execution
+		/// </summary>
+		/// <returns>Sorted list of jobs to execute</returns>
+		Task<List<IJob>> GetDispatchQueueAsync();
+
+		/// <summary>
+		/// Marks a job as skipped
+		/// </summary>
+		/// <param name="Job">The job to update</param>
+		/// <param name="Graph">Graph for the job</param>
+		/// <param name="Reason">Reason for this batch being failed</param>
+		/// <returns>Updated version of the job</returns>
+		Task<IJob?> SkipAllBatchesAsync(IJob? Job, IGraph Graph, JobStepBatchError Reason);
+
+		/// <summary>
+		/// Marks a batch as skipped
+		/// </summary>
+		/// <param name="Job">The job to update</param>
+		/// <param name="BatchIdx">The batch to mark as skipped</param>
+		/// <param name="Graph">Graph for the job</param>
+		/// <param name="Reason">Reason for this batch being failed</param>
+		/// <returns>Updated version of the job</returns>
+		Task<IJob?> SkipBatchAsync(IJob? Job, int BatchIdx, IGraph Graph, JobStepBatchError Reason);
+
+		/// <summary>
+		/// Abort an agent's lease, and update the payload accordingly
+		/// </summary>
+		/// <param name="Job">The job to update</param>
+		/// <param name="BatchIdx">Index of the batch to cancel</param>
+		/// <param name="Graph">Graph for the job</param>
+		/// <returns>True if the job is updated</returns>
+		Task<bool> TryFailBatchAsync(IJob Job, int BatchIdx, IGraph Graph);
+
+		/// <summary>
+		/// Attempt to assign a lease to execute a batch
+		/// </summary>
+		/// <param name="Job">The job containing the batch</param>
+		/// <param name="BatchIdx">Index of the batch</param>
+		/// <param name="PoolId">The pool id</param>
+		/// <param name="AgentId">New agent to execute the batch</param>
+		/// <param name="SessionId">Session of the agent that is to execute the batch</param>
+		/// <param name="LeaseId">The lease unique id</param>
+		/// <param name="LogId">Unique id of the log for the batch</param>
+		/// <returns>True if the batch is updated</returns>
+		Task<bool> TryAssignLeaseAsync(IJob Job, int BatchIdx, PoolId PoolId, AgentId AgentId, ObjectId SessionId, LeaseId LeaseId, LogId LogId);
+
+		/// <summary>
+		/// Cancel a lease reservation on a batch (before it has started)
+		/// </summary>
+		/// <param name="Job">The job containing the lease</param>
+		/// <param name="BatchIdx">Index of the batch to cancel</param>
+		/// <returns>True if the job is updated</returns>
+		Task<bool> TryCancelLeaseAsync(IJob Job, int BatchIdx);
+
+		/// <summary>
+		/// Upgrade all documents in the collection
+		/// </summary>
+		/// <returns>Async task</returns>
+		Task UpgradeDocumentsAsync();
+	}
+}
diff --git a/Engine/Source/Programs/Horde/Horde.Build/Jobs/IJobStepRefCollection.cs b/Engine/Source/Programs/Horde/Horde.Build/Jobs/IJobStepRefCollection.cs
new file mode 100644
--- /dev/null
+++ b/Engine/Source/Programs/Horde/Horde.Build/Jobs/IJobStepRefCollection.cs
@@ -0,0 +1,99 @@
+// Copyright Epic Games, Inc. All Rights Reserved.
+
+using HordeCommon;
+using HordeServer.Models;
+using HordeServer.Services;
+using MongoDB.Bson;
+using MongoDB.Bson.Serialization.Attributes;
+using MongoDB.Driver;
+using System;
+using System.Collections.Generic;
+using System.Linq;
+using System.Threading.Tasks;
+using HordeServer.Utilities;
+
+namespace HordeServer.Collections
+{
+	using LogId = ObjectId<ILogFile>;
+	using PoolId = StringId<IPool>;
+	using StreamId = StringId<IStream>;
+	using TemplateRefId = StringId<TemplateRef>;
+
+	/// <summary>
+	/// Interface for a collection of JobStepRef documents
+	/// </summary>
+	public interface IJobStepRefCollection
+	{
+		/// <summary>
+		/// Constructor
+		/// </summary>
+		/// <param name="Id">Unique id for the step being referenced</param>
+		/// <param name="JobName">Name of the job</param>
+		/// <param name="StepName">Name of the step</param>
+		/// <param name="StreamId">Unique id for the stream containing the job</param>
+		/// <param name="TemplateId"></param>
+		/// <param name="Change">The change number being built</param>
+		/// <param name="LogId">The log file id</param>
+		/// <param name="PoolId">The pool id</param>
+		/// <param name="AgentId">The agent id</param>
+		/// <param name="Outcome">Outcome of this step, if known</param>
+		/// <param name="WaitTime">Time taken for the batch containing this step to start</param>
+		/// <param name="InitTime">Time taken for the batch containing this step to initializer</param>
+		/// <param name="StartTimeUtc">Start time</param>
+		/// <param name="FinishTimeUtc">Finish time for the step, if known</param>
+		Task<IJobStepRef> InsertOrReplaceAsync(JobStepRefId Id, string JobName, string StepName, StreamId StreamId, TemplateRefId TemplateId, int Change, LogId? LogId, PoolId? PoolId, AgentId? AgentId, JobStepOutcome? Outcome, float WaitTime, float InitTime, DateTime StartTimeUtc, DateTime? FinishTimeUtc);
+
+		/// <summary>
+		/// Gets the history of a given node
+		/// </summary>
+		/// <param name="StreamId">Unique id for a stream</param>
+		/// <param name="TemplateId"></param>
+		/// <param name="NodeName">Name of the node</param>
+		/// <param name="Change">The current change</param>
+		/// <param name="IncludeFailed">Whether to include failed nodes</param>
+		/// <param name="Count">Number of results to return</param>
+		/// <returns>List of step references</returns>
+		Task<List<IJobStepRef>> GetStepsForNodeAsync(StreamId StreamId, TemplateRefId TemplateId, string NodeName, int? Change, bool IncludeFailed, int Count);
+
+		/// <summary>
+		/// Gets the previous job that ran a given step
+		/// </summary>
+		/// <param name="StreamId">Id of the stream to search</param>
+		/// <param name="TemplateId">The template id</param>
+		/// <param name="NodeName">Name of the step to find</param>
+		/// <param name="Change">The current changelist number</param>
+		/// <returns>The previous job, or null.</returns>
+		Task<IJobStepRef?> GetPrevStepForNodeAsync(StreamId StreamId, TemplateRefId TemplateId, string NodeName, int Change);
+
+		/// <summary>
+		/// Gets the next job that ran a given step
+		/// </summary>
+		/// <param name="StreamId">Id of the stream to search</param>
+		/// <param name="TemplateId">The template id</param>
+		/// <param name="NodeName">Name of the step to find</param>
+		/// <param name="Change">The current changelist number</param>
+		/// <returns>The previous job, or null.</returns>
+		Task<IJobStepRef?> GetNextStepForNodeAsync(StreamId StreamId, TemplateRefId TemplateId, string NodeName, int Change);
+	}
+
+	static class JobStepRefCollectionExtensions
+	{
+		public static Task UpdateAsync(this IJobStepRefCollection JobStepRefs, IJob Job, IJobStepBatch Batch, IJobStep Step, IGraph Graph)
+		{
+			if (Job.PreflightChange == 0)
+			{
+				float WaitTime = (float)(Batch.GetWaitTime() ?? TimeSpan.Zero).TotalSeconds;
+				float InitTime = (float)(Batch.GetInitTime() ?? TimeSpan.Zero).TotalSeconds;
+
+				string NodeName = Graph.Groups[Batch.GroupIdx].Nodes[Step.NodeIdx].Name;
+				JobStepOutcome? Outcome = Step.IsPending() ? (JobStepOutcome?)null : Step.Outcome;
+
+				return JobStepRefs.InsertOrReplaceAsync(new JobStepRefId(Job.Id, Batch.Id, Step.Id), Job.Name, NodeName, Job.StreamId, Job.TemplateId, Job.Change, Step.LogId, Batch.PoolId, Batch.AgentId, Outcome, WaitTime, InitTime, Step.StartTimeUtc ?? DateTime.UtcNow, Step.FinishTimeUtc);
+			}
+			else
+			{
+				return Task.CompletedTask;
+			}
+		}
+	}
+}
diff --git a/Engine/Source/Programs/Horde/Horde.Build/Jobs/IJobTimingCollection.cs b/Engine/Source/Programs/Horde/Horde.Build/Jobs/IJobTimingCollection.cs
new file mode 100644
--- /dev/null
+++ b/Engine/Source/Programs/Horde/Horde.Build/Jobs/IJobTimingCollection.cs
@@ -0,0 +1,102 @@
+// Copyright Epic Games, Inc. All Rights Reserved.
+
+using Amazon.Runtime.Internal.Util;
+using Amazon.S3.Model.Internal.MarshallTransformations;
+using HordeServer.Models;
+using HordeServer.Services;
+using HordeServer.Utilities;
+using MongoDB.Bson;
+using MongoDB.Bson.Serialization.Attributes;
+using MongoDB.Driver;
+using System;
+using System.Collections.Generic;
+using System.Diagnostics.CodeAnalysis;
+using System.Linq;
+using System.Threading.Tasks;
+
+namespace HordeServer.Collections
+{
+	using JobId = ObjectId<IJob>;
+
+	/// <summary>
+	/// Average timing information for a node
+	/// </summary>
+	public class JobStepTimingData
+	{
+		/// <summary>
+		/// Name of the node
+		/// </summary>
+		public string Name { get; set; }
+
+		/// <summary>
+		/// Wait time before executing the group containing this node
+		/// </summary>
+		[BsonIgnoreIfNull]
+		public float? AverageWaitTime { get; set; }
+
+		/// <summary>
+		/// Time taken for the group containing this node to initialize
+		/// </summary>
+		[BsonIgnoreIfNull]
+		public float? AverageInitTime { get; set; }
+
+		/// <summary>
+		/// Time spent executing this node
+		/// </summary>
+		[BsonIgnoreIfNull]
+		public float? AverageDuration { get; set; }
+
+		/// <summary>
+		/// Constructor for serialization
+		/// </summary>
+		[BsonConstructor]
+		private JobStepTimingData()
+		{
+			this.Name = null!;
+		}
+
+		/// <summary>
+		/// Constructor
+		/// </summary>
+		/// <param name="Name">Name of the node</param>
+		/// <param name="AverageWaitTime">Wait time before executing the group containing this node</param>
+		/// <param name="AverageInitTime">Time taken for the group containing this node to initialize</param>
+		/// <param name="AverageDuration">Time spent executing this node</param>
+		public JobStepTimingData(string Name, float? AverageWaitTime, float? AverageInitTime, float? AverageDuration)
+		{
+			this.Name = Name;
+			this.AverageWaitTime = AverageWaitTime;
+			this.AverageInitTime = AverageInitTime;
+			this.AverageDuration = AverageDuration;
+		}
+	}
+
+	/// <summary>
+	/// Interface for a collection of JobTiming documents
+	/// </summary>
+	public interface IJobTimingCollection
+	{
+		/// <summary>
+		/// Add timing information for the given job
+		/// </summary>
+		/// <param name="JobId">The job id</param>
+		/// <param name="Steps">List of timing info for each step</param>
+		/// <returns>New timing document</returns>
+		Task<IJobTiming?> TryAddAsync(JobId JobId, List<JobStepTimingData> Steps);
+
+		/// <summary>
+		/// Attempts to get the timing information for a particular job
+		/// </summary>
+		/// <param name="JobId">The unique job id</param>
+		/// <returns>Timing info for the requested jbo</returns>
+		Task<IJobTiming?> TryGetAsync(JobId JobId);
+
+		/// <summary>
+		/// Adds timing information for the particular job
+		/// </summary>
+		/// <param name="JobTiming">The current timing info for the job</param>
+		/// <param name="Steps">List of steps to add</param>
+		/// <returns>New timing document. Null if the document could not be updated.</returns>
+		Task<IJobTiming?> TryAddStepsAsync(IJobTiming JobTiming, List<JobStepTimingData> Steps);
+	}
+}
diff --git a/Engine/Source/Programs/Horde/Horde.Build/Jobs/ITestDataCollection.cs b/Engine/Source/Programs/Horde/Horde.Build/Jobs/ITestDataCollection.cs
new file mode 100644
--- /dev/null
+++ b/Engine/Source/Programs/Horde/Horde.Build/Jobs/ITestDataCollection.cs
@@ -0,0 +1,65 @@
+// Copyright Epic Games, Inc. All Rights Reserved.
+
+using HordeServer.Models;
+using HordeServer.Services;
+using HordeServer.Utilities;
+using MongoDB.Bson;
+using MongoDB.Bson.Serialization.Attributes;
+using MongoDB.Driver;
+using System;
+using System.Collections.Generic;
+using System.Diagnostics.CodeAnalysis;
+using System.Linq;
+using System.Threading.Tasks;
+
+namespace HordeServer.Collections
+{
+	using JobId = ObjectId<IJob>;
+	using ProjectId = StringId<IProject>;
+	using StreamId = StringId<IStream>;
+	using TemplateRefId = StringId<TemplateRef>;
+
+	/// <summary>
+	/// Collection of test data documents
+	/// </summary>
+	public interface ITestDataCollection
+	{
+		/// <summary>
+		/// Creates a new test data document
+		/// </summary>
+		/// <param name="Job">The job containing the step</param>
+		/// <param name="Step">The step producing the data</param>
+		/// <param name="Key">Key identifying the test</param>
+		/// <param name="Value">The data to store</param>
+		/// <returns>The new stream document</returns>
+		Task<ITestData> AddAsync(IJob Job, IJobStep Step, string Key, BsonDocument Value);
+
+		/// <summary>
+		/// Gets a stream by ID
+		/// </summary>
+		/// <param name="Id">Unique id of the stream</param>
+		/// <returns>The stream document</returns>
+		Task<ITestData?> GetAsync(ObjectId Id);
+
+		/// <summary>
+		/// Searches for test data that matches a set of criteria
+		/// </summary>
+		/// <param name="StreamId">The stream id</param>
+		/// <param name="MinChange">The minimum changelist number to return (inclusive)</param>
+		/// <param name="MaxChange">The maximum changelist number to return (inclusive)</param>
+		/// <param name="JobId">The job id</param>
+		/// <param name="StepId">The unique step id</param>
+		/// <param name="Key">Key identifying the result to return</param>
+		/// <param name="Index">Offset within the results to return</param>
+		/// <param name="Count">Number of results to return</param>
+		/// <returns>The stream document</returns>
+		Task<List<ITestData>> FindAsync(StreamId? StreamId, int? MinChange, int? MaxChange, JobId? JobId, SubResourceId? StepId, string? Key = null, int Index = 0, int Count = 10);
+
+		/// <summary>
+		/// Delete the test data
+		/// </summary>
+		/// <param name="Id">Unique id of the test data</param>
+		/// <returns>Async task</returns>
+		Task DeleteAsync(ObjectId Id);
+	}
+}
diff --git a/Engine/Source/Programs/Horde/Horde.Build/Jobs/Impl/GraphCollection.cs b/Engine/Source/Programs/Horde/Horde.Build/Jobs/Impl/GraphCollection.cs
new file mode 100644
--- /dev/null
+++ b/Engine/Source/Programs/Horde/Horde.Build/Jobs/Impl/GraphCollection.cs
@@ -0,0 +1,469 @@
+// Copyright Epic Games, Inc. All Rights Reserved.
+
+using EpicGames.Core;
+using HordeServer.Api;
+using HordeCommon;
+using HordeServer.Models;
+using HordeServer.Services;
+using MongoDB.Bson;
+using MongoDB.Bson.Serialization.Attributes;
+using MongoDB.Driver;
+using System;
+using System.Collections.Concurrent;
+using System.Collections.Generic;
+using System.Diagnostics;
+using System.Linq;
+using System.Threading;
+using System.Threading.Tasks;
+
+namespace HordeServer.Collections.Impl
+{
+	/// <summary>
+	/// Collection of graph documents
+	/// </summary>
+	public class GraphCollection : IGraphCollection
+	{
+		/// <summary>
+		/// Represents a node in the graph
+		/// </summary>
+		[DebuggerDisplay("{Name}")]
+		class Node : INode
+		{
+			[BsonRequired]
+			public string Name { get; set; }
+
+			[BsonIgnoreIfNull]
+			public NodeRef[] InputDependencies { get; set; }
+
+			[BsonIgnoreIfNull]
+			public NodeRef[] OrderDependencies { get; set; }
+
+			public Priority Priority { get; set; }
+			public bool AllowRetry { get; set; } = true;
+			public bool RunEarly { get; set; }
+			public bool Warnings { get; set; } = true;
+
+			[BsonIgnoreIfNull]
+			public Dictionary<string, string>? Credentials { get; set; }
+
+			[BsonIgnoreIfNull]
+			public Dictionary<string, string>? Properties { get; set; }
+
+			IReadOnlyDictionary<string, string>? INode.Credentials => Credentials;
+			IReadOnlyDictionary<string, string>? INode.Properties => Properties;
+
+			[BsonConstructor]
+			private Node()
+			{
+				Name = null!;
+				InputDependencies = null!;
+				OrderDependencies = null!;
+			}
+
+			public Node(string Name, NodeRef[] InputDependencies, NodeRef[] OrderDependencies, Priority Priority, bool AllowRetry, bool RunEarly, bool Warnings, Dictionary<string, string>? Credentials, Dictionary<string, string>? Properties)
+			{
+				this.Name = Name;
+				this.InputDependencies = InputDependencies;
+				this.OrderDependencies = OrderDependencies;
+				this.Priority = Priority;
+				this.AllowRetry = AllowRetry;
+				this.RunEarly = RunEarly;
+				this.Warnings = Warnings;
+				this.Credentials = Credentials;
+				this.Properties = Properties;
+			}
+		}
+
+		class NodeGroup : INodeGroup
+		{
+			public string AgentType { get; set; }
+			public List<Node> Nodes { get; set; }
+
+			IReadOnlyList<INode> INodeGroup.Nodes => Nodes;
+
+			/// <summary>
+			/// Private constructor for BSON serializer
+			/// </summary>
+			[BsonConstructor]
+			private NodeGroup()
+			{
+				AgentType = null!;
+				Nodes = null!;
+			}
+
+			/// <summary>
+			/// Constructor
+			/// </summary>
+			/// <param name="AgentType">The type of agent to execute this group</param>
+			/// <param name="Nodes">Nodes to execute</param>
+			public NodeGroup(string AgentType, List<Node> Nodes)
+			{
+				this.AgentType = AgentType;
+				this.Nodes = Nodes;
+			}
+		}
+
+		class Aggregate : IAggregate
+		{
+			public string Name { get; set; }
+			public List<NodeRef> Nodes { get; set; }
+
+			IReadOnlyList<NodeRef> IAggregate.Nodes => Nodes;
+
+			private Aggregate()
+			{
+				Name = null!;
+				Nodes = new List<NodeRef>();
+			}
+
+			public Aggregate(string Name, List<NodeRef> Nodes)
+			{
+				this.Name = Name;
+				this.Nodes = Nodes;
+			}
+		}
+
+		class Label : ILabel
+		{
+			[BsonIgnoreIfNull]
+			public string? Name { get; set; }
+
+			[BsonIgnoreIfNull]
+			public string? Category { get; set; }
+
+			[BsonIgnoreIfNull]
+			public string? DashboardName { get; set; }
+
+			[BsonIgnoreIfNull]
+			public string? DashboardCategory { get; set; }
+
+			[BsonIgnoreIfNull]
+			public string? UgsName { get; set; }
+
+			[BsonIgnoreIfNull]
+			public string? UgsProject { get; set; }
+
+			[BsonIgnoreIfNull]
+			public LabelChange Change { get; set; }
+
+			public List<NodeRef> RequiredNodes { get; set; }
+			public List<NodeRef> IncludedNodes { get; set; }
+
+			string? ILabel.DashboardName => DashboardName ?? Name;
+			string? ILabel.DashboardCategory => DashboardCategory ?? Category;
+
+			private Label()
+			{
+				RequiredNodes = new List<NodeRef>();
+				IncludedNodes = new List<NodeRef>();
+			}
+
+			public Label(string? DashboardName, string? DashboardCategory, string? UgsName, string? UgsProject, LabelChange Change, List<NodeRef> RequiredNodes, List<NodeRef> IncludedNodes)
+			{
+				this.DashboardName = DashboardName;
+				this.DashboardCategory = DashboardCategory;
+				this.UgsName = UgsName;
+				this.UgsProject = UgsProject;
+				this.Change = Change;
+				this.RequiredNodes = RequiredNodes;
+				this.IncludedNodes = IncludedNodes;
+			}
+		}
+
+		class GraphDocument : IGraph
+		{
+			public static GraphDocument Empty { get; } = new GraphDocument(new List<NodeGroup>(), new List<Aggregate>(), new List<Label>());
+
+			[BsonRequired, BsonId]
+			public ContentHash Id { get; private set; } = ContentHash.Empty;
+
+			public int Schema { get; set; }
+			public List<NodeGroup> Groups { get; private set; } = new List<NodeGroup>();
+			public List<Aggregate> Aggregates { get; private set; } = new List<Aggregate>();
+			public List<Label> Labels { get; private set; } = new List<Label>();
+
+			[BsonIgnore]
+			IReadOnlyDictionary<string, NodeRef>? CachedNodeNameToRef;
+
+			IReadOnlyList<INodeGroup> IGraph.Groups => Groups;
+			IReadOnlyList<IAggregate> IGraph.Aggregates => Aggregates;
+			IReadOnlyList<ILabel> IGraph.Labels => Labels;
+
+			[BsonConstructor]
+			private GraphDocument()
+			{
+			}
+
+			public GraphDocument(List<NodeGroup> Groups, List<Aggregate> Aggregates, List<Label> Labels)
+			{
+				this.Groups = Groups;
+				this.Aggregates = Aggregates;
+				this.Labels = Labels;
+				this.Id = ContentHash.SHA1(BsonExtensionMethods.ToBson(this));
+			}
+
+			public GraphDocument(GraphDocument BaseGraph, List<NewGroup>? NewGroupRequests, List<NewAggregate>? NewAggregateRequests, List<NewLabel>? NewLabelRequests)
+			{
+				Dictionary<string, NodeRef> NodeNameToRef = new Dictionary<string, NodeRef>(BaseGraph.GetNodeNameToRef(), StringComparer.OrdinalIgnoreCase);
+
+				// Update the new list of groups
+				List<NodeGroup> NewGroups = new List<NodeGroup>(BaseGraph.Groups);
+				if (NewGroupRequests != null)
+				{
+					foreach (NewGroup NewGroupRequest in NewGroupRequests)
+					{
+						List<Node> Nodes = new List<Node>();
+						foreach (NewNode NewNodeRequest in NewGroupRequest.Nodes)
+						{
+							int NodeIdx = Nodes.Count;
+
+							Priority Priority = NewNodeRequest.Priority ?? Priority.Normal;
+							bool bAllowRetry = NewNodeRequest.AllowRetry ?? true;
+							bool bRunEarly = NewNodeRequest.RunEarly ?? false;
+							bool bWarnings = NewNodeRequest.Warnings ?? true;
+
+							NodeRef[] InputDependencies = (NewNodeRequest.InputDependencies == null) ? Array.Empty<NodeRef>() : NewNodeRequest.InputDependencies.Select(x => NodeNameToRef[x]).ToArray();
+							NodeRef[] OrderDependencies = (NewNodeRequest.OrderDependencies == null) ? Array.Empty<NodeRef>() : NewNodeRequest.OrderDependencies.Select(x => NodeNameToRef[x]).ToArray();
+							OrderDependencies = OrderDependencies.Union(InputDependencies).ToArray();
+							Nodes.Add(new Node(NewNodeRequest.Name, InputDependencies, OrderDependencies, Priority, bAllowRetry, bRunEarly, bWarnings, NewNodeRequest.Credentials, NewNodeRequest.Properties));
+
+							NodeNameToRef.Add(NewNodeRequest.Name, new NodeRef(NewGroups.Count, NodeIdx));
+						}
+						NewGroups.Add(new NodeGroup(NewGroupRequest.AgentType, Nodes));
+					}
+				}
+
+				// Update the list of aggregates
+				List<Aggregate> NewAggregates = new List<Aggregate>(BaseGraph.Aggregates);
+				if (NewAggregateRequests != null)
+				{
+					foreach (NewAggregate NewAggregateRequest in NewAggregateRequests)
+					{
+						List<NodeRef> Nodes = NewAggregateRequest.Nodes.ConvertAll(x => NodeNameToRef[x]);
+						NewAggregates.Add(new Aggregate(NewAggregateRequest.Name, Nodes));
+					}
+				}
+
+				// Update the list of labels
+				List<Label> NewLabels = new List<Label>(BaseGraph.Labels);
+				if (NewLabelRequests != null)
+				{
+					foreach (NewLabel NewLabelRequest in NewLabelRequests)
+					{
+						List<NodeRef> RequiredNodes = NewLabelRequest.RequiredNodes.ConvertAll(x => NodeNameToRef[x]);
+						List<NodeRef> IncludedNodes = NewLabelRequest.IncludedNodes.ConvertAll(x => NodeNameToRef[x]);
+						NewLabels.Add(new Label(NewLabelRequest.DashboardName, NewLabelRequest.DashboardCategory, NewLabelRequest.UgsName, NewLabelRequest.UgsProject, NewLabelRequest.Change, RequiredNodes, IncludedNodes));
+					}
+				}
+
+				// Create the new arrays
+				Groups = NewGroups;
+				Aggregates = NewAggregates;
+				Labels = NewLabels;
+
+				// Create the new graph, and save the generated node lookup into it
+				CachedNodeNameToRef = NodeNameToRef;
+
+				// Compute the hash
+				Id = ContentHash.SHA1(BsonExtensionMethods.ToBson(this));
+			}
+
+			public IReadOnlyDictionary<string, NodeRef> GetNodeNameToRef()
+			{
+				if (CachedNodeNameToRef == null)
+				{
+					Dictionary<string, NodeRef> NodeNameToRef = new Dictionary<string, NodeRef>(StringComparer.OrdinalIgnoreCase);
+					for (int GroupIdx = 0; GroupIdx < Groups.Count; GroupIdx++)
+					{
+						List<Node> Nodes = Groups[GroupIdx].Nodes;
+						for (int NodeIdx = 0; NodeIdx < Nodes.Count; NodeIdx++)
+						{
+							Node Node = Nodes[NodeIdx];
+							NodeNameToRef[Node.Name] = new NodeRef(GroupIdx, NodeIdx);
+						}
+					}
+					CachedNodeNameToRef = NodeNameToRef;
+				}
+				return CachedNodeNameToRef;
+			}
+		}
+
+		/// <summary>
+		/// Stores information about a cached graph
+		/// </summary>
+		class CachedGraph
+		{
+			/// <summary>
+			/// Time at which the graph was last accessed
+			/// </summary>
+			public long LastAccessTime
+			{
+				get { return LastAccessTimePrivate; }
+			}
+
+			/// <summary>
+			/// Backing value for <see cref="LastAccessTime"/>
+			/// </summary>
+			private long LastAccessTimePrivate;
+
+			/// <summary>
+			/// The graph instance
+			/// </summary>
+			public IGraph Graph { get; }
+
+			/// <summary>
+			/// Constructor
+			/// </summary>
+			/// <param name="Graph">The graph to store</param>
+			public CachedGraph(IGraph Graph)
+			{
+				this.LastAccessTimePrivate = Stopwatch.GetTimestamp();
+				this.Graph = Graph;
+			}
+
+			/// <summary>
+			/// Update the last access time
+			/// </summary>
+			public void Touch()
+			{
+				for (; ; )
+				{
+					long Time = Stopwatch.GetTimestamp();
+					long LastAccessTimeCopy = LastAccessTimePrivate;
+					if (Time < LastAccessTimeCopy || Interlocked.CompareExchange(ref LastAccessTimePrivate, Time, LastAccessTimeCopy) == LastAccessTimeCopy)
+					{
+						break;
+					}
+				}
+			}
+		}
+
+		/// <summary>
+		/// The jobs collection
+		/// </summary>
+		IMongoCollection<GraphDocument> Graphs;
+
+		/// <summary>
+		/// Maximum number of graphs to keep in the cache
+		/// </summary>
+		const int MaxGraphs = 1000;
+
+		/// <summary>
+		/// Cache of recently accessed graphs
+		/// </summary>
+		ConcurrentDictionary<ContentHash, CachedGraph> CachedGraphs = new ConcurrentDictionary<ContentHash, CachedGraph>();
+
+		/// <summary>
+		/// Constructor
+		/// </summary>
+		/// <param name="DatabaseService">The database service singleton</param>
+		public GraphCollection(DatabaseService DatabaseService)
+		{
+			Graphs = DatabaseService.GetCollection<GraphDocument>("Graphs");
+		}
+
+		/// <summary>
+		/// Adds a new graph document
+		/// </summary>
+		/// <param name="Graph">The graph to add</param>
+		/// <returns>Async task</returns>
+		async Task AddAsync(GraphDocument Graph)
+		{
+			if (!await Graphs.Find(x => x.Id == Graph.Id).AnyAsync())
+			{
+				try
+				{
+					await Graphs.InsertOneAsync(Graph);
+				}
+				catch (MongoWriteException Ex)
+				{
+					if (Ex.WriteError.Category != ServerErrorCategory.DuplicateKey)
+					{
+						throw;
+					}
+				}
+			}
+		}
+
+		/// <inheritdoc/>
+		public async Task<IGraph> AddAsync(ITemplate Template)
+		{
+			Node Node = new Node(IJob.SetupNodeName, Array.Empty<NodeRef>(), Array.Empty<NodeRef>(), Priority.High, true, false, true, null, null);
+			NodeGroup Group = new NodeGroup(Template.InitialAgentType ?? "Win64", new List<Node> { Node });
+
+			GraphDocument Graph = new GraphDocument(new List<NodeGroup> { Group }, new List<Aggregate>(), new List<Label>());
+			await AddAsync(Graph);
+			return Graph;
+		}
+
+		/// <inheritdoc/>
+		public async Task<IGraph> AppendAsync(IGraph? BaseGraph, List<NewGroup>? NewGroupRequests, List<NewAggregate>? NewAggregateRequests, List<NewLabel>? NewLabelRequests)
+		{
+			GraphDocument Graph = new GraphDocument((GraphDocument?)BaseGraph ?? GraphDocument.Empty, NewGroupRequests, NewAggregateRequests, NewLabelRequests);
+			await AddAsync(Graph);
+			return Graph;
+		}
+
+		/// <inheritdoc/>
+		public async Task<IGraph> GetAsync(ContentHash? Hash)
+		{
+			// Special case for an empty graph request
+			if (Hash == null || Hash == ContentHash.Empty || Hash == GraphDocument.Empty.Id)
+			{
+				return GraphDocument.Empty;
+			}
+
+			// Try to read the graph from the cache
+			CachedGraph? CachedGraph;
+			if (CachedGraphs.TryGetValue(Hash, out CachedGraph))
+			{
+				// Update the last access time
+				CachedGraph.Touch();
+			}
+			else
+			{
+				// Trim the cache
+				while (CachedGraphs.Count > MaxGraphs)
+				{
+					ContentHash? RemoveHash = CachedGraphs.OrderBy(x => x.Value.LastAccessTime).Select(x => x.Key).FirstOrDefault();
+					if (RemoveHash == null || RemoveHash == ContentHash.Empty)
+					{
+						break;
+					}
+					CachedGraphs.TryRemove(RemoveHash, out _);
+				}
+
+				// Create the new entry
+				CachedGraph = new CachedGraph(await Graphs.Find<GraphDocument>(x => x.Id == Hash).FirstAsync());
+				CachedGraphs.TryAdd(Hash, CachedGraph);
+			}
+			return CachedGraph.Graph;
+		}
+
+		/// <inheritdoc/>
+		public async Task<List<IGraph>> FindAllAsync(ContentHash[]? Hashes, int? Index, int? Count)
+		{
+			FilterDefinitionBuilder<GraphDocument> FilterBuilder = Builders<GraphDocument>.Filter;
+
+			FilterDefinition<GraphDocument> Filter = FilterBuilder.Empty;
+			if (Hashes != null)
+			{
+				Filter &= FilterBuilder.In(x => x.Id, Hashes);
+			}
+
+			List<GraphDocument> Results;
+			IFindFluent<GraphDocument, GraphDocument> Search = Graphs.Find(Filter);
+			if (Index != null)
+			{
+				Search = Search.Skip(Index.Value);
+			}
+			if (Count != null)
+			{
+				Search = Search.Limit(Count.Value);
+			}
+
+			Results = await Search.ToListAsync();
+			return Results.ConvertAll<IGraph>(x => x);
+		}
+	}
+}
