Horde: Rename the HordeServer project to Horde.Build.

diff --git a/Engine/Source/Programs/Horde/Horde.Build/Jobs/Impl/JobCollection.cs b/Engine/Source/Programs/Horde/Horde.Build/Jobs/Impl/JobCollection.cs
new file mode 100644
--- /dev/null
+++ b/Engine/Source/Programs/Horde/Horde.Build/Jobs/Impl/JobCollection.cs
@@ -0,0 +1,1717 @@
+// Copyright Epic Games, Inc. All Rights Reserved.
+
+using EpicGames.Core;
+using HordeServer.Api;
+using HordeCommon;
+using HordeServer.Models;
+using HordeServer.Services;
+using HordeServer.Utilities;
+using Microsoft.Extensions.Logging;
+using MongoDB.Bson;
+using MongoDB.Bson.Serialization.Attributes;
+using MongoDB.Driver;
+using System;
+using System.Collections.Generic;
+using System.Diagnostics.CodeAnalysis;
+using System.IO;
+using System.Linq;
+using System.Threading.Tasks;
+using OpenTracing.Util;
+using OpenTracing;
+
+namespace HordeServer.Collections.Impl
+{
+	using JobId = ObjectId<IJob>;
+	using LeaseId = ObjectId<ILease>;
+	using LogId = ObjectId<ILogFile>;
+	using StreamId = StringId<IStream>;
+	using TemplateRefId = StringId<TemplateRef>;
+	using PoolId = StringId<IPool>;
+	using UserId = ObjectId<IUser>;
+
+	/// <summary>
+	/// Wrapper around the jobs collection in a mongo DB
+	/// </summary>
+	public class JobCollection : IJobCollection
+	{
+		/// <summary>
+		/// Embedded jobstep document
+		/// </summary>
+		[BsonIgnoreExtraElements]
+		class JobStepDocument : IJobStep
+		{
+			[BsonRequired]
+			public SubResourceId Id { get; set; }
+
+			[BsonRequired]
+			public int NodeIdx { get; set; }
+
+			[BsonRequired]
+			public JobStepState State { get; set; } = JobStepState.Waiting;
+
+			public JobStepOutcome Outcome { get; set; } = JobStepOutcome.Success;
+
+			[BsonIgnoreIfNull]
+			public LogId? LogId { get; set; }
+
+			[BsonIgnoreIfNull]
+			public ObjectId? NotificationTriggerId { get; set; }
+
+			[BsonIgnoreIfNull]
+			public DateTimeOffset? ReadyTime { get; set; }
+
+			[BsonIgnoreIfNull]
+			public DateTimeOffset? StartTime { get; set; }
+
+			[BsonIgnoreIfNull]
+			public DateTimeOffset? FinishTime { get; set; }
+
+			[BsonIgnoreIfNull]
+			public Priority? Priority { get; set; }
+
+			[BsonIgnoreIfDefault, BsonDefaultValue(false)]
+			public bool Retry { get; set; }
+
+			public UserId? RetriedByUserId { get; set; }
+
+			[BsonElement("RetryByUser")]
+			public string? RetriedByUser_DEPRECATED { get; set; }
+
+			public bool AbortRequested { get; set; } = false;
+
+			public UserId? AbortedByUserId { get; set; }
+			
+			[BsonElement("AbortByUser")]
+			public string? AbortedByUser_DEPRECATED { get; set; }
+
+			[BsonIgnoreIfNull]
+			public List<Report>? Reports { get; set; }
+			IReadOnlyList<IReport>? IJobStep.Reports => Reports;
+
+			[BsonIgnoreIfNull]
+			public Dictionary<string, string>? Properties { get; set; }
+
+			DateTime? IJobStep.ReadyTimeUtc => ReadyTime?.UtcDateTime;
+			DateTime? IJobStep.StartTimeUtc => StartTime?.UtcDateTime;
+			DateTime? IJobStep.FinishTimeUtc => FinishTime?.UtcDateTime;
+
+			[BsonConstructor]
+			private JobStepDocument()
+			{
+			}
+
+			public JobStepDocument(SubResourceId Id, int NodeIdx)
+			{
+				this.Id = Id;
+				this.NodeIdx = NodeIdx;
+			}
+		}
+
+		class JobStepBatchDocument : IJobStepBatch
+		{
+			[BsonRequired]
+			public SubResourceId Id { get; set; }
+
+			public LogId? LogId { get; set; }
+
+			[BsonRequired]
+			public int GroupIdx { get; set; }
+
+			[BsonRequired]
+			public JobStepBatchState State { get; set; }
+
+			[BsonIgnoreIfDefault, BsonDefaultValue(JobStepBatchError.None)]
+			public JobStepBatchError Error { get; set; }
+
+			public List<JobStepDocument> Steps { get; set; } = new List<JobStepDocument>();
+
+			[BsonIgnoreIfNull]
+			public PoolId? PoolId { get; set; }
+
+			[BsonIgnoreIfNull]
+			public AgentId? AgentId { get; set; }
+
+			[BsonIgnoreIfNull]
+			public ObjectId? SessionId { get; set; }
+
+			[BsonIgnoreIfNull]
+			public LeaseId? LeaseId { get; set; }
+
+			public int SchedulePriority { get; set; }
+
+			[BsonIgnoreIfNull]
+			public DateTimeOffset? ReadyTime { get; set; }
+
+			[BsonIgnoreIfNull]
+			public DateTimeOffset? StartTime { get; set; }
+
+			[BsonIgnoreIfNull]
+			public DateTimeOffset? FinishTime { get; set; }
+
+			IReadOnlyList<IJobStep> IJobStepBatch.Steps => Steps;
+			DateTime? IJobStepBatch.ReadyTimeUtc => ReadyTime?.UtcDateTime;
+			DateTime? IJobStepBatch.StartTimeUtc => StartTime?.UtcDateTime;
+			DateTime? IJobStepBatch.FinishTimeUtc => FinishTime?.UtcDateTime;
+
+			[BsonConstructor]
+			private JobStepBatchDocument()
+			{
+			}
+
+			public JobStepBatchDocument(SubResourceId Id, int GroupIdx)
+			{
+				this.Id = Id;
+				this.GroupIdx = GroupIdx;
+			}
+		}
+
+		class ChainedJobDocument : IChainedJob
+		{
+			public string Target { get; set; }
+			public TemplateRefId TemplateRefId { get; set; }
+			public JobId? JobId { get; set; }
+
+			[BsonConstructor]
+			private ChainedJobDocument()
+			{
+				Target = String.Empty;
+			}
+
+			public ChainedJobDocument(ChainedJobTemplate Trigger)
+			{
+				this.Target = Trigger.Trigger;
+				this.TemplateRefId = Trigger.TemplateRefId;
+			}
+		}
+
+		class LabelNotificationDocument
+		{
+			public int LabelIdx;
+			public ObjectId TriggerId;
+		}
+
+		class JobDocument : IJob
+		{
+			[BsonRequired, BsonId]
+			public JobId Id { get; set; }
+
+			public StreamId StreamId { get; set; }
+			public TemplateRefId TemplateId { get; set; }
+			public ContentHash? TemplateHash { get; set; }
+			public ContentHash GraphHash { get; set; }
+
+			[BsonIgnoreIfNull]
+			public UserId? StartedByUserId { get; set; }
+
+			[BsonIgnoreIfNull, BsonElement("StartedByUser")]
+			public string? StartedByUser_DEPRECATED { get; set; }
+
+			[BsonIgnoreIfNull]
+			public UserId? AbortedByUserId { get; set; }
+
+			[BsonIgnoreIfNull, BsonElement("AbortedByUser")]
+			public string? AbortedByUser_DEPRECATED { get; set; }
+
+			[BsonRequired]
+			public string Name { get; set; }
+
+			public int Change { get; set; }
+			public int CodeChange { get; set; }
+			public int PreflightChange { get; set; }
+			public int ClonedPreflightChange { get; set; }
+			public Priority Priority { get; set; }
+
+			[BsonIgnoreIfDefault]
+			public bool AutoSubmit { get; set; }
+
+			[BsonIgnoreIfNull]
+			public int? AutoSubmitChange { get; set; }
+
+			[BsonIgnoreIfNull]
+			public string? AutoSubmitMessage { get; set; }
+
+			public bool UpdateIssues { get; set; }
+
+			[BsonIgnoreIfNull]
+			public DateTimeOffset? CreateTime { get; set; }
+
+			[BsonIgnoreIfNull]
+			public DateTime? CreateTimeUtc { get; set; }
+
+			public int SchedulePriority { get; set; }
+			public List<JobStepBatchDocument> Batches { get; set; } = new List<JobStepBatchDocument>();
+			public List<Report>? Reports { get; set; }
+			public List<string> Arguments { get; set; } = new List<string>();
+			public List<int> ReferencedByIssues { get; set; } = new List<int>();
+			public ObjectId? NotificationTriggerId { get; set; }
+			public bool ShowUgsBadges { get; set; }
+			public bool ShowUgsAlerts { get; set; }
+			public string? NotificationChannel { get; set; }
+			public string? NotificationChannelFilter { get; set; }
+			public List<LabelNotificationDocument> LabelNotifications = new List<LabelNotificationDocument>();
+			public List<ChainedJobDocument> ChainedJobs { get; set; } = new List<ChainedJobDocument>();
+
+			[BsonIgnoreIfNull]
+			public List<NodeRef>? RetriedNodes { get; set; }
+			public SubResourceId NextSubResourceId { get; set; }
+
+			[BsonIgnoreIfNull]
+			public DateTimeOffset? UpdateTime { get; set; }
+
+			[BsonIgnoreIfNull]
+			public DateTime? UpdateTimeUtc { get; set; }
+
+			public Acl? Acl { get; set; }
+
+			[BsonRequired]
+			public int UpdateIndex { get; set; }
+
+			DateTime IJob.CreateTimeUtc => CreateTimeUtc ?? CreateTime?.UtcDateTime ?? DateTime.UnixEpoch;
+			DateTime IJob.UpdateTimeUtc => UpdateTimeUtc ?? UpdateTime?.UtcDateTime ?? DateTime.UnixEpoch;
+			IReadOnlyList<IJobStepBatch> IJob.Batches => (IReadOnlyList<JobStepBatchDocument>)Batches;
+			IReadOnlyList<IReport>? IJob.Reports => Reports;
+			IReadOnlyList<string> IJob.Arguments => Arguments;
+			IReadOnlyList<int> IJob.Issues => ReferencedByIssues;
+			IReadOnlyDictionary<int, ObjectId> IJob.LabelIdxToTriggerId => LabelNotifications.ToDictionary(x => x.LabelIdx, x => x.TriggerId);
+			IReadOnlyList<IChainedJob> IJob.ChainedJobs => ChainedJobs;
+
+			[BsonConstructor]
+			private JobDocument()
+			{
+				Name = null!;
+				GraphHash = null!;
+			}
+
+			public JobDocument(JobId Id, StreamId StreamId, TemplateRefId TemplateId, ContentHash TemplateHash, ContentHash GraphHash, string Name, int Change, int CodeChange, int PreflightChange, int ClonedPreflightChange, UserId? StartedByUserId, Priority? Priority, bool? AutoSubmit, bool? UpdateIssues, DateTime CreateTimeUtc, List<ChainedJobDocument> ChainedJobs, bool ShowUgsBadges, bool ShowUgsAlerts, string? NotificationChannel, string? NotificationChannelFilter, List<string>? Arguments)
+			{
+				this.Id = Id;
+				this.StreamId = StreamId;
+				this.TemplateId = TemplateId;
+				this.TemplateHash = TemplateHash;
+				this.GraphHash = GraphHash;
+				this.Name = Name;
+				this.Change = Change;
+				this.CodeChange = CodeChange;
+				this.PreflightChange = PreflightChange;
+				this.ClonedPreflightChange = ClonedPreflightChange;
+				this.StartedByUserId = StartedByUserId;
+				this.Priority = Priority ?? HordeCommon.Priority.Normal;
+				this.AutoSubmit = AutoSubmit ?? false;
+				this.UpdateIssues = UpdateIssues ?? (StartedByUserId == null && PreflightChange == 0);
+				this.CreateTimeUtc = CreateTimeUtc;
+				this.ChainedJobs = ChainedJobs;
+				this.ShowUgsBadges = ShowUgsBadges;
+				this.ShowUgsAlerts = ShowUgsAlerts;
+				this.NotificationChannel = NotificationChannel;
+				this.NotificationChannelFilter = NotificationChannelFilter;
+				this.Arguments = Arguments ?? this.Arguments;
+				this.NextSubResourceId = SubResourceId.Random();
+				this.UpdateTimeUtc = CreateTimeUtc;
+			}
+		}
+
+		/// <summary>
+		/// Projection of a job definition to just include permissions info
+		/// </summary>
+		[SuppressMessage("Design", "CA1812: Class is never instantiated")]
+		class JobPermissions : IJobPermissions
+		{
+			public static ProjectionDefinition<JobDocument> Projection { get; } = Builders<JobDocument>.Projection.Include(x => x.Acl).Include(x => x.StreamId);
+
+			public Acl? Acl { get; set; }
+			public StreamId StreamId { get; set; }
+		}
+
+		/// <summary>
+		/// Maximum number of times a step can be retried (after the original run)
+		/// </summary>
+		const int MaxRetries = 2;
+
+		/// <summary>
+		/// The jobs collection
+		/// </summary>
+		IMongoCollection<JobDocument> Jobs;
+
+		/// <summary>
+		/// The user collection
+		/// </summary>
+		IUserCollection UserCollection;
+
+		/// <summary>
+		/// Logger for output
+		/// </summary>
+		ILogger<JobCollection> Logger;
+
+		/// <summary>
+		/// Constructor
+		/// </summary>
+		/// <param name="DatabaseService">The database service singleton</param>
+		/// <param name="UserCollection"></param>
+		/// <param name="Logger">The logger instance</param>
+		public JobCollection(DatabaseService DatabaseService, IUserCollection UserCollection, ILogger<JobCollection> Logger)
+		{
+			this.UserCollection = UserCollection;
+			this.Logger = Logger;
+
+			Jobs = DatabaseService.GetCollection<JobDocument>("Jobs");
+
+			if (!DatabaseService.ReadOnlyMode)
+			{
+				Jobs.Indexes.CreateOne(new CreateIndexModel<JobDocument>(Builders<JobDocument>.IndexKeys.Ascending(x => x.StreamId)));
+				Jobs.Indexes.CreateOne(new CreateIndexModel<JobDocument>(Builders<JobDocument>.IndexKeys.Ascending(x => x.Change)));
+				Jobs.Indexes.CreateOne(new CreateIndexModel<JobDocument>(Builders<JobDocument>.IndexKeys.Ascending(x => x.PreflightChange)));
+				Jobs.Indexes.CreateOne(new CreateIndexModel<JobDocument>(Builders<JobDocument>.IndexKeys.Ascending(x => x.CreateTimeUtc)));
+				Jobs.Indexes.CreateOne(new CreateIndexModel<JobDocument>(Builders<JobDocument>.IndexKeys.Ascending(x => x.UpdateTimeUtc)));
+				Jobs.Indexes.CreateOne(new CreateIndexModel<JobDocument>(Builders<JobDocument>.IndexKeys.Ascending(x => x.Name)));
+				Jobs.Indexes.CreateOne(new CreateIndexModel<JobDocument>(Builders<JobDocument>.IndexKeys.Ascending(x => x.StartedByUserId)));
+				Jobs.Indexes.CreateOne(new CreateIndexModel<JobDocument>(Builders<JobDocument>.IndexKeys.Ascending(x => x.TemplateId)));
+				Jobs.Indexes.CreateOne(new CreateIndexModel<JobDocument>(Builders<JobDocument>.IndexKeys.Descending(x => x.SchedulePriority)));
+			}
+		}
+
+		async Task PostLoadAsync(JobDocument Job)
+		{
+			if (Job.GraphHash == ContentHash.Empty)
+			{
+				Job.Batches.Clear();
+			}
+
+			List<UpdateDefinition<JobDocument>> Updates = new List<UpdateDefinition<JobDocument>>();
+			if (Job.StartedByUser_DEPRECATED != null)
+			{
+				IUser StartedByUser = await UserCollection.FindOrAddUserByLoginAsync(Job.StartedByUser_DEPRECATED);
+				Job.StartedByUserId = StartedByUser.Id;
+				Job.StartedByUser_DEPRECATED = null;
+				Updates.Add(Builders<JobDocument>.Update.Set(x => x.StartedByUserId, StartedByUser.Id).Unset(x => x.StartedByUser_DEPRECATED));
+			}
+
+			if (Job.AbortedByUser_DEPRECATED != null)
+			{
+				IUser AbortedByUser = await UserCollection.FindOrAddUserByLoginAsync(Job.AbortedByUser_DEPRECATED);
+				Job.AbortedByUserId = AbortedByUser.Id;
+				Job.AbortedByUser_DEPRECATED = null;
+				Updates.Add(Builders<JobDocument>.Update.Set(x => x.AbortedByUserId, AbortedByUser.Id).Unset(x => x.AbortedByUser_DEPRECATED));
+			}
+
+
+			for (int BatchIdx = 0; BatchIdx < Job.Batches.Count; BatchIdx++)
+			{
+				int LocalBatchIdx = BatchIdx;
+				JobStepBatchDocument Batch = Job.Batches[BatchIdx];
+
+				for (int StepIdx = 0; StepIdx < Batch.Steps.Count; StepIdx++)
+				{
+					int LocalStepIdx = StepIdx;
+					JobStepDocument Step = Batch.Steps[StepIdx];
+
+					if (Step.AbortedByUser_DEPRECATED != null)
+					{
+						IUser AbortedByUser = await UserCollection.FindOrAddUserByLoginAsync(Step.AbortedByUser_DEPRECATED);
+						Step.AbortedByUserId = AbortedByUser.Id;
+						Step.AbortedByUser_DEPRECATED = null;
+						Updates.Add(Builders<JobDocument>.Update.Set(x => x.Batches[LocalBatchIdx].Steps[LocalStepIdx].AbortedByUserId, AbortedByUser.Id).Unset(x => x.Batches[LocalBatchIdx].Steps[LocalStepIdx].AbortedByUser_DEPRECATED));
+					}
+
+					if (Step.RetriedByUser_DEPRECATED != null)
+					{
+						IUser RetriedByUser = await UserCollection.FindOrAddUserByLoginAsync(Step.RetriedByUser_DEPRECATED);
+						Step.RetriedByUserId = RetriedByUser.Id;
+						Step.RetriedByUser_DEPRECATED = null;
+						Updates.Add(Builders<JobDocument>.Update.Set(x => x.Batches[LocalBatchIdx].Steps[LocalStepIdx].RetriedByUserId, RetriedByUser.Id).Unset(x => x.Batches[LocalBatchIdx].Steps[LocalStepIdx].RetriedByUser_DEPRECATED));
+					}
+				}
+			}
+
+
+			if (Updates.Count > 0)
+			{
+				await TryUpdateAsync(Job, Builders<JobDocument>.Update.Combine(Updates));
+			}
+		}
+
+		/// <inheritdoc/>
+		[SuppressMessage("Compiler", "CA1054:URI parameters should not be strings")]
+		public async Task<IJob> AddAsync(JobId JobId, StreamId StreamId, TemplateRefId TemplateRefId, ContentHash TemplateHash, IGraph Graph, string Name, int Change, int CodeChange, int? PreflightChange, int? ClonedPreflightChange, UserId? StartedByUserId, Priority? Priority, bool? AutoSubmit, bool? UpdateIssues, List<ChainedJobTemplate>? ChainedJobs, bool ShowUgsBadges, bool ShowUgsAlerts, string? NotificationChannel, string? NotificationChannelFilter, List<string>? Arguments)
+		{
+			List<ChainedJobDocument> JobTriggers = new List<ChainedJobDocument>();
+			if (ChainedJobs == null)
+			{
+				JobTriggers = new List<ChainedJobDocument>();
+			}
+			else
+			{
+				JobTriggers = ChainedJobs.ConvertAll(x => new ChainedJobDocument(x));
+			}
+
+			JobDocument NewJob = new JobDocument(JobId, StreamId, TemplateRefId, TemplateHash, Graph.Id, Name, Change, CodeChange, PreflightChange ?? 0, ClonedPreflightChange ?? 0, StartedByUserId, Priority, AutoSubmit, UpdateIssues, DateTime.UtcNow, JobTriggers, ShowUgsBadges, ShowUgsAlerts, NotificationChannel, NotificationChannelFilter, Arguments);
+			CreateBatches(NewJob, Graph, Logger);
+
+			await Jobs.InsertOneAsync(NewJob);
+
+			return NewJob;
+		}
+
+		/// <inheritdoc/>
+		public async Task<IJob?> GetAsync(JobId JobId)
+		{
+			JobDocument? Job = await Jobs.Find<JobDocument>(x => x.Id == JobId).FirstOrDefaultAsync();
+			if (Job != null)
+			{
+				await PostLoadAsync(Job);
+			}
+			return Job;
+		}
+
+		/// <inheritdoc/>
+		public async Task<bool> RemoveAsync(IJob Job)
+		{
+			DeleteResult Result = await Jobs.DeleteOneAsync(x => x.Id == Job.Id && x.UpdateIndex == Job.UpdateIndex);
+			return Result.DeletedCount > 0;
+		}
+
+		/// <inheritdoc/>
+		public async Task RemoveStreamAsync(StreamId StreamId)
+		{
+			await Jobs.DeleteManyAsync(x => x.StreamId == StreamId);
+		}
+
+		/// <inheritdoc/>
+		public async Task<IJobPermissions?> GetPermissionsAsync(JobId JobId)
+		{
+			return await Jobs.Find<JobDocument>(x => x.Id == JobId).Project<JobPermissions>(JobPermissions.Projection).FirstOrDefaultAsync();
+		}
+
+		/// <inheritdoc/>
+		public async Task<List<IJob>> FindAsync(JobId[]? JobIds, StreamId? StreamId, string? Name, TemplateRefId[]? Templates, int? MinChange, int? MaxChange, int? PreflightChange, UserId? PreflightStartedByUser, UserId? StartedByUser, DateTimeOffset? MinCreateTime, DateTimeOffset? MaxCreateTime, DateTimeOffset? ModifiedBefore, DateTimeOffset? ModifiedAfter, int? Index, int? Count)
+		{
+			FilterDefinitionBuilder<JobDocument> FilterBuilder = Builders<JobDocument>.Filter;
+
+			FilterDefinition<JobDocument> Filter = FilterBuilder.Empty;
+			if (JobIds != null && JobIds.Length > 0)
+			{
+				Filter &= FilterBuilder.In(x => x.Id, JobIds);
+			}
+			if (StreamId != null)
+			{
+				Filter &= FilterBuilder.Eq(x => x.StreamId, StreamId.Value);
+			}
+			if (Name != null)
+			{
+				Filter &= FilterBuilder.Eq(x => x.Name, Name);
+			}
+			if (Templates != null)
+			{
+				Filter &= FilterBuilder.In(x => x.TemplateId, Templates);
+			}
+			if (MinChange != null)
+			{
+				Filter &= FilterBuilder.Gte(x => x.Change, MinChange);
+			}
+			if (MaxChange != null)
+			{
+				Filter &= FilterBuilder.Lte(x => x.Change, MaxChange);
+			}
+			if (PreflightChange != null)
+			{
+				Filter &= FilterBuilder.Eq(x => x.PreflightChange, PreflightChange);
+			}
+			if (PreflightStartedByUser != null)
+			{
+				Filter &= FilterBuilder.Or(FilterBuilder.Eq(x => x.PreflightChange, 0), FilterBuilder.Eq(x => x.StartedByUserId, PreflightStartedByUser));
+			}
+			if (StartedByUser != null)
+			{
+				Filter &= FilterBuilder.Eq(x => x.StartedByUserId, StartedByUser);
+			}
+			if (MinCreateTime != null)
+			{
+				Filter &= FilterBuilder.Gte(x => x.CreateTimeUtc!, MinCreateTime.Value.UtcDateTime);
+			}
+			if (MaxCreateTime != null)
+			{
+				Filter &= FilterBuilder.Lte(x => x.CreateTimeUtc!, MaxCreateTime.Value.UtcDateTime);
+			}
+			if (ModifiedBefore != null)
+			{
+				Filter &= FilterBuilder.Lte(x => x.UpdateTimeUtc!, ModifiedBefore.Value.UtcDateTime);
+			}
+			if (ModifiedAfter != null)
+			{
+				Filter &= FilterBuilder.Gte(x => x.UpdateTimeUtc!, ModifiedAfter.Value.UtcDateTime);
+			}
+
+			List<JobDocument> Results;
+			using (IScope Scope = GlobalTracer.Instance.BuildSpan("Jobs.Find").StartActive())
+			{
+				IFindFluent<JobDocument, JobDocument> Query = Jobs.Find<JobDocument>(Filter).SortByDescending(x => x.CreateTimeUtc!);
+				if (Index != null)
+				{
+					Query = Query.Skip(Index.Value);
+				}
+				if (Count != null)
+				{
+					Query = Query.Limit(Count.Value);
+				}
+				Results = await Query.ToListAsync();
+			}
+			foreach (JobDocument Result in Results)
+			{
+				await PostLoadAsync(Result);
+			}
+			return Results.ConvertAll<JobDocument, IJob>(x => x);
+		}
+
+		/// <inheritdoc/>
+		public async Task<bool> TryUpdateJobAsync(IJob Job, IGraph Graph, string? Name, Priority? Priority, bool? AutoSubmit, int? AutoSubmitChange, string? AutoSubmitMessage, UserId? AbortedByUserId, ObjectId? NotificationTriggerId, List<Report>? Reports, List<string>? Arguments, KeyValuePair<int, ObjectId>? LabelIdxToTriggerId, KeyValuePair<TemplateRefId, JobId>? JobTrigger)
+		{
+			// Create the update 
+			UpdateDefinitionBuilder<JobDocument> UpdateBuilder = Builders<JobDocument>.Update;
+			List<UpdateDefinition<JobDocument>> Updates = new List<UpdateDefinition<JobDocument>>();
+
+			// Flag for whether to update batches
+			bool bUpdateBatches = false;
+
+			// Build the update list
+			JobDocument JobDocument = (JobDocument)Job;
+			if (Name != null)
+			{
+				JobDocument.Name = Name;
+				Updates.Add(UpdateBuilder.Set(x => x.Name, Job.Name));
+			}
+			if (Priority != null)
+			{
+				JobDocument.Priority = Priority.Value;
+				Updates.Add(UpdateBuilder.Set(x => x.Priority, Job.Priority));
+			}
+			if (AutoSubmit != null)
+			{
+				JobDocument.AutoSubmit = AutoSubmit.Value;
+				Updates.Add(UpdateBuilder.Set(x => x.AutoSubmit, Job.AutoSubmit));
+			}
+			if (AutoSubmitChange != null)
+			{
+				JobDocument.AutoSubmitChange = AutoSubmitChange.Value;
+				Updates.Add(UpdateBuilder.Set(x => x.AutoSubmitChange, Job.AutoSubmitChange));
+			}
+			if (AutoSubmitMessage != null)
+			{
+				JobDocument.AutoSubmitMessage = (AutoSubmitMessage.Length == 0)? null : AutoSubmitMessage;
+				Updates.Add(UpdateBuilder.SetOrUnsetNullRef(x => x.AutoSubmitMessage, Job.AutoSubmitMessage));
+			}
+			if (AbortedByUserId != null && Job.AbortedByUserId == null)
+			{
+				JobDocument.AbortedByUserId = AbortedByUserId;
+				Updates.Add(UpdateBuilder.Set(x => x.AbortedByUserId, Job.AbortedByUserId));
+				bUpdateBatches = true;
+			}
+			if (NotificationTriggerId != null)
+			{
+				JobDocument.NotificationTriggerId = NotificationTriggerId.Value;
+				Updates.Add(UpdateBuilder.Set(x => x.NotificationTriggerId, NotificationTriggerId));
+			}
+			if (LabelIdxToTriggerId != null)
+			{
+				if (JobDocument.LabelNotifications.Any(x => x.LabelIdx == LabelIdxToTriggerId.Value.Key))
+				{
+					throw new ArgumentException("Cannot update label trigger that already exists");
+				}
+				JobDocument.LabelNotifications.Add(new LabelNotificationDocument { LabelIdx = LabelIdxToTriggerId.Value.Key, TriggerId = LabelIdxToTriggerId.Value.Value });
+				Updates.Add(UpdateBuilder.Set(x => x.LabelNotifications, JobDocument.LabelNotifications));
+			}
+			if (JobTrigger != null)
+			{
+				for (int Idx = 0; Idx < JobDocument.ChainedJobs.Count; Idx++)
+				{
+					ChainedJobDocument JobTriggerDocument = JobDocument.ChainedJobs[Idx];
+					if (JobTriggerDocument.TemplateRefId == JobTrigger.Value.Key)
+					{
+						int LocalIdx = Idx;
+						JobTriggerDocument.JobId = JobTrigger.Value.Value;
+						Updates.Add(UpdateBuilder.Set(x => x.ChainedJobs[LocalIdx].JobId, JobTrigger.Value.Value));
+					}
+				}
+			}
+			if (Reports != null)
+			{
+				JobDocument.Reports ??= new List<Report>();
+				JobDocument.Reports.RemoveAll(x => Reports.Any(y => y.Name == x.Name));
+				JobDocument.Reports.AddRange(Reports);
+				Updates.Add(UpdateBuilder.Set(x => x.Reports, JobDocument.Reports));
+			}
+			if (Arguments != null)
+			{
+				HashSet<string> ModifiedArguments = new HashSet<string>(Job.Arguments);
+				ModifiedArguments.SymmetricExceptWith(Arguments);
+
+				foreach (string ModifiedArgument in ModifiedArguments)
+				{
+					if (ModifiedArgument.StartsWith(IJob.TargetArgumentPrefix, StringComparison.OrdinalIgnoreCase))
+					{
+						bUpdateBatches = true;
+					}
+				}
+
+				JobDocument.Arguments = Arguments.ToList();
+				Updates.Add(UpdateBuilder.Set(x => x.Arguments, JobDocument.Arguments));
+			}
+
+			// Update the batches
+			if (bUpdateBatches)
+			{
+				UpdateBatches(JobDocument, Graph, Updates, Logger);
+			}
+
+			// Update the new list of job steps
+			return (Updates.Count == 0 || await TryUpdateAsync(JobDocument, UpdateBuilder.Combine(Updates)));
+		}
+
+		/// <inheritdoc/>
+		public async Task<bool> TryUpdateBatchAsync(IJob Job, IGraph Graph, SubResourceId BatchId, LogId? NewLogId, JobStepBatchState? NewState, JobStepBatchError? NewError)
+		{
+			JobDocument JobDocument = (JobDocument)Job;
+
+			// Find the index of the appropriate batch
+			int BatchIdx = JobDocument.Batches.FindIndex(x => x.Id == BatchId);
+			if (BatchIdx == -1)
+			{
+				return false;
+			}
+
+			// If we're marking the batch as complete and there are still steps to run (eg. because the agent crashed), we need to mark all the steps as complete first
+			JobStepBatchDocument Batch = JobDocument.Batches[BatchIdx];
+
+			// Create the update 
+			UpdateDefinitionBuilder<JobDocument> UpdateBuilder = Builders<JobDocument>.Update;
+			List<UpdateDefinition<JobDocument>> Updates = new List<UpdateDefinition<JobDocument>>();
+
+			// Update the batch
+			if (NewLogId != null)
+			{
+				Batch.LogId = NewLogId.Value;
+				Updates.Add(UpdateBuilder.Set(x => x.Batches[BatchIdx].LogId, Batch.LogId));
+			}
+			if (NewState != null)
+			{
+				Batch.State = NewState.Value;
+				Updates.Add(UpdateBuilder.Set(x => x.Batches[BatchIdx].State, Batch.State));
+
+				if (Batch.StartTime == null && NewState >= JobStepBatchState.Starting)
+				{
+					Batch.StartTime = DateTimeOffset.Now;
+					Updates.Add(UpdateBuilder.Set(x => x.Batches[BatchIdx].StartTime, Batch.StartTime));
+				}
+				if (NewState == JobStepBatchState.Complete)
+				{
+					Batch.FinishTime = DateTimeOffset.Now;
+					Updates.Add(UpdateBuilder.Set(x => x.Batches[BatchIdx].FinishTime, Batch.FinishTime));
+				}
+			}
+			if (NewError != null)
+			{
+				Batch.Error = NewError.Value;
+				Updates.Add(UpdateBuilder.Set(x => x.Batches[BatchIdx].Error, Batch.Error));
+			}
+
+			// If the batch is being marked as incomplete, see if we can reschedule any of the work
+			if (NewError == JobStepBatchError.Incomplete)
+			{
+				// Get the new list of retried nodes
+				List<NodeRef> RetriedNodes = JobDocument.RetriedNodes ?? new List<NodeRef>();
+
+				// Check if there are any steps that need to be run again
+				bool bUpdateState = false;
+				foreach (JobStepDocument Step in Batch.Steps)
+				{
+					if (Step.State == JobStepState.Ready || Step.State == JobStepState.Waiting)
+					{
+						if (CanRetryNode(JobDocument, Batch.GroupIdx, Step.NodeIdx))
+						{
+							RetriedNodes.Add(new NodeRef(Batch.GroupIdx, Step.NodeIdx));
+						}
+						else
+						{
+							Step.State = JobStepState.Skipped;
+						}
+						bUpdateState = true;
+					}
+				}
+
+				// Update the steps
+				if (bUpdateState)
+				{
+					Updates.Clear();
+					UpdateBatches(JobDocument, Graph, Updates, Logger);
+
+					JobDocument.RetriedNodes = RetriedNodes;
+					Updates.Add(UpdateBuilder.Set(x => x.RetriedNodes, JobDocument.RetriedNodes));
+				}
+			}
+
+			// Update the new list of job steps
+			return (Updates.Count == 0 || await TryUpdateAsync(JobDocument, UpdateBuilder.Combine(Updates)));
+		}
+
+		/// <inheritdoc/>
+		public async Task<bool> TryUpdateStepAsync(IJob Job, IGraph Graph, SubResourceId BatchId, SubResourceId StepId, JobStepState NewState, JobStepOutcome NewOutcome, bool? NewAbortRequested, UserId? NewAbortByUserId, LogId? NewLogId, ObjectId? NewNotificationTriggerId, UserId? NewRetryByUserId, Priority? NewPriority, List<Report>? NewReports, Dictionary<string, string?>? NewProperties)
+		{
+			JobDocument JobDocument = (JobDocument)Job;
+
+			// Create the update 
+			UpdateDefinitionBuilder<JobDocument> UpdateBuilder = Builders<JobDocument>.Update;
+			List<UpdateDefinition<JobDocument>> Updates = new List<UpdateDefinition<JobDocument>>();
+
+			// Update the appropriate batch
+			bool bRefreshBatches = false;
+			bool bRefreshDependentJobSteps = false;
+			for (int LoopBatchIdx = 0; LoopBatchIdx < JobDocument.Batches.Count; LoopBatchIdx++)
+			{
+				int BatchIdx = LoopBatchIdx; // For lambda capture
+				JobStepBatchDocument Batch = JobDocument.Batches[BatchIdx];
+				if (Batch.Id == BatchId)
+				{
+					for (int LoopStepIdx = 0; LoopStepIdx < Batch.Steps.Count; LoopStepIdx++)
+					{
+						int StepIdx = LoopStepIdx; // For lambda capture
+						JobStepDocument Step = Batch.Steps[StepIdx];
+						if (Step.Id == StepId)
+						{
+							// Update the state
+							if (NewState != JobStepState.Unspecified && Step.State != NewState)
+							{
+								Step.State = NewState;
+								Updates.Add(UpdateBuilder.Set(x => x.Batches[BatchIdx].Steps[StepIdx].State, Step.State));
+
+								if (Step.State == JobStepState.Running)
+								{
+									Step.StartTime = DateTimeOffset.Now;
+									Updates.Add(UpdateBuilder.Set(x => x.Batches[BatchIdx].Steps[StepIdx].StartTime, Step.StartTime));
+								}
+								else if (Step.State == JobStepState.Completed || Step.State == JobStepState.Aborted)
+								{
+									Step.FinishTime = DateTimeOffset.Now;
+									Updates.Add(UpdateBuilder.Set(x => x.Batches[BatchIdx].Steps[StepIdx].FinishTime, Step.FinishTime));
+								}
+
+								bRefreshDependentJobSteps = true;
+							}
+
+							// Update the job outcome
+							if (NewOutcome != JobStepOutcome.Unspecified && Step.Outcome != NewOutcome)
+							{
+								Step.Outcome = NewOutcome;
+								Updates.Add(UpdateBuilder.Set(x => x.Batches[BatchIdx].Steps[StepIdx].Outcome, Step.Outcome));
+
+								bRefreshDependentJobSteps = true;
+							}
+
+							// Update the request abort status
+							if (NewAbortRequested != null && Step.AbortRequested == false)
+							{
+								Step.AbortRequested = NewAbortRequested.Value;
+								Updates.Add(UpdateBuilder.Set(x => x.Batches[BatchIdx].Steps[StepIdx].AbortRequested, Step.AbortRequested));
+
+								bRefreshDependentJobSteps = true;
+							}
+
+							// Update the user that requested the abort
+							if (NewAbortByUserId != null && Step.AbortedByUserId == null)
+							{
+								Step.AbortedByUserId = NewAbortByUserId;
+								Updates.Add(UpdateBuilder.Set(x => x.Batches[BatchIdx].Steps[StepIdx].AbortedByUserId, Step.AbortedByUserId));
+
+								bRefreshDependentJobSteps = true;
+							}
+
+							// Update the log id
+							if (NewLogId != null && Step.LogId != NewLogId.Value)
+							{
+								Step.LogId = NewLogId.Value;
+								Updates.Add(UpdateBuilder.Set(x => x.Batches[BatchIdx].Steps[StepIdx].LogId, Step.LogId));
+
+								bRefreshDependentJobSteps = true;
+							}
+
+							// Update the notification trigger id
+							if (NewNotificationTriggerId != null && Step.NotificationTriggerId == null)
+							{
+								Step.NotificationTriggerId = NewNotificationTriggerId.Value;
+								Updates.Add(UpdateBuilder.Set(x => x.Batches[BatchIdx].Steps[StepIdx].NotificationTriggerId, Step.NotificationTriggerId));
+							}
+
+							// Update the retry flag
+							if (NewRetryByUserId != null && Step.RetriedByUserId == null)
+							{
+								Step.RetriedByUserId = NewRetryByUserId;
+								Updates.Add(UpdateBuilder.Set(x => x.Batches[BatchIdx].Steps[StepIdx].RetriedByUserId, Step.RetriedByUserId));
+
+								bRefreshBatches = true;
+							}
+
+							// Update the priority
+							if (NewPriority != null && NewPriority.Value != Step.Priority)
+							{
+								Step.Priority = NewPriority.Value;
+								Updates.Add(UpdateBuilder.Set(x => x.Batches[BatchIdx].Steps[StepIdx].Priority, Step.Priority));
+
+								bRefreshBatches = true;
+							}
+
+							// Add any new reports
+							if (NewReports != null)
+							{
+								Step.Reports ??= new List<Report>();
+								Step.Reports.RemoveAll(x => NewReports.Any(y => y.Name == x.Name));
+								Step.Reports.AddRange(NewReports);
+								Updates.Add(UpdateBuilder.Set(x => x.Batches[BatchIdx].Steps[StepIdx].Reports, Step.Reports));
+							}
+
+							// Apply any property updates
+							if (NewProperties != null)
+							{
+								if (Step.Properties == null)
+								{
+									Step.Properties = new Dictionary<string, string>(StringComparer.Ordinal);
+								}
+
+								foreach (KeyValuePair<string, string?> Pair in NewProperties)
+								{
+									if (Pair.Value == null)
+									{
+										Step.Properties.Remove(Pair.Key);
+									}
+									else
+									{
+										Step.Properties[Pair.Key] = Pair.Value;
+									}
+								}
+
+								if (Step.Properties.Count == 0)
+								{
+									Step.Properties = null;
+									Updates.Add(UpdateBuilder.Unset(x => x.Batches[BatchIdx].Steps[StepIdx].Properties));
+								}
+								else
+								{
+									foreach (KeyValuePair<string, string?> Pair in NewProperties)
+									{
+										if (Pair.Value == null)
+										{
+											Updates.Add(UpdateBuilder.Unset(x => x.Batches[BatchIdx].Steps[StepIdx].Properties![Pair.Key]));
+										}
+										else
+										{
+											Updates.Add(UpdateBuilder.Set(x => x.Batches[BatchIdx].Steps[StepIdx].Properties![Pair.Key], Pair.Value));
+										}
+									}
+								}
+							}
+							break;
+						}
+					}
+					break;
+				}
+			}
+
+			// Update the batches
+			if (bRefreshBatches)
+			{
+				Updates.Clear(); // UpdateBatches will update the entire batches list. We need to remove all the individual step updates to avoid an exception.
+				UpdateBatches(JobDocument, Graph, Updates, Logger);
+			}
+
+			// Update the state of dependent jobsteps
+			if (bRefreshDependentJobSteps)
+			{
+				RefreshDependentJobSteps(JobDocument, Graph, Updates, Logger);
+				RefreshJobPriority(JobDocument, Updates);
+			}
+
+			// Update the new list of job steps
+			return (Updates.Count == 0 || await TryUpdateAsync(JobDocument, UpdateBuilder.Combine(Updates)));
+		}
+
+		/// <inheritdoc/>
+		async Task<bool> TryUpdateAsync(JobDocument Job, UpdateDefinition<JobDocument> Update)
+		{
+			int NewUpdateIndex = Job.UpdateIndex + 1;
+			Update = Update.Set(x => x.UpdateIndex, NewUpdateIndex);
+
+			DateTime NewUpdateTimeUtc = DateTime.UtcNow;
+			Update = Update.Set(x => x.UpdateTimeUtc, NewUpdateTimeUtc);
+
+			UpdateResult Result = await Jobs.UpdateOneAsync<JobDocument>(x => x.Id == Job.Id && x.UpdateIndex == Job.UpdateIndex, Update);
+			if (Result.ModifiedCount > 0)
+			{
+				Job.UpdateIndex = NewUpdateIndex;
+				Job.UpdateTimeUtc = NewUpdateTimeUtc;
+				return true;
+			}
+			return false;
+		}
+
+		/// <inheritdoc/>
+		public Task<bool> TryRemoveFromDispatchQueueAsync(IJob Job)
+		{
+			JobDocument JobDocument = (JobDocument)Job;
+
+			// Create the update 
+			UpdateDefinitionBuilder<JobDocument> UpdateBuilder = Builders<JobDocument>.Update;
+			List<UpdateDefinition<JobDocument>> Updates = new List<UpdateDefinition<JobDocument>>();
+
+			JobDocument.SchedulePriority = 0;
+			Updates.Add(UpdateBuilder.Set(x => x.SchedulePriority, JobDocument.SchedulePriority));
+
+			// Update the new list of job steps
+			return TryUpdateAsync(JobDocument, UpdateBuilder.Combine(Updates));
+		}
+
+		/// <inheritdoc/>
+		public Task<bool> TryUpdateGraphAsync(IJob Job, IGraph NewGraph)
+		{
+			JobDocument JobDocument = (JobDocument)Job;
+
+			// Create the update 
+			UpdateDefinitionBuilder<JobDocument> UpdateBuilder = Builders<JobDocument>.Update;
+			List<UpdateDefinition<JobDocument>> Updates = new List<UpdateDefinition<JobDocument>>();
+
+			JobDocument.GraphHash = NewGraph.Id;
+			Updates.Add(UpdateBuilder.Set(x => x.GraphHash, Job.GraphHash));
+
+			UpdateBatches(JobDocument, NewGraph, Updates, Logger);
+
+			// Update the new list of job steps
+			return TryUpdateAsync(JobDocument, UpdateBuilder.Combine(Updates));
+		}
+
+		/// <inheritdoc/>
+		public async Task AddIssueToJobAsync(JobId JobId, int IssueId)
+		{
+			FilterDefinition<JobDocument> JobFilter = Builders<JobDocument>.Filter.Eq(x => x.Id, JobId);
+			UpdateDefinition<JobDocument> JobUpdate = Builders<JobDocument>.Update.AddToSet(x => x.ReferencedByIssues, IssueId).Inc(x => x.UpdateIndex, 1).Max(x => x.UpdateTimeUtc, DateTime.UtcNow);
+			await Jobs.UpdateOneAsync(JobFilter, JobUpdate);
+		}
+
+		/// <inheritdoc/>
+		public async Task<List<IJob>> GetDispatchQueueAsync()
+		{
+			List<JobDocument> NewJobs = await Jobs.Find(x => x.SchedulePriority > 0).SortByDescending(x => x.SchedulePriority).ThenBy(x => x.CreateTimeUtc).ToListAsync();
+			foreach (JobDocument Result in NewJobs)
+			{
+				await PostLoadAsync(Result);
+			}
+			return NewJobs.ConvertAll<JobDocument, IJob>(x => x);
+		}
+
+		/// <summary>
+		/// Marks a job as skipped
+		/// </summary>
+		/// <param name="Job">The job to update</param>
+		/// <param name="Graph">Graph for the job</param>
+		/// <param name="Reason">Reason for this batch being failed</param>
+		/// <returns>Updated version of the job</returns>
+		public async Task<IJob?> SkipAllBatchesAsync(IJob? Job, IGraph Graph, JobStepBatchError Reason)
+		{
+			while (Job != null)
+			{
+				JobDocument JobDocument = (JobDocument)Job;
+
+				for (int BatchIdx = 0; BatchIdx < JobDocument.Batches.Count; BatchIdx++)
+				{
+					JobStepBatchDocument Batch = JobDocument.Batches[BatchIdx];
+					if (Batch.State == JobStepBatchState.Ready || Batch.State == JobStepBatchState.Waiting)
+					{
+						Batch.State = JobStepBatchState.Complete;
+						Batch.Error = Reason;
+						Batch.FinishTime = DateTimeOffset.UtcNow;
+
+						for (int StepIdx = 0; StepIdx < Batch.Steps.Count; StepIdx++)
+						{
+							JobStepDocument Step = Batch.Steps[StepIdx];
+							if (Step.State == JobStepState.Ready || Step.State == JobStepState.Waiting)
+							{
+								Step.State = JobStepState.Completed;
+								Step.Outcome = JobStepOutcome.Failure;
+							}
+						}
+					}
+				}
+
+				List<UpdateDefinition<JobDocument>> Updates = new List<UpdateDefinition<JobDocument>>();
+				UpdateBatches(JobDocument, Graph, Updates, Logger);
+
+				if (await TryUpdateAsync(JobDocument, Builders<JobDocument>.Update.Combine(Updates)))
+				{
+					break;
+				}
+
+				Job = await GetAsync(Job.Id);
+			}
+			return Job;
+		}
+
+		/// <inheritdoc/>
+		public async Task<IJob?> SkipBatchAsync(IJob? Job, int BatchIdx, IGraph Graph, JobStepBatchError Reason)
+		{
+			while (Job != null)
+			{
+				JobDocument JobDocument = (JobDocument)Job;
+				JobStepBatchDocument Batch = JobDocument.Batches[BatchIdx];
+
+				Batch.State = JobStepBatchState.Complete;
+				Batch.Error = Reason;
+				Batch.FinishTime = DateTimeOffset.UtcNow;
+
+				foreach (JobStepDocument Step in Batch.Steps)
+				{
+					if (Step.State != JobStepState.Skipped)
+					{
+						Step.State = JobStepState.Skipped;
+						Step.Outcome = JobStepOutcome.Failure;
+					}
+				}
+
+				List<UpdateDefinition<JobDocument>> Updates = new List<UpdateDefinition<JobDocument>>();
+				UpdateBatches(JobDocument, Graph, Updates, Logger);
+
+				if (await TryUpdateAsync(JobDocument, Builders<JobDocument>.Update.Combine(Updates)))
+				{
+					break;
+				}
+
+				Job = await GetAsync(Job.Id);
+			}
+			return Job;
+		}
+
+		/// <inheritdoc/>
+		public async Task<bool> TryAssignLeaseAsync(IJob Job, int BatchIdx, PoolId PoolId, AgentId AgentId, ObjectId SessionId, LeaseId LeaseId, LogId LogId)
+		{
+			// Try to update the job with this agent id
+			UpdateDefinitionBuilder<JobDocument> UpdateBuilder = Builders<JobDocument>.Update;
+			List<UpdateDefinition<JobDocument>> Updates = new List<UpdateDefinition<JobDocument>>();
+
+			Updates.Add(UpdateBuilder.Set(x => x.Batches[BatchIdx].PoolId, PoolId));
+			Updates.Add(UpdateBuilder.Set(x => x.Batches[BatchIdx].AgentId, AgentId));
+			Updates.Add(UpdateBuilder.Set(x => x.Batches[BatchIdx].SessionId, SessionId));
+			Updates.Add(UpdateBuilder.Set(x => x.Batches[BatchIdx].LeaseId, LeaseId));
+			Updates.Add(UpdateBuilder.Set(x => x.Batches[BatchIdx].LogId, LogId));
+
+			// Extra logs for catching why session IDs sometimes doesn't match in prod. Resulting in PermissionDenied errors.
+			if (BatchIdx < Job.Batches.Count && Job.Batches[BatchIdx].SessionId != null)
+			{
+				string CurrentSessionId = Job.Batches[BatchIdx].SessionId!.Value.ToString();
+				Logger.LogError("Attempt to replace current session ID {CurrSessionId} with {NewSessionId} for batch {JobId}:{BatchId}", CurrentSessionId, SessionId.ToString(), Job.Id.ToString(), Job.Batches[BatchIdx].Id);
+				return false;
+			}
+
+			JobDocument JobDocument = (JobDocument)Job;
+			if (!await TryUpdateAsync(JobDocument, UpdateBuilder.Combine(Updates)))
+			{
+				return false;
+			}
+
+			JobStepBatchDocument Batch = JobDocument.Batches[BatchIdx];
+			Batch.AgentId = AgentId;
+			Batch.SessionId = SessionId;
+			Batch.LeaseId = LeaseId;
+			Batch.LogId = LogId;
+
+			return true;
+		}
+
+		/// <inheritdoc/>
+		public async Task<bool> TryCancelLeaseAsync(IJob Job, int BatchIdx)
+		{
+			JobDocument JobDocument = (JobDocument)Job;
+
+			UpdateDefinition<JobDocument> Update = Builders<JobDocument>.Update.Unset(x => x.Batches[BatchIdx].AgentId).Unset(x => x.Batches[BatchIdx].SessionId).Unset(x => x.Batches[BatchIdx].LeaseId);
+			if (await TryUpdateAsync(JobDocument, Update))
+			{
+				JobStepBatchDocument Batch = JobDocument.Batches[BatchIdx];
+				Batch.AgentId = null;
+				Batch.SessionId = null;
+				Batch.LeaseId = null;
+				return true;
+			}
+			return false;
+		}
+
+		/// <inheritdoc/>
+		public async Task<bool> TryFailBatchAsync(IJob Job, int BatchIdx, IGraph Graph)
+		{
+			JobDocument JobDocument = (JobDocument)Job;
+			JobStepBatchDocument Batch = JobDocument.Batches[BatchIdx];
+
+			UpdateDefinitionBuilder<JobDocument> UpdateBuilder = new UpdateDefinitionBuilder<JobDocument>();
+			List<UpdateDefinition<JobDocument>> Updates = new List<UpdateDefinition<JobDocument>>();
+
+			if (Batch.State != JobStepBatchState.Complete)
+			{
+				Batch.State = JobStepBatchState.Complete;
+				Updates.Add(UpdateBuilder.Set(x => x.Batches[BatchIdx].State, Batch.State));
+
+				Batch.Error = JobStepBatchError.LostConnection;
+				Updates.Add(UpdateBuilder.Set(x => x.Batches[BatchIdx].Error, Batch.Error));
+
+				Batch.FinishTime = DateTimeOffset.UtcNow;
+				Updates.Add(UpdateBuilder.Set(x => x.Batches[BatchIdx].FinishTime, Batch.FinishTime));
+			}
+
+			for (int StepIdx = 0; StepIdx < Batch.Steps.Count; StepIdx++)
+			{
+				JobStepDocument Step = Batch.Steps[StepIdx];
+				if (Step.State == JobStepState.Running)
+				{
+					int StepIdxCopy = StepIdx;
+
+					Step.State = JobStepState.Aborted;
+					Updates.Add(UpdateBuilder.Set(x => x.Batches[BatchIdx].Steps[StepIdxCopy].State, Step.State));
+
+					Step.Outcome = JobStepOutcome.Failure;
+					Updates.Add(UpdateBuilder.Set(x => x.Batches[BatchIdx].Steps[StepIdxCopy].Outcome, Step.Outcome));
+
+					Step.FinishTime = DateTimeOffset.Now;
+					Updates.Add(UpdateBuilder.Set(x => x.Batches[BatchIdx].Steps[StepIdxCopy].FinishTime, Step.FinishTime));
+				}
+				else if (Step.State == JobStepState.Ready)
+				{
+					int StepIdxCopy = StepIdx;
+
+					Step.State = JobStepState.Skipped;
+					Updates.Add(UpdateBuilder.Set(x => x.Batches[BatchIdx].Steps[StepIdxCopy].State, Step.State));
+
+					Step.Outcome = JobStepOutcome.Failure;
+					Updates.Add(UpdateBuilder.Set(x => x.Batches[BatchIdx].Steps[StepIdxCopy].Outcome, Step.Outcome));
+				}
+			}
+
+			RefreshDependentJobSteps(JobDocument, Graph, Updates, Logger);
+			RefreshJobPriority(JobDocument, Updates);
+
+			return (Updates.Count == 0 || await TryUpdateAsync(JobDocument, UpdateBuilder.Combine(Updates)));
+		}
+
+		/// <summary>
+		/// Populate the list of batches for a job
+		/// </summary>
+		/// <param name="Job">The job to update</param>
+		/// <param name="Graph">The graph definition for this job</param>
+		/// <param name="Logger">Logger for output messages</param>
+		private static void CreateBatches(JobDocument Job, IGraph Graph, ILogger Logger)
+		{
+			UpdateBatches(Job, Graph, new List<UpdateDefinition<JobDocument>>(), Logger);
+		}
+
+		/// <summary>
+		/// Updates the list of batches for a job
+		/// </summary>
+		/// <param name="Job">Job to update</param>
+		/// <param name="Graph">Graph definition for this job</param>
+		/// <param name="Updates">List of updates for the job</param>
+		/// <param name="Logger">Logger for updates</param>
+		static void UpdateBatches(JobDocument Job, IGraph Graph, List<UpdateDefinition<JobDocument>> Updates, ILogger Logger)
+		{
+			UpdateDefinitionBuilder<JobDocument> UpdateBuilder = Builders<JobDocument>.Update;
+
+			// Update the list of batches
+			CreateOrUpdateBatches(Job, Graph);
+			Updates.Add(UpdateBuilder.Set(x => x.Batches, Job.Batches));
+			Updates.Add(UpdateBuilder.Set(x => x.NextSubResourceId, Job.NextSubResourceId));
+
+			// Update all the dependencies
+			RefreshDependentJobSteps(Job, Graph, new List<UpdateDefinition<JobDocument>>(), Logger);
+			RefreshJobPriority(Job, Updates);
+		}
+
+		/// <summary>
+		/// Update the jobsteps for the given node graph 
+		/// </summary>
+		/// <param name="Job">The job to update</param>
+		/// <param name="Graph">The graph for this job</param>
+		private static void CreateOrUpdateBatches(JobDocument Job, IGraph Graph)
+		{
+			// Find the priorities of each node, incorporating all the per-step overrides
+			Dictionary<INode, Priority> NodePriorities = new Dictionary<INode, Priority>();
+			foreach (INodeGroup Group in Graph.Groups)
+			{
+				foreach (INode Node in Group.Nodes)
+				{
+					NodePriorities[Node] = Node.Priority;
+				}
+			}
+			foreach (JobStepBatchDocument Batch in Job.Batches)
+			{
+				INodeGroup Group = Graph.Groups[Batch.GroupIdx];
+				foreach (JobStepDocument Step in Batch.Steps)
+				{
+					if (Step.Priority != null)
+					{
+						INode Node = Group.Nodes[Step.NodeIdx];
+						NodePriorities[Node] = Step.Priority.Value;
+					}
+				}
+			}
+
+			// Remove any steps and batches that haven't started yet
+			foreach (JobStepBatchDocument Batch in Job.Batches)
+			{
+				Batch.Steps.RemoveAll(x => x.State == JobStepState.Waiting || x.State == JobStepState.Ready);
+			}
+
+			// Remove any skipped nodes whose skipped state is no longer valid
+			HashSet<INode> FailedNodes = new HashSet<INode>();
+			foreach (JobStepBatchDocument Batch in Job.Batches)
+			{
+				INodeGroup Group = Graph.Groups[Batch.GroupIdx];
+				foreach (JobStepDocument Step in Batch.Steps)
+				{
+					INode Node = Group.Nodes[Step.NodeIdx];
+					if (Step.Retry)
+					{
+						FailedNodes.Remove(Node);
+					}
+					else if (Step.State == JobStepState.Skipped && (Node.InputDependencies.Any(x => FailedNodes.Contains(Graph.GetNode(x))) || !CanRetryNode(Job, Batch.GroupIdx, Step.NodeIdx)))
+					{
+						FailedNodes.Add(Node);
+					}
+					else if (Step.Outcome == JobStepOutcome.Failure)
+					{
+						FailedNodes.Add(Node);
+					}
+					else
+					{
+						FailedNodes.Remove(Node);
+					}
+				}
+				Batch.Steps.RemoveAll(x => x.State == JobStepState.Skipped && !FailedNodes.Contains(Group.Nodes[x.NodeIdx]));
+			}
+
+			// Remove any batches which are now empty
+			Job.Batches.RemoveAll(x => x.Steps.Count == 0 && x.Error == JobStepBatchError.None);
+
+			// Find all the targets in this job
+			HashSet<string> Targets = new HashSet<string>(StringComparer.OrdinalIgnoreCase);
+			if (Job.AbortedByUserId == null)
+			{
+				foreach (string Argument in Job.Arguments)
+				{
+					if (Argument.StartsWith(IJob.TargetArgumentPrefix, StringComparison.OrdinalIgnoreCase))
+					{
+						Targets.UnionWith(Argument.Substring(IJob.TargetArgumentPrefix.Length).Split(';'));
+					}
+				}
+				Targets.Add(IJob.SetupNodeName);
+			}
+
+			// Add all the referenced aggregates
+			HashSet<INode> NewNodesToExecute = new HashSet<INode>();
+			foreach (IAggregate Aggregate in Graph.Aggregates)
+			{
+				if (Targets.Contains(Aggregate.Name))
+				{
+					NewNodesToExecute.UnionWith(Aggregate.Nodes.Select(x => Graph.GetNode(x)));
+				}
+			}
+
+			// Add any individual nodes
+			foreach (INode Node in Graph.Groups.SelectMany(x => x.Nodes))
+			{
+				if (Targets.Contains(Node.Name))
+				{
+					NewNodesToExecute.Add(Node);
+				}
+			}
+
+			// Also add any dependencies of these nodes
+			for (int GroupIdx = Graph.Groups.Count - 1; GroupIdx >= 0; GroupIdx--)
+			{
+				INodeGroup Group = Graph.Groups[GroupIdx];
+				for (int NodeIdx = Group.Nodes.Count - 1; NodeIdx >= 0; NodeIdx--)
+				{
+					INode Node = Group.Nodes[NodeIdx];
+					if (NewNodesToExecute.Contains(Node))
+					{
+						foreach (NodeRef Dependency in Node.InputDependencies)
+						{
+							NewNodesToExecute.Add(Graph.GetNode(Dependency));
+						}
+					}
+				}
+			}
+
+			// Cancel any batches which are still running but are no longer required
+			foreach (JobStepBatchDocument Batch in Job.Batches)
+			{
+				if (Batch.State == JobStepBatchState.Starting || Batch.State == JobStepBatchState.Running)
+				{
+					INodeGroup Group = Graph.Groups[Batch.GroupIdx];
+					if (!Batch.Steps.Any(x => NewNodesToExecute.Contains(Group.Nodes[x.NodeIdx])))
+					{
+						Batch.Error = JobStepBatchError.Cancelled;
+					}
+				}
+			}
+
+			// Remove all the nodes which have already succeeded
+			foreach (JobStepBatchDocument Batch in Job.Batches)
+			{
+				foreach (IJobStep Step in Batch.Steps)
+				{
+					if ((Step.State == JobStepState.Running && Step.RetriedByUserId == null) || (Step.State == JobStepState.Completed && Step.RetriedByUserId == null) || (Step.State == JobStepState.Aborted && Step.RetriedByUserId == null) || Step.State == JobStepState.Skipped)
+					{
+						NewNodesToExecute.Remove(Graph.Groups[Batch.GroupIdx].Nodes[Step.NodeIdx]);
+					}
+				}
+			}
+
+			// Re-add all the nodes that have input dependencies in the same group.
+			for (int GroupIdx = Graph.Groups.Count - 1; GroupIdx >= 0; GroupIdx--)
+			{
+				INodeGroup Group = Graph.Groups[GroupIdx];
+				for (int NodeIdx = Group.Nodes.Count - 1; NodeIdx >= 0; NodeIdx--)
+				{
+					INode Node = Group.Nodes[NodeIdx];
+					if (NewNodesToExecute.Contains(Node))
+					{
+						foreach (NodeRef Dependency in Node.InputDependencies)
+						{
+							if (Dependency.GroupIdx == GroupIdx)
+							{
+								NewNodesToExecute.Add(Group.Nodes[Dependency.NodeIdx]);
+							}
+						}
+					}
+				}
+			}
+
+			// Build a list of nodes which are currently set to be executed
+			HashSet<INode> ExistingNodesToExecute = new HashSet<INode>();
+			foreach (JobStepBatchDocument Batch in Job.Batches)
+			{
+				INodeGroup Group = Graph.Groups[Batch.GroupIdx];
+				foreach (IJobStep Step in Batch.Steps)
+				{
+					INode Node = Group.Nodes[Step.NodeIdx];
+					ExistingNodesToExecute.Add(Node);
+				}
+			}
+
+			// Figure out the existing batch for each group
+			JobStepBatchDocument?[] AppendToBatches = new JobStepBatchDocument?[Graph.Groups.Count];
+			foreach (JobStepBatchDocument Batch in Job.Batches)
+			{
+				if (Batch.CanBeAppendedTo())
+				{
+					INodeGroup Group = Graph.Groups[Batch.GroupIdx];
+					INode FirstNode = Group.Nodes[Batch.Steps[0].NodeIdx];
+					AppendToBatches[Batch.GroupIdx] = Batch;
+				}
+			}
+
+			// Invalidate all the entries for groups where we're too late to append new entries (ie. we need to execute an earlier node that wasn't executed previously)
+			for (int GroupIdx = 0; GroupIdx < Graph.Groups.Count; GroupIdx++)
+			{
+				INodeGroup Group = Graph.Groups[GroupIdx];
+				for (int NodeIdx = 0; NodeIdx < Group.Nodes.Count; NodeIdx++)
+				{
+					INode Node = Group.Nodes[NodeIdx];
+					if (NewNodesToExecute.Contains(Node) && !ExistingNodesToExecute.Contains(Node))
+					{
+						IJobStepBatch? Batch = AppendToBatches[GroupIdx];
+						if (Batch != null)
+						{
+							IJobStep LastStep = Batch.Steps[Batch.Steps.Count - 1];
+							if (NodeIdx < LastStep.NodeIdx)
+							{
+								AppendToBatches[GroupIdx] = null;
+							}
+						}
+					}
+				}
+			}
+
+			// Create all the new jobsteps
+			for (int GroupIdx = 0; GroupIdx < Graph.Groups.Count; GroupIdx++)
+			{
+				INodeGroup Group = Graph.Groups[GroupIdx];
+				for (int NodeIdx = 0; NodeIdx < Group.Nodes.Count; NodeIdx++)
+				{
+					INode Node = Group.Nodes[NodeIdx];
+					if (NewNodesToExecute.Contains(Node))
+					{
+						JobStepBatchDocument? Batch = AppendToBatches[GroupIdx];
+						if (Batch == null)
+						{
+							Job.NextSubResourceId = Job.NextSubResourceId.Next();
+
+							Batch = new JobStepBatchDocument(Job.NextSubResourceId, GroupIdx);
+							Job.Batches.Add(Batch);
+
+							AppendToBatches[GroupIdx] = Batch;
+						}
+						if (Batch.Steps.Count == 0 || NodeIdx > Batch.Steps[Batch.Steps.Count - 1].NodeIdx)
+						{
+							Job.NextSubResourceId = Job.NextSubResourceId.Next();
+
+							JobStepDocument Step = new JobStepDocument(Job.NextSubResourceId, NodeIdx);
+							Batch.Steps.Add(Step);
+						}
+					}
+				}
+			}
+
+			// Find the priority of each node, propagating dependencies from dependent nodes
+			for (int GroupIdx = 0; GroupIdx < Graph.Groups.Count; GroupIdx++)
+			{
+				INodeGroup Group = Graph.Groups[GroupIdx];
+				for (int NodeIdx = 0; NodeIdx < Group.Nodes.Count; NodeIdx++)
+				{
+					INode Node = Group.Nodes[NodeIdx];
+					Priority NodePriority = NodePriorities[Node];
+
+					foreach (NodeRef DependencyRef in Node.OrderDependencies)
+					{
+						INode Dependency = Graph.Groups[DependencyRef.GroupIdx].Nodes[DependencyRef.NodeIdx];
+						if (NodePriorities[Node] > NodePriority)
+						{
+							NodePriorities[Dependency] = NodePriority;
+						}
+					}
+				}
+			}
+			foreach (JobStepBatchDocument Batch in Job.Batches)
+			{
+				if (Batch.Steps.Count > 0)
+				{
+					INodeGroup Group = Graph.Groups[Batch.GroupIdx];
+					Priority NodePriority = Batch.Steps.Max(x => NodePriorities[Group.Nodes[x.NodeIdx]]);
+					Batch.SchedulePriority = ((int)Job.Priority * 10) + (int)NodePriority + 1; // Reserve '0' for none.
+				}
+			}
+
+			// Check we're not running a node which doesn't allow retries more than once
+			Dictionary<INode, int> NodeExecutionCount = new Dictionary<INode, int>();
+			foreach (IJobStepBatch Batch in Job.Batches)
+			{
+				INodeGroup Group = Graph.Groups[Batch.GroupIdx];
+				foreach (IJobStep Step in Batch.Steps)
+				{
+					INode Node = Group.Nodes[Step.NodeIdx];
+
+					int Count;
+					NodeExecutionCount.TryGetValue(Node, out Count);
+
+					if (!Node.AllowRetry && Count > 0)
+					{
+						throw new RetryNotAllowedException(Node.Name);
+					}
+
+					NodeExecutionCount[Node] = Count + 1;
+				}
+			}
+		}
+
+		/// <summary>
+		/// Tests whether a node can be retried again
+		/// </summary>
+		static bool CanRetryNode(JobDocument Job, int GroupIdx, int NodeIdx)
+		{
+			return Job.RetriedNodes == null || Job.RetriedNodes.Count(x => x.GroupIdx == GroupIdx && x.NodeIdx == NodeIdx) < MaxRetries;
+		}
+
+		/// <summary>
+		/// Gets the scheduling priority of this job
+		/// </summary>
+		/// <param name="Job">Job to consider</param>
+		public static int GetSchedulePriority(IJob Job)
+		{
+			int NewSchedulePriority = 0;
+			foreach (IJobStepBatch Batch in Job.Batches)
+			{
+				if (Batch.State == JobStepBatchState.Ready)
+				{
+					NewSchedulePriority = Math.Max(Batch.SchedulePriority, NewSchedulePriority);
+				}
+			}
+			return NewSchedulePriority;
+		}
+
+		/// <summary>
+		/// Update the state of any jobsteps that are dependent on other jobsteps (eg. transition them from waiting to ready based on other steps completing)
+		/// </summary>
+		/// <param name="Job">The job to update</param>
+		/// <param name="Graph">The graph for this job</param>
+		/// <param name="Updates">List of updates to the job</param>
+		/// <param name="Logger">Logger instance</param>
+		static void RefreshDependentJobSteps(JobDocument Job, IGraph Graph, List<UpdateDefinition<JobDocument>> Updates, ILogger Logger)
+		{
+			// Update the batches
+			UpdateDefinitionBuilder<JobDocument> UpdateBuilder = Builders<JobDocument>.Update;
+			if (Job.Batches != null)
+			{
+				Dictionary<INode, IJobStep> StepForNode = new Dictionary<INode, IJobStep>();
+				for (int LoopBatchIdx = 0; LoopBatchIdx < Job.Batches.Count; LoopBatchIdx++)
+				{
+					int BatchIdx = LoopBatchIdx; // For lambda capture
+					JobStepBatchDocument Batch = Job.Batches[BatchIdx];
+
+					for (int LoopStepIdx = 0; LoopStepIdx < Batch.Steps.Count; LoopStepIdx++)
+					{
+						int StepIdx = LoopStepIdx; // For lambda capture
+						JobStepDocument Step = Batch.Steps[StepIdx];
+
+						JobStepState NewState = Step.State;
+						JobStepOutcome NewOutcome = Step.Outcome;
+
+						INode Node = Graph.Groups[Batch.GroupIdx].Nodes[Step.NodeIdx];
+						if (NewState == JobStepState.Waiting)
+						{
+							List<IJobStep> Steps = GetDependentSteps(Graph, Node, StepForNode);
+							if (Steps.Any(x => x.IsFailedOrSkipped()))
+							{
+								NewState = JobStepState.Skipped;
+								NewOutcome = JobStepOutcome.Failure;
+							}
+							else if (!Steps.Any(x => x.IsPending()))
+							{
+								Logger.LogDebug("Transitioning job {JobId}, batch {BatchId}, step {StepId} to ready state ({Dependencies})", Job.Id, Batch.Id, Step.Id, String.Join(", ", Steps.Select(x => x.Id.ToString())));
+								NewState = JobStepState.Ready;
+							}
+						}
+
+						if (NewState != Step.State)
+						{
+							Step.State = NewState;
+							Updates.Add(UpdateBuilder.Set(x => x.Batches[BatchIdx].Steps[StepIdx].State, NewState));
+						}
+
+						if (NewOutcome != Step.Outcome)
+						{
+							Step.Outcome = NewOutcome;
+							Updates.Add(UpdateBuilder.Set(x => x.Batches[BatchIdx].Steps[StepIdx].Outcome, NewOutcome));
+						}
+
+						StepForNode[Node] = Step;
+					}
+
+					if (Batch.State == JobStepBatchState.Waiting || Batch.State == JobStepBatchState.Ready)
+					{
+						DateTime? NewReadyTime;
+						JobStepBatchState NewState = GetBatchState(Job, Graph, Batch, StepForNode, out NewReadyTime);
+						if (Batch.State != NewState)
+						{
+							Batch.State = NewState;
+							Updates.Add(UpdateBuilder.Set(x => x.Batches[BatchIdx].State, Batch.State));
+						}
+						if (Batch.ReadyTime != NewReadyTime)
+						{
+							Batch.ReadyTime = NewReadyTime;
+							Updates.Add(UpdateBuilder.Set(x => x.Batches[BatchIdx].ReadyTime, Batch.ReadyTime));
+						}
+					}
+				}
+			}
+		}
+
+		/// <summary>
+		/// Updates the schedule priority of a job
+		/// </summary>
+		/// <param name="Job"></param>
+		/// <param name="Updates"></param>
+		static void RefreshJobPriority(JobDocument Job, List<UpdateDefinition<JobDocument>> Updates)
+		{
+			// Update the weighted priority for the job
+			int NewSchedulePriority = GetSchedulePriority(Job);
+			if (Job.SchedulePriority != NewSchedulePriority)
+			{
+				Job.SchedulePriority = NewSchedulePriority;
+				Updates.Add(Builders<JobDocument>.Update.Set(x => x.SchedulePriority, NewSchedulePriority));
+			}
+		}
+
+		/// <summary>
+		/// Gets the steps that a node depends on
+		/// </summary>
+		/// <param name="Graph">The graph for this job</param>
+		/// <param name="Node">The node to test</param>
+		/// <param name="StepForNode">Map of node to step</param>
+		/// <returns></returns>
+		static List<IJobStep> GetDependentSteps(IGraph Graph, INode Node, Dictionary<INode, IJobStep> StepForNode)
+		{
+			List<IJobStep> Steps = new List<IJobStep>();
+			foreach (NodeRef OrderDependencyRef in Node.OrderDependencies)
+			{
+				IJobStep? Step;
+				if (StepForNode.TryGetValue(Graph.GetNode(OrderDependencyRef), out Step))
+				{
+					Steps.Add(Step);
+				}
+			}
+			return Steps;
+		}
+
+		/// <summary>
+		/// Gets the new state for a batch
+		/// </summary>
+		/// <param name="Job">The job being executed</param>
+		/// <param name="Graph">Graph for the job</param>
+		/// <param name="Batch">List of nodes in the job</param>
+		/// <param name="StepForNode">Array mapping each index to the appropriate step for that node</param>
+		/// <param name="OutReadyTimeUtc">Receives the time at which the batch was ready to execute</param>
+		/// <returns>True if the batch is ready, false otherwise</returns>
+		static JobStepBatchState GetBatchState(IJob Job, IGraph Graph, IJobStepBatch Batch, Dictionary<INode, IJobStep> StepForNode, out DateTime? OutReadyTimeUtc)
+		{
+			// Check if the batch is already complete
+			if (Batch.Steps.All(x => x.State == JobStepState.Skipped || x.State == JobStepState.Completed || x.State == JobStepState.Aborted))
+			{
+				OutReadyTimeUtc = Batch.ReadyTimeUtc;
+				return JobStepBatchState.Complete;
+			}
+
+			// Get the dependencies for this batch to start. Some steps may be "after" dependencies that are optional parts of the graph.
+			List<INode> NodeDependencies = Batch.GetStartDependencies(Graph.Groups).ToList();
+
+			// Check if we're still waiting on anything
+			DateTime ReadyTimeUtc = Job.CreateTimeUtc;
+			foreach (INode NodeDependency in NodeDependencies)
+			{
+				IJobStep? StepDependency;
+				if (StepForNode.TryGetValue(NodeDependency, out StepDependency))
+				{
+					if (StepDependency.State != JobStepState.Completed && StepDependency.State != JobStepState.Skipped && StepDependency.State != JobStepState.Aborted)
+					{
+						OutReadyTimeUtc = null;
+						return JobStepBatchState.Waiting;
+					}
+
+					if (StepDependency.FinishTimeUtc != null && StepDependency.FinishTimeUtc.Value > ReadyTimeUtc)
+					{
+						ReadyTimeUtc = StepDependency.FinishTimeUtc.Value;
+					}
+				}
+			}
+
+			// Otherwise return the ready state
+			OutReadyTimeUtc = ReadyTimeUtc;
+			return JobStepBatchState.Ready;
+		}
+
+		/// <inheritdoc/>
+		public async Task UpgradeDocumentsAsync()
+		{
+			IAsyncCursor<JobDocument> Cursor = await Jobs.Find(Builders<JobDocument>.Filter.Eq(x => x.UpdateTimeUtc, null)).ToCursorAsync();
+
+			int NumUpdated = 0;
+			while (await Cursor.MoveNextAsync())
+			{
+				foreach (JobDocument Document in Cursor.Current)
+				{
+					UpdateDefinition<JobDocument> Update = Builders<JobDocument>.Update.Set(x => x.CreateTimeUtc, ((IJob)Document).CreateTimeUtc).Set(x => x.UpdateTimeUtc, ((IJob)Document).UpdateTimeUtc);
+					await Jobs.UpdateOneAsync(Builders<JobDocument>.Filter.Eq(x => x.Id, Document.Id), Update);
+					NumUpdated++;
+				}
+				Logger.LogInformation("Updated {NumDocuments} documents", NumUpdated);
+			}
+		}
+	}
+}
