Horde: Rename the HordeServer project to Horde.Build.

diff --git a/Engine/Source/Programs/Horde/Horde.Build/Jobs/Impl/JobService.cs b/Engine/Source/Programs/Horde/Horde.Build/Jobs/Impl/JobService.cs
new file mode 100644
--- /dev/null
+++ b/Engine/Source/Programs/Horde/Horde.Build/Jobs/Impl/JobService.cs
@@ -0,0 +1,1246 @@
+// Copyright Epic Games, Inc. All Rights Reserved.
+
+using EpicGames.Core;
+using HordeServer.Api;
+using HordeServer.Collections;
+using HordeCommon;
+using HordeServer.Models;
+using HordeCommon.Rpc;
+using HordeServer.Utilities;
+using Microsoft.AspNetCore.Http;
+using Microsoft.Extensions.Logging;
+using Microsoft.Extensions.Options;
+using MongoDB.Bson;
+using MongoDB.Bson.Serialization.Serializers;
+using MongoDB.Driver;
+using System;
+using System.Collections.Generic;
+using System.ComponentModel;
+using System.Diagnostics;
+using System.Diagnostics.CodeAnalysis;
+using System.Linq;
+using System.Security.Claims;
+using System.Text.RegularExpressions;
+using System.Threading;
+using System.Threading.Tasks;
+using System.Globalization;
+using HordeServer.Tasks.Impl;
+using OpenTracing.Util;
+using OpenTracing;
+
+namespace HordeServer.Services
+{
+	using JobId = ObjectId<IJob>;
+	using LeaseId = ObjectId<ILease>;
+	using LogId = ObjectId<ILogFile>;
+	using StreamId = StringId<IStream>;
+	using TemplateRefId = StringId<TemplateRef>;
+	using UserId = ObjectId<IUser>;
+
+	/// <summary>
+	/// Exception thrown when attempting to retry executing a node that does not allow retries
+	/// </summary>
+	public class RetryNotAllowedException : Exception
+	{
+		/// <summary>
+		/// The node name
+		/// </summary>
+		public string NodeName { get; private set; }
+
+		/// <summary>
+		/// Constructor
+		/// </summary>
+		/// <param name="NodeName">Name of the node that does not permit retries</param>
+		public RetryNotAllowedException(string NodeName)
+			: base($"Node '{NodeName}' does not permit retries")
+		{
+			this.NodeName = NodeName;
+		}
+	}
+
+	/// <summary>
+	/// Cache of information about job ACLs
+	/// </summary>
+	public class JobPermissionsCache : StreamPermissionsCache
+	{
+		/// <summary>
+		/// Map of job id to permissions for that job
+		/// </summary>
+		public Dictionary<JobId, IJobPermissions?> Jobs { get; set; } = new Dictionary<JobId, IJobPermissions?>();
+	}
+
+	/// <summary>
+	/// Wraps funtionality for manipulating jobs, jobsteps, and jobstep runs
+	/// </summary>
+	[SuppressMessage("Compiler", "CA1054:URI parameters should not be strings")]
+	public class JobService
+	{
+		/// <summary>
+		/// Collection of job documents
+		/// </summary>
+		IJobCollection Jobs;
+
+		/// <summary>
+		/// Collection of graph documents
+		/// </summary>
+		IGraphCollection Graphs;
+
+		/// <summary>
+		/// Collection of agent documents
+		/// </summary>
+		IAgentCollection Agents;
+
+		/// <summary>
+		/// Collection of jobstepref documents
+		/// </summary>
+		IJobStepRefCollection JobStepRefs;
+
+		/// <summary>
+		/// Collection of job timing documents
+		/// </summary>
+		IJobTimingCollection JobTimings;
+
+		/// <summary>
+		/// Collection of users
+		/// </summary>
+		IUserCollection UserCollection;
+
+		/// <summary>
+		/// Collection of notification triggers
+		/// </summary>
+		INotificationTriggerCollection TriggerCollection;
+
+		/// <summary>
+		/// The singleton instance of the queue service
+		/// </summary>
+		JobTaskSource JobTaskSource;
+
+		/// <summary>
+		/// Singleton instance of the stream service
+		/// </summary>
+		StreamService StreamService;
+
+		/// <summary>
+		/// Singleton instance of the template service
+		/// </summary>
+		ITemplateCollection TemplateCollection;
+
+		/// <summary>
+		/// Singleton instance of the issue service
+		/// </summary>
+		IIssueService? IssueService;
+
+		/// <summary>
+		/// The perforce service
+		/// </summary>
+		IPerforceService PerforceService;
+
+		/// <summary>
+		/// The server settings
+		/// </summary>
+		IOptionsMonitor<ServerSettings> Settings;
+
+		/// <summary>
+		/// Log output device
+		/// </summary>
+		ILogger Logger;
+
+		/// <summary>
+		/// Delegate for label update events
+		/// </summary>
+		public delegate void LabelUpdateEvent(IJob Job, IReadOnlyList<(LabelState, LabelOutcome)> OldLabelStates, IReadOnlyList<(LabelState, LabelOutcome)> NewLabelStates);
+
+		/// <summary>
+		/// Event triggered when a label state changes
+		/// </summary>
+		public event LabelUpdateEvent? OnLabelUpdate;
+
+		/// <summary>
+		/// Delegate for job step complete events
+		/// </summary>
+		public delegate void JobStepCompleteEvent(IJob Job, IGraph Graph, SubResourceId BatchId, SubResourceId StepId);
+
+		/// <summary>
+		/// Event triggered when a job step completes
+		/// </summary>
+		public event JobStepCompleteEvent? OnJobStepComplete;
+
+		/// <summary>
+		/// Constructor
+		/// </summary>
+		/// <param name="Jobs">The jobs collection</param>
+		/// <param name="Graphs">The graphs collection</param>
+		/// <param name="Agents">The agents collection</param>
+		/// <param name="JobStepRefs">The jobsteprefs collection</param>
+		/// <param name="JobTimings">The job timing document collection</param>
+		/// <param name="UserCollection">User profiles</param>
+		/// <param name="TriggerCollection">Trigger collection</param>
+		/// <param name="JobTaskSource">The queue service</param>
+		/// <param name="StreamService">The stream service</param>
+		/// <param name="TemplateCollection">The template service</param>
+		/// <param name="IssueService">The issue service</param>
+		/// <param name="PerforceService">The perforce service</param>
+		/// <param name="Settings">Settings instance</param>
+		/// <param name="Logger">Log output</param>
+		public JobService(IJobCollection Jobs, IGraphCollection Graphs, IAgentCollection Agents, IJobStepRefCollection JobStepRefs, IJobTimingCollection JobTimings, IUserCollection UserCollection, INotificationTriggerCollection TriggerCollection, JobTaskSource JobTaskSource, StreamService StreamService, ITemplateCollection TemplateCollection, IIssueService IssueService, IPerforceService PerforceService, IOptionsMonitor<ServerSettings> Settings, ILogger<JobService> Logger)
+		{
+			this.Jobs = Jobs;
+			this.Graphs = Graphs;
+			this.Agents = Agents;
+			this.JobStepRefs = JobStepRefs;
+			this.JobTimings = JobTimings;
+			this.UserCollection = UserCollection;
+			this.TriggerCollection = TriggerCollection;
+			this.JobTaskSource = JobTaskSource;
+			this.StreamService = StreamService;
+			this.TemplateCollection = TemplateCollection;
+			this.IssueService = IssueService;
+			this.PerforceService = PerforceService;
+			this.Settings = Settings;
+			this.Logger = Logger;
+		}
+
+#pragma warning disable CA1801
+		static bool ShouldClonePreflightChange(StreamId StreamId)
+		{
+			return false; //			return StreamId == new StreamId("ue5-main");
+		}
+#pragma warning restore CA1801
+
+		/// <summary>
+		/// Creates a new job
+		/// </summary>
+		/// <param name="JobId">A requested job id</param>
+		/// <param name="Stream">The stream that this job belongs to</param>
+		/// <param name="TemplateRefId">Name of the template ref</param>
+		/// <param name="TemplateHash">Template for this job</param>
+		/// <param name="Graph">The graph for the new job</param>
+		/// <param name="Name">Name of the job</param>
+		/// <param name="Change">The change to build</param>
+		/// <param name="CodeChange">The corresponding code changelist</param>
+		/// <param name="PreflightChange">Optional changelist to preflight</param>
+		/// <param name="ClonedPreflightChange">Duplicated preflight change</param>
+		/// <param name="StartedByUserId">Id of the user that started the job</param>
+		/// <param name="Priority">Priority of the job</param>
+		/// <param name="AutoSubmit">Whether to automatically submit the preflighted change on completion</param>
+		/// <param name="UpdateIssues">Whether to update issues when this job completes</param>
+		/// <param name="JobTriggers">List of downstream job triggers</param>
+		/// <param name="ShowUgsBadges">Whether to show badges in UGS for this job</param>
+		/// <param name="ShowUgsAlerts">Whether to show alerts in UGS for this job</param>
+		/// <param name="NotificationChannel">Notification Channel for this job</param>
+		/// <param name="NotificationChannelFilter">Notification Channel filter for this job</param>
+		/// <param name="Arguments">Arguments for the job</param>
+		/// <returns>Unique id representing the job</returns>
+		public async Task<IJob> CreateJobAsync(JobId? JobId, IStream Stream, TemplateRefId TemplateRefId, ContentHash TemplateHash, IGraph Graph, string Name, int Change, int CodeChange, int? PreflightChange, int? ClonedPreflightChange, UserId? StartedByUserId, Priority? Priority, bool? AutoSubmit, bool? UpdateIssues, List<ChainedJobTemplate>? JobTriggers, bool ShowUgsBadges, bool ShowUgsAlerts, string? NotificationChannel, string? NotificationChannelFilter, IReadOnlyList<string> Arguments)
+		{
+			JobId JobIdValue = JobId ?? HordeServer.Utilities.ObjectId<IJob>.GenerateNewId();
+			using IDisposable Scope = Logger.BeginScope("CreateJobAsync({JobId})", JobIdValue);
+
+			if (PreflightChange != null && ShouldClonePreflightChange(Stream.Id))
+			{
+				ClonedPreflightChange = await CloneShelvedChangeAsync(Stream.ClusterName, ClonedPreflightChange ?? PreflightChange.Value);
+			}
+
+			Logger.LogInformation("Creating job at CL {Change}, code CL {CodeChange}, preflight CL {PreflightChange}, cloned CL {ClonedPreflightChange}", Change, CodeChange, PreflightChange, ClonedPreflightChange);
+
+			Dictionary<string, string> Properties = new Dictionary<string, string>();
+			Properties["Change"] = Change.ToString(CultureInfo.InvariantCulture);
+			Properties["CodeChange"] = CodeChange.ToString(CultureInfo.InvariantCulture);
+			Properties["PreflightChange"] = PreflightChange?.ToString(CultureInfo.InvariantCulture) ?? String.Empty;
+			Properties["ClonedPreflightChange"] = ClonedPreflightChange?.ToString(CultureInfo.InvariantCulture) ?? String.Empty;
+			Properties["StreamId"] = Stream.Id.ToString();
+			Properties["TemplateId"] = TemplateRefId.ToString();
+			Properties["JobId"] = JobIdValue.ToString();
+
+			List<string> ExpandedArguments = new List<string>();
+			if (Arguments != null)
+			{
+				foreach (string Argument in Arguments)
+				{
+					string ExpandedArgument = StringUtils.ExpandProperties(Argument, Properties);
+					ExpandedArguments.Add(ExpandedArgument);
+				}
+			}
+
+			Name = StringUtils.ExpandProperties(Name, Properties);
+
+			IJob NewJob = await Jobs.AddAsync(JobIdValue, Stream.Id, TemplateRefId, TemplateHash, Graph, Name, Change, CodeChange, PreflightChange, ClonedPreflightChange, StartedByUserId, Priority, AutoSubmit, UpdateIssues, JobTriggers, ShowUgsBadges, ShowUgsAlerts, NotificationChannel, NotificationChannelFilter, ExpandedArguments);
+			JobTaskSource.UpdateQueuedJob(NewJob, Graph);
+
+			await JobTaskSource.UpdateUgsBadges(NewJob, Graph, new List<(LabelState, LabelOutcome)>());
+
+			if (StartedByUserId != null)
+			{
+				await UserCollection.UpdateSettingsAsync(StartedByUserId.Value, AddPinnedJobIds: new[] { NewJob.Id });
+			}
+
+			await AbortAnyDuplicateJobs(NewJob);
+
+			return NewJob;
+		}
+
+		private async Task AbortAnyDuplicateJobs(IJob NewJob)
+		{
+			List<IJob> JobsToAbort = new List<IJob>();
+			if (NewJob.PreflightChange > 0)
+			{
+				JobsToAbort = await Jobs.FindAsync(PreflightChange: NewJob.PreflightChange);
+			}
+
+			foreach (IJob Job in JobsToAbort)
+			{
+				if (Job.GetState() == JobState.Complete) continue;
+				if (Job.Id == NewJob.Id) continue; // Don't remove the new job
+				if (Job.TemplateId != NewJob.TemplateId) continue;
+				if (Job.TemplateHash != NewJob.TemplateHash) continue;
+				if (Job.StartedByUserId != NewJob.StartedByUserId) continue;
+				if (string.Join(",", Job.Arguments) != string.Join(",", NewJob.Arguments)) continue;
+
+				bool bUpdateSuccessful = await UpdateJobAsync(Job, null, null, null, KnownUsers.System, null, null, null);
+				if (!bUpdateSuccessful)
+				{
+					Logger.LogError("Failed marking duplicate job as aborted! Job ID: {JobId}", Job.Id);
+				}
+			}
+		}
+
+		/// <summary>
+		/// Deletes a job
+		/// </summary>
+		/// <param name="Job">The job to delete</param>
+		public async Task<bool> DeleteJobAsync(IJob Job)
+		{
+			using IDisposable Scope = Logger.BeginScope("DeleteJobAsync({JobId})", Job.Id);
+
+			// Delete the job
+			while (!await Jobs.RemoveAsync(Job))
+			{
+				IJob? NewJob = await Jobs.GetAsync(Job.Id);
+				if (NewJob == null)
+				{
+					return false;
+				}
+				Job = NewJob;
+			}
+
+			// Remove all the triggers from it
+			List<ObjectId> TriggerIds = new List<ObjectId>();
+			if (Job.NotificationTriggerId != null)
+			{
+				TriggerIds.Add(Job.NotificationTriggerId.Value);
+			}
+			foreach (IJobStep Step in Job.Batches.SelectMany(x => x.Steps))
+			{
+				if (Step.NotificationTriggerId != null)
+				{
+					TriggerIds.Add(Step.NotificationTriggerId.Value);
+				}
+			}
+			await TriggerCollection.DeleteAsync(TriggerIds);
+
+			return true;
+		}
+
+		/// <summary>
+		/// Delete all the jobs for a stream
+		/// </summary>
+		/// <param name="StreamId">Unique id of the stream</param>
+		/// <returns>Async task</returns>
+		public async Task DeleteJobsForStreamAsync(StreamId StreamId)
+		{
+			await Jobs.RemoveStreamAsync(StreamId);
+		}
+
+		/// <summary>
+		/// Updates a new job
+		/// </summary>
+		/// <param name="Job">The job document to update</param>
+		/// <param name="Name">Name of the job</param>
+		/// <param name="Priority">Priority of the job</param>
+		/// <param name="AutoSubmit">Whether to automatically submit the preflighted change on completion</param>
+		/// <param name="AbortedByUserId">Name of the user that aborted this job</param>
+		/// <param name="OnCompleteTriggerId">Object id for a notification trigger</param>
+		/// <param name="Reports">New reports to add</param>
+		/// <param name="Arguments">New arguments for the job</param>
+		/// <param name="LabelIdxToTriggerId">New trigger ID for a label in the job</param>
+		public async Task<bool> UpdateJobAsync(IJob Job, string? Name = null, Priority? Priority = null, bool? AutoSubmit = null, UserId? AbortedByUserId = null, ObjectId? OnCompleteTriggerId = null, List<Report>? Reports = null, List<string>? Arguments = null, KeyValuePair<int, ObjectId>? LabelIdxToTriggerId = null)
+		{
+			using IDisposable Scope = Logger.BeginScope("UpdateJobAsync({JobId})", Job.Id);
+			for (; ; )
+			{
+				IGraph Graph = await GetGraphAsync(Job);
+
+				// Capture the previous label states
+				IReadOnlyList<(LabelState, LabelOutcome)> OldLabelStates = Job.GetLabelStates(Graph);
+
+				// Update the new list of job steps
+				if (await Jobs.TryUpdateJobAsync(Job, Graph, Name, Priority, AutoSubmit, null, null, AbortedByUserId, OnCompleteTriggerId, Reports, Arguments, LabelIdxToTriggerId))
+				{
+					// Update any badges that have been modified
+					await JobTaskSource.UpdateUgsBadges(Job, Graph, OldLabelStates);
+
+					// Cancel any leases which are no longer required
+					foreach (IJobStepBatch Batch in Job.Batches)
+					{
+						if (Batch.Error == JobStepBatchError.Cancelled && (Batch.State == JobStepBatchState.Starting || Batch.State == JobStepBatchState.Running) && Batch.AgentId != null && Batch.LeaseId != null)
+						{
+							await CancelLeaseAsync(Batch.AgentId.Value, Batch.LeaseId.Value);
+						}
+					}
+					return true;
+				}
+
+				// Get the new job state
+				IJob? NewJob = await GetJobAsync(Job.Id);
+				if (NewJob == null)
+				{
+					return false;
+				}
+				Job = NewJob;
+			}
+		}
+
+		/// <summary>
+		/// Cancel an active lease
+		/// </summary>
+		/// <param name="AgentId">The agent to retreive</param>
+		/// <param name="LeaseId">The lease id to update</param>
+		/// <returns></returns>
+		async Task CancelLeaseAsync(AgentId AgentId, LeaseId LeaseId)
+		{
+			for (; ; )
+			{
+				IAgent? Agent = await Agents.GetAsync(AgentId);
+				if (Agent == null)
+				{
+					break;
+				}
+
+				int Index = 0;
+				while (Index < Agent.Leases.Count && Agent.Leases[Index].Id != LeaseId)
+				{
+					Index++;
+				}
+				if (Index == Agent.Leases.Count)
+				{
+					break;
+				}
+
+				AgentLease Lease = Agent.Leases[Index];
+				if (Lease.State != LeaseState.Active)
+				{
+					break;
+				}
+
+				IAgent? NewAgent = await Agents.TryCancelLeaseAsync(Agent, Index);
+				if (NewAgent != null)
+				{
+					JobTaskSource.CancelLongPollForAgent(Agent.Id);
+					break;
+				}
+			}
+		}
+
+		/// <summary>
+		/// Gets a job with the given unique id
+		/// </summary>
+		/// <param name="JobId">Job id to search for</param>
+		/// <returns>Information about the given job</returns>
+		public Task<IJob?> GetJobAsync(JobId JobId)
+		{
+			return Jobs.GetAsync(JobId);
+		}
+
+		/// <summary>
+		/// Gets the graph for a job
+		/// </summary>
+		/// <param name="Job">Job to retrieve the graph for</param>
+		/// <returns>The graph for this job</returns>
+		public Task<IGraph> GetGraphAsync(IJob Job)
+		{
+			return Graphs.GetAsync(Job.GraphHash);
+		}
+
+		/// <summary>
+		/// Gets a job's permissions info by ID
+		/// </summary>
+		/// <param name="JobId">Unique id of the job</param>
+		/// <returns>The job document</returns>
+		public Task<IJobPermissions?> GetJobPermissionsAsync(JobId JobId)
+		{
+			return Jobs.GetPermissionsAsync(JobId);
+		}
+
+		/// <summary>
+		/// Searches for jobs matching the given criteria
+		/// </summary>
+		/// <param name="JobIds">List of job ids to return</param>
+		/// <param name="StreamId">The stream containing the job</param>
+		/// <param name="Name">Name of the job</param>
+		/// <param name="Templates">Templates to look for</param>
+		/// <param name="MinChange">The minimum changelist number</param>
+		/// <param name="MaxChange">The maximum changelist number</param>		
+		/// <param name="PreflightChange">The preflight change to look for</param>
+		/// <param name="PreflightStartedByUser">User for which to include preflight jobs</param>		
+		/// <param name="StartedByUser">User for which to include jobs</param>
+		/// <param name="MinCreateTime">The minimum creation time</param>
+		/// <param name="MaxCreateTime">The maximum creation time</param>
+		/// <param name="Target">The target to query</param>
+		/// <param name="State">State to query</param>
+		/// <param name="Outcome">Outcomes to return</param>
+		/// <param name="ModifiedBefore">Filter the results by last modified time</param>
+		/// <param name="ModifiedAfter">Filter the results by last modified time</param>
+		/// <param name="Index">Index of the first result to return</param>
+		/// <param name="Count">Number of results to return</param>
+		/// <returns>List of jobs matching the given criteria</returns>
+		public async Task<List<IJob>> FindJobsAsync(JobId[]? JobIds = null, StreamId? StreamId = null, string? Name = null, TemplateRefId[]? Templates = null, int? MinChange = null, int? MaxChange = null, int? PreflightChange = null, UserId? PreflightStartedByUser = null, UserId? StartedByUser = null, DateTimeOffset ? MinCreateTime = null, DateTimeOffset? MaxCreateTime = null, string? Target = null, JobStepState[]? State = null, JobStepOutcome[]? Outcome = null, DateTimeOffset? ModifiedBefore = null, DateTimeOffset? ModifiedAfter = null, int? Index = null, int? Count = null)
+		{
+			if (Target == null && (State == null || State.Length == 0) && (Outcome == null || Outcome.Length == 0))
+			{
+				return await Jobs.FindAsync(JobIds, StreamId, Name, Templates, MinChange, MaxChange, PreflightChange, PreflightStartedByUser, StartedByUser, MinCreateTime, MaxCreateTime, ModifiedBefore, ModifiedAfter, Index, Count);
+			}
+			else
+			{
+				List<IJob> Results = new List<IJob>();
+				Logger.LogInformation("Performing scan for job with ");
+
+				int MaxCount = (Count ?? 1);
+				while (Results.Count < MaxCount)
+				{
+					List<IJob> ScanJobs = await Jobs.FindAsync(JobIds, StreamId, Name, Templates, MinChange, MaxChange, PreflightChange, PreflightStartedByUser, StartedByUser, MinCreateTime, MaxCreateTime, ModifiedBefore, ModifiedAfter, 0, 5);
+					if (ScanJobs.Count == 0)
+					{
+						break;
+					}
+
+					foreach (IJob Job in ScanJobs.OrderByDescending(x => x.Change))
+					{
+						(JobStepState, JobStepOutcome)? Result;
+						if (Target == null)
+						{
+							Result = Job.GetTargetState();
+						}
+						else
+						{
+							Result = Job.GetTargetState(await GetGraphAsync(Job), Target);
+						}
+
+						if (Result != null)
+						{
+							(JobStepState JobState, JobStepOutcome JobOutcome) = Result.Value;
+							if ((State == null || State.Length == 0 || State.Contains(JobState)) && (Outcome == null || Outcome.Length == 0 || Outcome.Contains(JobOutcome)))
+							{
+								Results.Add(Job);
+								if (Results.Count == MaxCount)
+								{
+									break;
+								}
+							}
+						}
+					}
+
+					MaxChange = ScanJobs.Min(x => x.Change) - 1;
+				}
+
+				return Results;
+			}
+		}
+
+		/// <summary>
+		/// Attempts to update the node groups to be executed for a job. Fails if another write happens in the meantime.
+		/// </summary>
+		/// <param name="Job">The job to update</param>
+		/// <param name="NewGraph">New graph for this job</param>
+		/// <returns>True if the groups were updated to the given list. False if another write happened first.</returns>
+		public async Task<bool> TryUpdateGraphAsync(IJob Job, IGraph NewGraph)
+		{
+			using IDisposable Scope = Logger.BeginScope("TryUpdateGraphAsync({JobId})", Job.Id);
+
+			IReadOnlyList<(LabelState, LabelOutcome)> OldLabelStates = Job.GetLabelStates(NewGraph);
+			if (await Jobs.TryUpdateGraphAsync(Job, NewGraph))
+			{
+				await JobTaskSource.UpdateUgsBadges(Job, NewGraph, OldLabelStates);
+
+				JobTaskSource.UpdateQueuedJob(Job, NewGraph);
+				return true;
+			}
+			return false;
+		}
+
+		/// <summary>
+		/// Gets the timing info for a particular job.
+		/// </summary>
+		/// <param name="Job">The job to get timing info for</param>
+		/// <returns>Timing info for the given job</returns>
+		public async Task<IJobTiming> GetJobTimingAsync(IJob Job)
+		{
+			using IDisposable Scope = Logger.BeginScope("GetJobTimingAsync({JobId})", Job.Id);
+
+			IGraph Graph = await Graphs.GetAsync(Job.GraphHash);
+
+			Dictionary<string, JobStepTimingData> CachedNewSteps = new Dictionary<string, JobStepTimingData>();
+			for (; ; )
+			{
+				// Try to get the current timing document
+				IJobTiming? JobTiming = await JobTimings.TryGetAsync(Job.Id);
+
+				// Calculate timing info for any missing steps
+				List<JobStepTimingData> NewSteps = new List<JobStepTimingData>();
+				foreach (IJobStepBatch Batch in Job.Batches)
+				{
+					INodeGroup Group = Graph.Groups[Batch.GroupIdx];
+					foreach (IJobStep Step in Batch.Steps)
+					{
+						INode Node = Group.Nodes[Step.NodeIdx];
+						if (JobTiming == null || !JobTiming.TryGetStepTiming(Node.Name, out _))
+						{
+							JobStepTimingData? StepTimingData;
+							if (!CachedNewSteps.TryGetValue(Node.Name, out StepTimingData))
+							{
+								StepTimingData = await GetStepTimingInfo(Job.StreamId, Job.TemplateId, Node.Name, Job.Change);
+							}
+							NewSteps.Add(StepTimingData);
+						}
+					}
+				}
+
+				// Try to add or update the timing document
+				if (JobTiming == null)
+				{
+					IJobTiming? NewJobTiming = await JobTimings.TryAddAsync(Job.Id, NewSteps);
+					if (NewJobTiming != null)
+					{
+						return NewJobTiming;
+					}
+				}
+				else if (NewSteps.Count == 0)
+				{
+					return JobTiming;
+				}
+				else
+				{
+					IJobTiming? NewJobTiming = await JobTimings.TryAddStepsAsync(JobTiming, NewSteps);
+					if (NewJobTiming != null)
+					{
+						return NewJobTiming;
+					}
+				}
+
+				// Update the cached list of steps for the next iteration
+				foreach (JobStepTimingData NewStep in NewSteps)
+				{
+					CachedNewSteps[NewStep.Name] = NewStep;
+				}
+			}
+		}
+
+		/// <summary>
+		/// Finds the average duration for the given step
+		/// </summary>
+		/// <param name="StreamId">The stream to search</param>
+		/// <param name="TemplateId">The template id</param>
+		/// <param name="NodeName">Name of the node</param>
+		/// <param name="Change">Maximum changelist to consider</param>
+		/// <returns>Expected duration for the given step</returns>
+		async Task<JobStepTimingData> GetStepTimingInfo(StreamId StreamId, TemplateRefId TemplateId, string NodeName, int? Change)
+		{
+			// Find all the steps matching the given criteria
+			List<IJobStepRef> Steps = await JobStepRefs.GetStepsForNodeAsync(StreamId, TemplateId, NodeName, Change, false, 10);
+
+			// Sum up all the durations and wait times
+			int Count = 0;
+			float WaitTimeSum = 0;
+			float InitTimeSum = 0;
+			float DurationSum = 0;
+
+			foreach (IJobStepRef Step in Steps)
+			{
+				if (Step.FinishTimeUtc != null)
+				{
+					WaitTimeSum += Step.BatchWaitTime;
+					InitTimeSum += Step.BatchInitTime;
+					DurationSum += (float)(Step.FinishTimeUtc.Value - Step.StartTimeUtc).TotalSeconds;
+					Count++;
+				}
+			}
+
+			// Compute the averages
+			if (Count == 0)
+			{
+				return new JobStepTimingData(NodeName, null, null, null);
+			}
+			else
+			{
+				return new JobStepTimingData(NodeName, WaitTimeSum / Count, InitTimeSum / Count, DurationSum / Count);
+			}
+		}
+
+		/// <summary>
+		/// Updates the state of a batch
+		/// </summary>
+		/// <param name="Job">Job to update</param>
+		/// <param name="BatchId">Unique id of the batch to update</param>
+		/// <param name="NewLogId">The new log file id</param>
+		/// <param name="NewState">New state of the jobstep</param>
+		/// <returns>True if the job was updated, false if it was deleted</returns>
+		public async Task<bool> UpdateBatchAsync(IJob Job, SubResourceId BatchId, LogId? NewLogId = null, JobStepBatchState? NewState = null)
+		{
+			using IDisposable Scope = Logger.BeginScope("UpdateBatchAsync({JobId})", Job.Id);
+
+			bool bCheckForBadAgent = true;
+			for (; ; )
+			{
+				// Find the index of the appropriate batch
+				int BatchIdx = Job.Batches.FindIndex(x => x.Id == BatchId);
+				if (BatchIdx == -1)
+				{
+					break;
+				}
+
+				IJobStepBatch Batch = Job.Batches[BatchIdx];
+
+				// If the batch has already been marked as complete, error out
+				if (Batch.State == JobStepBatchState.Complete)
+				{
+					break;
+				}
+
+				// If we're marking the batch as complete before the agent has run everything, mark it to conform
+				JobStepBatchError? NewError = null;
+				if (NewState == JobStepBatchState.Complete && Batch.AgentId != null)
+				{
+					if (Batch.Steps.Any(x => x.State == JobStepState.Waiting || x.State == JobStepState.Ready))
+					{
+						// Mark the batch as incomplete
+						NewError = JobStepBatchError.Incomplete;
+
+						// Find the agent and set the conform flag
+						if (bCheckForBadAgent)
+						{
+							for (; ; )
+							{
+								IAgent? Agent = await Agents.GetAsync(Batch.AgentId.Value);
+								if (Agent == null || Agent.RequestConform)
+								{
+									break;
+								}
+								if (await Agents.TryUpdateSettingsAsync(Agent, bRequestConform: true) != null)
+								{
+									Logger.LogError("Agent {AgentId} did not complete lease; marking for conform", Agent.Id);
+									break;
+								}
+							}
+							bCheckForBadAgent = false;
+						}
+					}
+				}
+
+				// Update the batch state
+				if (await TryUpdateBatchAsync(Job, BatchId, NewLogId, NewState, NewError))
+				{
+					return true;
+				}
+
+				// Update the job
+				IJob? NewJob = await GetJobAsync(Job.Id);
+				if (NewJob == null)
+				{
+					break;
+				}
+				Job = NewJob;
+			}
+			return false;
+		}
+
+		/// <summary>
+		/// Attempts to update the state of a batch
+		/// </summary>
+		/// <param name="Job">Job to update</param>
+		/// <param name="BatchId">Unique id of the batch to update</param>
+		/// <param name="NewLogId">The new log file id</param>
+		/// <param name="NewState">New state of the jobstep</param>
+		/// <param name="NewError">New error state</param>
+		/// <returns>The updated job, otherwise null</returns>
+		public async Task<bool> TryUpdateBatchAsync(IJob Job, SubResourceId BatchId, LogId? NewLogId = null, JobStepBatchState? NewState = null, JobStepBatchError? NewError = null)
+		{
+			IGraph Graph = await GetGraphAsync(Job);
+			if (!await Jobs.TryUpdateBatchAsync(Job, Graph, BatchId, NewLogId, NewState, NewError))
+			{
+				return false;
+			}
+			return true;
+		}
+
+		/// <summary>
+		/// Update a jobstep state
+		/// </summary>
+		/// <param name="Job">Job to update</param>
+		/// <param name="BatchId">Unique id of the batch containing the step</param>
+		/// <param name="StepId">Unique id of the step to update</param>
+		/// <param name="NewState">New state of the jobstep</param>
+		/// <param name="NewOutcome">New outcome of the jobstep</param>
+		/// <param name="NewAbortRequested">New state of abort request</param>
+		/// <param name="NewAbortByUserId">New user that requested the abort</param>
+		/// <param name="NewLogId">New log id for the jobstep</param>
+		/// <param name="NewNotificationTriggerId">New notification trigger id for the jobstep</param>
+		/// <param name="NewRetryByUserId">Whether the step should be retried</param>
+		/// <param name="NewPriority">New priority for this step</param>
+		/// <param name="NewReports">New list of reports</param>
+		/// <param name="NewProperties">Property changes. Any properties with a null value will be removed.</param>
+		/// <returns>True if the job was updated, false if it was deleted in the meantime</returns>
+		public async Task<IJob?> UpdateStepAsync(IJob Job, SubResourceId BatchId, SubResourceId StepId, JobStepState NewState = JobStepState.Unspecified, JobStepOutcome NewOutcome = JobStepOutcome.Unspecified, bool? NewAbortRequested = null, UserId? NewAbortByUserId = null, LogId? NewLogId = null, ObjectId? NewNotificationTriggerId = null, UserId? NewRetryByUserId = null, Priority? NewPriority = null, List<Report>? NewReports = null, Dictionary<string, string?>? NewProperties = null)
+		{
+			using IDisposable Scope = Logger.BeginScope("UpdateStepAsync({JobId})", Job.Id);
+			for (; ;)
+			{
+				if (await TryUpdateStepAsync(Job, BatchId, StepId, NewState, NewOutcome, NewAbortRequested, NewAbortByUserId, NewLogId, NewNotificationTriggerId, NewRetryByUserId, NewPriority, NewReports, NewProperties))
+				{
+					return Job;
+				}
+
+				IJob? NewJob = await GetJobAsync(Job.Id);
+				if(NewJob == null)
+				{
+					return null;
+				}
+
+				Job = NewJob;
+			}
+		}
+
+		/// <summary>
+		/// Update a jobstep state
+		/// </summary>
+		/// <param name="Job">Job to update</param>
+		/// <param name="BatchId">Unique id of the batch containing the step</param>
+		/// <param name="StepId">Unique id of the step to update</param>
+		/// <param name="NewState">New state of the jobstep</param>
+		/// <param name="NewOutcome">New outcome of the jobstep</param>
+		/// <param name="NewAbortRequested">New state for request abort</param>
+		/// <param name="NewAbortByUserId">New name of user that requested the abort</param>
+		/// <param name="NewLogId">New log id for the jobstep</param>
+		/// <param name="NewTriggerId">New trigger id for the jobstep</param>
+		/// <param name="NewRetryByUserId">Whether the step should be retried</param>
+		/// <param name="NewPriority">New priority for this step</param>
+		/// <param name="NewReports">New reports</param>
+		/// <param name="NewProperties">Property changes. Any properties with a null value will be removed.</param>
+		/// <returns>True if the job was updated, false if it was deleted in the meantime</returns>
+		public async Task<bool> TryUpdateStepAsync(IJob Job, SubResourceId BatchId, SubResourceId StepId, JobStepState NewState = JobStepState.Unspecified, JobStepOutcome NewOutcome = JobStepOutcome.Unspecified, bool? NewAbortRequested = null, UserId? NewAbortByUserId = null, LogId? NewLogId = null, ObjectId? NewTriggerId = null, UserId? NewRetryByUserId = null, Priority? NewPriority = null, List<Report>? NewReports = null, Dictionary<string, string?>? NewProperties = null)
+		{
+			using IDisposable Scope = Logger.BeginScope("TryUpdateStepAsync({JobId})", Job.Id);
+
+			// Get the graph for this job
+			IGraph Graph = await GetGraphAsync(Job);
+
+			// Make sure the job state is set to unspecified
+			JobStepState OldState = JobStepState.Unspecified;
+			JobStepOutcome OldOutcome = JobStepOutcome.Unspecified;
+			if (NewState != JobStepState.Unspecified)
+			{
+				if (Job.TryGetStep(BatchId, StepId, out IJobStep? Step))
+				{
+					OldState = Step.State;
+					OldOutcome = Step.Outcome;
+				}
+			}
+
+			// If we're changing the state of a step, capture the label states beforehand so we can update any that change
+			IReadOnlyList<(LabelState, LabelOutcome)>? OldLabelStates = null;
+			if (OldState != NewState || OldOutcome != NewOutcome)
+			{
+				OldLabelStates = Job.GetLabelStates(Graph);
+			}
+
+			// Update the step
+			if (await Jobs.TryUpdateStepAsync(Job, Graph, BatchId, StepId, NewState, NewOutcome, NewAbortRequested, NewAbortByUserId, NewLogId, NewTriggerId, NewRetryByUserId, NewPriority, NewReports, NewProperties))
+			{
+				using IScope DdScope = GlobalTracer.Instance.BuildSpan("TryUpdateStepAsync").StartActive();
+				DdScope.Span.SetTag("JobId", Job.Id.ToString());
+				DdScope.Span.SetTag("BatchId", BatchId.ToString());
+				DdScope.Span.SetTag("StepId", StepId.ToString());
+				
+				if (OldState != NewState || OldOutcome != NewOutcome)
+				{
+					Logger.LogDebug("Transitioned job {JobId}, batch {BatchId}, step {StepId} from {OldState} to {NewState}", Job.Id, BatchId, StepId, OldState, NewState);
+
+					// Send any updates for modified badges
+					if (OldLabelStates != null)
+					{
+						using IScope _ = GlobalTracer.Instance.BuildSpan("Send badge updates").StartActive();
+						IReadOnlyList<(LabelState, LabelOutcome)> NewLabelStates = Job.GetLabelStates(Graph);
+						OnLabelUpdate?.Invoke(Job, OldLabelStates, NewLabelStates);
+						await JobTaskSource.UpdateUgsBadges(Job, Graph, OldLabelStates, NewLabelStates);
+					}
+
+					// Submit the change if auto-submit is enabled
+					if (Job.PreflightChange != 0 && NewState == JobStepState.Completed)
+					{
+						(JobStepState State, JobStepOutcome Outcome) = Job.GetTargetState();
+						if (State == JobStepState.Completed)
+						{
+							if (Job.AutoSubmit && Outcome == JobStepOutcome.Success && Job.AbortedByUserId == null)
+							{
+								Job = await AutoSubmitChangeAsync(Job, Graph);
+							}
+							else if (Job.ClonedPreflightChange != 0)
+							{
+								IStream? Stream = await StreamService.GetCachedStream(Job.StreamId);
+								if (Stream != null)
+								{
+									await DeleteShelvedChangeAsync(Stream.ClusterName, Job.ClonedPreflightChange);
+								}
+							}
+						}
+					}
+
+					// Notify subscribers
+					if (NewState == JobStepState.Completed)
+					{
+						using IScope _ = GlobalTracer.Instance.BuildSpan("Notify subscribers").StartActive();
+						OnJobStepComplete?.Invoke(Job, Graph, BatchId, StepId);
+					}
+
+					// Create any downstream jobs
+					if (NewState == JobStepState.Completed && NewOutcome != JobStepOutcome.Failure)
+					{
+						using IScope _ = GlobalTracer.Instance.BuildSpan("Create downstream jobs").StartActive();
+						for (int Idx = 0; Idx < Job.ChainedJobs.Count; Idx++)
+						{
+							IChainedJob JobTrigger = Job.ChainedJobs[Idx];
+							if (JobTrigger.JobId == null)
+							{
+								JobStepOutcome JobTriggerOutcome = Job.GetTargetOutcome(Graph, JobTrigger.Target);
+								if (JobTriggerOutcome == JobStepOutcome.Success || JobTriggerOutcome == JobStepOutcome.Warnings)
+								{
+									Job = await FireJobTriggerAsync(Job, Graph, JobTrigger);
+								}
+							}
+						}
+					}
+
+					// Update the jobstep ref if it completed
+					if (NewState == JobStepState.Running || NewState == JobStepState.Completed || NewState == JobStepState.Aborted)
+					{
+						using IScope _ = GlobalTracer.Instance.BuildSpan("Update job step ref").StartActive();
+						if (Job.TryGetBatch(BatchId, out IJobStepBatch? Batch) && Batch.TryGetStep(StepId, out IJobStep? Step) && Step.StartTimeUtc != null)
+						{
+							await JobStepRefs.UpdateAsync(Job, Batch, Step, Graph);
+						}
+					}
+
+					// Update any issues that depend on this step
+					if (NewState == JobStepState.Completed && Job.UpdateIssues)
+					{
+						if (IssueService != null)
+						{
+							using IScope _ = GlobalTracer.Instance.BuildSpan("Update issues (V2)").StartActive();
+							try
+							{
+								await IssueService.UpdateCompleteStep(Job, Graph, BatchId, StepId);
+							}
+							catch(Exception Ex)
+							{
+								Logger.LogError(Ex, "Exception while updating issue service for job {JobId}:{BatchId}:{StepId}: {Message}", Job.Id, BatchId, StepId, Ex.Message);
+							}
+						}
+					}
+				}
+
+				using (IScope DispatchScope = GlobalTracer.Instance.BuildSpan("Update queued jobs").StartActive())
+				{
+					// Notify the dispatch service that the job has changed
+					JobTaskSource.UpdateQueuedJob(Job, Graph);
+				}
+
+				return true;
+			}
+
+			return false;
+		}
+
+		/// <summary>
+		/// Submit the given change for a preflight
+		/// </summary>
+		/// <param name="Job">The job being run</param>
+		/// <param name="Graph">Graph for the job</param>
+		/// <returns></returns>
+		[SuppressMessage("Design", "CA1031:Do not catch general exception types", Justification = "<Pending>")]
+		private async Task<IJob> AutoSubmitChangeAsync(IJob Job, IGraph Graph)
+		{
+			int? Change;
+			string Message;
+			try
+			{
+				IStream? Stream = await StreamService.GetCachedStream(Job.StreamId);
+				if (Stream != null)
+				{
+					int ClonedPreflightChange = Job.ClonedPreflightChange;
+					if (ClonedPreflightChange == 0)
+					{
+						if (ShouldClonePreflightChange(Job.StreamId))
+						{
+							ClonedPreflightChange = await CloneShelvedChangeAsync(Stream.ClusterName, Job.PreflightChange);
+						}
+						else
+						{
+							ClonedPreflightChange = Job.PreflightChange;
+						}
+					}
+
+					Logger.LogInformation("Updating description for {ClonedPreflightChange}", ClonedPreflightChange);
+
+					ChangeDetails Details = await PerforceService.GetChangeDetailsAsync(Stream.ClusterName, Stream.Name, ClonedPreflightChange, null);
+					await PerforceService.UpdateChangelistDescription(Stream.ClusterName, ClonedPreflightChange, Details.Description.TrimEnd() + $"\n#preflight {Job.Id}");
+
+					Logger.LogInformation("Submitting change {Change} (through {ChangeCopy}) after successful completion of {JobId}", Job.PreflightChange, ClonedPreflightChange, Job.Id);
+					(Change, Message) = await PerforceService.SubmitShelvedChangeAsync(Stream.ClusterName, ClonedPreflightChange, Job.PreflightChange);
+
+					Logger.LogInformation("Attempt to submit {Change} (through {ChangeCopy}): {Message}", Job.PreflightChange, ClonedPreflightChange, Message);
+
+					if (ShouldClonePreflightChange(Job.StreamId))
+					{
+						if (Change != null && Job.ClonedPreflightChange != 0)
+						{
+							await DeleteShelvedChangeAsync(Stream.ClusterName, Job.PreflightChange);
+						}
+					}
+					else
+					{
+						if (Change != null && Job.PreflightChange != 0)
+						{
+							await DeleteShelvedChangeAsync(Stream.ClusterName, Job.PreflightChange);
+						}
+					}
+				}
+				else
+				{
+					(Change, Message) = ((int?)null, "Stream not found");
+				}
+			}
+			catch (Exception Ex)
+			{
+				(Change, Message) = ((int?)null, "Internal error");
+				Logger.LogError(Ex, "Unable to submit shelved change");
+			}
+
+			for (; ; )
+			{
+				if (await Jobs.TryUpdateJobAsync(Job, Graph, AutoSubmitChange: Change, AutoSubmitMessage: Message))
+				{
+					return Job;
+				}
+
+				IJob? NewJob = await GetJobAsync(Job.Id);
+				if (NewJob == null)
+				{
+					return Job;
+				}
+
+				Job = NewJob;
+			}
+		}
+
+		/// <summary>
+		/// Clone a shelved changelist for running a preflight
+		/// </summary>
+		/// <param name="ClusterName">Name of the Perforce cluster</param>
+		/// <param name="Change">The changelist to clone</param>
+		/// <returns></returns>
+		private async Task<int> CloneShelvedChangeAsync(string ClusterName, int Change)
+		{
+			int ClonedChange;
+			try
+			{
+				ClonedChange = await PerforceService.DuplicateShelvedChangeAsync(ClusterName, Change);
+				Logger.LogInformation("CL {Change} was duplicated into {ClonedChange}", Change, ClonedChange);
+			}
+			catch (Exception Ex)
+			{
+				Logger.LogError(Ex, "Unable to CL {Change} for preflight: {Message}", Change, Ex.Message);
+				throw;
+			}
+			return ClonedChange;
+		}
+
+		/// <summary>
+		/// Deletes a shelved changelist, catching and logging any exceptions
+		/// </summary>
+		/// <param name="ClusterName"></param>
+		/// <param name="Change">The changelist to delete</param>
+		/// <returns>True if the change was deleted successfully, false otherwise</returns>
+		[SuppressMessage("Design", "CA1031:Do not catch general exception types", Justification = "<Pending>")]
+		private async Task<bool> DeleteShelvedChangeAsync(string ClusterName, int Change)
+		{
+			Logger.LogInformation("Removing shelf {Change}", Change);
+			try
+			{
+				await PerforceService.DeleteShelvedChangeAsync(ClusterName, Change);
+				return true;
+			}
+			catch (Exception Ex)
+			{
+				Logger.LogError(Ex, "Unable to delete shelved change {Change}", Change);
+				return false;
+			}
+		}
+
+		/// <summary>
+		/// Fires a trigger for a chained job
+		/// </summary>
+		/// <param name="Job">The job containing the trigger</param>
+		/// <param name="Graph">Graph for the job containing the trigger</param>
+		/// <param name="JobTrigger">The trigger object to fire</param>
+		/// <returns>New job instance</returns>
+		private async Task<IJob> FireJobTriggerAsync(IJob Job, IGraph Graph, IChainedJob JobTrigger)
+		{
+			for (; ; )
+			{
+				// Update the job
+				JobId ChainedJobId = JobId.GenerateNewId();
+				if (await Jobs.TryUpdateJobAsync(Job, Graph, JobTrigger: new KeyValuePair<TemplateRefId, JobId>(JobTrigger.TemplateRefId, ChainedJobId)))
+				{
+					IStream? Stream = await StreamService.GetStreamAsync(Job.StreamId);
+					if(Stream == null)
+					{
+						Logger.LogWarning("Cannot find stream {StreamId} for downstream job", Job.StreamId);
+						break;
+					}
+
+					TemplateRef? TemplateRef;
+					if (!Stream.Templates.TryGetValue(JobTrigger.TemplateRefId, out TemplateRef))
+					{
+						Logger.LogWarning("Cannot find template {TemplateRefId} in stream {StreamId}", JobTrigger.TemplateRefId, Job.StreamId);
+						break;
+					}
+
+					ITemplate? Template = await TemplateCollection.GetAsync(TemplateRef.Hash);
+					if (Template == null)
+					{
+						Logger.LogWarning("Cannot find template {TemplateHash}", TemplateRef.Hash);
+						break;
+					}
+
+					IGraph TriggerGraph = await Graphs.AddAsync(Template);
+					Logger.LogInformation("Creating downstream job {ChainedJobId} from job {JobId}", ChainedJobId, Job.Id);
+
+					await CreateJobAsync(ChainedJobId, Stream, JobTrigger.TemplateRefId, TemplateRef.Hash, TriggerGraph, TemplateRef.Name, Job.Change, Job.CodeChange, Job.PreflightChange, Job.ClonedPreflightChange, Job.StartedByUserId, Template.Priority, null, Job.UpdateIssues, TemplateRef.ChainedJobs, false, false, TemplateRef.NotificationChannel, TemplateRef.NotificationChannelFilter, Template.Arguments);
+					break;
+				}
+
+				// Fetch the job again
+				IJob? NewJob = await Jobs.GetAsync(Job.Id);
+				if(NewJob == null)
+				{
+					break;
+				}
+				Job = NewJob;
+			}
+			return Job;
+		}
+
+		/// <summary>
+		/// Determines if the user is authorized to perform an action on a particular stream
+		/// </summary>
+		/// <param name="Acl">The ACL to check</param>
+		/// <param name="StreamId">The stream containing this job</param>
+		/// <param name="Action">The action being performed</param>
+		/// <param name="User">The principal to authorize</param>
+		/// <param name="Cache">Cache of stream permissions</param>
+		/// <returns>True if the action is authorized</returns>
+		private Task<bool> AuthorizeAsync(Acl? Acl, StreamId StreamId, AclAction Action, ClaimsPrincipal User, StreamPermissionsCache? Cache)
+		{
+			// Do the regular authorization
+			bool? Result = Acl?.Authorize(Action, User);
+			if (Result == null)
+			{
+				return StreamService.AuthorizeAsync(StreamId, Action, User, Cache);
+			}
+			else
+			{
+				return Task.FromResult(Result.Value);
+			}
+		}
+
+		/// <summary>
+		/// Determines if the user is authorized to perform an action on a particular stream
+		/// </summary>
+		/// <param name="Job">The job to check</param>
+		/// <param name="Action">The action being performed</param>
+		/// <param name="User">The principal to authorize</param>
+		/// <param name="Cache">Cache of stream permissions</param>
+		/// <returns>True if the action is authorized</returns>
+		public Task<bool> AuthorizeAsync(IJob Job, AclAction Action, ClaimsPrincipal User, StreamPermissionsCache? Cache)
+		{
+			return AuthorizeAsync(Job.Acl, Job.StreamId, Action, User, Cache);
+		}
+
+		/// <summary>
+		/// Determines if the user is authorized to perform an action on a particular stream
+		/// </summary>
+		/// <param name="JobId">The job to check</param>
+		/// <param name="Action">The action being performed</param>
+		/// <param name="User">The principal to authorize</param>
+		/// <param name="Cache">Cache of job permissions</param>
+		/// <returns>True if the action is authorized</returns>
+		public async Task<bool> AuthorizeAsync(JobId JobId, AclAction Action, ClaimsPrincipal User, JobPermissionsCache? Cache)
+		{
+			IJobPermissions? Permissions;
+			if (Cache == null)
+			{
+				Permissions = await GetJobPermissionsAsync(JobId);
+			}
+			else if (!Cache.Jobs.TryGetValue(JobId, out Permissions))
+			{
+				Permissions = await GetJobPermissionsAsync(JobId);
+				Cache.Jobs.Add(JobId, Permissions);
+			}
+			return Permissions != null && await AuthorizeAsync(Permissions.Acl, Permissions.StreamId, Action, User, Cache);
+		}
+
+		/// <summary>
+		/// Determines if the user is authorized to perform an action on a particular stream
+		/// </summary>
+		/// <param name="Job">The job to check</param>
+		/// <param name="User">The principal to authorize</param>
+		/// <returns>True if the action is authorized</returns>
+		public static bool AuthorizeSession(IJob Job, ClaimsPrincipal User)
+		{
+			return Job.Batches.Any(x => AuthorizeSession(x, User));
+		}
+
+		/// <summary>
+		/// Determines if the user is authorized to perform an action on a particular stream
+		/// </summary>
+		/// <param name="Batch">The batch being executed</param>
+		/// <param name="User">The principal to authorize</param>
+		/// <returns>True if the action is authorized</returns>
+		public static bool AuthorizeSession(IJobStepBatch Batch, ClaimsPrincipal User)
+		{
+			if(Batch.SessionId != null)
+			{
+				foreach (Claim Claim in User.Claims)
+				{
+					if(Claim.Type == HordeClaimTypes.AgentSessionId)
+					{
+						ObjectId SessionId;
+						if (ObjectId.TryParse(Claim.Value, out SessionId) && SessionId == Batch.SessionId.Value)
+						{
+							return true;
+						}
+					}
+				}
+			}
+			return false;
+		}
+		
+		/// <summary>
+		/// Get a list of each batch's session ID formatted as a string for debugging purposes
+		/// </summary>
+		/// <param name="Job">The job to list</param>
+		/// <returns>List of session IDs</returns>
+		public static string GetAllBatchSessionIds(IJob Job)
+		{
+			return string.Join(",", Job.Batches.Select(b => b.SessionId));
+		}
+	}
+}
diff --git a/Engine/Source/Programs/Horde/Horde.Build/Jobs/Impl/JobStepRefCollection.cs b/Engine/Source/Programs/Horde/Horde.Build/Jobs/Impl/JobStepRefCollection.cs
new file mode 100644
--- /dev/null
+++ b/Engine/Source/Programs/Horde/Horde.Build/Jobs/Impl/JobStepRefCollection.cs
@@ -0,0 +1,142 @@
+// Copyright Epic Games, Inc. All Rights Reserved.
+
+using HordeCommon;
+using HordeServer.Models;
+using HordeServer.Services;
+using MongoDB.Bson;
+using MongoDB.Bson.Serialization.Attributes;
+using MongoDB.Driver;
+using System;
+using System.Collections.Generic;
+using System.Linq;
+using System.Threading.Tasks;
+using HordeServer.Utilities;
+
+namespace HordeServer.Collections.Impl
+{
+	using LogId = ObjectId<ILogFile>;
+	using PoolId = StringId<IPool>;
+	using StreamId = StringId<IStream>;
+	using TemplateRefId = StringId<TemplateRef>;
+
+	/// <summary>
+	/// Collection of JobStepRef documents
+	/// </summary>
+	public class JobStepRefCollection : IJobStepRefCollection
+	{
+		class JobStepRef : IJobStepRef
+		{
+			[BsonId]
+			public JobStepRefId Id { get; set; }
+
+			public string JobName { get; set; } = "Unknown";
+
+			public string Name { get; set; }
+
+			public StreamId StreamId { get; set; }
+			public TemplateRefId TemplateId { get; set; }
+			public int Change { get; set; }
+			public LogId? LogId { get; set; }
+			public PoolId? PoolId { get; set; }
+			public AgentId? AgentId { get; set; }
+			public JobStepOutcome? Outcome { get; set; }
+			public float BatchWaitTime { get; set; }
+			public float BatchInitTime { get; set; }
+
+			[BsonIgnoreIfNull]
+			public DateTimeOffset? StartTime { get; set; }
+
+			[BsonIgnoreIfNull]
+			public DateTime? StartTimeUtc { get; set; }
+
+			[BsonIgnoreIfNull]
+			public DateTimeOffset? FinishTime { get; set; }
+
+			[BsonIgnoreIfNull]
+			public DateTime? FinishTimeUtc { get; set; }
+
+			DateTime IJobStepRef.StartTimeUtc => StartTimeUtc ?? StartTime?.UtcDateTime ?? default;
+			DateTime? IJobStepRef.FinishTimeUtc => FinishTimeUtc ?? FinishTime?.UtcDateTime;
+			string IJobStepRef.NodeName => Name;
+
+			public JobStepRef(JobStepRefId Id, string JobName, string NodeName, StreamId StreamId, TemplateRefId TemplateId, int Change, LogId? LogId, PoolId? PoolId, AgentId? AgentId, JobStepOutcome? Outcome, float BatchWaitTime, float BatchInitTime, DateTime StartTimeUtc, DateTime? FinishTimeUtc)
+			{
+				this.Id = Id;
+				this.JobName = JobName;
+				this.Name = NodeName;
+				this.StreamId = StreamId;
+				this.TemplateId = TemplateId;
+				this.Change = Change;
+				this.LogId = LogId;
+				this.PoolId = PoolId;
+				this.AgentId = AgentId;
+				this.Outcome = Outcome;
+				this.BatchWaitTime = BatchWaitTime;
+				this.BatchInitTime = BatchInitTime;
+				this.StartTime = StartTimeUtc;
+				this.StartTimeUtc = StartTimeUtc;
+				this.FinishTime = FinishTimeUtc;
+				this.FinishTimeUtc = FinishTimeUtc;
+			}
+		}
+
+		IMongoCollection<JobStepRef> JobStepRefs;
+
+		/// <summary>
+		/// Constructor
+		/// </summary>
+		/// <param name="DatabaseService">The database service instance</param>
+		public JobStepRefCollection(DatabaseService DatabaseService)
+		{
+			JobStepRefs = DatabaseService.GetCollection<JobStepRef>("JobStepRefs");
+
+			if (!DatabaseService.ReadOnlyMode)
+			{
+				JobStepRefs.Indexes.CreateOne(new CreateIndexModel<JobStepRef>(Builders<JobStepRef>.IndexKeys.Ascending(x => x.StreamId).Ascending(x => x.TemplateId).Ascending(x => x.Name).Descending(x => x.Change)));
+			}
+		}
+
+		/// <inheritdoc/>
+		public async Task<IJobStepRef> InsertOrReplaceAsync(JobStepRefId Id, string JobName, string StepName, StreamId StreamId, TemplateRefId TemplateId, int Change, LogId? LogId, PoolId? PoolId, AgentId? AgentId, JobStepOutcome? Outcome, float WaitTime, float InitTime, DateTime StartTimeUtc, DateTime? FinishTimeUtc)
+		{
+			JobStepRef NewJobStepRef = new JobStepRef(Id, JobName, StepName, StreamId, TemplateId, Change, LogId, PoolId, AgentId, Outcome, WaitTime, InitTime, StartTimeUtc, FinishTimeUtc);
+			await JobStepRefs.ReplaceOneAsync(Builders<JobStepRef>.Filter.Eq(x => x.Id, NewJobStepRef.Id), NewJobStepRef, new ReplaceOptions { IsUpsert = true });
+			return NewJobStepRef;
+		}
+
+		/// <inheritdoc/>
+		public async Task<List<IJobStepRef>> GetStepsForNodeAsync(StreamId StreamId, TemplateRefId TemplateId, string NodeName, int? Change, bool IncludeFailed, int Count)
+		{
+			// Find all the steps matching the given criteria
+			FilterDefinitionBuilder<JobStepRef> FilterBuilder = Builders<JobStepRef>.Filter;
+
+			FilterDefinition<JobStepRef> Filter = FilterDefinition<JobStepRef>.Empty;
+			Filter &= FilterBuilder.Eq(x => x.StreamId, StreamId);
+			Filter &= FilterBuilder.Eq(x => x.TemplateId, TemplateId);
+			Filter &= FilterBuilder.Eq(x => x.Name, NodeName);
+			if (Change != null)
+			{
+				Filter &= FilterBuilder.Lte(x => x.Change, Change.Value);
+			}
+			if (!IncludeFailed)
+			{
+				Filter &= FilterBuilder.Ne(x => x.Outcome, JobStepOutcome.Failure);
+			}
+
+			List<JobStepRef> Steps = await JobStepRefs.Find(Filter).SortByDescending(x => x.Change).ThenByDescending(x => x.StartTime).Limit(Count).ToListAsync();
+			return Steps.ConvertAll<IJobStepRef>(x => x);
+		}
+
+		/// <inheritdoc/>
+		public async Task<IJobStepRef?> GetPrevStepForNodeAsync(StreamId StreamId, TemplateRefId TemplateId, string NodeName, int Change)
+		{
+			return await JobStepRefs.Find(x => x.StreamId == StreamId && x.TemplateId == TemplateId && x.Name == NodeName && x.Change < Change && x.Outcome != null).SortByDescending(x => x.Change).FirstOrDefaultAsync();
+		}
+
+		/// <inheritdoc/>
+		public async Task<IJobStepRef?> GetNextStepForNodeAsync(StreamId StreamId, TemplateRefId TemplateId, string NodeName, int Change)
+		{
+			return await JobStepRefs.Find(x => x.StreamId == StreamId && x.TemplateId == TemplateId && x.Name == NodeName && x.Change > Change && x.Outcome != null).SortBy(x => x.Change).FirstOrDefaultAsync();
+		}
+	}
+}
diff --git a/Engine/Source/Programs/Horde/Horde.Build/Jobs/Impl/JobTaskSource.cs b/Engine/Source/Programs/Horde/Horde.Build/Jobs/Impl/JobTaskSource.cs
new file mode 100644
--- /dev/null
+++ b/Engine/Source/Programs/Horde/Horde.Build/Jobs/Impl/JobTaskSource.cs
@@ -0,0 +1,1101 @@
+// Copyright Epic Games, Inc. All Rights Reserved.
+
+using EpicGames.Core;
+using EpicGames.Horde.Common;
+using Google.Protobuf;
+using Google.Protobuf.Reflection;
+using Google.Protobuf.WellKnownTypes;
+using HordeServer.Api;
+using HordeServer.Collections;
+using HordeCommon;
+using HordeServer.Models;
+using HordeCommon.Rpc;
+using HordeCommon.Rpc.Tasks;
+using HordeServer.Tasks;
+using HordeServer.Utilities;
+using Microsoft.Extensions.Hosting;
+using Microsoft.Extensions.Logging;
+using Microsoft.Extensions.Options;
+using MongoDB.Bson;
+using MongoDB.Driver;
+using System;
+using System.Collections.Generic;
+using System.Diagnostics;
+using System.Diagnostics.CodeAnalysis;
+using System.Linq;
+using System.Runtime.CompilerServices;
+using System.Text;
+using System.Threading;
+using System.Threading.Tasks;
+using HordeServer.Services;
+using EpicGames.Horde.Compute;
+
+namespace HordeServer.Tasks.Impl
+{
+	using JobId = ObjectId<IJob>;
+	using LeaseId = ObjectId<ILease>;
+	using LogId = ObjectId<ILogFile>;
+	using PoolId = StringId<IPool>;
+	using StreamId = StringId<IStream>;
+
+	/// <summary>
+	/// Background service to dispatch pending work to agents in priority order.
+	/// </summary>
+	public sealed class JobTaskSource : TaskSourceBase<ExecuteJobTask>, IDisposable
+	{
+		/// <inheritdoc/>
+		public override string Type => "Job";
+
+		/// <summary>
+		/// An item in the queue to be executed
+		/// </summary>
+		class QueueItem
+		{
+			/// <summary>
+			/// The stream for this job
+			/// </summary>
+			public IStream Stream;
+
+			/// <summary>
+			/// The job instance
+			/// </summary>
+			public IJob Job;
+
+			/// <summary>
+			/// Index of the batch within this job to be executed
+			/// </summary>
+			public int BatchIdx;
+
+			/// <summary>
+			/// The pool of machines to allocate from
+			/// </summary>
+			public PoolId PoolId;
+
+			/// <summary>
+			/// The type of workspace that this item should run in
+			/// </summary>
+			public AgentWorkspace Workspace;
+
+			/// <summary>
+			/// Task for creating a lease and assigning to a waiter
+			/// </summary>
+			public Task? AssignTask;
+
+			/// <summary>
+			/// Accessor for the batch referenced by this item
+			/// </summary>
+			public IJobStepBatch Batch
+			{
+				get { return Job.Batches[BatchIdx]; }
+			}
+
+			/// <summary>
+			/// Returns an identifier describing this unique batch
+			/// </summary>
+			public (JobId, SubResourceId) Id
+			{
+				get { return (Job.Id, Batch.Id); }
+			}
+
+			/// <summary>
+			/// Constructor
+			/// </summary>
+			/// <param name="Stream">The stream containing this job</param>
+			/// <param name="Job">The job instance</param>
+			/// <param name="BatchIdx">The batch index to execute</param>
+			/// <param name="PoolId">Unique id of the pool of machines to allocate from</param>
+			/// <param name="Workspace">The workspace that this job should run in</param>
+			public QueueItem(IStream Stream, IJob Job, int BatchIdx, PoolId PoolId, AgentWorkspace Workspace)
+			{
+				this.Stream = Stream;
+				this.Job = Job;
+				this.BatchIdx = BatchIdx;
+				this.PoolId = PoolId;
+				this.Workspace = Workspace;
+			}
+		}
+
+		/// <summary>
+		/// Comparer for items in the queue
+		/// </summary>
+		class QueueItemComparer : IComparer<QueueItem>
+		{
+			/// <summary>
+			/// Compare two items
+			/// </summary>
+			/// <param name="X">First item to compare</param>
+			/// <param name="Y">Second item to compare</param>
+			/// <returns>Negative value if X is a higher priority than Y</returns>
+			public int Compare([AllowNull] QueueItem X, [AllowNull] QueueItem Y)
+			{
+				if (X == null)
+				{
+					return (Y == null) ? 0 : -1;
+				}
+				else if (Y == null)
+				{
+					return 1;
+				}
+
+				int Delta = Y.Batch.SchedulePriority - X.Batch.SchedulePriority;
+				if (Delta == 0)
+				{
+					Delta = X.Job.Id.CompareTo(Y.Job.Id);
+					if (Delta == 0)
+					{
+						Delta = (int)X.Batch.Id.Value - (int)Y.Batch.Id.Value;
+					}
+				}
+				return Delta;
+			}
+		}
+
+		/// <summary>
+		/// Information about an agent waiting for work
+		/// </summary>
+		class QueueWaiter
+		{
+			/// <summary>
+			/// The agent performing the wait
+			/// </summary>
+			public IAgent Agent { get; }
+
+			/// <summary>
+			/// Task to wait for a lease to be assigned
+			/// </summary>
+			public Task<AgentLease?> Task => LeaseSource.Task;
+
+			/// <summary>
+			/// Completion source for the waiting agent. If a new queue item becomes available, the result will be passed through this.
+			/// </summary>
+			public TaskCompletionSource<AgentLease?> LeaseSource { get; } = new TaskCompletionSource<AgentLease?>();
+
+			/// <summary>
+			/// Constructor
+			/// </summary>
+			/// <param name="Agent">The agent waiting for a task</param>
+			public QueueWaiter(IAgent Agent)
+			{
+				this.Agent = Agent;
+			}
+		}
+
+		/// <summary>
+		/// The database service instance
+		/// </summary>
+		DatabaseService DatabaseService;
+
+		/// <summary>
+		/// The stream service instance
+		/// </summary>
+		StreamService StreamService;
+
+		/// <summary>
+		/// The log file service instance
+		/// </summary>
+		ILogFileService LogFileService;
+
+		/// <summary>
+		/// Collection of agent documents
+		/// </summary>
+		IAgentCollection AgentsCollection;
+
+		/// <summary>
+		/// Collection of job documents
+		/// </summary>
+		IJobCollection Jobs;
+
+		/// <summary>
+		/// Collection of jobstepref documents
+		/// </summary>
+		IJobStepRefCollection JobStepRefs;
+
+		/// <summary>
+		/// Collection of graph documents
+		/// </summary>
+		IGraphCollection Graphs;
+
+		/// <summary>
+		/// Collection of pool documents
+		/// </summary>
+		IPoolCollection PoolCollection;
+
+		/// <summary>
+		/// Collection of UGS metadata documents
+		/// </summary>
+		IUgsMetadataCollection UgsMetadataCollection;
+
+		/// <summary>
+		/// The Perforce load balancer
+		/// </summary>
+		PerforceLoadBalancer PerforceLoadBalancer;
+
+		/// <summary>
+		/// The application lifetime
+		/// </summary>
+		IHostApplicationLifetime ApplicationLifetime;
+
+		/// <summary>
+		/// Settings instance
+		/// </summary>
+		IOptionsMonitor<ServerSettings> Settings;
+
+		/// <summary>
+		/// Log writer
+		/// </summary>
+		ILogger<JobTaskSource> Logger;
+
+		/// <summary>
+		/// Object used for ensuring mutual exclusion to the queues
+		/// </summary>
+		object LockObject = new object();
+
+		/// <summary>
+		/// List of items waiting to be executed
+		/// </summary>
+		SortedSet<QueueItem> Queue = new SortedSet<QueueItem>(new QueueItemComparer());
+
+		/// <summary>
+		/// Map from batch id to the corresponding queue item
+		/// </summary>
+		Dictionary<(JobId, SubResourceId), QueueItem> BatchIdToQueueItem = new Dictionary<(JobId, SubResourceId), QueueItem>();
+
+		/// <summary>
+		/// Set of long-poll tasks waiting to be satisfied 
+		/// </summary>
+		HashSet<QueueWaiter> Waiters = new HashSet<QueueWaiter>();
+
+		/// <summary>
+		/// During a background queue refresh operation, any updated batches are added to this dictionary for merging into the updated queue.
+		/// </summary>
+		List<QueueItem>? NewQueueItemsDuringUpdate;
+
+		/// <summary>
+		/// Cache of pools
+		/// </summary>
+		Dictionary<PoolId, IPool> CachedPoolIdToInstance = new Dictionary<PoolId, IPool>();
+
+		/// <summary>
+		/// Cache of stream objects. Used to resolve agent types.
+		/// </summary>
+		private Dictionary<StreamId, IStream> Streams = new Dictionary<StreamId, IStream>();
+
+		/// <summary>
+		/// Background tick registration
+		/// </summary>
+		BackgroundTick Ticker;
+
+		/// <summary>
+		/// Interval between querying the database for jobs to execute
+		/// </summary>
+		static readonly TimeSpan RefreshInterval = TimeSpan.FromSeconds(5.0);
+
+		/// <summary>
+		/// Constructor
+		/// </summary>
+		/// <param name="DatabaseService">The database service instance</param>
+		/// <param name="Agents">The agents collection</param>
+		/// <param name="Jobs">The jobs collection</param>
+		/// <param name="JobStepRefs">Collection of jobstepref documents</param>
+		/// <param name="Graphs">The graphs collection</param>
+		/// <param name="Pools">The pools collection</param>
+		/// <param name="UgsMetadataCollection">Ugs metadata collection</param>
+		/// <param name="StreamService">The stream service instance</param>
+		/// <param name="LogFileService">The log file service instance</param>
+		/// <param name="PerforceLoadBalancer">Perforce load balancer</param>
+		/// <param name="ApplicationLifetime">Application lifetime interface</param>
+		/// <param name="Settings">Settings for the server</param>
+		/// <param name="Logger">Log writer</param>
+		public JobTaskSource(DatabaseService DatabaseService, IAgentCollection Agents, IJobCollection Jobs, IJobStepRefCollection JobStepRefs, IGraphCollection Graphs, IPoolCollection Pools, IUgsMetadataCollection UgsMetadataCollection, StreamService StreamService, ILogFileService LogFileService, PerforceLoadBalancer PerforceLoadBalancer, IHostApplicationLifetime ApplicationLifetime, IOptionsMonitor<ServerSettings> Settings, ILogger<JobTaskSource> Logger)
+		{
+			this.DatabaseService = DatabaseService;
+			this.AgentsCollection = Agents;
+			this.Jobs = Jobs;
+			this.JobStepRefs = JobStepRefs;
+			this.Graphs = Graphs;
+			this.PoolCollection = Pools;
+			this.UgsMetadataCollection = UgsMetadataCollection;
+			this.StreamService = StreamService;
+			this.LogFileService = LogFileService;
+			this.PerforceLoadBalancer = PerforceLoadBalancer;
+			this.ApplicationLifetime = ApplicationLifetime;
+			this.Ticker = new BackgroundTick(TickAsync, RefreshInterval, Logger);
+			this.Settings = Settings;
+			this.Logger = Logger;
+
+			OnLeaseStartedProperties.Add(nameof(ExecuteJobTask.JobId), x => new JobId(x.JobId)).Add(nameof(ExecuteJobTask.BatchId), x => SubResourceId.Parse(x.BatchId)).Add(nameof(ExecuteJobTask.LogId), x => new LogId(x.LogId));
+		}
+
+		/// <inheritdoc/>
+		public async ValueTask DisposeAsync()
+		{
+			await Ticker.DisposeAsync();
+		}
+
+		/// <inheritdoc/>
+		public void Dispose()
+		{
+			Ticker.Dispose();
+		}
+
+		/// <summary>
+		/// Gets an object containing the stats of the queue for diagnostic purposes.
+		/// </summary>
+		/// <returns>Status object</returns>
+		public object GetStatus()
+		{
+			lock (LockObject)
+			{
+				List<object> OutputItems = new List<object>();
+				foreach (QueueItem QueueItem in Queue)
+				{
+					OutputItems.Add(new { JobId = QueueItem.Job.Id.ToString(), BatchId = QueueItem.Batch.Id.ToString(), PoolId = QueueItem.PoolId.ToString(), Workspace = QueueItem.Workspace });
+				}
+
+				List<object> OutputWaiters = new List<object>();
+				foreach (QueueWaiter Waiter in Waiters)
+				{
+					OutputWaiters.Add(new { Id = Waiter.Agent.Id.ToString(), Pools = Waiter.Agent.GetPools().Select(x => x.ToString()).ToList(), Workspaces = Waiter.Agent.Workspaces });
+				}
+
+				return new { Items = OutputItems, Waiters = OutputWaiters };
+			}
+		}
+
+		/// <summary>
+		/// Cancel any pending wait for an agent, allowing it to cycle its session state immediately
+		/// </summary>
+		/// <param name="AgentId">The agent id</param>
+		public void CancelLongPollForAgent(AgentId AgentId)
+		{
+			QueueWaiter? Waiter;
+			lock (LockObject)
+			{
+				Waiter = Waiters.FirstOrDefault(x => x.Agent.Id == AgentId);
+			}
+			if(Waiter != null)
+			{
+				Waiter.LeaseSource.TrySetResult(null);
+			}
+		}
+
+		/// <summary>
+		/// Background task
+		/// </summary>
+		/// <param name="StoppingToken">Token that indicates that the service should shut down</param>
+		/// <returns>Async task</returns>
+		async Task TickAsync(CancellationToken StoppingToken)
+		{
+			// Set the NewBatchIdToQueueItem member, so we capture any updated jobs during the DB query.
+			lock (LockObject)
+			{
+				NewQueueItemsDuringUpdate = new List<QueueItem>();
+			}
+
+			// Query all the current streams
+			List<IStream> StreamsList = await StreamService.GetStreamsAsync();
+			Streams = StreamsList.ToDictionary(x => x.Id, x => x);
+
+			// Find all the pools which are valid (ie. have at least one online agent)
+			DateTime UtcNow = DateTime.UtcNow;
+			List<IAgent> Agents = await AgentsCollection.FindAsync();
+			List<IPool> Pools = await PoolCollection.GetAsync();
+
+			// Find all the pools which are currently online
+			HashSet<PoolId> OnlinePools = new HashSet<PoolId>(Agents.Where(x => x.IsSessionValid(UtcNow) && x.Enabled).SelectMany(x => x.ExplicitPools));
+			foreach (IPool Pool in Pools)
+			{
+				if (Pool.Condition != null && !OnlinePools.Contains(Pool.Id) && Agents.Any(x => x.IsSessionValid(UtcNow) && x.SatisfiesCondition(Pool.Condition) && x.Enabled))
+				{
+					OnlinePools.Add(Pool.Id);
+				}
+			}
+
+			// Find lists of valid pools and online pools
+			HashSet<PoolId> ValidPools = new HashSet<PoolId>(OnlinePools.Union(Agents.Where(x => !x.IsSessionValid(UtcNow) || !x.Enabled).SelectMany(x => x.ExplicitPools)));
+			foreach (IPool Pool in Pools)
+			{
+				if (Pool.Condition != null && !ValidPools.Contains(Pool.Id) && Agents.Any(x => !x.IsSessionValid(UtcNow) && x.SatisfiesCondition(Pool.Condition) && x.Enabled))
+				{
+					ValidPools.Add(Pool.Id);
+				}
+			}
+
+			// Query all the current pools
+			CachedPoolIdToInstance = Pools.ToDictionary(x => x.Id, x => x);
+
+			// New list of queue items
+			SortedSet<QueueItem> NewQueue = new SortedSet<QueueItem>(Queue.Comparer);
+			Dictionary<(JobId, SubResourceId), QueueItem> NewBatchIdToQueueItem = new Dictionary<(JobId, SubResourceId), QueueItem>();
+
+			// Query for a new list of jobs for the queue
+			List<IJob> NewJobs = await Jobs.GetDispatchQueueAsync();
+			for(int Idx = 0; Idx < NewJobs.Count; Idx++)
+			{
+				IJob? NewJob = NewJobs[Idx];
+
+				if (NewJob.GraphHash == null)
+				{
+					Logger.LogError("Job {JobId} has a null graph hash and can't be started.", NewJob.Id);
+					await Jobs.TryRemoveFromDispatchQueueAsync(NewJob);
+					continue;
+				}
+				if (NewJob.AbortedByUserId != null)
+				{
+					Logger.LogError("Job {JobId} was aborted but not removed from dispatch queue", NewJob.Id);
+					await Jobs.TryRemoveFromDispatchQueueAsync(NewJob);
+					continue;
+				}
+
+				// Get the graph for this job
+				IGraph Graph = await Graphs.GetAsync(NewJob.GraphHash);
+
+				// Get the stream. If it fails, skip the whole job.
+				IStream? Stream;
+				if (!Streams.TryGetValue(NewJob.StreamId, out Stream))
+				{
+					NewJob = await Jobs.SkipAllBatchesAsync(NewJob, Graph, JobStepBatchError.UnknownStream);
+					continue;
+				}
+
+				// Update all the batches
+				bool IsRunning = false;
+				for (int BatchIdx = 0; NewJob != null && BatchIdx < NewJob.Batches.Count; BatchIdx++)
+				{
+					IJobStepBatch Batch = NewJob.Batches[BatchIdx];
+					IsRunning |= (Batch.State == JobStepBatchState.Ready || Batch.State == JobStepBatchState.Starting || Batch.State == JobStepBatchState.Running || Batch.State == JobStepBatchState.Stopping);
+
+					// Skip any batches which aren't ready to execute
+					if (Batch.State == JobStepBatchState.Ready)
+					{
+						// Validate the agent type and workspace settings
+						IPool? Pool;
+						if (!Stream.AgentTypes.TryGetValue(Graph.Groups[Batch.GroupIdx].AgentType, out AgentType? AgentType))
+						{
+							NewJob = await SkipBatchAsync(NewJob, BatchIdx, Graph, JobStepBatchError.UnknownAgentType);
+						}
+						else if (!CachedPoolIdToInstance.TryGetValue(AgentType.Pool, out Pool))
+						{
+							NewJob = await SkipBatchAsync(NewJob, BatchIdx, Graph, JobStepBatchError.UnknownPool);
+						}
+						else if (!ValidPools.Contains(AgentType.Pool))
+						{
+							NewJob = await SkipBatchAsync(NewJob, BatchIdx, Graph, JobStepBatchError.NoAgentsInPool);
+						}
+						else if (!OnlinePools.Contains(AgentType.Pool))
+						{
+							NewJob = await SkipBatchAsync(NewJob, BatchIdx, Graph, JobStepBatchError.NoAgentsOnline);
+						}
+						else if (!Stream.TryGetAgentWorkspace(AgentType, out AgentWorkspace? Workspace))
+						{
+							NewJob = await SkipBatchAsync(NewJob, BatchIdx, Graph, JobStepBatchError.UnknownWorkspace);
+						}
+						else
+						{
+							QueueItem NewQueueItem = new QueueItem(Stream, NewJob, BatchIdx, AgentType.Pool, Workspace);
+							NewQueue.Add(NewQueueItem);
+							NewBatchIdToQueueItem[(NewJob.Id, Batch.Id)] = NewQueueItem;
+						}
+					}
+				}
+
+				// Add a warning if a job looks to be idle
+				if (!IsRunning && NewJob != null)
+				{
+					Logger.LogError("Job {JobId} is in dispatch queue but not currently executing", NewJob.Id);
+					await Jobs.TryRemoveFromDispatchQueueAsync(NewJob);
+					continue;
+				}
+			}
+
+			// Update the queue
+			lock (LockObject)
+			{
+				Queue = NewQueue;
+				BatchIdToQueueItem = NewBatchIdToQueueItem;
+
+				// Merge the new queue items with the queue
+				foreach (QueueItem NewQueueItem in NewQueueItemsDuringUpdate)
+				{
+					QueueItem? ExistingQueueItem;
+					if (!NewBatchIdToQueueItem.TryGetValue((NewQueueItem.Job.Id, NewQueueItem.Batch.Id), out ExistingQueueItem))
+					{
+						// Always just add this item
+						Queue.Add(NewQueueItem);
+						BatchIdToQueueItem[NewQueueItem.Id] = NewQueueItem;
+					}
+					else if (NewQueueItem.Job.UpdateIndex > ExistingQueueItem.Job.UpdateIndex)
+					{
+						// Replace the existing item
+						Queue.Remove(ExistingQueueItem);
+						Queue.Add(NewQueueItem);
+						BatchIdToQueueItem[NewQueueItem.Id] = NewQueueItem;
+					}
+				}
+
+				// Clear out the list to capture queue items during an update
+				NewQueueItemsDuringUpdate = null;
+			}
+		}
+
+		private async Task<IJob?> SkipBatchAsync(IJob Job, int BatchIdx, IGraph Graph, JobStepBatchError Reason)
+		{
+			Logger.LogInformation("Skipping batch {BatchIdx} for job {JobId} (reason: {Reason})", BatchIdx, Job.Id, Reason);
+
+			IReadOnlyList<(LabelState, LabelOutcome)> OldLabelStates = Job.GetLabelStates(Graph);
+			IJob? NewJob = await Jobs.SkipBatchAsync(Job, BatchIdx, Graph, Reason);
+			if(NewJob != null)
+			{
+				IReadOnlyList<(LabelState, LabelOutcome)> NewLabelStates = NewJob.GetLabelStates(Graph);
+				await UpdateUgsBadges(NewJob, Graph, OldLabelStates, NewLabelStates);
+			}
+			return NewJob;
+		}
+
+		/// <summary>
+		/// Updates the current state of a job
+		/// </summary>
+		/// <param name="Job">The job that has been updated</param>
+		/// <param name="Graph">Graph for the job</param>
+		/// <returns>Async task</returns>
+		public void UpdateQueuedJob(IJob Job, IGraph Graph)
+		{
+			IStream? Stream;
+			Streams.TryGetValue(Job.StreamId, out Stream);
+			UpdateQueuedJob(Job, Graph, Stream);
+		}
+
+		void AssignAnyQueueItemToWaiter(QueueWaiter Waiter)
+		{
+			lock (Waiters)
+			{
+				foreach(QueueItem Item in BatchIdToQueueItem.Values)
+				{
+					if (TryAssignItemToWaiter(Item, Waiter))
+					{
+						break;
+					}
+				}
+			}
+		}
+
+		/// <summary>
+		/// Attempt to find a waiter that can handle the given queue item
+		/// </summary>
+		/// <param name="Item">The queue item</param>
+		/// <returns></returns>
+		void AssignQueueItemToAnyWaiter(QueueItem Item)
+		{
+			if (Item.AssignTask == null && Item.Batch.SessionId == null)
+			{
+				lock (Waiters)
+				{
+					if (Item.AssignTask == null && Item.Batch.SessionId == null)
+					{
+						foreach (QueueWaiter Waiter in Waiters)
+						{
+							if (TryAssignItemToWaiter(Item, Waiter))
+							{
+								break;
+							}
+						}
+					}
+				}
+			}
+		}
+
+		/// <summary>
+		/// Attempts to assign a queue item to an agent waiting for work
+		/// </summary>
+		/// <param name="Item"></param>
+		/// <param name="Waiter"></param>
+		/// <returns></returns>
+		bool TryAssignItemToWaiter(QueueItem Item, QueueWaiter Waiter)
+		{
+			if (Item.AssignTask == null && Item.Batch.SessionId == null && Waiter.Agent.Enabled && Waiter.Agent.Leases.Count == 0 && Waiter.Agent.IsInPool(Item.PoolId))
+			{
+				Task StartTask = new Task<Task>(() => TryCreateLeaseAsync(Item, Waiter));
+				Task ExecuteTask = StartTask.ContinueWith(Task => Task, TaskScheduler.Default);
+
+				if (Interlocked.CompareExchange(ref Item.AssignTask, ExecuteTask, null) == null)
+				{
+					StartTask.Start(TaskScheduler.Default);
+					return true;
+				}
+			}
+			return false;
+		}
+
+		/// <summary>
+		/// Updates the current state of a job
+		/// </summary>
+		/// <param name="Job">The job that has been updated</param>
+		/// <param name="Graph">Graph for the job</param>
+		/// <param name="Stream">The stream containing the job</param>
+		public void UpdateQueuedJob(IJob Job, IGraph Graph, IStream? Stream)
+		{
+			List<TaskCompletionSource<bool>> CompleteWaiters = new List<TaskCompletionSource<bool>>();
+			lock (LockObject)
+			{
+				for (int BatchIdx = 0; BatchIdx < Job.Batches.Count; BatchIdx++)
+				{
+					IJobStepBatch Batch = Job.Batches[BatchIdx];
+					if (Batch.State == JobStepBatchState.Ready && Stream != null && Batch.AgentId == null)
+					{
+						// Check if this item is already in the list.
+						QueueItem? ExistingItem;
+						if (BatchIdToQueueItem.TryGetValue((Job.Id, Batch.Id), out ExistingItem))
+						{
+							// Make sure this is a newer version of the job. There's no guarantee that this is the latest revision.
+							if (Job.UpdateIndex > ExistingItem.Job.UpdateIndex)
+							{
+								if (Batch.SchedulePriority == ExistingItem.Batch.SchedulePriority)
+								{
+									ExistingItem.Job = Job;
+									ExistingItem.BatchIdx = BatchIdx;
+								}
+								else
+								{
+									RemoveQueueItem(ExistingItem);
+									InsertQueueItem(Stream, Job, BatchIdx, ExistingItem.PoolId, ExistingItem.Workspace);
+								}
+							}
+							continue;
+						}
+
+						// Get the group being executed by this batch
+						INodeGroup Group = Graph.Groups[Batch.GroupIdx];
+
+						// Get the requirements for the new queue item
+						AgentType? AgentType;
+						if (Stream.AgentTypes.TryGetValue(Group.AgentType, out AgentType))
+						{
+							AgentWorkspace? Workspace;
+							if (Stream.TryGetAgentWorkspace(AgentType, out Workspace))
+							{
+								InsertQueueItem(Stream, Job, BatchIdx, AgentType.Pool, Workspace);
+							}
+						}
+					}
+					else
+					{
+						// Check if this item is already in the list. Remove it if it is.
+						QueueItem? ExistingItem;
+						if (BatchIdToQueueItem.TryGetValue((Job.Id, Batch.Id), out ExistingItem))
+						{
+							if (Job.UpdateIndex > ExistingItem.Job.UpdateIndex)
+							{
+								RemoveQueueItem(ExistingItem);
+							}
+						}
+					}
+				}
+			}
+
+			// Awake all the threads that have been assigned new work items. Has do be done outside the lock to prevent continuations running within it (see Waiter.CompletionSource for more info).
+			foreach (TaskCompletionSource<bool> CompleteWaiter in CompleteWaiters)
+			{
+				CompleteWaiter.TrySetResult(true);
+			}
+		}
+
+		/// <inheritdoc/>
+		public override async Task<AgentLease?> AssignLeaseAsync(IAgent Agent, CancellationToken CancellationToken)
+		{
+			QueueWaiter Waiter = new QueueWaiter(Agent);
+			lock (LockObject)
+			{
+				AssignAnyQueueItemToWaiter(Waiter);
+				if (Waiter.LeaseSource.Task.IsCompleted)
+				{
+					return Waiter.Task.Result;
+				}
+				Waiters.Add(Waiter);
+			}
+
+			AgentLease? Lease;
+			using (CancellationToken.Register(() => Waiter.LeaseSource.TrySetResult(null)))
+			{
+				Lease = await Waiter.Task;
+			}
+
+			lock (LockObject)
+			{
+				Waiters.Remove(Waiter);
+			}
+
+			return Lease;
+		}
+
+		/// <inheritdoc/>
+		public override Task CancelLeaseAsync(IAgent Agent, LeaseId LeaseId, ExecuteJobTask Task)
+		{
+			return CancelLeaseAsync(Agent, new JobId(Task.JobId), Task.BatchId.ToSubResourceId());
+		}
+
+		/// <summary>
+		/// Assign a new batch to be executed by the given agent
+		/// </summary>
+		/// <param name="Item">The item to create a lease for</param>
+		/// <param name="Waiter">The agent waiting for work</param>
+		/// <returns>New work to execute</returns>
+		private async Task<AgentLease?> TryCreateLeaseAsync(QueueItem Item, QueueWaiter Waiter)
+		{
+			IJob Job = Item.Job;
+			IJobStepBatch Batch = Job.Batches[Item.BatchIdx];
+			IAgent Agent = Waiter.Agent;
+			Logger.LogDebug("Assigning job {JobId}, batch {BatchId} to waiter (agent {AgentID})", Job.Id, Batch.Id, Agent.Id);
+
+			// Generate a new unique id for the lease
+			LeaseId LeaseId = LeaseId.GenerateNewId();
+
+			// The next time to try assigning to another agent
+			DateTime BackOffTime = DateTime.UtcNow + TimeSpan.FromMinutes(1.0);
+
+			// Try to update the job with this agent id
+			LogId LogId = (await LogFileService.CreateLogFileAsync(Job.Id, Agent.SessionId, LogType.Json)).Id;
+			if (await Jobs.TryAssignLeaseAsync(Item.Job, Item.BatchIdx, Item.PoolId, Agent.Id, Agent.SessionId!.Value, LeaseId, LogId))
+			{
+				// Get the lease name
+				StringBuilder LeaseName = new StringBuilder($"{Item.Stream.Name} - ");
+				if (Job.PreflightChange > 0)
+				{
+					LeaseName.Append((Job.Change > 0) ? $"Preflight CL {Job.PreflightChange} against CL {Job.Change}" : $"Preflight CL {Job.PreflightChange} against latest");
+				}
+				else
+				{
+					LeaseName.Append((Job.Change > 0) ? $"CL {Job.Change}" : "Latest CL");
+				}
+				LeaseName.Append($" - {Job.Name}");
+
+				// Get the global settings
+				Globals Globals = await DatabaseService.GetGlobalsAsync();
+
+				// Encode the payload
+				ExecuteJobTask? Task = await CreateExecuteJobTaskAsync(Item.Stream, Job, Batch, Agent, Item.Workspace, LogId);
+				if (Task != null)
+				{
+					byte[] Payload = Any.Pack(Task).ToByteArray();
+
+					// Create the lease and try to set it on the waiter. If this fails, the waiter has already moved on, and the lease can be cancelled.
+					AgentLease Lease = new AgentLease(LeaseId, LeaseName.ToString(), Job.StreamId, Item.PoolId, LogId, LeaseState.Pending, null, true, Payload);
+					if (Waiter.LeaseSource.TrySetResult(Lease))
+					{
+						Logger.LogDebug("Assigned lease {LeaseId} to agent {AgentId}", LeaseId, Agent.Id);
+						return Lease;
+					}
+				}
+
+				// Cancel the lease
+				Logger.LogDebug("Unable to assign lease {LeaseId} to agent {AgentId}, cancelling", LeaseId, Agent.Id);
+				await CancelLeaseAsync(Waiter.Agent, Job.Id, Batch.Id);
+			}
+			else
+			{
+				// Unable to assign job
+				Logger.LogTrace("Refreshing queue entries for job {JobId}", Job.Id);
+
+				// Get the new copy of the job
+				IJob? NewJob = await Jobs.GetAsync(Job.Id);
+				if (NewJob == null)
+				{
+					lock (LockObject)
+					{
+						List<QueueItem> RemoveItems = Queue.Where(x => x.Job == Job).ToList();
+						foreach (QueueItem RemoveItem in RemoveItems)
+						{
+							RemoveQueueItem(RemoveItem);
+						}
+					}
+				}
+				else
+				{
+					IGraph Graph = await Graphs.GetAsync(NewJob.GraphHash);
+					UpdateQueuedJob(NewJob, Graph);
+				}
+			}
+
+			// Clear out the assignment for this item, and try to reassign it
+			Item.AssignTask = null;
+			AssignQueueItemToAnyWaiter(Item);
+			return null;
+		}
+
+		async Task<ExecuteJobTask?> CreateExecuteJobTaskAsync(IStream Stream, IJob Job, IJobStepBatch Batch, IAgent Agent, AgentWorkspace Workspace, LogId LogId)
+		{
+			// Get the lease name
+			StringBuilder LeaseName = new StringBuilder($"{Stream.Name} - ");
+			if (Job.PreflightChange > 0)
+			{
+				LeaseName.Append((Job.Change > 0) ? $"Preflight CL {Job.PreflightChange} against CL {Job.Change}" : $"Preflight CL {Job.PreflightChange} against latest");
+			}
+			else
+			{
+				LeaseName.Append((Job.Change > 0) ? $"CL {Job.Change}" : "Latest CL");
+			}
+			LeaseName.Append($" - {Job.Name}");
+
+			// Get the global settings
+			Globals Globals = await DatabaseService.GetGlobalsAsync();
+
+			// Encode the payload
+			ExecuteJobTask Task = new ExecuteJobTask();
+			Task.JobId = Job.Id.ToString();
+			Task.BatchId = Batch.Id.ToString();
+			Task.LogId = LogId.ToString();
+			Task.JobName = LeaseName.ToString();
+
+			List<HordeCommon.Rpc.Messages.AgentWorkspace> Workspaces = new List<HordeCommon.Rpc.Messages.AgentWorkspace>();
+
+			PerforceCluster? Cluster = Globals.FindPerforceCluster(Workspace.Cluster);
+			if (Cluster == null)
+			{
+				return null;
+			}
+
+			AgentWorkspace? AutoSdkWorkspace = Agent.GetAutoSdkWorkspace(Cluster);
+			if (AutoSdkWorkspace != null)
+			{
+				if (!await Agent.TryAddWorkspaceMessage(AutoSdkWorkspace, Cluster, PerforceLoadBalancer, Workspaces))
+				{
+					return null;
+				}
+				Task.AutoSdkWorkspace = Workspaces.Last();
+			}
+
+			if (!await Agent.TryAddWorkspaceMessage(Workspace, Cluster, PerforceLoadBalancer, Workspaces))
+			{
+				return null;
+			}
+
+			Task.Workspace = Workspaces.Last();
+			return Task;
+		}
+
+		/// <summary>
+		/// Send any badge updates for this job
+		/// </summary>
+		/// <param name="Job">The job being updated</param>
+		/// <param name="Graph">Graph for the job</param>
+		/// <param name="OldLabelStates">Previous badge states for the job</param>
+		/// <returns>Async task</returns>
+		public async Task UpdateUgsBadges(IJob Job, IGraph Graph, IReadOnlyList<(LabelState, LabelOutcome)> OldLabelStates)
+		{
+			await UpdateUgsBadges(Job, Graph, OldLabelStates, Job.GetLabelStates(Graph));
+		}
+
+		/// <summary>
+		/// Send any badge updates for this job
+		/// </summary>
+		/// <param name="Job">The job being updated</param>
+		/// <param name="Graph">Graph for the job</param>
+		/// <param name="OldLabelStates">Previous badge states for the job</param>
+		/// <param name="NewLabelStates">The new badge states for the job</param>
+		/// <returns>Async task</returns>
+		public async Task UpdateUgsBadges(IJob Job, IGraph Graph, IReadOnlyList<(LabelState, LabelOutcome)> OldLabelStates, IReadOnlyList<(LabelState, LabelOutcome)> NewLabelStates)
+		{
+			if (!Job.ShowUgsBadges || Job.PreflightChange != 0)
+			{
+				return;
+			}
+
+			IReadOnlyDictionary<int, UgsBadgeState> OldStates = Job.GetUgsBadgeStates(Graph, OldLabelStates);
+			IReadOnlyDictionary<int, UgsBadgeState> NewStates = Job.GetUgsBadgeStates(Graph, NewLabelStates);
+
+			// Figure out a list of all the badges that have been modified
+			List<int> UpdateLabels = new List<int>();
+			foreach (KeyValuePair<int, UgsBadgeState> Pair in OldStates)
+			{
+				if (!NewStates.ContainsKey(Pair.Key))
+				{
+					UpdateLabels.Add(Pair.Key);
+				}
+			}
+			foreach (KeyValuePair<int, UgsBadgeState> Pair in NewStates)
+			{
+				if (!OldStates.TryGetValue(Pair.Key, out UgsBadgeState PrevState) || PrevState != Pair.Value)
+				{
+					UpdateLabels.Add(Pair.Key);
+				}
+			}
+
+			// Cached stream for this job
+			IStream? Stream = null;
+
+			// Send all the updates
+			Dictionary<int, IUgsMetadata> MetadataCache = new Dictionary<int, IUgsMetadata>();
+			foreach (int LabelIdx in UpdateLabels)
+			{
+				ILabel Label = Graph.Labels[LabelIdx];
+
+				// Skip if this label has no UGS name.
+				if (Label.UgsName == null)
+				{
+					continue;
+				}
+
+				// Get the new state
+				if (!NewStates.TryGetValue(LabelIdx, out UgsBadgeState NewState))
+				{
+					NewState = UgsBadgeState.Skipped;
+				}
+
+				// Get the stream
+				if (Stream == null)
+				{
+					Stream = await StreamService.GetStreamAsync(Job.StreamId);
+					if (Stream == null)
+					{
+						Logger.LogError("Unable to fetch definition for stream {StreamId}", Job.StreamId);
+						break;
+					}
+				}
+
+				// The changelist number to display the badge for
+				int Change;
+				if (Label.Change == LabelChange.Code)
+				{
+					Change = Job.CodeChange;
+				}
+				else
+				{
+					Change = Job.Change;
+				}
+
+				// Get the current metadata state
+				IUgsMetadata? Metadata;
+				if (!MetadataCache.TryGetValue(Change, out Metadata))
+				{
+					Metadata = await UgsMetadataCollection.FindOrAddAsync(Stream.Name, Change, Label.UgsProject);
+					MetadataCache[Change] = Metadata;
+				}
+
+				// Apply the update
+				Uri LabelUrl = new Uri(Settings.CurrentValue.DashboardUrl, $"job/{Job.Id}?label={LabelIdx}");
+				Logger.LogInformation("Updating state of badge {BadgeName} at {Change} to {NewState} ({LabelUrl})", Label.UgsName, Change, NewState, LabelUrl);
+				Metadata = await UgsMetadataCollection.UpdateBadgeAsync(Metadata, Label.UgsName!, LabelUrl, NewState);
+				MetadataCache[Change] = Metadata;
+			}
+		}
+
+		/// <inheritdoc/>
+		public override async Task OnLeaseFinishedAsync(IAgent Agent, LeaseId LeaseId, ExecuteJobTask Task, LeaseOutcome Outcome, ReadOnlyMemory<byte> Output, ILogger Logger)
+		{
+			if (Outcome != LeaseOutcome.Success)
+			{
+				AgentId AgentId = Agent.Id;
+				JobId JobId = new JobId(Task.JobId);
+				SubResourceId BatchId = Task.BatchId.ToSubResourceId();
+
+				// Update the batch
+				for (; ; )
+				{
+					IJob? Job = await Jobs.GetAsync(JobId);
+					if (Job == null)
+					{
+						break;
+					}
+
+					int BatchIdx = Job.Batches.FindIndex(x => x.Id == BatchId);
+					if (BatchIdx == -1)
+					{
+						break;
+					}
+
+					IJobStepBatch Batch = Job.Batches[BatchIdx];
+					if (Batch.AgentId != AgentId)
+					{
+						break;
+					}
+
+					int RunningStepIdx = Batch.Steps.FindIndex(x => x.State == JobStepState.Running);
+
+					IGraph Graph = await Graphs.GetAsync(Job.GraphHash);
+					if (await Jobs.TryFailBatchAsync(Job, BatchIdx, Graph))
+					{
+						if (Batch.Error != JobStepBatchError.None)
+						{
+							Logger.LogInformation("Failed job {JobId}, batch {BatchId} with error {Error}", Job.Id, Batch.Id, Batch.Error);
+						}
+						if (RunningStepIdx != -1)
+						{
+							await JobStepRefs.UpdateAsync(Job, Batch, Batch.Steps[RunningStepIdx], Graph);
+						}
+						break;
+					}
+				}
+			}
+		}
+
+		/// <summary>
+		/// 
+		/// </summary>
+		/// <param name="Agent"></param>
+		/// <param name="JobId"></param>
+		/// <param name="BatchId"></param>
+		/// <returns></returns>
+		async Task CancelLeaseAsync(IAgent Agent, JobId JobId, SubResourceId BatchId)
+		{
+			Logger.LogDebug("Cancelling lease for job {JobId}, batch {BatchId}", JobId, BatchId);
+
+			// Update the batch
+			for (; ; )
+			{
+				IJob? Job = await Jobs.GetAsync(JobId);
+				if (Job == null)
+				{
+					break;
+				}
+
+				int BatchIdx = Job.Batches.FindIndex(x => x.Id == BatchId);
+				if (BatchIdx == -1)
+				{
+					break;
+				}
+
+				IJobStepBatch Batch = Job.Batches[BatchIdx];
+				if (Batch.AgentId != Agent.Id)
+				{
+					break;
+				}
+
+				if (await Jobs.TryCancelLeaseAsync(Job, BatchIdx))
+				{
+					break;
+				}
+			}
+		}
+
+		/// <summary>
+		/// Inserts an item into the queue
+		/// </summary>
+		/// <param name="Stream">The stream containing the job</param>
+		/// <param name="Job"></param>
+		/// <param name="BatchIdx"></param>
+		/// <param name="PoolId">The pool to use</param>
+		/// <param name="Workspace">The workspace for this item to run in</param>
+		/// <returns></returns>
+		void InsertQueueItem(IStream Stream, IJob Job, int BatchIdx, PoolId PoolId, AgentWorkspace Workspace)
+		{
+			Logger.LogDebug("Adding queued job {JobId}, batch {BatchId} [Pool: {Pool}, Workspace: {Workspace}]", Job.Id, Job.Batches[BatchIdx].Id, PoolId, Workspace.Identifier);
+
+			QueueItem NewItem = new QueueItem(Stream, Job, BatchIdx, PoolId, Workspace);
+			BatchIdToQueueItem[NewItem.Id] = NewItem;
+			Queue.Add(NewItem);
+
+			AssignQueueItemToAnyWaiter(NewItem);
+		}
+
+		/// <summary>
+		/// Removes an item from the queue
+		/// </summary>
+		/// <param name="Item">Item to remove</param>
+		void RemoveQueueItem(QueueItem Item)
+		{
+			Logger.LogDebug("Removing queued job {JobId}, batch {BatchId}", Item.Job.Id, Item.Batch.Id);
+
+			Queue.Remove(Item);
+			BatchIdToQueueItem.Remove(Item.Id);
+		}
+	}
+}
diff --git a/Engine/Source/Programs/Horde/Horde.Build/Jobs/Impl/JobTimingCollection.cs b/Engine/Source/Programs/Horde/Horde.Build/Jobs/Impl/JobTimingCollection.cs
new file mode 100644
--- /dev/null
+++ b/Engine/Source/Programs/Horde/Horde.Build/Jobs/Impl/JobTimingCollection.cs
@@ -0,0 +1,147 @@
+// Copyright Epic Games, Inc. All Rights Reserved.
+
+using Amazon.Runtime.Internal.Util;
+using Amazon.S3.Model.Internal.MarshallTransformations;
+using HordeServer.Models;
+using HordeServer.Services;
+using HordeServer.Utilities;
+using MongoDB.Bson;
+using MongoDB.Bson.Serialization.Attributes;
+using MongoDB.Driver;
+using System;
+using System.Collections.Generic;
+using System.Diagnostics.CodeAnalysis;
+using System.Linq;
+using System.Threading.Tasks;
+
+namespace HordeServer.Collections.Impl
+{
+	using JobId = ObjectId<IJob>;
+
+	/// <summary>
+	/// Concrete implementation of IJobTimingCollection
+	/// </summary>
+	public class JobTimingCollection : IJobTimingCollection
+	{
+		/// <summary>
+		/// Timing information for a particular step
+		/// </summary>
+		class JobStepTimingDocument : IJobStepTiming
+		{
+			public string Name { get; set; }
+			public float? AverageWaitTime { get; set; }
+			public float? AverageInitTime { get; set; }
+			public float? AverageDuration { get; set; }
+
+			[BsonConstructor]
+			private JobStepTimingDocument()
+			{
+				this.Name = null!;
+			}
+
+			public JobStepTimingDocument(string Name, float? AverageWaitTime, float? AverageInitTime, float? AverageDuration)
+			{
+				this.Name = Name;
+				this.AverageWaitTime = AverageWaitTime;
+				this.AverageInitTime = AverageInitTime;
+				this.AverageDuration = AverageDuration;
+			}
+		}
+
+		/// <summary>
+		/// Timing information for a particular job
+		/// </summary>
+		class JobTimingDocument : IJobTiming
+		{
+			public JobId Id { get; set; }
+			public List<JobStepTimingDocument> Steps { get; set; } = new List<JobStepTimingDocument>();
+			public int UpdateIndex { get; set; }
+
+			[BsonIgnore]
+			public Dictionary<string, IJobStepTiming>? NameToStep;
+
+			public bool TryGetStepTiming(string Name, [NotNullWhen(true)] out IJobStepTiming? Timing)
+			{
+				if (NameToStep == null)
+				{
+					NameToStep = new Dictionary<string, IJobStepTiming>();
+					foreach (JobStepTimingDocument Step in Steps)
+					{
+						if (NameToStep.ContainsKey(Step.Name))
+						{
+							Serilog.Log.Logger.Warning("Step {Name} appears twice in job timing document {Id}", Step.Name, Id);
+						}
+						NameToStep[Step.Name] = Step;
+					}
+				}
+				return NameToStep.TryGetValue(Name, out Timing);
+			}
+		}
+
+		/// <summary>
+		/// Collection of timing documents
+		/// </summary>
+		IMongoCollection<JobTimingDocument> Collection;
+
+		/// <summary>
+		/// Constructor
+		/// </summary>
+		/// <param name="DatabaseService">The database service singleton</param>
+		public JobTimingCollection(DatabaseService DatabaseService)
+		{
+			Collection = DatabaseService.GetCollection<JobTimingDocument>("JobTiming");
+		}
+
+		/// <inheritdoc/>
+		public async Task<IJobTiming?> TryAddAsync(JobId JobId, List<JobStepTimingData> Steps)
+		{
+			JobTimingDocument JobTiming = new JobTimingDocument();
+			JobTiming.Id = JobId;
+			JobTiming.Steps.AddRange(Steps.Select(x => new JobStepTimingDocument(x.Name, x.AverageWaitTime, x.AverageInitTime, x.AverageDuration)));
+			JobTiming.UpdateIndex = 1;
+
+			try
+			{
+				await Collection.InsertOneAsync(JobTiming);
+				return JobTiming;
+			}
+			catch (MongoWriteException Ex)
+			{
+				if (Ex.WriteError.Category == ServerErrorCategory.DuplicateKey)
+				{
+					return null;
+				}
+				else
+				{
+					throw;
+				}
+			}
+		}
+
+		/// <inheritdoc/>
+		public async Task<IJobTiming?> TryGetAsync(JobId JobId)
+		{
+			return await Collection.Find(x => x.Id == JobId).FirstOrDefaultAsync();
+		}
+
+		/// <inheritdoc/>
+		public async Task<IJobTiming?> TryAddStepsAsync(IJobTiming JobTiming, List<JobStepTimingData> Steps)
+		{
+			JobTimingDocument JobTimingDocument = (JobTimingDocument)JobTiming;
+			List<JobStepTimingDocument> StepDocuments = Steps.ConvertAll(x => new JobStepTimingDocument(x.Name, x.AverageWaitTime, x.AverageInitTime, x.AverageDuration));
+
+			FilterDefinition<JobTimingDocument> Filter = Builders<JobTimingDocument>.Filter.Expr(x => x.Id == JobTimingDocument.Id && x.UpdateIndex == JobTimingDocument.UpdateIndex);
+			UpdateDefinition<JobTimingDocument> Update = Builders<JobTimingDocument>.Update.Set(x => x.UpdateIndex, JobTimingDocument.UpdateIndex + 1).PushEach(x => x.Steps, StepDocuments);
+
+			UpdateResult Result = await Collection.UpdateOneAsync(Filter, Update);
+			if (Result.ModifiedCount > 0)
+			{
+				JobTimingDocument.NameToStep = null;
+				JobTimingDocument.Steps.AddRange(StepDocuments);
+				JobTimingDocument.UpdateIndex++;
+				return JobTimingDocument;
+			}
+			return null;
+		}
+	}
+}
