Horde: Rename the HordeServer project to Horde.Build.

diff --git a/Engine/Source/Programs/Horde/Horde.Build/Logs/Impl/LogFileService.cs b/Engine/Source/Programs/Horde/Horde.Build/Logs/Impl/LogFileService.cs
new file mode 100644
--- /dev/null
+++ b/Engine/Source/Programs/Horde/Horde.Build/Logs/Impl/LogFileService.cs
@@ -0,0 +1,1608 @@
+// Copyright Epic Games, Inc. All Rights Reserved.
+
+using EpicGames.Core;
+using HordeCommon;
+using HordeServer.Api;
+using HordeServer.Collections;
+using HordeServer.Logs;
+using HordeServer.Models;
+using HordeServer.Storage;
+using HordeServer.Utilities;
+using Microsoft.Extensions.Caching.Memory;
+using Microsoft.Extensions.Logging;
+using MongoDB.Bson;
+using MongoDB.Driver;
+using System;
+using System.Collections.Concurrent;
+using System.Collections.Generic;
+using System.Collections.Immutable;
+using System.Diagnostics;
+using System.Diagnostics.CodeAnalysis;
+using System.Globalization;
+using System.IO;
+using System.Linq;
+using System.Runtime.CompilerServices;
+using System.Security.Claims;
+using System.Text.Json;
+using System.Text;
+using System.Threading.Tasks;
+using System.Threading;
+
+using ILogger = Microsoft.Extensions.Logging.ILogger;
+using Stream = System.IO.Stream;
+using OpenTracing.Util;
+using OpenTracing;
+
+[assembly: InternalsVisibleTo("HordeServerTests")]
+
+namespace HordeServer.Services
+{
+	using JobId = ObjectId<IJob>;
+	using LogId = ObjectId<ILogFile>;
+
+	/// <summary>
+	/// Metadata about a log file
+	/// </summary>
+	public class LogMetadata
+	{
+		/// <summary>
+		/// Length of the log file
+		/// </summary>
+		public long Length { get; set; }
+
+		/// <summary>
+		/// Number of lines in the log file
+		/// </summary>
+		public int MaxLineIndex { get; set; }
+	}
+
+	/// <summary>
+	/// Interface for the log file service
+	/// </summary>
+	public interface ILogFileService
+	{
+		/// <summary>
+		/// Creates a new log
+		/// </summary>
+		/// <param name="JobId">Unique id of the job that owns this log file</param>
+		/// <param name="SessionId">Agent session allowed to update the log</param>
+		/// <param name="Type">Type of events to be stored in the log</param>
+		/// <returns>The new log file document</returns>
+		Task<ILogFile> CreateLogFileAsync(JobId JobId, ObjectId? SessionId, LogType Type);
+
+		/// <summary>
+		/// Gets a logfile by ID
+		/// </summary>
+		/// <param name="LogFileId">Unique id of the log file</param>
+		/// <returns>The logfile document</returns>
+		Task<ILogFile?> GetLogFileAsync(LogId LogFileId);
+
+		/// <summary>
+		/// Gets a logfile by ID, returning a cached copy if available. This should only be used to retrieve constant properties set at creation, such as the session or job it's associated with.
+		/// </summary>
+		/// <param name="LogFileId">Unique id of the log file</param>
+		/// <returns>The logfile document</returns>
+		Task<ILogFile?> GetCachedLogFileAsync(LogId LogFileId);
+
+		/// <summary>
+		/// Returns a list of log files
+		/// </summary>
+		/// <param name="Index">Index of the first result to return</param>
+		/// <param name="Count">Number of results to return</param>
+		/// <returns>List of logfile documents</returns>
+		Task<List<ILogFile>> GetLogFilesAsync(int? Index = null, int? Count = null);
+
+		/// <summary>
+		/// Writes out chunk data and assigns to a file
+		/// </summary>
+		/// <param name="LogFile">The log file</param>
+		/// <param name="Offset">Offset within the file of data</param>
+		/// <param name="LineIndex">Current line index of the data (need not be the starting of the line)</param>
+		/// <param name="Data">the data to add</param>
+		/// <param name="Flush">Whether the current chunk is complete and should be flushed</param>
+		/// <param name="MaxChunkLength">The maximum chunk length. Defaults to 128kb.</param>
+		/// <param name="MaxSubChunkLineCount">Maximum number of lines in each sub-chunk.</param>
+		/// <returns></returns>
+		Task<ILogFile?> WriteLogDataAsync(ILogFile LogFile, long Offset, int LineIndex, ReadOnlyMemory<byte> Data, bool Flush, int MaxChunkLength = 256 * 1024, int MaxSubChunkLineCount = 128);
+
+		/// <summary>
+		/// Gets metadata about the log file
+		/// </summary>
+		/// <param name="LogFile">The log file to query</param>
+		/// <returns>Metadata about the log file</returns>
+		Task<LogMetadata> GetMetadataAsync(ILogFile LogFile);
+
+		/// <summary>
+		/// Creates new log events
+		/// </summary>
+		/// <param name="NewEvents">List of events</param>
+		/// <returns>Async task</returns>
+		Task CreateEventsAsync(List<NewLogEventData> NewEvents);
+
+		/// <summary>
+		/// Find events for a particular log file
+		/// </summary>
+		/// <param name="LogFile">The log file instance</param>
+		/// <param name="Index">Index of the first event to retrieve</param>
+		/// <param name="Count">Number of events to retrieve</param>
+		/// <returns>List of log events</returns>
+		Task<List<ILogEvent>> FindLogEventsAsync(ILogFile LogFile, int? Index = null, int? Count = null);
+
+		/// <summary>
+		/// Adds events to a log span
+		/// </summary>
+		/// <param name="Events">The events to add</param>
+		/// <param name="SpanId">The span id</param>
+		/// <returns>Async task</returns>
+		Task AddSpanToEventsAsync(IEnumerable<ILogEvent> Events, ObjectId SpanId);
+
+		/// <summary>
+		/// Find events for an issue
+		/// </summary>
+		/// <param name="SpanIds">The span ids</param>
+		/// <param name="LogIds">Log ids to include</param>
+		/// <param name="Index">Index within the events for results to return</param>
+		/// <param name="Count">Number of results to return</param>
+		/// <returns>Async task</returns>
+		Task<List<ILogEvent>> FindEventsForSpansAsync(IEnumerable<ObjectId> SpanIds, LogId[]? LogIds, int Index, int Count);
+
+		/// <summary>
+		/// Gets the data for an event
+		/// </summary>
+		/// <param name="LogFile">The log file instance</param>
+		/// <param name="LineIndex">Index of the line in the file</param>
+		/// <param name="LineCount">Number of lines in the event</param>
+		/// <returns>New event data instance</returns>
+		Task<ILogEventData> GetEventDataAsync(ILogFile LogFile, int LineIndex, int LineCount);
+
+		/// <summary>
+		/// Gets lines from the given log 
+		/// </summary>
+		/// <param name="LogFile">The log file</param>
+		/// <param name="Offset">Offset of the data to return</param>
+		/// <param name="Length">Length of the data to return</param>
+		/// <returns>Data for the requested range</returns>
+		Task<Stream> OpenRawStreamAsync(ILogFile LogFile, long Offset, long Length);
+
+		/// <summary>
+		/// Gets the offset of the given line number
+		/// </summary>
+		/// <param name="LogFile">The log file to search</param>
+		/// <param name="LineIdx">The line index to retrieve the offset for</param>
+		/// <returns>The actual clamped line number and offset</returns>
+		Task<(int, long)> GetLineOffsetAsync(ILogFile LogFile, int LineIdx);
+
+		/// <summary>
+		/// Search for the specified text in a log file
+		/// </summary>
+		/// <param name="LogFile">The log file to search</param>
+		/// <param name="Text">Text to search for</param>
+		/// <param name="FirstLine">Line to start search from</param>
+		/// <param name="Count">Number of results to return</param>
+		/// <param name="Stats">Receives stats for the search</param>
+		/// <returns>List of line numbers containing the given term</returns>
+		Task<List<int>> SearchLogDataAsync(ILogFile LogFile, string Text, int FirstLine, int Count, LogSearchStats Stats);
+	}
+
+	/// <summary>
+	/// Extension methods for dealing with log files
+	/// </summary>
+	public static class LogFileServiceExtensions
+	{
+		/// <summary>
+		/// Parses a stream of json text and outputs plain text
+		/// </summary>
+		/// <param name="LogFileService">The log file service</param>
+		/// <param name="LogFile">The log file to query</param>
+		/// <param name="Offset">Offset within the log file to copy</param>
+		/// <param name="Length">Length of the data to copy</param>
+		/// <param name="OutputStream">Output stream to receive the text data</param>
+		/// <returns>Async text</returns>
+		public static async Task CopyRawStreamAsync(this ILogFileService LogFileService, ILogFile LogFile, long Offset, long Length, Stream OutputStream)
+		{
+			using (Stream Stream = await LogFileService.OpenRawStreamAsync(LogFile, Offset, Length))
+			{
+				await Stream.CopyToAsync(OutputStream);
+			}
+		}
+
+		/// <summary>
+		/// Parses a stream of json text and outputs plain text
+		/// </summary>
+		/// <param name="LogFileService">The log file service</param>
+		/// <param name="LogFile">The log file to query</param>
+		/// <param name="Offset">Offset within the data to copy from</param>
+		/// <param name="Length">Length of the data to copy</param>
+		/// <param name="OutputStream">Output stream to receive the text data</param>
+		/// <returns>Async text</returns>
+		public static async Task CopyPlainTextStreamAsync(this ILogFileService LogFileService, ILogFile LogFile, long Offset, long Length, Stream OutputStream)
+		{
+			using (Stream Stream = await LogFileService.OpenRawStreamAsync(LogFile, 0, long.MaxValue))
+			{
+				byte[] ReadBuffer = new byte[4096];
+				int ReadBufferLength = 0;
+
+				byte[] WriteBuffer = new byte[4096];
+				int WriteBufferLength = 0;
+
+				while(Length > 0)
+				{
+					// Add more data to the buffer
+					int ReadBytes = await Stream.ReadAsync(ReadBuffer, ReadBufferLength, ReadBuffer.Length - ReadBufferLength);
+					ReadBufferLength += ReadBytes;
+
+					// Copy as many lines as possible to the output
+					int ConvertedBytes = 0;
+					for(int EndIdx = 1; EndIdx < ReadBufferLength; EndIdx++)
+					{
+						if (ReadBuffer[EndIdx] == '\n')
+						{
+							WriteBufferLength = LogText.ConvertToPlainText(ReadBuffer.AsSpan(ConvertedBytes, EndIdx - ConvertedBytes), WriteBuffer, WriteBufferLength);
+							ConvertedBytes = EndIdx + 1;
+						}
+					}
+
+					// If there's anything in the write buffer, write it out
+					if (WriteBufferLength > 0)
+					{
+						if (Offset < WriteBufferLength)
+						{
+							int WriteLength = (int)Math.Min((long)WriteBufferLength - Offset, Length);
+							await OutputStream.WriteAsync(WriteBuffer, (int)Offset, WriteLength);
+							Length -= WriteLength;
+						}
+						Offset = Math.Max(Offset - WriteBufferLength, 0);
+						WriteBufferLength = 0;
+					}
+
+					// If we were able to read something, shuffle down the rest of the buffer. Otherwise expand the read buffer.
+					if (ConvertedBytes > 0)
+					{
+						Buffer.BlockCopy(ReadBuffer, ConvertedBytes, ReadBuffer, 0, ReadBufferLength - ConvertedBytes);
+						ReadBufferLength -= ConvertedBytes;
+					}
+					else if(ReadBufferLength > 0)
+					{
+						Array.Resize(ref ReadBuffer, ReadBuffer.Length + 128);
+						WriteBuffer = new byte[ReadBuffer.Length];
+					}
+
+					// Exit if we didn't read anything in this iteration
+					if(ReadBytes == 0)
+					{
+						break;
+					}
+				}
+			}
+		}
+	}
+	
+	/// <summary>
+	/// Wraps functionality for manipulating logs
+	/// </summary>
+	public class LogFileService : TickedBackgroundService, ILogFileService
+	{
+		/// <summary>
+		/// Information Logger
+		/// </summary>
+		private readonly ILogger<LogFileService> Logger;
+
+		/// <summary>
+		/// Collection of log documents
+		/// </summary>
+		private readonly ILogFileCollection LogFiles;
+
+		/// <summary>
+		/// Collection of log events
+		/// </summary>
+		private readonly ILogEventCollection LogEvents;
+
+		/// <summary>
+		/// Interface for the log reader
+		/// </summary>
+		private readonly ILogStorage Storage;
+
+		/// <summary>
+		/// Interface for the log builder
+		/// </summary>
+		private readonly ILogBuilder Builder;
+
+		/// <summary>
+		/// Lock object for the <see cref="WriteTasks"/> and <see cref="WriteChunks"/> members
+		/// </summary>
+		private object WriteLock = new object();
+
+		/// <summary>
+		/// List of active write tasks
+		/// </summary>
+		private List<Task> WriteTasks = new List<Task>();
+
+		/// <summary>
+		/// Set of chunks which is currently being written to
+		/// </summary>
+		private HashSet<(LogId, long)> WriteChunks = new HashSet<(LogId, long)>();
+
+		/// <summary>
+		/// Cache of session ids for each log
+		/// </summary>
+		private IMemoryCache LogFileCache;
+
+		/// <summary>
+		/// Streams log data to a caller
+		/// </summary>
+		class ResponseStream : Stream
+		{
+			/// <summary>
+			/// The log file service that created this stream
+			/// </summary>
+			LogFileService LogFileService;
+
+			/// <summary>
+			/// The log file being read
+			/// </summary>
+			ILogFile LogFile;
+
+			/// <summary>
+			/// Starting offset within the file of the data to return 
+			/// </summary>
+			long ResponseOffset;
+
+			/// <summary>
+			/// Length of data to return
+			/// </summary>
+			long ResponseLength;
+
+			/// <summary>
+			/// Current offset within the stream
+			/// </summary>
+			long CurrentOffset;
+
+			/// <summary>
+			/// The current chunk index
+			/// </summary>
+			int ChunkIdx;
+
+			/// <summary>
+			/// Buffer containing a message for missing data
+			/// </summary>
+			ReadOnlyMemory<byte> SourceBuffer;
+
+			/// <summary>
+			/// Offset within the source buffer
+			/// </summary>
+			int SourcePos;
+
+			/// <summary>
+			/// Length of the source buffer being copied from
+			/// </summary>
+			int SourceEnd;
+
+			/// <summary>
+			/// Constructor
+			/// </summary>
+			/// <param name="LogFileService">The log file service, for q</param>
+			/// <param name="LogFile"></param>
+			/// <param name="Offset"></param>
+			/// <param name="Length"></param>
+			public ResponseStream(LogFileService LogFileService, ILogFile LogFile, long Offset, long Length)
+			{
+				this.LogFileService = LogFileService;
+				this.LogFile = LogFile;
+
+				this.ResponseOffset = Offset;
+				this.ResponseLength = Length;
+
+				this.CurrentOffset = Offset;
+
+				this.ChunkIdx = LogFile.Chunks.GetChunkForOffset(Offset);
+				this.SourceBuffer = null!;
+			}
+
+			/// <inheritdoc/>
+			public override bool CanRead => true;
+
+			/// <inheritdoc/>
+			public override bool CanSeek => false;
+
+			/// <inheritdoc/>
+			public override bool CanWrite => false;
+
+			/// <inheritdoc/>
+			public override long Length
+			{
+				get { return ResponseLength; }
+			}
+
+			/// <inheritdoc/>
+			public override long Position
+			{
+				get { return CurrentOffset - ResponseOffset; }
+				set => throw new NotImplementedException();
+			}
+
+			/// <inheritdoc/>
+			public override void Flush()
+			{
+			}
+
+			/// <inheritdoc/>
+			public override int Read(byte[] Buffer, int Offset, int Count)
+			{
+				return ReadAsync(Buffer, Offset, Count, CancellationToken.None).Result;
+			}
+
+			/// <inheritdoc/>
+			public override async Task<int> ReadAsync(byte[] Buffer, int Offset, int Length, CancellationToken CancellationToken)
+			{
+				int ReadBytes = 0;
+				while (ReadBytes < Length)
+				{
+					if (SourcePos < SourceEnd)
+					{
+						// Try to copy from the current buffer
+						int BlockSize = Math.Min(SourceEnd - SourcePos, Length - ReadBytes);
+						SourceBuffer.Slice(SourcePos, BlockSize).Span.CopyTo(Buffer.AsSpan(Offset + ReadBytes));
+						CurrentOffset += BlockSize;
+						ReadBytes += BlockSize;
+						SourcePos += BlockSize;
+					}
+					else if (CurrentOffset < ResponseOffset + ResponseLength)
+					{
+						// Move to the right chunk
+						while (ChunkIdx + 1 < LogFile.Chunks.Count && CurrentOffset >= LogFile.Chunks[ChunkIdx + 1].Offset)
+						{
+							ChunkIdx++;
+						}
+
+						// Get the end of this chunk
+						long NextOffset = ResponseOffset + ResponseLength;
+						if (ChunkIdx + 1 < LogFile.Chunks.Count)
+						{
+							ILogChunk NextChunk = LogFile.Chunks[ChunkIdx + 1];
+							NextOffset = Math.Min(NextOffset, NextChunk.Offset);
+						}
+
+						// Get the chunk data
+						ILogChunk Chunk = LogFile.Chunks[ChunkIdx];
+						LogChunkData ChunkData = await LogFileService.ReadChunkAsync(LogFile, ChunkIdx);
+
+						// Figure out which sub-chunk to use
+						int SubChunkIdx = ChunkData.GetSubChunkForOffsetWithinChunk((int)(CurrentOffset - Chunk.Offset));
+						LogSubChunkData SubChunkData = ChunkData.SubChunks[SubChunkIdx];
+
+						// Get the source data
+						long SubChunkOffset = Chunk.Offset + ChunkData.SubChunkOffset[SubChunkIdx];
+						SourceBuffer = SubChunkData.InflateText().Data;
+						SourcePos = (int)(CurrentOffset - SubChunkOffset);
+						SourceEnd = (int)Math.Min(SourceBuffer.Length, (ResponseOffset + ResponseLength) - SubChunkOffset);
+					}
+					else
+					{
+						// End of the log
+						break;
+					}
+				}
+				return ReadBytes;
+			}
+
+			/// <inheritdoc/>
+			public override long Seek(long offset, SeekOrigin origin) => throw new NotImplementedException();
+
+			/// <inheritdoc/>
+			public override void SetLength(long value) => throw new NotImplementedException();
+
+			/// <inheritdoc/>
+			public override void Write(byte[] buffer, int offset, int count) => throw new NotImplementedException();
+		}
+
+		/// <summary>
+		/// Constructor
+		/// </summary>
+		/// <param name="LogFiles">The log file collection</param>
+		/// <param name="LogEvents">The log events collection</param>
+		/// <param name="Builder">The log builder</param>
+		/// <param name="Storage">THe log storage hierarchy</param>
+		/// <param name="Logger">Log interface</param>
+		public LogFileService(ILogFileCollection LogFiles, ILogEventCollection LogEvents, ILogBuilder Builder, ILogStorage Storage, ILogger<LogFileService> Logger)
+			: base(TimeSpan.FromSeconds(30.0), Logger)
+		{
+			this.LogFiles = LogFiles;
+			this.LogEvents = LogEvents;
+			this.Logger = Logger;
+			this.LogFileCache = new MemoryCache(new MemoryCacheOptions());
+			this.Builder = Builder;
+			this.Storage = Storage;
+		}
+
+		/// <inheritdoc/>
+		public Task<ILogFile> CreateLogFileAsync(JobId JobId, ObjectId? SessionId, LogType Type)
+		{
+			return LogFiles.CreateLogFileAsync(JobId, SessionId, Type);
+		}
+
+		/// <inheritdoc/>
+		public async Task<ILogFile?> GetLogFileAsync(LogId LogFileId)
+		{
+			ILogFile? LogFile = await LogFiles.GetLogFileAsync(LogFileId);
+			if(LogFile != null)
+			{
+				AddCachedLogFile(LogFile);
+			}
+			return LogFile;
+		}
+
+		/// <summary>
+		/// Adds a log file to the cache
+		/// </summary>
+		/// <param name="LogFile">The log file to cache</param>
+		void AddCachedLogFile(ILogFile LogFile)
+		{
+			MemoryCacheEntryOptions Options = new MemoryCacheEntryOptions().SetSlidingExpiration(TimeSpan.FromSeconds(30));
+			LogFileCache.Set(LogFile.Id, LogFile, Options);
+		}
+
+		/// <summary>
+		/// Gets a cached log file by id
+		/// </summary>
+		/// <param name="LogFileId">The log file id</param>
+		/// <returns>New log file, or null if not found</returns>
+		public async Task<ILogFile?> GetCachedLogFileAsync(LogId LogFileId)
+		{
+			object? LogFile;
+			if (!LogFileCache.TryGetValue(LogFileId, out LogFile))
+			{
+				LogFile = await GetLogFileAsync(LogFileId);
+			}
+			return (ILogFile?)LogFile;
+		}
+
+		/// <inheritdoc/>
+		public Task<List<ILogFile>> GetLogFilesAsync(int? Index = null, int? Count = null)
+		{
+			return LogFiles.GetLogFilesAsync(Index, Count);
+		}
+
+		class WriteState
+		{
+			public long Offset;
+			public int LineIndex;
+			public ReadOnlyMemory<byte> Memory;
+
+			public WriteState(long Offset, int LineIndex, ReadOnlyMemory<byte> Memory)
+			{
+				this.Offset = Offset;
+				this.LineIndex = LineIndex;
+				this.Memory = Memory;
+			}
+		}
+
+		/// <inheritdoc/>
+		public async Task<ILogFile?> WriteLogDataAsync(ILogFile LogFile, long Offset, int LineIndex, ReadOnlyMemory<byte> Data, bool Flush, int MaxChunkLength, int MaxSubChunkLineCount)
+		{
+			using IScope Scope = GlobalTracer.Instance.BuildSpan("WriteLogDataAsync").StartActive();
+			Scope.Span.SetTag("LogId", LogFile.Id.ToString());
+			Scope.Span.SetTag("Offset", Offset.ToString(CultureInfo.InvariantCulture));
+			Scope.Span.SetTag("Length", Data.Length.ToString(CultureInfo.InvariantCulture));
+			Scope.Span.SetTag("LineIndex", LineIndex.ToString(CultureInfo.InvariantCulture));
+
+			// Make sure the data ends in a newline
+			if (Data.Length > 0 && Data.Span[Data.Length - 1] != '\n')
+			{
+				throw new ArgumentException("Log data must consist of a whole number of lines", nameof(Data));
+			}
+
+			// Make sure the line count is a power of two
+			if ((MaxSubChunkLineCount & (MaxSubChunkLineCount - 1)) != 0)
+			{
+				throw new ArgumentException("Maximum line count per sub-chunk must be a power of two", nameof(MaxSubChunkLineCount));
+			}
+
+			// List of the flushed chunks
+			List<long> CompleteOffsets = new List<long>();
+
+			// Add the data to new chunks
+			WriteState State = new WriteState(Offset, LineIndex, Data);
+			while (State.Memory.Length > 0)
+			{
+				// Find an existing chunk to append to
+				int ChunkIdx = LogFile.Chunks.GetChunkForOffset(State.Offset);
+				if (ChunkIdx >= 0)
+				{
+					ILogChunk Chunk = LogFile.Chunks[ChunkIdx];
+					if (await WriteLogChunkDataAsync(LogFile, Chunk, State, CompleteOffsets, MaxChunkLength, MaxSubChunkLineCount))
+					{
+						continue;
+					}
+				}
+
+				// Create a new chunk. Ensure that there's a chunk at the start of the file, even if the current write is beyond it.
+				ILogFile? NewLogFile;
+				if (LogFile.Chunks.Count == 0)
+				{
+					NewLogFile = await LogFiles.TryAddChunkAsync(LogFile, 0, 0);
+				}
+				else
+				{
+					NewLogFile = await LogFiles.TryAddChunkAsync(LogFile, State.Offset, State.LineIndex);
+				}
+
+				// Try to add a new chunk at the new location
+				if (NewLogFile == null)
+				{
+					NewLogFile = await LogFiles.GetLogFileAsync(LogFile.Id);
+					if (NewLogFile == null)
+					{
+						Logger.LogError("Unable to update log file {LogId}", LogFile.Id);
+						return null;
+					}
+					LogFile = NewLogFile;
+				}
+				else
+				{
+					// Logger.LogDebug("Added new chunk at offset {Offset} to log {LogId}", State.Offset, LogFile.Id);
+					LogFile = NewLogFile;
+				}
+			}
+
+			// Flush any pending chunks on this log file
+			if (Flush)
+			{
+				foreach(ILogChunk Chunk in LogFile.Chunks)
+				{
+					if (Chunk.Length == 0 && !CompleteOffsets.Contains(Chunk.Offset))
+					{
+						await Builder.CompleteChunkAsync(LogFile.Id, Chunk.Offset);
+						CompleteOffsets.Add(Chunk.Offset);
+					}
+				}
+			}
+
+			// Write all the chunks
+			if (CompleteOffsets.Count > 0 || Flush)
+			{
+				ILogFile? NewLogFile = await WriteCompleteChunksForLogAsync(LogFile, CompleteOffsets, Flush);
+				if (NewLogFile == null)
+				{
+					return null;
+				}
+				LogFile = NewLogFile;
+			}
+			return LogFile;
+		}
+
+		/// <summary>
+		/// Append data to an existing chunk.
+		/// </summary>
+		/// <param name="LogFile">The log file to append to</param>
+		/// <param name="Chunk">Chunk within the log file to update</param>
+		/// <param name="State">Data remaining to be written</param>
+		/// <param name="CompleteOffsets">List of complete chunks</param>
+		/// <param name="MaxChunkLength">Maximum length of each chunk</param>
+		/// <param name="MaxSubChunkLineCount">Maximum number of lines in each subchunk</param>
+		/// <returns>True if data was appended to </returns>
+		private async Task<bool> WriteLogChunkDataAsync(ILogFile LogFile, ILogChunk Chunk, WriteState State, List<long> CompleteOffsets, int MaxChunkLength, int MaxSubChunkLineCount)
+		{
+			// Don't allow data to be appended if the chunk is complete
+			if(Chunk.Length > 0)
+			{
+				return false;
+			}
+
+			// Otherwise keep appending subchunks
+			bool bResult = false;
+			for (; ; )
+			{
+				// Flush the current sub-chunk if we're on a boundary
+				if (State.LineIndex > 0 && (State.LineIndex & (MaxSubChunkLineCount - 1)) == 0)
+				{
+					Logger.LogDebug("Completing log {LogId} chunk offset {Offset} sub-chunk at line {LineIndex}", LogFile.Id, Chunk.Offset, State.LineIndex);
+					await Builder.CompleteSubChunkAsync(LogFile.Id, Chunk.Offset);
+				}
+
+				// Figure out the max length to write to the current chunk
+				int MaxLength = Math.Min((int)((Chunk.Offset + MaxChunkLength) - State.Offset), State.Memory.Length);
+
+				// Figure out the maximum line index for the current sub chunk
+				int MinLineIndex = State.LineIndex;
+				int MaxLineIndex = (MinLineIndex & ~(MaxSubChunkLineCount - 1)) + MaxSubChunkLineCount;
+
+				// Append this data
+				(int Length, int LineCount) = GetWriteLength(State.Memory.Span, MaxLength, MaxLineIndex - MinLineIndex, State.Offset == Chunk.Offset);
+				if (Length > 0)
+				{
+					// Append this data
+					ReadOnlyMemory<byte> AppendData = State.Memory.Slice(0, Length);
+					if (!await Builder.AppendAsync(LogFile.Id, Chunk.Offset, State.Offset, State.LineIndex, LineCount, AppendData, LogFile.Type))
+					{
+						break;
+					}
+
+					// Update the state
+					//Logger.LogDebug("Append to log {LogId} chunk offset {Offset} (LineIndex={LineIndex}, LineCount={LineCount}, Offset={WriteOffset}, Length={WriteLength})", LogFile.Id, Chunk.Offset, State.LineIndex, LineCount, State.Offset, Length);
+					State.Offset += Length;
+					State.LineIndex += LineCount;
+					State.Memory = State.Memory.Slice(Length);
+					bResult = true;
+
+					// If this is the end of the data, bail out
+					if(State.Memory.Length == 0)
+					{
+						break;
+					}
+				}
+
+				// Flush the sub-chunk if it's full
+				if (State.LineIndex < MaxLineIndex)
+				{
+					Logger.LogDebug("Completing chunk for log {LogId} at offset {Offset}", LogFile.Id, Chunk.Offset);
+					await Builder.CompleteChunkAsync(LogFile.Id, Chunk.Offset);
+					CompleteOffsets.Add(Chunk.Offset);
+					break;
+				}
+			}
+			return bResult;
+		}
+
+		/// <summary>
+		/// Get the amount of data to write from the given span
+		/// </summary>
+		/// <param name="Span">Data to write</param>
+		/// <param name="MaxLength">Maximum length of the data to write</param>
+		/// <param name="MaxLineCount">Maximum number of lines to write</param>
+		/// <param name="bIsEmptyChunk">Whether the current chunk is empty</param>
+		/// <returns>A tuple consisting of the amount of data to write and number of lines in it</returns>
+		private static (int, int) GetWriteLength(ReadOnlySpan<byte> Span, int MaxLength, int MaxLineCount, bool bIsEmptyChunk)
+		{
+			int Length = 0;
+			int LineCount = 0;
+			for (int Idx = 0; Idx < MaxLength || bIsEmptyChunk; Idx++)
+			{
+				if (Span[Idx] == '\n')
+				{
+					Length = Idx + 1;
+					LineCount++;
+					bIsEmptyChunk = false;
+
+					if (LineCount >= MaxLineCount)
+					{
+						break;
+					}
+				}
+			}
+			return (Length, LineCount);
+		}
+
+		/// <inheritdoc/>
+		public async Task<LogMetadata> GetMetadataAsync(ILogFile LogFile)
+		{
+			LogMetadata Metadata = new LogMetadata();
+			if (LogFile.Chunks.Count > 0)
+			{
+				ILogChunk Chunk = LogFile.Chunks[LogFile.Chunks.Count - 1];
+				if (LogFile.MaxLineIndex == null || Chunk.Length == 0)
+				{
+					LogChunkData ChunkData = await ReadChunkAsync(LogFile, LogFile.Chunks.Count - 1);
+					Metadata.Length = Chunk.Offset + ChunkData.Length;
+					Metadata.MaxLineIndex = Chunk.LineIndex + ChunkData.LineCount;
+				}
+				else
+				{
+					Metadata.Length = Chunk.Offset + Chunk.Length;
+					Metadata.MaxLineIndex = LogFile.MaxLineIndex.Value;
+				}
+			}
+			return Metadata;
+		}
+
+		/// <inheritdoc/>
+		public Task CreateEventsAsync(List<NewLogEventData> NewEvents)
+		{
+			return LogEvents.AddManyAsync(NewEvents);
+		}
+
+		/// <inheritdoc/>
+		public Task<List<ILogEvent>> FindLogEventsAsync(ILogFile LogFile, int? Index = null, int? Count = null)
+		{
+			return LogEvents.FindAsync(LogFile.Id, Index, Count);
+		}
+
+		class LogEventLine : ILogEventLine
+		{
+			LogLevel Level;
+			public EventId? EventId { get; }
+			public string Message { get; }
+			public JsonElement Data { get; }
+
+			LogLevel ILogEventLine.Level => Level;
+
+			public LogEventLine(ReadOnlySpan<byte> Data)
+				: this(JsonSerializer.Deserialize<JsonElement>(Data))
+			{
+			}
+
+			[SuppressMessage("Design", "CA1031:Do not catch general exception types")]
+			public LogEventLine(JsonElement Data)
+			{
+				this.Data = Data;
+
+				JsonElement LevelElement;
+				if (!Data.TryGetProperty("level", out LevelElement) || !Enum.TryParse(LevelElement.GetString(), out Level))
+				{
+					Level = LogLevel.Information;
+				}
+
+				JsonElement IdElement;
+				if (Data.TryGetProperty("id", out IdElement))
+				{
+					int IdValue;
+					if (IdElement.TryGetInt32(out IdValue))
+					{
+						EventId = IdValue;
+					}
+				}
+
+				JsonElement MessageElement;
+				if (Data.TryGetProperty("renderedMessage", out MessageElement) || Data.TryGetProperty("message", out MessageElement))
+				{
+					Message = MessageElement.GetString() ?? "(Invalid)";
+				}
+				else
+				{
+					Message = "(Missing message or renderedMessage field)";
+				}
+			}
+		}
+
+		class LogEventData : ILogEventData
+		{
+			public IReadOnlyList<ILogEventLine> Lines { get; }
+
+			EventId? ILogEventData.EventId => (Lines.Count > 0) ? Lines[0].EventId : null;
+			EventSeverity ILogEventData.Severity => (Lines.Count == 0) ? EventSeverity.Information : (Lines[0].Level == LogLevel.Warning) ? EventSeverity.Warning : EventSeverity.Error;
+			string ILogEventData.Message => String.Join("\n", Lines.Select(x => x.Message));
+
+			public LogEventData(IReadOnlyList<ILogEventLine> Lines)
+			{
+				this.Lines = Lines;
+			}
+		}
+
+		/// <inheritdoc/>
+		public Task AddSpanToEventsAsync(IEnumerable<ILogEvent> Events, ObjectId SpanId)
+		{
+			return LogEvents.AddSpanToEventsAsync(Events, SpanId);
+		}
+
+		/// <inheritdoc/>
+		public Task<List<ILogEvent>> FindEventsForSpansAsync(IEnumerable<ObjectId> SpanIds, LogId[]? LogIds, int Index, int Count)
+		{
+			return LogEvents.FindEventsForSpansAsync(SpanIds, LogIds, Index, Count);
+		}
+
+		/// <inheritdoc/>
+		public async Task<ILogEventData> GetEventDataAsync(ILogFile LogFile, int LineIndex, int LineCount)
+		{
+			using IScope Scope = GlobalTracer.Instance.BuildSpan("GetEventDataAsync").StartActive();
+			Scope.Span.SetTag("LogId", LogFile.Id.ToString());
+			Scope.Span.SetTag("LineIndex", LineIndex.ToString(CultureInfo.InvariantCulture));
+			Scope.Span.SetTag("LineCount", LineCount.ToString(CultureInfo.InvariantCulture));
+
+			(_, long MinOffset) = await GetLineOffsetAsync(LogFile, LineIndex);
+			(_, long MaxOffset) = await GetLineOffsetAsync(LogFile, LineIndex + LineCount);
+
+			byte[] Data = new byte[MaxOffset - MinOffset];
+			using (Stream Stream = await OpenRawStreamAsync(LogFile, MinOffset, MaxOffset - MinOffset))
+			{
+				int Length = await Stream.ReadAsync(Data.AsMemory());
+				if(Length != Data.Length)
+				{
+					Logger.LogWarning("Read less than expected from log stream (Expected {Expected}, Got {Got})", Data.Length, Length);
+				}
+				return ParseEventData(Data.AsSpan(0, Length));
+			}
+		}
+
+		/// <summary>
+		/// Parses event data from a buffer
+		/// </summary>
+		/// <param name="Data">Data to parse</param>
+		/// <returns>Parsed event data</returns>
+		private ILogEventData ParseEventData(ReadOnlySpan<byte> Data)
+		{
+			List<LogEventLine> Lines = new List<LogEventLine>();
+
+			ReadOnlySpan<byte> RemainingData = Data;
+			while(RemainingData.Length > 0)
+			{
+				int EndOfLine = RemainingData.IndexOf((byte)'\n');
+				if(EndOfLine == -1)
+				{
+					break;
+				}
+
+				ReadOnlySpan<byte> LineData = RemainingData.Slice(0, EndOfLine);
+				try
+				{
+					Lines.Add(new LogEventLine(LineData));
+				}
+				catch(JsonException Ex)
+				{
+					Logger.LogWarning(Ex, "Unable to parse line from log file: {Line}", Encoding.UTF8.GetString(LineData));
+				}
+				RemainingData = RemainingData.Slice(EndOfLine + 1);
+			}
+
+			return new LogEventData(Lines);
+		}
+
+		/// <inheritdoc/>
+		public async Task<Stream> OpenRawStreamAsync(ILogFile LogFile, long Offset, long Length)
+		{
+			if (LogFile.Chunks.Count == 0)
+			{
+				return new MemoryStream(Array.Empty<byte>(), false);
+			}
+			else
+			{
+				int LastChunkIdx = LogFile.Chunks.Count - 1;
+
+				// Clamp the length of the request
+				ILogChunk LastChunk = LogFile.Chunks[LastChunkIdx];
+				if (Length > LastChunk.Offset)
+				{
+					long LastChunkLength = LastChunk.Length;
+					if (LastChunkLength <= 0)
+					{
+						LogChunkData LastChunkData = await ReadChunkAsync(LogFile, LastChunkIdx);
+						LastChunkLength = LastChunkData.Length;
+					}
+					Length = Math.Min(Length, (LastChunk.Offset + LastChunkLength) - Offset);
+				}
+
+				// Create the new stream
+				return new ResponseStream(this, LogFile, Offset, Length);
+			}
+		}
+
+		/// <inheritdoc/>
+		public async Task<(int, long)> GetLineOffsetAsync(ILogFile LogFile, int LineIdx)
+		{
+			int ChunkIdx = LogFile.Chunks.GetChunkForLine(LineIdx);
+
+			ILogChunk Chunk = LogFile.Chunks[ChunkIdx];
+			LogChunkData ChunkData = await ReadChunkAsync(LogFile, ChunkIdx);
+
+			if (LineIdx < Chunk.LineIndex)
+			{
+				LineIdx = Chunk.LineIndex;
+			}
+
+			int MaxLineIndex = Chunk.LineIndex + ChunkData.LineCount;
+			if (LineIdx >= MaxLineIndex)
+			{
+				LineIdx = MaxLineIndex;
+			}
+
+			long Offset = Chunk.Offset + ChunkData.GetLineOffsetWithinChunk(LineIdx - Chunk.LineIndex);
+			return (LineIdx, Offset);
+		}
+
+		/// <summary>
+		/// Executes a background task
+		/// </summary>
+		/// <param name="StoppingToken">Cancellation token</param>
+		protected override async Task TickAsync(CancellationToken StoppingToken)
+		{
+			lock (WriteLock)
+			{
+				try
+				{
+					WriteTasks.RemoveCompleteTasks();
+				}
+				catch (Exception Ex)
+				{
+					Logger.LogError(Ex, "Exception while waiting for write tasks to complete");
+				}
+			}
+			await IncrementalFlush();
+		}
+		
+		/// <summary>
+		/// Executes a background task. Publicly exposed for tests.
+		/// </summary>
+		public async Task TickOnlyForTestingAsync()
+		{
+			await TickAsync(new CancellationToken());
+		}
+		
+		/// <summary>
+		/// Flushes complete chunks to the storage provider
+		/// </summary>
+		/// <returns>Async task</returns>
+		private async Task IncrementalFlush()
+		{
+			// Get all the chunks older than 20 minutes
+			List<(LogId, long)> FlushChunks = await Builder.TouchChunksAsync(TimeSpan.FromMinutes(10.0));
+			Logger.LogDebug("Performing incremental flush of log builder ({NumChunks} chunks)", FlushChunks.Count);
+
+			// Mark them all as complete
+			foreach ((LogId LogId, long Offset) in FlushChunks)
+			{
+				await Builder.CompleteChunkAsync(LogId, Offset);
+			}
+
+			// Add tasks for flushing all the chunks
+			WriteCompleteChunks(FlushChunks, true);
+		}
+
+		/// <summary>
+		/// Called when the service is shutting down
+		/// </summary>
+		/// <param name="CancellationToken"></param>
+		/// <returns></returns>
+		[SuppressMessage("Design", "CA1031:Do not catch general exception types")]
+		public override async Task StopAsync(CancellationToken CancellationToken)
+		{
+			Logger.LogInformation("Stopping log file service");
+			if (Builder.FlushOnShutdown)
+			{
+				await FlushAsync();
+			}
+			Logger.LogInformation("Log service stopped");
+		}
+
+		/// <summary>
+		/// Flushes the write cache
+		/// </summary>
+		/// <returns>Async task</returns>
+		public async Task FlushAsync()
+		{
+			Logger.LogInformation("Forcing flush of pending log chunks...");
+
+			// Mark everything in the cache as complete
+			List<(LogId, long)> WriteChunks = await Builder.TouchChunksAsync(TimeSpan.Zero);
+			WriteCompleteChunks(WriteChunks, true);
+
+			// Wait for everything to flush
+			await FlushPendingWritesAsync();
+		}
+
+		/// <summary>
+		/// Flush any writes in progress
+		/// </summary>
+		/// <returns>Async task</returns>
+		public async Task FlushPendingWritesAsync()
+		{
+			for(; ;)
+			{
+				// Capture the current contents of the WriteTasks list
+				List<Task> Tasks;
+				lock (WriteLock)
+				{
+					WriteTasks.RemoveCompleteTasks();
+					Tasks = new List<Task>(WriteTasks);
+				}
+				if (Tasks.Count == 0)
+				{
+					break;
+				}
+
+				// Also add a delay so we'll periodically refresh the list
+				Tasks.Add(Task.Delay(TimeSpan.FromSeconds(5.0)));
+				await Task.WhenAny(Tasks);
+			}
+		}
+
+		/// <summary>
+		/// Adds tasks for writing a list of complete chunks
+		/// </summary>
+		/// <param name="ChunksToWrite">List of chunks to write</param>
+		/// <param name="bCreateIndex">Create an index for the log</param>
+		private void WriteCompleteChunks(List<(LogId, long)> ChunksToWrite, bool bCreateIndex)
+		{
+			foreach (IGrouping<LogId, long> Group in ChunksToWrite.GroupBy(x => x.Item1, x => x.Item2))
+			{
+				LogId LogId = Group.Key;
+
+				// Find offsets of new chunks to write
+				List<long> Offsets = new List<long>();
+				lock (WriteLock)
+				{
+					foreach (long Offset in Group.OrderBy(x => x))
+					{
+						if (WriteChunks.Add((LogId, Offset)))
+						{
+							Offsets.Add(Offset);
+						}
+					}
+				}
+
+				// Create the write task
+				if (Offsets.Count > 0)
+				{
+					Task Task = Task.Run(() => WriteCompleteChunksForLogAsync(LogId, Offsets, bCreateIndex));
+					lock (WriteLock)
+					{
+						WriteTasks.Add(Task);
+					}
+				}
+			}
+		}
+
+		/// <summary>
+		/// Writes a set of chunks to the database
+		/// </summary>
+		/// <param name="LogId">Log file to update</param>
+		/// <param name="Offsets">Chunks to write</param>
+		/// <param name="bCreateIndex">Whether to create the index for this log</param>
+		/// <returns>Async task</returns>
+		private async Task<ILogFile?> WriteCompleteChunksForLogAsync(LogId LogId, List<long> Offsets, bool bCreateIndex)
+		{
+			ILogFile? LogFile = await LogFiles.GetLogFileAsync(LogId);
+			if(LogFile != null)
+			{
+				LogFile = await WriteCompleteChunksForLogAsync(LogFile, Offsets, bCreateIndex);
+			}
+			return LogFile;
+		}
+
+		/// <summary>
+		/// Writes a set of chunks to the database
+		/// </summary>
+		/// <param name="LogFileInterface">Log file to update</param>
+		/// <param name="Offsets">Chunks to write</param>
+		/// <param name="bCreateIndex">Whether to create the index for this log</param>
+		/// <returns>Async task</returns>
+		[SuppressMessage("Design", "CA1031:Do not catch general exception types", Justification = "<Pending>")]
+		private async Task<ILogFile?> WriteCompleteChunksForLogAsync(ILogFile LogFileInterface, List<long> Offsets, bool bCreateIndex)
+		{
+			// Write the data to the storage provider
+			List<Task<LogChunkData?>> ChunkWriteTasks = new List<Task<LogChunkData?>>();
+			foreach (long Offset in Offsets)
+			{
+				int ChunkIdx = LogFileInterface.Chunks.BinarySearch(x => x.Offset, Offset);
+				if (ChunkIdx >= 0)
+				{
+					Logger.LogDebug("Queuing write of log {LogId} chunk {ChunkIdx} offset {Offset}", LogFileInterface.Id, ChunkIdx, Offset);
+					int LineIndex = LogFileInterface.Chunks[ChunkIdx].LineIndex;
+					ChunkWriteTasks.Add(Task.Run(() => WriteChunkAsync(LogFileInterface.Id, Offset, LineIndex)));
+				}
+			}
+
+			// Wait for the tasks to complete, periodically updating the log file object
+			ILogFile? LogFile = LogFileInterface;
+			while (ChunkWriteTasks.Count > 0)
+			{
+				// Wait for all tasks to be complete OR (any task has completed AND 30 seconds has elapsed)
+				Task AllCompleteTask = Task.WhenAll(ChunkWriteTasks);
+				Task AnyCompleteTask = Task.WhenAny(ChunkWriteTasks);
+				await Task.WhenAny(AllCompleteTask, Task.WhenAll(AnyCompleteTask, Task.Delay(TimeSpan.FromSeconds(30.0))));
+
+				// Update the log file with the written chunks
+				List<LogChunkData?> WrittenChunks = ChunkWriteTasks.RemoveCompleteTasks();
+				while (LogFile != null)
+				{
+					// Update the length of any complete chunks
+					List<CompleteLogChunkUpdate> Updates = new List<CompleteLogChunkUpdate>();
+					foreach (LogChunkData? ChunkData in WrittenChunks)
+					{
+						if (ChunkData != null)
+						{
+							int ChunkIdx = LogFile.Chunks.GetChunkForOffset(ChunkData.Offset);
+							if (ChunkIdx >= 0)
+							{
+								ILogChunk Chunk = LogFile.Chunks[ChunkIdx];
+								if (Chunk.Offset == ChunkData.Offset)
+								{
+									CompleteLogChunkUpdate Update = new CompleteLogChunkUpdate(ChunkIdx, ChunkData.Length, ChunkData.LineCount);
+									Updates.Add(Update);
+								}
+							}
+						}
+					}
+
+					// Try to apply the updates
+					ILogFile? NewLogFile = await LogFiles.TryCompleteChunksAsync(LogFile, Updates);
+					if (NewLogFile != null)
+					{
+						LogFile = NewLogFile;
+						break;
+					}
+
+					// Update the log file
+					LogFile = await GetLogFileAsync(LogFile.Id);
+				}
+			}
+
+			// Create the index if necessary
+			if (bCreateIndex && LogFile != null)
+			{
+				try
+				{
+					LogFile = await CreateIndexAsync(LogFile);
+				}
+				catch(Exception Ex)
+				{
+					Logger.LogError(Ex, "Failed to create index for log {LogId}", LogFileInterface.Id);
+				}
+			}
+
+			return LogFile;
+		}
+
+		/// <summary>
+		/// Creates an index for the given log file
+		/// </summary>
+		/// <param name="LogFile">The log file object</param>
+		/// <returns>Updated log file</returns>
+		private async Task<ILogFile?> CreateIndexAsync(ILogFile LogFile)
+		{
+			if(LogFile.Chunks.Count == 0)
+			{
+				return LogFile;
+			}
+
+			// Get the new length of the log, and early out if it won't be any longer
+			ILogChunk LastChunk = LogFile.Chunks[LogFile.Chunks.Count - 1];
+			if(LastChunk.Offset + LastChunk.Length <= (LogFile.IndexLength ?? 0))
+			{
+				return LogFile;
+			}
+
+			// Save stats for the index creation
+			using IScope Scope = GlobalTracer.Instance.BuildSpan("CreateIndexAsync").StartActive();
+			Scope.Span.SetTag("LogId", LogFile.Id.ToString());
+			Scope.Span.SetTag("Length", (LastChunk.Offset + LastChunk.Length).ToString(CultureInfo.InvariantCulture));
+
+			long NewLength = 0;
+			int NewLineCount = 0;
+
+			// Read the existing index if there is one
+			List<LogIndexData> Indexes = new List<LogIndexData>();
+			if (LogFile.IndexLength != null)
+			{
+				LogIndexData? ExistingIndex = await ReadIndexAsync(LogFile, LogFile.IndexLength.Value);
+				if(ExistingIndex != null)
+				{
+					Indexes.Add(ExistingIndex);
+					NewLineCount = ExistingIndex.LineCount;
+				}
+			}
+
+			// Add all the new chunks
+			int ChunkIdx = LogFile.Chunks.GetChunkForLine(NewLineCount);
+			if (ChunkIdx < 0)
+			{
+				int FirstLine = (LogFile.Chunks.Count > 0) ? LogFile.Chunks[0].LineIndex : -1;
+				throw new Exception($"Invalid chunk index {ChunkIdx}. Index.LineCount={NewLineCount}, Chunks={LogFile.Chunks.Count}, First line={FirstLine}");
+			}
+
+			for (; ChunkIdx < LogFile.Chunks.Count; ChunkIdx++)
+			{
+				ILogChunk Chunk = LogFile.Chunks[ChunkIdx];
+				LogChunkData ChunkData = await ReadChunkAsync(LogFile, ChunkIdx);
+
+				int SubChunkIdx = ChunkData.GetSubChunkForLine(Math.Max(NewLineCount - Chunk.LineIndex, 0));
+				if(SubChunkIdx < 0)
+				{
+					throw new Exception($"Invalid subchunk index {SubChunkIdx}. Chunk {ChunkIdx}/{LogFile.Chunks.Count}. Index.LineCount={NewLineCount}, Chunk.LineIndex={Chunk.LineIndex}, First subchunk {ChunkData.SubChunkLineIndex[0]}");
+				}
+
+				for (; SubChunkIdx < ChunkData.SubChunks.Count; SubChunkIdx++)
+				{
+					LogSubChunkData SubChunkData = ChunkData.SubChunks[SubChunkIdx];
+					if (SubChunkData.LineIndex >= NewLineCount)
+					{
+						try
+						{
+							Indexes.Add(SubChunkData.BuildIndex());
+						}
+						catch (Exception Ex)
+						{
+							throw new Exception($"Failed to create index block - log {LogFile.Id}, chunk {ChunkIdx} ({LogFile.Chunks.Count}), subchunk {SubChunkIdx} ({ChunkData.SubChunks.Count}), index lines: {NewLineCount}, chunk index: {Chunk.LineIndex}, subchunk index: {Chunk.LineIndex + ChunkData.SubChunkLineIndex[SubChunkIdx]}, subchunk count: {SubChunkData.LineCount}", Ex);
+						}
+
+						NewLength = SubChunkData.Offset + SubChunkData.Length;
+						NewLineCount = SubChunkData.LineIndex + SubChunkData.LineCount;
+					}
+				}
+			}
+
+			// Try to update the log file
+			ILogFile? NewLogFile = LogFile;
+			if (NewLength > (LogFile.IndexLength ?? 0))
+			{
+				LogIndexData Index = LogIndexData.Merge(Indexes);
+				Logger.LogDebug("Writing index for log {LogId} covering {Length} (index length {IndexLength})", LogFile.Id, NewLength, Index.GetSerializedSize());
+
+				await WriteIndexAsync(LogFile.Id, NewLength, Index);
+
+				while(NewLogFile != null && NewLength > (NewLogFile.IndexLength ?? 0))
+				{
+					NewLogFile = await LogFiles.TryUpdateIndexAsync(NewLogFile, NewLength);
+					if(NewLogFile != null)
+					{
+						break;
+					}
+					NewLogFile = await LogFiles.GetLogFileAsync(LogFile.Id);
+				}
+			}
+			return NewLogFile;
+		}
+
+		/// <summary>
+		/// Reads a chunk from storage
+		/// </summary>
+		/// <param name="LogFile">Log file to read from</param>
+		/// <param name="ChunkIdx">The chunk to read</param>
+		/// <returns>Chunk data</returns>
+		[SuppressMessage("Design", "CA1031:Do not catch general exception types", Justification = "<Pending>")]
+		private async Task<LogChunkData> ReadChunkAsync(ILogFile LogFile, int ChunkIdx)
+		{
+			ILogChunk Chunk = LogFile.Chunks[ChunkIdx];
+
+			// Try to read the chunk data from storage
+			LogChunkData? ChunkData = null;
+			try
+			{
+				// If the chunk is not yet complete, query the log builder
+				if (Chunk.Length == 0)
+				{
+					ChunkData = await Builder.GetChunkAsync(LogFile.Id, Chunk.Offset, Chunk.LineIndex);
+				}
+
+				// Otherwise go directly to the log storage
+				if (ChunkData == null)
+				{
+					ChunkData = await Storage.ReadChunkAsync(LogFile.Id, Chunk.Offset, Chunk.LineIndex);
+				}
+			}
+			catch (Exception Ex)
+			{
+				Logger.LogError(Ex, "Unable to read log {LogId} at offset {Offset}", LogFile.Id, Chunk.Offset);
+			}
+
+			// Get the minimum length and line count for the chunk
+			if (ChunkIdx + 1 < LogFile.Chunks.Count)
+			{
+				ILogChunk NextChunk = LogFile.Chunks[ChunkIdx + 1];
+				ChunkData = RepairChunkData(LogFile, ChunkIdx, ChunkData, (int)(NextChunk.Offset - Chunk.Offset), NextChunk.LineIndex - Chunk.LineIndex);
+			}
+			else
+			{
+				if (LogFile.MaxLineIndex != null && Chunk.Length != 0)
+				{
+					ChunkData = RepairChunkData(LogFile, ChunkIdx, ChunkData, Chunk.Length, LogFile.MaxLineIndex.Value - Chunk.LineIndex);
+				}
+				else if(ChunkData == null)
+				{
+					ChunkData = RepairChunkData(LogFile, ChunkIdx, ChunkData, 1024, 1);
+				}
+			}
+
+			return ChunkData;
+		}
+
+		/// <summary>
+		/// Validates the given chunk data, and fix it up if necessary
+		/// </summary>
+		/// <param name="LogFile">The log file instance</param>
+		/// <param name="ChunkIdx">Index of the chunk within the logfile</param>
+		/// <param name="ChunkData">The chunk data that was read</param>
+		/// <param name="Length">Expected length of the data</param>
+		/// <param name="LineCount">Expected number of lines in the data</param>
+		/// <returns>Repaired chunk data</returns>
+		LogChunkData RepairChunkData(ILogFile LogFile, int ChunkIdx, LogChunkData? ChunkData, int Length, int LineCount)
+		{
+			int CurrentLength = 0;
+			int CurrentLineCount = 0;
+			if(ChunkData != null)
+			{
+				CurrentLength = ChunkData.Length;
+				CurrentLineCount = ChunkData.LineCount;
+			}
+
+			if (ChunkData == null || CurrentLength < Length || CurrentLineCount < LineCount)
+			{
+				Logger.LogWarning("Creating placeholder subchunk for log {LogId} chunk {ChunkIdx} (length {Length} vs expected {ExpLength}, lines {LineCount} vs expected {ExpLineCount})", LogFile.Id, ChunkIdx, CurrentLength, Length, CurrentLineCount, LineCount);
+
+				List<LogSubChunkData> SubChunks = new List<LogSubChunkData>();
+				if (ChunkData != null && ChunkData.Length < Length && ChunkData.LineCount < LineCount)
+				{
+					SubChunks.AddRange(ChunkData.SubChunks);
+				}
+
+				LogText Text = new LogText();
+				Text.AppendMissingDataInfo(ChunkIdx, LogFile.Chunks[ChunkIdx].Server, Length - CurrentLength, LineCount - CurrentLineCount);
+				SubChunks.Add(new LogSubChunkData(LogFile.Type, CurrentLength, CurrentLineCount, Text));
+
+				ILogChunk Chunk = LogFile.Chunks[ChunkIdx];
+				ChunkData = new LogChunkData(Chunk.Offset, Chunk.LineIndex, SubChunks);
+			}
+			return ChunkData;
+		}
+
+		/// <summary>
+		/// Writes a set of chunks to the database
+		/// </summary>
+		/// <param name="LogFileId">Unique id of the log file</param>
+		/// <param name="Offset">Offset of the chunk to write</param>
+		/// <param name="LineIndex">First line index of the chunk</param>
+		/// <returns>Chunk daata</returns>
+		[SuppressMessage("Design", "CA1031:Do not catch general exception types", Justification = "<Pending>")]
+		private async Task<LogChunkData?> WriteChunkAsync(LogId LogFileId, long Offset, int LineIndex)
+		{
+			// Write the chunk to storage
+			LogChunkData? ChunkData = await Builder.GetChunkAsync(LogFileId, Offset, LineIndex);
+			if (ChunkData == null)
+			{
+				Logger.LogDebug("Log {LogId} offset {Offset} not found in log builder", LogFileId, Offset);
+			}
+			else
+			{
+				try
+				{
+					await Storage.WriteChunkAsync(LogFileId, ChunkData.Offset, ChunkData);
+				}
+				catch (Exception Ex)
+				{
+					Logger.LogError(Ex, "Unable to write log {LogId} at offset {Offset}", LogFileId, ChunkData.Offset);
+				}
+			}
+
+			// Remove it from the log builder
+			try
+			{
+				await Builder.RemoveChunkAsync(LogFileId, Offset);
+			}
+			catch (Exception Ex)
+			{
+				Logger.LogError(Ex, "Unable to remove log {LogId} at offset {Offset} from log builder", LogFileId, Offset);
+			}
+
+			return ChunkData;
+		}
+
+		/// <summary>
+		/// Reads a chunk from storage
+		/// </summary>
+		/// <param name="LogFile">Log file to read from</param>
+		/// <param name="Length">Length of the log covered by the index</param>
+		/// <returns>Chunk data</returns>
+		[SuppressMessage("Design", "CA1031:Do not catch general exception types", Justification = "<Pending>")]
+		private async Task<LogIndexData?> ReadIndexAsync(ILogFile LogFile, long Length)
+		{
+			try
+			{
+				LogIndexData? Index = await Storage.ReadIndexAsync(LogFile.Id, Length);
+				return Index;
+			}
+			catch (Exception Ex)
+			{
+				Logger.LogError(Ex, "Unable to read log {LogId} index at length {Length}", LogFile.Id, Length);
+				return null;
+			}
+		}
+
+		/// <summary>
+		/// Writes an index to the database
+		/// </summary>
+		/// <param name="LogFileId">Unique id of the log file</param>
+		/// <param name="Length">Length of the data covered by the index</param>
+		/// <param name="Index">Index to write</param>
+		/// <returns>Async task</returns>
+		[SuppressMessage("Design", "CA1031:Do not catch general exception types", Justification = "<Pending>")]
+		private async Task WriteIndexAsync(LogId LogFileId, long Length, LogIndexData Index)
+		{
+			try
+			{
+				await Storage.WriteIndexAsync(LogFileId, Length, Index);
+			}
+			catch(Exception Ex)
+			{
+				Logger.LogError(Ex, "Unable to write index for log {LogId}", LogFileId);
+			}
+		}
+
+		/// <summary>
+		/// Determines if the user is authorized to perform an action on a particular template
+		/// </summary>
+		/// <param name="LogFile">The template to check</param>
+		/// <param name="User">The principal to authorize</param>
+		/// <returns>True if the action is authorized</returns>
+		public static bool AuthorizeForSession(ILogFile LogFile, ClaimsPrincipal User)
+		{
+			if(LogFile.SessionId != null)
+			{
+				return User.HasClaim(HordeClaimTypes.AgentSessionId, LogFile.SessionId.Value.ToString());
+			}
+			else
+			{
+				return false;
+			}
+		}
+
+		/// <inheritdoc/>
+		public async Task<List<int>> SearchLogDataAsync(ILogFile LogFile, string Text, int FirstLine, int Count, LogSearchStats SearchStats)
+		{
+			Stopwatch Timer = Stopwatch.StartNew();
+
+			using IScope Scope = GlobalTracer.Instance.BuildSpan("SearchLogDataAsync").StartActive();
+			Scope.Span.SetTag("LogId", LogFile.Id.ToString());
+			Scope.Span.SetTag("Text", Text);
+			Scope.Span.SetTag("Count", Count.ToString(CultureInfo.InvariantCulture));
+
+			List<int> Results = new List<int>();
+			if (Count > 0)
+			{
+				IAsyncEnumerator<int> Enumerator = SearchLogDataInternalAsync(LogFile, Text, FirstLine, SearchStats).GetAsyncEnumerator();
+				while (await Enumerator.MoveNextAsync() && Results.Count < Count)
+				{
+					Results.Add(Enumerator.Current);
+				}
+			}
+
+			Logger.LogDebug("Search for \"{SearchText}\" in log {LogId} found {NumResults}/{MaxResults} results, took {Time}ms ({@Stats})", Text, LogFile.Id, Results.Count, Count, Timer.ElapsedMilliseconds, SearchStats);
+			return Results;
+		}
+
+		/// <inheritdoc/>
+		public async IAsyncEnumerable<int> SearchLogDataInternalAsync(ILogFile LogFile, string Text, int FirstLine, LogSearchStats SearchStats)
+		{
+			SearchText SearchText = new SearchText(Text);
+
+			// Read the index for this log file
+			if (LogFile.IndexLength != null)
+			{
+				LogIndexData? IndexData = await ReadIndexAsync(LogFile, LogFile.IndexLength.Value);
+				if(IndexData != null && FirstLine < IndexData.LineCount)
+				{
+					using IScope IndexScope = GlobalTracer.Instance.BuildSpan("Indexed").StartActive();
+					IndexScope.Span.SetTag("LineCount", IndexData.LineCount.ToString(CultureInfo.InvariantCulture));
+
+					foreach(int LineIndex in IndexData.Search(FirstLine, SearchText, SearchStats))
+					{
+						yield return LineIndex;
+					}
+
+					FirstLine = IndexData.LineCount;
+				}
+			}
+
+			// Manually search through the rest of the log
+			int ChunkIdx = LogFile.Chunks.GetChunkForLine(FirstLine);
+			for (; ChunkIdx < LogFile.Chunks.Count; ChunkIdx++)
+			{
+				ILogChunk Chunk = LogFile.Chunks[ChunkIdx];
+
+				// Read the chunk data
+				LogChunkData ChunkData = await ReadChunkAsync(LogFile, ChunkIdx);
+				if (FirstLine < ChunkData.LineIndex + ChunkData.LineCount)
+				{
+					// Find the first sub-chunk we're looking for
+					int SubChunkIdx = 0;
+					if (FirstLine > Chunk.LineIndex)
+					{
+						SubChunkIdx = ChunkData.GetSubChunkForLine(FirstLine - Chunk.LineIndex);
+					}
+
+					// Search through the sub-chunks
+					for (; SubChunkIdx < ChunkData.SubChunks.Count; SubChunkIdx++)
+					{
+						LogSubChunkData SubChunkData = ChunkData.SubChunks[SubChunkIdx];
+						if (FirstLine < SubChunkData.LineIndex + SubChunkData.LineCount)
+						{
+							// Create an index containing just this sub-chunk
+							LogIndexData Index = SubChunkData.BuildIndex();
+							foreach (int LineIndex in Index.Search(FirstLine, SearchText, SearchStats))
+							{
+								yield return LineIndex;
+							}
+						}
+					}
+				}
+			}
+		}
+	}
+}
+
diff --git a/Engine/Source/Programs/Horde/Horde.Build/Logs/Storage/ILogStorage.cs b/Engine/Source/Programs/Horde/Horde.Build/Logs/Storage/ILogStorage.cs
new file mode 100644
--- /dev/null
+++ b/Engine/Source/Programs/Horde/Horde.Build/Logs/Storage/ILogStorage.cs
@@ -0,0 +1,58 @@
+// Copyright Epic Games, Inc. All Rights Reserved.
+
+using HordeServer.Models;
+using HordeServer.Services;
+using HordeServer.Storage;
+using HordeServer.Utilities;
+using Microsoft.Extensions.Caching.Memory;
+using MongoDB.Bson;
+using System;
+using System.Collections.Generic;
+using System.Linq;
+using System.Threading.Tasks;
+
+namespace HordeServer.Logs
+{
+	using LogId = ObjectId<ILogFile>;
+
+	/// <summary>
+	/// Caching interface for reading and writing log data.
+	/// </summary>
+	public interface ILogStorage : IDisposable
+	{
+		/// <summary>
+		/// Attempts to read an index for the given log file
+		/// </summary>
+		/// <param name="LogId">Unique id of the log file</param>
+		/// <param name="Length">Length of the file covered by the index</param>
+		/// <returns>Index for the log file</returns>
+		Task<LogIndexData?> ReadIndexAsync(LogId LogId, long Length);
+
+		/// <summary>
+		/// Log file to write an index for
+		/// </summary>
+		/// <param name="LogId">Unique id of the log file</param>
+		/// <param name="Length">Length of the file covered by the index</param>
+		/// <param name="Index">The log file index</param>
+		/// <returns>Async task</returns>
+		Task WriteIndexAsync(LogId LogId, long Length, LogIndexData Index);
+
+		/// <summary>
+		/// Retrieves an item from the cache
+		/// </summary>
+		/// <param name="LogId">Unique id of the log file</param>
+		/// <param name="Offset">Offset of the chunk to read</param>
+		/// <param name="LineIndex">First line of the chunk</param>
+		/// <returns>Data for the given key, or null if it's not present</returns>
+		Task<LogChunkData?> ReadChunkAsync(LogId LogId, long Offset, int LineIndex);
+
+		/// <summary>
+		/// Writes a chunk to storage
+		/// </summary>
+		/// <param name="LogId">Unique id of the log file</param>
+		/// <param name="Offset">Offset of the chunk to write</param>
+		/// <param name="ChunkData">Information about the chunk data</param>
+		/// <returns>Async task</returns>
+		Task WriteChunkAsync(LogId LogId, long Offset, LogChunkData ChunkData);
+	}
+}
diff --git a/Engine/Source/Programs/Horde/Horde.Build/Logs/Storage/Impl/LocalLogStorage.cs b/Engine/Source/Programs/Horde/Horde.Build/Logs/Storage/Impl/LocalLogStorage.cs
new file mode 100644
--- /dev/null
+++ b/Engine/Source/Programs/Horde/Horde.Build/Logs/Storage/Impl/LocalLogStorage.cs
@@ -0,0 +1,147 @@
+// Copyright Epic Games, Inc. All Rights Reserved.
+
+using HordeServer.Models;
+using HordeServer.Utilities;
+using Microsoft.Extensions.Caching.Memory;
+using MongoDB.Bson;
+using System;
+using System.Collections.Generic;
+using System.Linq;
+using System.Threading.Tasks;
+
+namespace HordeServer.Logs.Storage
+{
+	using LogId = ObjectId<ILogFile>;
+
+	/// <summary>
+	/// In-memory cache for chunk and index data
+	/// </summary>
+	class LocalLogStorage : ILogStorage
+	{
+		/// <summary>
+		/// The memory cache instance
+		/// </summary>
+		MemoryCache Cache;
+
+		/// <summary>
+		/// The inner storage
+		/// </summary>
+		ILogStorage Inner;
+
+		/// <summary>
+		/// Constructor
+		/// </summary>
+		/// <param name="NumItems">Maximum size for the memory cache</param>
+		/// <param name="Inner">The inner storage provider</param>
+		public LocalLogStorage(int NumItems, ILogStorage Inner)
+		{
+			MemoryCacheOptions Options = new MemoryCacheOptions();
+			Options.SizeLimit = NumItems;
+
+			this.Cache = new MemoryCache(Options);
+			this.Inner = Inner;
+		}
+
+		/// <inheritdoc/>
+		public void Dispose()
+		{
+			Cache.Dispose();
+			Inner.Dispose();
+		}
+
+		/// <summary>
+		/// Gets the cache key for a particular index
+		/// </summary>
+		/// <param name="LogId">The log file to retrieve an index for</param>
+		/// <param name="Length">Length of the file covered by the index</param>
+		/// <returns>Cache key for the index</returns>
+		static string IndexKey(LogId LogId, long Length)
+		{
+			return $"{LogId}/index-{Length}";
+		}
+
+		/// <summary>
+		/// Gets the cache key for a particular chunk
+		/// </summary>
+		/// <param name="LogId">The log file to retrieve an index for</param>
+		/// <param name="Offset">The chunk offset</param>
+		/// <returns>Cache key for the chunk</returns>
+		static string ChunkKey(LogId LogId, long Offset)
+		{
+			return $"{LogId}/chunk-{Offset}";
+		}
+
+		/// <summary>
+		/// Adds an entry to the cache
+		/// </summary>
+		/// <param name="Key">The cache key</param>
+		/// <param name="Value">Value to store</param>
+		void AddEntry(string Key, object? Value)
+		{
+			using (ICacheEntry Entry = Cache.CreateEntry(Key))
+			{
+				if (Value == null)
+				{
+					Entry.SetAbsoluteExpiration(TimeSpan.FromSeconds(30.0));
+					Entry.SetSize(0);
+				}
+				else
+				{
+					Entry.SetSlidingExpiration(TimeSpan.FromHours(4));
+					Entry.SetSize(1);
+				}
+				Entry.SetValue(Value);
+			}
+		}
+
+		/// <summary>
+		/// Reads a value from the cache, or from the inner storage provider, adding the result to the cache
+		/// </summary>
+		/// <typeparam name="T">Type of object to read</typeparam>
+		/// <param name="Key">The cache key</param>
+		/// <param name="ReadInner">Delegate to read from the inner storage provider</param>
+		/// <returns>The retrieved value</returns>
+		async Task<T?> ReadValueAsync<T>(string Key, Func<Task<T?>> ReadInner) where T : class
+		{
+			T? Value;
+			if (!Cache.TryGetValue(Key, out Value))
+			{
+				try
+				{
+					Value = await ReadInner();
+				}
+				finally
+				{
+					AddEntry(Key, Value);
+				}
+			}
+			return Value;
+		}
+
+		/// <inheritdoc/>
+		public Task<LogIndexData?> ReadIndexAsync(LogId LogId, long Length)
+		{
+			return ReadValueAsync(IndexKey(LogId, Length), () => Inner.ReadIndexAsync(LogId, Length));
+		}
+
+		/// <inheritdoc/>
+		public Task WriteIndexAsync(LogId LogId, long Length, LogIndexData IndexData)
+		{
+			AddEntry(IndexKey(LogId, Length), IndexData);
+			return Inner.WriteIndexAsync(LogId, Length, IndexData);
+		}
+
+		/// <inheritdoc/>
+		public Task<LogChunkData?> ReadChunkAsync(LogId LogId, long Offset, int LineIndex)
+		{
+			return ReadValueAsync(ChunkKey(LogId, Offset), () => Inner.ReadChunkAsync(LogId, Offset, LineIndex));
+		}
+
+		/// <inheritdoc/>
+		public Task WriteChunkAsync(LogId LogId, long Offset, LogChunkData ChunkData)
+		{
+			AddEntry(ChunkKey(LogId, Offset), ChunkData);
+			return Inner.WriteChunkAsync(LogId, Offset, ChunkData);
+		}
+	}
+}
diff --git a/Engine/Source/Programs/Horde/Horde.Build/Logs/Storage/Impl/NullLogStorage.cs b/Engine/Source/Programs/Horde/Horde.Build/Logs/Storage/Impl/NullLogStorage.cs
new file mode 100644
--- /dev/null
+++ b/Engine/Source/Programs/Horde/Horde.Build/Logs/Storage/Impl/NullLogStorage.cs
@@ -0,0 +1,49 @@
+// Copyright Epic Games, Inc. All Rights Reserved.
+
+using HordeServer.Models;
+using HordeServer.Utilities;
+using MongoDB.Bson;
+using System;
+using System.Collections.Generic;
+using System.Linq;
+using System.Threading.Tasks;
+
+namespace HordeServer.Logs.Storage.Impl
+{
+	using LogId = ObjectId<ILogFile>;
+
+	/// <summary>
+	/// Empty implementation of log storage
+	/// </summary>
+	public sealed class NullLogStorage : ILogStorage
+	{
+		/// <inheritdoc/>
+		public void Dispose()
+		{
+		}
+
+		/// <inheritdoc/>
+		public Task<LogIndexData?> ReadIndexAsync(LogId LogId, long Length)
+		{
+			return Task.FromResult<LogIndexData?>(null);
+		}
+
+		/// <inheritdoc/>
+		public Task WriteIndexAsync(LogId LogId, long Length, LogIndexData Index)
+		{
+			return Task.CompletedTask;
+		}
+
+		/// <inheritdoc/>
+		public Task<LogChunkData?> ReadChunkAsync(LogId LogId, long Offset, int LineIndex)
+		{
+			return Task.FromResult<LogChunkData?>(null);
+		}
+
+		/// <inheritdoc/>
+		public Task WriteChunkAsync(LogId LogId, long Offset, LogChunkData ChunkData)
+		{
+			return Task.CompletedTask;
+		}
+	}
+}
diff --git a/Engine/Source/Programs/Horde/Horde.Build/Logs/Storage/Impl/PersistentLogStorage.cs b/Engine/Source/Programs/Horde/Horde.Build/Logs/Storage/Impl/PersistentLogStorage.cs
new file mode 100644
--- /dev/null
+++ b/Engine/Source/Programs/Horde/Horde.Build/Logs/Storage/Impl/PersistentLogStorage.cs
@@ -0,0 +1,111 @@
+// Copyright Epic Games, Inc. All Rights Reserved.
+
+using EpicGames.Core;
+using HordeServer.Models;
+using HordeServer.Services;
+using HordeServer.Storage;
+using HordeServer.Utilities;
+using Microsoft.Extensions.Logging;
+using MongoDB.Bson;
+using System;
+using System.Collections.Generic;
+using System.Linq;
+using System.Threading.Tasks;
+
+namespace HordeServer.Logs.Readers
+{
+	using LogId = ObjectId<ILogFile>;
+
+	/// <summary>
+	/// Bulk storage for log file data
+	/// </summary>
+	class PersistentLogStorage : ILogStorage
+	{
+		/// <summary>
+		/// The bulk storage provider to use
+		/// </summary>
+		IStorageBackend StorageProvider;
+
+		/// <summary>
+		/// Log provider
+		/// </summary>
+		ILogger Logger;
+
+		/// <summary>
+		/// Constructor
+		/// </summary>
+		/// <param name="StorageProvider">The storage provider</param>
+		/// <param name="Logger">Logging provider</param>
+		public PersistentLogStorage(IStorageBackend<PersistentLogStorage> StorageProvider, ILogger<PersistentLogStorage> Logger)
+		{
+			this.StorageProvider = StorageProvider;
+			this.Logger = Logger;
+		}
+
+		/// <inheritdoc/>
+		public void Dispose()
+		{
+		}
+
+		/// <inheritdoc/>
+		public async Task<LogIndexData?> ReadIndexAsync(LogId LogId, long Length)
+		{
+			Logger.LogDebug("Reading log {LogId} index length {Length} from persistent storage", LogId, Length);
+
+			string Path = $"{LogId}/index_{Length}";
+			ReadOnlyMemory<byte>? Data = await StorageProvider.ReadBytesAsync(Path);
+			if (Data == null)
+			{
+				return null;
+			}
+			return LogIndexData.FromMemory(Data.Value);
+		}
+
+		/// <inheritdoc/>
+		public Task WriteIndexAsync(LogId LogId, long Length, LogIndexData IndexData)
+		{
+			Logger.LogDebug("Writing log {LogId} index length {Length} to persistent storage", LogId, Length);
+
+			string Path = $"{LogId}/index_{Length}";
+			ReadOnlyMemory<byte> Data = IndexData.ToByteArray();
+			return StorageProvider.WriteBytesAsync(Path, Data);
+		}
+
+		/// <inheritdoc/>
+		public async Task<LogChunkData?> ReadChunkAsync(LogId LogId, long Offset, int LineIndex)
+		{
+			Logger.LogDebug("Reading log {LogId} chunk offset {Offset} from persistent storage", LogId, Offset);
+
+			string Path = $"{LogId}/offset_{Offset}";
+			ReadOnlyMemory<byte>? Data = await StorageProvider.ReadBytesAsync(Path);
+			if(Data == null)
+			{
+				return null;
+			}
+
+			MemoryReader Reader = new MemoryReader(Data.Value);
+			LogChunkData ChunkData = Reader.ReadLogChunkData(Offset, LineIndex);
+
+			if (Reader.Offset != Data.Value.Length)
+			{
+				throw new Exception($"Serialization of persistent chunk {Path} is not at expected offset (expected {Data.Value.Length}, actual {Reader.Offset})");
+			}
+
+			return ChunkData;
+		}
+
+		/// <inheritdoc/>
+		public Task WriteChunkAsync(LogId LogId, long Offset, LogChunkData ChunkData)
+		{
+			Logger.LogDebug("Writing log {LogId} chunk offset {Offset} to persistent storage", LogId, Offset);
+
+			string Path = $"{LogId}/offset_{Offset}";
+			byte[] Data = new byte[ChunkData.GetSerializedSize()];
+			MemoryWriter Writer = new MemoryWriter(Data);
+			Writer.WriteLogChunkData(ChunkData);
+			Writer.CheckOffset(Data.Length);
+
+			return StorageProvider.WriteBytesAsync(Path, Data);
+		}
+	}
+}
diff --git a/Engine/Source/Programs/Horde/Horde.Build/Logs/Storage/Impl/RedisLogStorage.cs b/Engine/Source/Programs/Horde/Horde.Build/Logs/Storage/Impl/RedisLogStorage.cs
new file mode 100644
--- /dev/null
+++ b/Engine/Source/Programs/Horde/Horde.Build/Logs/Storage/Impl/RedisLogStorage.cs
@@ -0,0 +1,175 @@
+// Copyright Epic Games, Inc. All Rights Reserved.
+
+using HordeServer.Models;
+using HordeServer.Utilities;
+using Microsoft.Extensions.Logging;
+using MongoDB.Bson;
+using StackExchange.Redis;
+using System;
+using System.Collections.Generic;
+using System.Diagnostics.CodeAnalysis;
+using System.Linq;
+using System.Threading.Tasks;
+
+namespace HordeServer.Logs.Storage.Impl
+{
+	using LogId = ObjectId<ILogFile>;
+
+	/// <summary>
+	/// Redis log file storage
+	/// </summary>
+	public sealed class RedisLogStorage : ILogStorage
+	{
+		/// <summary>
+		/// The Redis database
+		/// </summary>
+		IDatabase RedisDb;
+
+		/// <summary>
+		/// Logging device
+		/// </summary>
+		ILogger Logger;
+
+		/// <summary>
+		/// Inner storage layer
+		/// </summary>
+		ILogStorage Inner;
+
+		/// <summary>
+		/// Constructor
+		/// </summary>
+		/// <param name="RedisDb">Redis database</param>
+		/// <param name="Logger">Logger instance</param>
+		/// <param name="Inner">Inner storage layer</param>
+		public RedisLogStorage(IDatabase RedisDb, ILogger Logger, ILogStorage Inner)
+		{
+			this.RedisDb = RedisDb;
+			this.Logger = Logger;
+			this.Inner = Inner;
+		}
+
+		/// <inheritdoc/>
+		public void Dispose()
+		{
+			Inner.Dispose();
+		}
+
+		/// <summary>
+		/// Gets the key for a log file's index
+		/// </summary>
+		/// <param name="LogId">The log file id</param>
+		/// <param name="Length">Length of the file covered by the index</param>
+		/// <returns>The index key</returns>
+		static string IndexKey(LogId LogId, long Length)
+		{
+			return $"log-{LogId}-index-{Length}";
+		}
+
+		/// <summary>
+		/// Gets the key for a chunk's data
+		/// </summary>
+		/// <param name="LogId">The log file id</param>
+		/// <param name="Offset">Offset of the chunk within the log file</param>
+		/// <returns>The chunk key</returns>
+		static string ChunkKey(LogId LogId, long Offset)
+		{
+			return $"log-{LogId}-chunk-{Offset}";
+		}
+
+		/// <summary>
+		/// Adds an index to the cache
+		/// </summary>
+		/// <param name="Key">Key for the item to store</param>
+		/// <param name="Value">Value to be stored</param>
+		/// <returns>Async task</returns>
+		[SuppressMessage("Design", "CA1031:Do not catch general exception types")]
+		async Task AddAsync(string Key, ReadOnlyMemory<byte> Value)
+		{
+			try
+			{
+				await RedisDb.StringSetAsync(Key, Value, expiry: TimeSpan.FromHours(1.0));
+				Logger.LogDebug("Added key {Key} to Redis cache ({Size} bytes)", Key, Value.Length);
+			}
+			catch (Exception Ex)
+			{
+				Logger.LogError(Ex, "Error writing Redis key {Key}", Key);
+			}
+		}
+
+		/// <summary>
+		/// Adds an index to the cache
+		/// </summary>
+		/// <param name="Key">Key for the item to store</param>
+		/// <param name="Deserialize">Delegate to deserialize the item</param>
+		/// <returns>Async task</returns>
+		[SuppressMessage("Design", "CA1031:Do not catch general exception types")]
+		async Task<T?> GetAsync<T>(string Key, Func<ReadOnlyMemory<byte>, T> Deserialize) where T : class
+		{
+			try
+			{
+				RedisValue Value = await RedisDb.StringGetAsync(Key);
+				if (Value.IsNullOrEmpty)
+				{
+					Logger.LogDebug("Redis cache miss for {Key}", Key);
+					return null;
+				}
+				else
+				{
+					Logger.LogDebug("Redis cache hit for {Key}", Key);
+					return Deserialize(Value);
+				}
+			}
+			catch (Exception Ex)
+			{
+				Logger.LogError(Ex, "Error reading Redis key {Key}", Key);
+				return null;
+			}
+		}
+
+		/// <inheritdoc/>
+		public async Task<LogIndexData?> ReadIndexAsync(LogId LogId, long Length)
+		{
+			string Key = IndexKey(LogId, Length);
+
+			LogIndexData? IndexData = await GetAsync(Key, Memory => LogIndexData.FromMemory(Memory));
+			if(IndexData == null)
+			{
+				IndexData = await Inner.ReadIndexAsync(LogId, Length);
+				if(IndexData != null)
+				{
+					await AddAsync(Key, IndexData.ToByteArray());
+				}
+			}
+			return IndexData;
+		}
+
+		/// <inheritdoc/>
+		public async Task WriteIndexAsync(LogId LogId, long Length, LogIndexData Index)
+		{
+			await Inner.WriteIndexAsync(LogId, Length, Index);
+		}
+
+		/// <inheritdoc/>
+		public async Task<LogChunkData?> ReadChunkAsync(LogId LogId, long Offset, int LineIndex)
+		{
+			string Key = ChunkKey(LogId, Offset);
+
+			LogChunkData? ChunkData = await GetAsync(Key, Memory => LogChunkData.FromMemory(Memory, Offset, LineIndex));
+			if (ChunkData == null)
+			{
+				ChunkData = await Inner.ReadChunkAsync(LogId, Offset, LineIndex);
+				if (ChunkData != null)
+				{
+					await AddAsync(Key, ChunkData.ToByteArray());
+				}
+			}
+			return ChunkData;
+		}
+
+		/// <inheritdoc/>
+		public async Task WriteChunkAsync(LogId LogId, long Offset, LogChunkData ChunkData)
+		{
+			await Inner.WriteChunkAsync(LogId, Offset, ChunkData);
+		}
+	}
+}
diff --git a/Engine/Source/Programs/Horde/Horde.Build/Logs/Storage/Impl/SequencedLogStorage.cs b/Engine/Source/Programs/Horde/Horde.Build/Logs/Storage/Impl/SequencedLogStorage.cs
new file mode 100644
--- /dev/null
+++ b/Engine/Source/Programs/Horde/Horde.Build/Logs/Storage/Impl/SequencedLogStorage.cs
@@ -0,0 +1,196 @@
+// Copyright Epic Games, Inc. All Rights Reserved.
+
+using HordeServer.Models;
+using HordeServer.Utilities;
+using MongoDB.Bson;
+using System;
+using System.Collections.Generic;
+using System.Linq;
+using System.Threading.Tasks;
+
+namespace HordeServer.Logs.Storage.Impl
+{
+	using LogId = ObjectId<ILogFile>;
+
+	/// <summary>
+	/// Storage layer which caches pending read tasks, to avoid fetching the same item more than once
+	/// </summary>
+	class SequencedLogStorage : ILogStorage
+	{
+		/// <summary>
+		/// Inner storage implementation
+		/// </summary>
+		ILogStorage Inner;
+
+		/// <summary>
+		/// Pending index reads
+		/// </summary>
+		Dictionary<(LogId, long), Task<LogIndexData?>> IndexReadTasks = new Dictionary<(LogId, long), Task<LogIndexData?>>();
+
+		/// <summary>
+		/// Pending index reads
+		/// </summary>
+		Dictionary<(LogId, long), Task> IndexWriteTasks = new Dictionary<(LogId, long), Task>();
+
+		/// <summary>
+		/// Pending chunk reads
+		/// </summary>
+		Dictionary<(LogId, long), Task<LogChunkData?>> ChunkReadTasks = new Dictionary<(LogId, long), Task<LogChunkData?>>();
+
+		/// <summary>
+		/// Pending chunk reads
+		/// </summary>
+		Dictionary<(LogId, long), Task> ChunkWriteTasks = new Dictionary<(LogId, long), Task>();
+
+		/// <summary>
+		/// Constructor
+		/// </summary>
+		/// <param name="Inner">The inner storage provider</param>
+		public SequencedLogStorage(ILogStorage Inner)
+		{
+			this.Inner = Inner;
+		}
+
+		/// <inheritdoc/>
+		public void Dispose()
+		{
+			Inner.Dispose();
+		}
+
+		/// <inheritdoc/>
+		public Task<LogIndexData?> ReadIndexAsync(LogId LogId, long Length)
+		{
+			Task<LogIndexData?>? Task;
+			lock (IndexReadTasks)
+			{
+				if (!IndexReadTasks.TryGetValue((LogId, Length), out Task))
+				{
+					Task = InnerReadIndexAsync(LogId, Length);
+					IndexReadTasks.Add((LogId, Length), Task);
+				}
+			}
+			return Task;
+		}
+
+		/// <inheritdoc/>
+		public Task WriteIndexAsync(LogId LogId, long Length, LogIndexData IndexData)
+		{
+			Task? Task;
+			lock (IndexWriteTasks)
+			{
+				if (!IndexWriteTasks.TryGetValue((LogId, Length), out Task))
+				{
+					Task = InnerWriteIndexAsync(LogId, Length, IndexData);
+					IndexWriteTasks.Add((LogId, Length), Task);
+				}
+			}
+			return Task;
+		}
+
+		/// <inheritdoc/>
+		public Task<LogChunkData?> ReadChunkAsync(LogId LogId, long Offset, int LineIndex)
+		{
+			Task<LogChunkData?>? Task;
+			lock (ChunkReadTasks)
+			{
+				if (!ChunkReadTasks.TryGetValue((LogId, Offset), out Task))
+				{
+					Task = InnerReadChunkAsync(LogId, Offset, LineIndex);
+					ChunkReadTasks[(LogId, Offset)] = Task;
+				}
+			}
+			return Task;
+		}
+
+		/// <inheritdoc/>
+		public Task WriteChunkAsync(LogId LogId, long Offset, LogChunkData ChunkData)
+		{
+			Task? Task;
+			lock (ChunkWriteTasks)
+			{
+				if (!ChunkWriteTasks.TryGetValue((LogId, Offset), out Task))
+				{
+					Task = InnerWriteChunkAsync(LogId, Offset, ChunkData);
+					ChunkWriteTasks[(LogId, Offset)] = Task;
+				}
+			}
+			return Task;
+		}
+
+		/// <summary>
+		/// Wrapper for reading an index from the inner storage provider
+		/// </summary>
+		/// <param name="LogId">The log file to read the index for</param>
+		/// <param name="Length">Length of the file that's indexed</param>
+		/// <returns>The index data</returns>
+		async Task<LogIndexData?> InnerReadIndexAsync(LogId LogId, long Length)
+		{
+			await Task.Yield();
+
+			LogIndexData? IndexData = await Inner.ReadIndexAsync(LogId, Length);
+			lock (IndexReadTasks)
+			{
+				IndexReadTasks.Remove((LogId, Length));
+			}
+			return IndexData;
+		}
+
+		/// <summary>
+		/// Wrapper for reading an index from the inner storage provider
+		/// </summary>
+		/// <param name="LogId">The log file to read the index for</param>
+		/// <param name="Length">Length of the indexed data</param>
+		/// <param name="IndexData">The index data to write</param>
+		/// <returns>The index data</returns>
+		async Task InnerWriteIndexAsync(LogId LogId, long Length, LogIndexData IndexData)
+		{
+			await Task.Yield();
+
+			await Inner.WriteIndexAsync(LogId, Length, IndexData);
+
+			lock (IndexWriteTasks)
+			{
+				IndexWriteTasks.Remove((LogId, Length));
+			}
+		}
+
+		/// <summary>
+		/// Wrapper for reading a chunk from the inner storage provider
+		/// </summary>
+		/// <param name="LogId">The log file to read the index for</param>
+		/// <param name="Offset">Offset of the chunk to read</param>
+		/// <param name="LineIndex">Index of the first line in this chunk</param>
+		/// <returns>The index data</returns>
+		async Task<LogChunkData?> InnerReadChunkAsync(LogId LogId, long Offset, int LineIndex)
+		{
+			await Task.Yield();
+
+			LogChunkData? ChunkData = await Inner.ReadChunkAsync(LogId, Offset, LineIndex);
+			lock (ChunkReadTasks)
+			{
+				ChunkReadTasks.Remove((LogId, Offset));
+			}
+			return ChunkData;
+		}
+
+		/// <summary>
+		/// Wrapper for reading a chunk from the inner storage provider
+		/// </summary>
+		/// <param name="LogId">The log file to write the chunk for</param>
+		/// <param name="Offset">Offset of the chunk within the log</param>
+		/// <param name="ChunkData">The chunk data</param>
+		/// <returns>The index data</returns>
+		async Task<LogChunkData> InnerWriteChunkAsync(LogId LogId, long Offset, LogChunkData ChunkData)
+		{
+			await Task.Yield();
+
+			await Inner.WriteChunkAsync(LogId, Offset, ChunkData);
+
+			lock (ChunkWriteTasks)
+			{
+				ChunkWriteTasks.Remove((LogId, Offset));
+			}
+			return ChunkData;
+		}
+	}
+}
diff --git a/Engine/Source/Programs/Horde/Horde.Build/Models/Globals.cs b/Engine/Source/Programs/Horde/Horde.Build/Models/Globals.cs
new file mode 100644
--- /dev/null
+++ b/Engine/Source/Programs/Horde/Horde.Build/Models/Globals.cs
@@ -0,0 +1,411 @@
+// Copyright Epic Games, Inc. All Rights Reserved.
+
+using Google.Protobuf.WellKnownTypes;
+using HordeServer.Services;
+using HordeServer.Utilities;
+using Microsoft.IdentityModel.Tokens;
+using MongoDB.Bson;
+using MongoDB.Bson.Serialization.Attributes;
+using System;
+using System.Collections.Generic;
+using System.ComponentModel.DataAnnotations;
+using System.Diagnostics;
+using System.Linq;
+using System.Security.Cryptography;
+using System.Threading.Tasks;
+
+namespace HordeServer.Models
+{
+	/// <summary>
+	/// Base class for singleton documents
+	/// </summary>
+	public abstract class SingletonBase
+	{
+		/// <summary>
+		/// Unique id for this singleton
+		/// </summary>
+		[BsonId]
+		public ObjectId Id { get; set; }
+
+		/// <summary>
+		/// The revision index of this document
+		/// </summary>
+		public int Revision { get; set; }
+
+		/// <summary>
+		/// Callback to allow the singleton to fix up itself after being read
+		/// </summary>
+		public virtual void PostLoad()
+		{
+		}
+	}
+
+	/// <summary>
+	/// How frequently the maintence window repeats
+	/// </summary>
+	public enum ScheduledDowntimeFrequency
+	{
+		/// <summary>
+		/// Once
+		/// </summary>
+		Once,
+
+		/// <summary>
+		/// Every day
+		/// </summary>
+		Daily,
+
+		/// <summary>
+		/// Every week
+		/// </summary>
+		Weekly,
+	}
+
+	/// <summary>
+	/// Settings for the maintenance window
+	/// </summary>
+	public class ScheduledDowntime
+	{
+		/// <summary>
+		/// Start time
+		/// </summary>
+		public DateTimeOffset StartTime { get; set; }
+
+		/// <summary>
+		/// Finish time
+		/// </summary>
+		public DateTimeOffset FinishTime { get; set; }
+
+		/// <summary>
+		/// Frequency that the window repeats
+		/// </summary>
+		public ScheduledDowntimeFrequency Frequency { get; set; } = ScheduledDowntimeFrequency.Once;
+
+		/// <summary>
+		/// Gets the next scheduled downtime
+		/// </summary>
+		/// <param name="Now">The current time</param>
+		/// <returns>Start and finish time</returns>
+		public (DateTimeOffset StartTime, DateTimeOffset FinishTime) GetNext(DateTimeOffset Now)
+		{
+			TimeSpan Offset = TimeSpan.Zero;
+			if (Frequency == ScheduledDowntimeFrequency.Daily)
+			{
+				double Days = (Now - StartTime).TotalDays;
+				if (Days >= 1.0)
+				{
+					Days -= Days % 1.0;
+				}
+				Offset = TimeSpan.FromDays(Days);
+			}
+			else if (Frequency == ScheduledDowntimeFrequency.Weekly)
+			{
+				double Days = (Now - StartTime).TotalDays;
+				if (Days >= 7.0)
+				{
+					Days -= Days % 7.0;
+				}
+				Offset = TimeSpan.FromDays(Days);
+			}
+			return (StartTime + Offset, FinishTime + Offset);
+		}
+
+		/// <summary>
+		/// Determines if this schedule is active
+		/// </summary>
+		/// <param name="Now">The current time</param>
+		/// <returns>True if downtime is active</returns>
+		public bool IsActive(DateTimeOffset Now)
+		{
+			if (Frequency == ScheduledDowntimeFrequency.Once)
+			{
+				return Now >= StartTime && Now < FinishTime;
+			}
+			else if (Frequency == ScheduledDowntimeFrequency.Daily)
+			{
+				double Days = (Now - StartTime).TotalDays;
+				if (Days < 0.0)
+				{
+					return false;
+				}
+				else
+				{
+					return (Days % 1.0) < (FinishTime - StartTime).TotalDays;
+				}
+			}
+			else if (Frequency == ScheduledDowntimeFrequency.Weekly)
+			{
+				double Days = (Now - StartTime).TotalDays;
+				if(Days < 0.0)
+				{
+					return false;
+				}
+				else
+				{
+					return (Days % 7.0) < (FinishTime - StartTime).TotalDays;
+				}
+			}
+			else
+			{
+				return false;
+			}
+		}
+	}
+
+	/// <summary>
+	/// User notice
+	/// </summary>
+	public class Notice
+	{
+		/// <summary>
+		/// Unique id for this notice
+		/// </summary>
+		public string Id { get; set; } = String.Empty;
+
+		/// <summary>
+		/// Start time to display this message
+		/// </summary>
+		public DateTime? StartTime { get; set; }
+
+		/// <summary>
+		/// Finish time to display this message
+		/// </summary>
+		public DateTime? FinishTime { get; set; }
+
+		/// <summary>
+		/// Message to display
+		/// </summary>
+		public string Message { get; set; } = String.Empty;
+	}
+
+	/// <summary>
+	/// Path to a platform and stream to use for syncing AutoSDK
+	/// </summary>
+	public class AutoSdkWorkspace
+	{
+		/// <summary>
+		/// Name of this workspace
+		/// </summary>
+		public string? Name { get; set; }
+
+		/// <summary>
+		/// The agent properties to check (eg. "OSFamily=Windows")
+		/// </summary>
+		public List<string> Properties { get; set; } = new List<string>();
+
+		/// <summary>
+		/// Username for logging in to the server
+		/// </summary>
+		public string? UserName { get; set; }
+
+		/// <summary>
+		/// Stream to use
+		/// </summary>
+		[Required]
+		public string? Stream { get; set; }
+	}
+
+	/// <summary>
+	/// Information about an individual Perforce server
+	/// </summary>
+	public class PerforceServer
+	{
+		/// <summary>
+		/// The server and port. The server may be a DNS entry with multiple records, in which case it will be actively load balanced.
+		/// </summary>
+		public string ServerAndPort { get; set; } = "perforce:1666";
+
+		/// <summary>
+		/// Whether to query the healthcheck address under each server
+		/// </summary>
+		public bool HealthCheck { get; set; }
+
+		/// <summary>
+		/// Whether to resolve the DNS entries and load balance between different hosts
+		/// </summary>
+		public bool ResolveDns { get; set; }
+
+		/// <summary>
+		/// Maximum number of simultaneous conforms on this server
+		/// </summary>
+		public int MaxConformCount { get; set; }
+
+		/// <summary>
+		/// List of properties for an agent to be eligable to use this server
+		/// </summary>
+		public List<string>? Properties { get; set; }
+	}
+
+	/// <summary>
+	/// Credentials for a Perforce user
+	/// </summary>
+	public class PerforceCredentials
+	{
+		/// <summary>
+		/// The username
+		/// </summary>
+		public string UserName { get; set; } = String.Empty;
+
+		/// <summary>
+		/// Password for the user
+		/// </summary>
+		public string Password { get; set; } = String.Empty;
+	}
+
+	/// <summary>
+	/// Information about a cluster of Perforce servers. 
+	/// </summary>
+	public class PerforceCluster
+	{
+		/// <summary>
+		/// The default cluster name
+		/// </summary>
+		public const string DefaultName = "Default";
+
+		/// <summary>
+		/// Name of the cluster
+		/// </summary>
+		[Required]
+		public string Name { get; set; } = null!;
+
+		/// <summary>
+		/// Username for Horde to log in to this server. Will use the default user if not set.
+		/// </summary>
+		public string? ServiceAccount { get; set; }
+
+		/// <summary>
+		/// Whether the service account can impersonate other users
+		/// </summary>
+		public bool CanImpersonate { get; set; } = true;
+
+		/// <summary>
+		/// List of servers
+		/// </summary>
+		public List<PerforceServer> Servers { get; set; } = new List<PerforceServer>();
+
+		/// <summary>
+		/// List of server credentials
+		/// </summary>
+		public List<PerforceCredentials> Credentials { get; set; } = new List<PerforceCredentials>();
+
+		/// <summary>
+		/// List of autosdk streams
+		/// </summary>
+		public List<AutoSdkWorkspace> AutoSdk { get; set; } = new List<AutoSdkWorkspace>();
+	}
+
+	/// <summary>
+	/// Global server settings
+	/// </summary>
+	[SingletonDocument("5e3981cb28b8ec59cd07184a")]
+	public class Globals : SingletonBase
+	{
+		/// <summary>
+		/// The unique id for all globals objects
+		/// </summary>
+		public static ObjectId StaticId { get; } = new ObjectId("5e3981cb28b8ec59cd07184a");
+
+		/// <summary>
+		/// Unique instance id of this database
+		/// </summary>
+		public ObjectId InstanceId { get; set; }
+
+		/// <summary>
+		/// Set to true to drop the database on restart.
+		/// </summary>
+		public bool ForceReset { get; set; }
+
+		/// <summary>
+		/// The config revision
+		/// </summary>
+		public string? ConfigRevision { get; set; }
+
+		/// <summary>
+		/// Manually added status messages
+		/// </summary>
+		public List<Notice> Notices { get; set; } = new List<Notice>();
+
+		/// <summary>
+		/// List of Perforce clusters
+		/// </summary>
+		public List<PerforceCluster> PerforceClusters { get; set; } = new List<PerforceCluster>();
+
+		/// <summary>
+		/// Maximum number of simultaneous conforms
+		/// </summary>
+		public int MaxConformCount { get; set; }
+
+		/// <summary>
+		/// Randomly generated JWT signing key
+		/// </summary>
+		public byte[]? JwtSigningKey { get; set; }
+
+		/// <summary>
+		/// The current schema version
+		/// </summary>
+		public int? SchemaVersion { get; set; }
+
+		/// <summary>
+		/// The scheduled downtime
+		/// </summary>
+		public List<ScheduledDowntime> ScheduledDowntime { get; set; } = new List<ScheduledDowntime>();
+
+		/// <summary>
+		/// Private constructor for serialization
+		/// </summary>
+		public Globals()
+		{
+			InstanceId = ObjectId.GenerateNewId();
+		}
+
+		/// <summary>
+		/// Rotate the signing key
+		/// </summary>
+		public void RotateSigningKey()
+		{
+			using (RNGCryptoServiceProvider Generator = new RNGCryptoServiceProvider())
+			{
+				JwtSigningKey = new byte[128];
+				Generator.GetBytes(JwtSigningKey);
+			}
+		}
+
+		/// <summary>
+		/// Finds a perforce cluster with the given name
+		/// </summary>
+		/// <param name="Name">Name of the cluster</param>
+		/// <returns></returns>
+		public PerforceCluster? FindPerforceCluster(string? Name)
+		{
+			List<PerforceCluster> Clusters = PerforceClusters;
+			if (Clusters.Count == 0)
+			{
+				Clusters = DefaultClusters;
+			}
+
+			if (Name == null)
+			{
+				return Clusters.FirstOrDefault();
+			}
+			else
+			{
+				return Clusters.FirstOrDefault(x => String.Equals(x.Name, Name, StringComparison.OrdinalIgnoreCase));
+			}
+		}
+
+		static List<PerforceCluster> DefaultClusters { get; } = GetDefaultClusters();
+
+		static List<PerforceCluster> GetDefaultClusters()
+		{
+			PerforceServer Server = new PerforceServer();
+			Server.ServerAndPort = Perforce.P4.P4Server.Get("P4PORT") ?? "perforce:1666";
+
+			PerforceCluster Cluster = new PerforceCluster();
+			Cluster.Name = "Default";
+			Cluster.CanImpersonate = false;
+			Cluster.Servers.Add(Server);
+
+			return new List<PerforceCluster> { Cluster };
+		}
+	}
+}
diff --git a/Engine/Source/Programs/Horde/Horde.Build/Notifications/INotificationService.cs b/Engine/Source/Programs/Horde/Horde.Build/Notifications/INotificationService.cs
new file mode 100644
--- /dev/null
+++ b/Engine/Source/Programs/Horde/Horde.Build/Notifications/INotificationService.cs
@@ -0,0 +1,84 @@
+// Copyright Epic Games, Inc. All Rights Reserved.
+
+using HordeServer.Api;
+using HordeServer.Models;
+using MongoDB.Bson;
+using System;
+using System.Collections.Generic;
+using System.Linq;
+using System.Security.Claims;
+using System.Threading.Tasks;
+
+namespace HordeServer.Notifications
+{
+	/// <summary>
+	/// Interface for the notification service
+	/// </summary>
+	public interface INotificationService
+	{
+		/// <summary>
+		/// Updates a subscription to a trigger
+		/// </summary>
+		/// <param name="TriggerId"></param>
+		/// <param name="User"></param>
+		/// <param name="Email">Whether to receive email notifications</param>
+		/// <param name="Slack">Whether to receive Slack notifications</param>
+		/// <returns></returns>
+		Task<bool> UpdateSubscriptionsAsync(ObjectId TriggerId, ClaimsPrincipal User, bool? Email, bool? Slack);
+
+		/// <summary>
+		/// Gets the current subscriptions for a user
+		/// </summary>
+		/// <param name="TriggerId"></param>
+		/// <param name="User"></param>
+		/// <returns>Subscriptions for that user</returns>
+		Task<INotificationSubscription?> GetSubscriptionsAsync(ObjectId TriggerId, ClaimsPrincipal User);
+
+		/// <summary>
+		/// Notify all subscribers that a job step has finished
+		/// </summary>
+		/// <param name="Job">The job containing the step that has finished</param>
+		/// <param name="Graph"></param>
+		/// <param name="BatchId">The batch id</param>
+		/// <param name="StepId">The step id</param>
+		/// <returns>Async task</returns>
+		void NotifyJobStepComplete(IJob Job, IGraph Graph, SubResourceId BatchId, SubResourceId StepId);
+
+		/// <summary>
+		/// Notify all subscribers that a job step's outcome has changed
+		/// </summary>
+		/// <param name="Job">The job containing the step that has changed</param>
+		/// <param name="OldLabelStates">The old label states for the job</param>
+		/// <param name="NewLabelStates">The new label states for the job</param>
+		/// <returns>Async task</returns>
+		void NotifyLabelUpdate(IJob Job, IReadOnlyList<(LabelState, LabelOutcome)> OldLabelStates, IReadOnlyList<(LabelState, LabelOutcome)> NewLabelStates);
+
+		/// <summary>
+		/// Notify slack channel about a stream update failure
+		/// </summary>
+		/// <param name="ErrorMessage">Error message passed back</param>
+		/// <param name="FileName"></param>
+		/// <param name="Change">Latest change number of the file</param>
+		/// <param name="Author">Author of the change, if found from p4 service</param>
+		/// <param name="Description">Description of the change, if found from p4 service</param>
+		void NotifyConfigUpdateFailure(string ErrorMessage, string FileName, int? Change = null, IUser? Author = null, string? Description = null);
+
+		/// <summary>
+		/// Sends a notification to a user regarding a build health issue.
+		/// </summary>
+		/// <param name="Issue">The new issue that was created</param>
+		void NotifyIssueUpdated(IIssue Issue);
+
+		/// <summary>
+        /// Send a device service notification
+        /// </summary>
+        /// <param name="Message">The message to send</param>
+        /// <param name="Device">The device</param>
+        /// <param name="Pool">The pool</param>
+        /// <param name="Stream"></param>
+        /// <param name="Job"></param>
+        /// <param name="Step"></param>
+        /// <param name="Node"></param>
+        void NotifyDeviceService(string Message, IDevice? Device = null, IDevicePool? Pool = null, IStream? Stream = null, IJob? Job = null, IJobStep? Step = null, INode? Node = null);
+    }
+}
diff --git a/Engine/Source/Programs/Horde/Horde.Build/Notifications/INotificationSink.cs b/Engine/Source/Programs/Horde/Horde.Build/Notifications/INotificationSink.cs
new file mode 100644
--- /dev/null
+++ b/Engine/Source/Programs/Horde/Horde.Build/Notifications/INotificationSink.cs
@@ -0,0 +1,97 @@
+// Copyright Epic Games, Inc. All Rights Reserved.
+
+using HordeCommon;
+using HordeServer.Api;
+using HordeServer.Models;
+using MongoDB.Bson;
+using System;
+using System.Collections.Generic;
+using System.Linq;
+using System.Threading.Tasks;
+
+namespace HordeServer.Notifications
+{
+	/// <summary>
+	/// Implements a notification method
+	/// </summary>
+	public interface INotificationSink
+	{
+		/// <summary>
+		/// Send notifications that a job has completed
+		/// </summary>
+		/// <param name="JobStream"></param>
+		/// <param name="Job">The job containing the step</param>
+		/// <param name="Graph"></param>
+		/// <param name="Outcome"></param>
+		/// <returns>Async task</returns>
+		Task NotifyJobCompleteAsync(IStream JobStream, IJob Job, IGraph Graph, LabelOutcome Outcome);
+
+		/// <summary>
+		/// Send notifications that a job has completed
+		/// </summary>
+		/// <param name="User">User to notify</param>
+		/// <param name="JobStream"></param>
+		/// <param name="Job">The job containing the step</param>
+		/// <param name="Graph"></param>
+		/// <param name="Outcome"></param>
+		/// <returns>Async task</returns>
+		Task NotifyJobCompleteAsync(IUser User, IStream JobStream, IJob Job, IGraph Graph, LabelOutcome Outcome);
+
+		/// <summary>
+		/// Send notifications that a job step has completed
+		/// </summary>
+		/// <param name="User">User to notify</param>
+		/// <param name="JobStream">Stream containing the job</param>
+		/// <param name="Job">The job containing the step</param>
+		/// <param name="Batch">Unique id of the batch</param>
+		/// <param name="Step">The step id</param>
+		/// <param name="Node">Corresponding node for the step</param>
+		/// <param name="JobStepEventData"></param>
+		/// <returns>Async task</returns>
+		Task NotifyJobStepCompleteAsync(IUser User, IStream JobStream, IJob Job, IJobStepBatch Batch, IJobStep Step, INode Node, List<ILogEventData> JobStepEventData);
+
+		/// <summary>
+		/// Send notifications that a job step has completed
+		/// </summary>
+		/// <param name="User">User to notify</param>
+		/// <param name="Job">The job containing the step</param>
+		/// <param name="Stream"></param>
+		/// <param name="Label"></param>
+		/// <param name="LabelIdx"></param>
+		/// <param name="Outcome"></param>
+		/// <param name="StepData"></param>
+		/// <returns>Async task</returns>
+		Task NotifyLabelCompleteAsync(IUser User, IJob Job, IStream Stream, ILabel Label, int LabelIdx, LabelOutcome Outcome, List<(string, JobStepOutcome, Uri)> StepData);
+
+		/// <summary>
+		/// Send notifications that a new issue has been created or assigned
+		/// </summary>
+		/// <param name="Issue"></param>
+		/// <returns></returns>
+		Task NotifyIssueUpdatedAsync(IIssue Issue);
+
+		/// <summary>
+		/// Notification that a stream has failed to update
+		/// </summary>
+		/// <param name="ErrorMessage"></param>
+		/// <param name="FileName"></param>
+		/// <param name="Change"></param>
+		/// <param name="Author"></param>
+		/// <param name="Description"></param>
+		/// <returns></returns>
+		Task NotifyConfigUpdateFailureAsync(string ErrorMessage, string FileName, int? Change = null, IUser? Author = null, string? Description = null);
+		
+		/// <summary>
+		/// Notification for device service
+		/// </summary>
+		/// <param name="Message"></param>
+		/// <param name="Device"></param>
+		/// <param name="Pool"></param>
+		/// <param name="Stream"></param>
+		/// <param name="Job"></param>
+        /// <param name="Step"></param>
+        /// <param name="Node"></param>
+		/// <returns></returns>
+        Task NotifyDeviceServiceAsync(string Message, IDevice? Device = null, IDevicePool? Pool = null, IStream? Stream = null, IJob? Job = null, IJobStep? Step = null, INode? Node = null);
+    }
+}
