Horde: Rename the HordeServer project to Horde.Build.

diff --git a/Engine/Source/Programs/Horde/Horde.Build/Notifications/Impl/NotificationService.cs b/Engine/Source/Programs/Horde/Horde.Build/Notifications/Impl/NotificationService.cs
new file mode 100644
--- /dev/null
+++ b/Engine/Source/Programs/Horde/Horde.Build/Notifications/Impl/NotificationService.cs
@@ -0,0 +1,619 @@
+// Copyright Epic Games, Inc. All Rights Reserved.
+
+using System;
+using System.Collections.Concurrent;
+using System.Collections.Generic;
+using System.Drawing;
+using System.Linq;
+using System.Net;
+using System.Net.Http;
+using System.Net.Mail;
+using System.Security.Claims;
+using System.Text;
+using System.Threading;
+using System.Threading.Tasks;
+using EpicGames.Core;
+using HordeServer.Api;
+using HordeServer.Collections;
+using HordeCommon;
+using HordeServer.Models;
+using HordeServer.Utilities;
+using HordeServer.Utilities.BlockKit;
+using HordeServer.Utilities.Slack.BlockKit;
+using Microsoft.Extensions.Hosting;
+using Microsoft.Extensions.Logging;
+using Microsoft.Extensions.Options;
+using MongoDB.Bson;
+using MongoDB.Driver;
+using Newtonsoft.Json;
+using Newtonsoft.Json.Linq;
+using StatsdClient;
+using JsonSerializer = System.Text.Json.JsonSerializer;
+using HordeServer.Notifications;
+using HordeServer.Services;
+
+namespace HordeServer.Notifications.Impl
+{
+	using UserId = ObjectId<IUser>;
+
+	/// <summary>
+	/// Wraps funtionality for delivering notifications.
+	/// </summary>
+	public class NotificationService : BackgroundService, INotificationService
+	{
+		/// <summary>
+		/// The available notification sinks
+		/// </summary>
+		private List<INotificationSink> Sinks;
+
+		/// <summary>
+		/// Collection of subscriptions
+		/// </summary>
+		private readonly ISubscriptionCollection SubscriptionCollection;
+
+		/// <summary>
+		/// Collection of notification request documents.
+		/// </summary>
+		private readonly INotificationTriggerCollection TriggerCollection;
+
+		/// <summary>
+		/// Instance of the <see cref="GraphCollection"/>.
+		/// </summary>
+		private readonly IGraphCollection GraphCollection;
+
+		/// <summary>
+		/// 
+		/// </summary>
+		private readonly IUserCollection UserCollection;
+
+		/// <summary>
+		/// Job service instance
+		/// </summary>
+		private readonly JobService JobService;
+
+		/// <summary>
+		/// Instance of the <see cref="StreamService"/>.
+		/// </summary>
+		private readonly StreamService StreamService;
+
+		/// <summary>
+		/// 
+		/// </summary>
+		private readonly IIssueService IssueService;
+
+		/// <summary>
+		/// Instance of the <see cref="LogFileService"/>.
+		/// </summary>
+		private readonly ILogFileService LogFileService;
+		
+		/// <summary>
+		/// Instance of the <see cref="IDogStatsd"/>.
+		/// </summary>
+		private readonly IDogStatsd DogStatsd;
+		
+		/// <summary>
+		/// Settings for the application.
+		/// </summary>
+		private readonly IOptionsMonitor<ServerSettings> Settings;
+
+		/// <summary>
+		/// List of asychronous tasks currently executing
+		/// </summary>
+		private readonly ConcurrentQueue<Task> NewTasks = new ConcurrentQueue<Task>();
+
+		/// <summary>
+		/// Set when there are new tasks to wait for
+		/// </summary>
+		private AsyncEvent NewTaskEvent = new AsyncEvent();
+
+		/// <summary>
+		/// Settings for the application.
+		/// </summary>
+		private readonly ILogger<NotificationService> Logger;
+
+		/// <summary>
+		/// Constructor
+		/// </summary>
+		public NotificationService(IEnumerable<INotificationSink> Sinks, IOptionsMonitor<ServerSettings> Settings, ILogger<NotificationService> Logger, IGraphCollection GraphCollection, ISubscriptionCollection SubscriptionCollection, INotificationTriggerCollection TriggerCollection, IUserCollection UserCollection, JobService JobService, StreamService StreamService, IIssueService IssueService, ILogFileService LogFileService, IDogStatsd DogStatsd)
+		{
+			this.Sinks = Sinks.ToList();
+			this.Settings = Settings;
+			this.Logger = Logger;
+			this.GraphCollection = GraphCollection;
+			this.SubscriptionCollection = SubscriptionCollection;
+			this.TriggerCollection = TriggerCollection;
+			this.UserCollection = UserCollection;
+			this.JobService = JobService;
+			this.StreamService = StreamService;
+			this.IssueService = IssueService;
+			this.LogFileService = LogFileService;
+			this.DogStatsd = DogStatsd;
+
+			IssueService.OnIssueUpdated += NotifyIssueUpdated;
+			JobService.OnJobStepComplete += NotifyJobStepComplete;
+			JobService.OnLabelUpdate += NotifyLabelUpdate;
+		}
+
+		/// <inheritdoc/>
+		public override void Dispose()
+		{
+			base.Dispose();
+
+			IssueService.OnIssueUpdated -= NotifyIssueUpdated;
+			JobService.OnJobStepComplete -= NotifyJobStepComplete;
+			JobService.OnLabelUpdate -= NotifyLabelUpdate;
+
+			GC.SuppressFinalize(this);
+		}
+
+		/// <inheritdoc/>
+		public async Task<bool> UpdateSubscriptionsAsync(ObjectId TriggerId, ClaimsPrincipal User, bool? Email, bool? Slack)
+		{
+			UserId? UserId = User.GetUserId();
+			if (UserId == null)
+			{
+				Logger.LogWarning("Unable to find username for principal {User}", User.Identity?.Name);
+				return false;
+			}
+
+			INotificationTrigger Trigger = await TriggerCollection.FindOrAddAsync(TriggerId);
+			await TriggerCollection.UpdateSubscriptionsAsync(Trigger, UserId.Value, Email, Slack);
+			return true;
+		}
+
+		/// <inheritdoc/>
+		public async Task<INotificationSubscription?> GetSubscriptionsAsync(ObjectId TriggerId, ClaimsPrincipal User)
+		{
+			UserId? UserId = User.GetUserId();
+			if (UserId == null)
+			{
+				return null;
+			}
+
+			INotificationTrigger? Trigger = await TriggerCollection.GetAsync(TriggerId);
+			if(Trigger == null)
+			{
+				return null;
+			}
+
+			return Trigger.Subscriptions.FirstOrDefault(x => x.UserId == UserId.Value);
+		}
+
+		/// <inheritdoc/>
+		public void NotifyJobStepComplete(IJob Job, IGraph Graph, SubResourceId BatchId, SubResourceId StepId)
+		{
+			// Enqueue job step complete notifications if needed
+			if (Job.TryGetStep(BatchId, StepId, out IJobStep? Step))
+			{
+				Logger.LogInformation("Queuing step notifications for {JobId}:{BatchId}:{StepId}", Job.Id, BatchId, StepId);
+				EnqueueTask(() => SendJobStepNotificationsAsync(Job, BatchId, StepId));
+			}
+
+			// Enqueue job complete notifications if needed
+			if (Job.GetState() == JobState.Complete)
+			{
+				Logger.LogInformation("Queuing job notifications for {JobId}:{BatchId}:{StepId}", Job.Id, BatchId, StepId);
+				EnqueueTask(() => SendJobNotificationsAsync(Job, Graph));
+				EnqueueTask(() => RecordJobCompleteMetrics(Job));
+			}
+		}
+
+		/// <inheritdoc/>
+		public void NotifyLabelUpdate(IJob Job, IReadOnlyList<(LabelState, LabelOutcome)> OldLabelStates, IReadOnlyList<(LabelState, LabelOutcome)> NewLabelStates)
+		{
+			// If job has any label trigger IDs, send label complete notifications if needed
+			for (int Idx = 0; Idx < OldLabelStates.Count && Idx < NewLabelStates.Count; Idx++)
+			{
+				if (OldLabelStates[Idx] != NewLabelStates[Idx])
+				{
+					EnqueueTask(() => SendAllLabelNotificationsAsync(Job, OldLabelStates, NewLabelStates));
+					break;
+				}
+			}
+		}
+
+		/// <inheritdoc/>
+		public void NotifyIssueUpdated(IIssue Issue)
+		{
+			Logger.LogInformation("Issue {IssueId} updated", Issue.Id);
+			EnqueueTasks(Sink => Sink.NotifyIssueUpdatedAsync(Issue));
+		}
+
+		/// <inheritdoc/>
+		public void NotifyConfigUpdateFailure(string ErrorMessage, string FileName, int? Change = null, IUser? Author = null, string? Description = null)
+		{
+			EnqueueTasks(Sink => Sink.NotifyConfigUpdateFailureAsync(ErrorMessage, FileName, Change, Author, Description));
+		}
+
+		/// <inheritdoc/>
+		public void NotifyDeviceService(string Message, IDevice? Device = null, IDevicePool? Pool = null, IStream? Stream = null, IJob? Job = null, IJobStep? Step = null, INode? Node = null)
+		{
+			EnqueueTasks(Sink => Sink.NotifyDeviceServiceAsync(Message, Device, Pool, Stream, Job, Step, Node));
+		}
+
+		/// <summary>
+		/// Enqueues an async task
+		/// </summary>
+		/// <param name="TaskFunc">Function to generate an async task</param>
+		void EnqueueTask(Func<Task> TaskFunc)
+		{
+			NewTasks.Enqueue(Task.Run(TaskFunc));
+		}
+
+		/// <summary>
+		/// Enqueues an async task
+		/// </summary>
+		/// <param name="TaskFunc">Function to generate an async task</param>
+		void EnqueueTasks(Func<INotificationSink, Task> TaskFunc)
+		{
+			foreach (INotificationSink Sink in Sinks)
+			{
+				EnqueueTask(() => TaskFunc(Sink));
+			}
+		}
+
+		/// <inheritdoc/>
+		[System.Diagnostics.CodeAnalysis.SuppressMessage("Design", "CA1031:Do not catch general exception types", Justification = "<Pending>")]
+		protected async override Task ExecuteAsync(CancellationToken StoppingToken)
+		{
+			// This background service just waits for tasks to finish and prints any exception info. The only reason to do this is to
+			// ensure we finish processing everything before shutdown.
+			using (CancellationTask StoppingTask = new CancellationTask(StoppingToken))
+			{
+				List<Task> Tasks = new List<Task>();
+				Tasks.Add(NewTaskEvent.Task);
+
+				for (; ; )
+				{
+					// Add any new tasks to be monitored
+					Task? NewTask;
+					while (NewTasks.TryDequeue(out NewTask))
+					{
+						Tasks.Add(NewTask);
+					}
+
+					// If we don't have any
+					if (Tasks.Count == 1)
+					{
+						await Task.WhenAny(NewTaskEvent.Task, StoppingTask.Task);
+						if (StoppingToken.IsCancellationRequested)
+						{
+							break;
+						}
+					}
+					else
+					{
+						// Wait for something to finish
+						Task Task = await Task.WhenAny(Tasks);
+						if (Task == NewTaskEvent.Task)
+						{
+							NewTaskEvent = new AsyncEvent();
+							Tasks[0] = NewTaskEvent.Task;
+						}
+						else
+						{
+							try
+							{
+								await Task;
+							}
+							catch (Exception Ex)
+							{
+								Logger.LogError(Ex, "Exception while executing notification");
+							}
+							Tasks.Remove(Task);
+						}
+					}
+				}
+			}
+		}
+
+		internal Task ExecuteBackgroundForTest(CancellationToken StoppingToken)
+		{
+			return ExecuteAsync(StoppingToken);
+		}
+
+		/// <summary>
+		/// Gets the <see cref="INotificationTrigger"/> for a given trigger ID, if any.
+		/// </summary>
+		/// <param name="TriggerId"></param>
+		/// <param name="bFireTrigger">If true, the trigger is fired and cannot be reused</param>
+		/// <returns></returns>
+		private async Task<INotificationTrigger?> GetNotificationTrigger(ObjectId? TriggerId, bool bFireTrigger)
+		{
+			if (TriggerId == null)
+			{
+				return null;
+			}
+
+			INotificationTrigger? Trigger = await TriggerCollection.GetAsync(TriggerId.Value);
+			if (Trigger == null)
+			{
+				return null;
+			}
+
+			return bFireTrigger ? await TriggerCollection.FireAsync(Trigger) : Trigger;
+		}
+	
+		private async Task SendJobNotificationsAsync(IJob Job, IGraph Graph)
+		{
+			using IDisposable Scope = Logger.BeginScope("Sending notifications for job {JobId}", Job.Id);
+
+			Job.GetJobState(Job.GetStepForNodeMap(), out _, out LabelOutcome Outcome);
+			JobCompleteEventRecord JobCompleteEvent = new JobCompleteEventRecord(Job.StreamId, Job.TemplateId, Outcome);
+
+			IStream? JobStream = await StreamService.GetStreamAsync(Job.StreamId);
+			if (JobStream == null)
+			{
+				Logger.LogError("Unable to get stream {StreamId}", Job.StreamId);
+				return;
+			}
+
+			List<IUser> UsersToNotify = await GetUsersToNotify(JobCompleteEvent, Job.NotificationTriggerId, true);
+			foreach (IUser UserToNotify in UsersToNotify)
+			{
+				if(Job.PreflightChange != 0)
+				{
+					if(UserToNotify.Id != Job.StartedByUserId)
+					{
+						continue;
+					}
+				}
+				EnqueueTasks(Sink => Sink.NotifyJobCompleteAsync(UserToNotify, JobStream, Job, Graph, Outcome));
+			}
+
+			if (Job.PreflightChange == 0)
+			{
+				EnqueueTasks(Sink => Sink.NotifyJobCompleteAsync(JobStream, Job, Graph, Outcome));
+			}
+
+			Logger.LogDebug("Finished sending notifications for job {JobId}", Job.Id);
+		}
+
+		private Task RecordJobCompleteMetrics(IJob Job)
+		{
+			void RecordMetric(string Type, JobStepOutcome Outcome, DateTimeOffset? StartTime, DateTimeOffset? FinishTime)
+			{
+				string OutcomeStr = Outcome switch
+				{
+					JobStepOutcome.Unspecified => "unspecified",
+					JobStepOutcome.Failure => "failure",
+					JobStepOutcome.Warnings => "warnings",
+					JobStepOutcome.Success => "success",
+					_ => "unspecified"
+				};
+				
+				string[] Tags = {"stream:" + Job.StreamId, "template:" + Job.TemplateId};
+				DogStatsd.Increment($"horde.{Type}.{OutcomeStr}.count", 1, tags: Tags);
+
+				if (StartTime == null || FinishTime == null)
+				{
+					Logger.LogDebug("Completed job or step is missing start or finish time, cannot record duration metric. Job ID={JobId}", Job.Id);
+					return;
+				}
+
+				TimeSpan Duration = FinishTime.Value - StartTime.Value;
+				DogStatsd.Timer($"horde.{Type}.{OutcomeStr}.duration", Duration.TotalSeconds, tags: Tags);
+			}
+
+			JobStepOutcome JobOutcome = Job.Batches.SelectMany(x => x.Steps).Min(x => x.Outcome);
+			DateTime? StartTime = Job.Batches.Select(x => x.StartTimeUtc).Min();
+			DateTime? FinishTime = Job.Batches.Select(x => x.FinishTimeUtc).Max();
+			RecordMetric("job", JobOutcome, StartTime, FinishTime);
+
+			// TODO: record metrics for individual steps
+			// foreach (IJobStepBatch Batch in Job.Batches)
+			// {
+			// 	foreach (IJobStep Step in Batch.Steps)
+			// 	{
+			// 	}
+			// }
+
+			return Task.CompletedTask;
+		}
+
+		private async Task<List<IUser>> GetUsersToNotify(EventRecord? Event, ObjectId? NotificationTriggerId, bool bFireTrigger)
+		{
+			List<UserId> UserIds = new List<UserId>();
+
+			// Find the notifications for all steps of this type
+			if (Event != null)
+			{
+				List<ISubscription> Subscriptions = await SubscriptionCollection.FindSubscribersAsync(Event);
+				foreach (ISubscription Subscription in Subscriptions)
+				{
+					if (Subscription.NotificationType == NotificationType.Slack)
+					{
+						UserIds.Add(Subscription.UserId);
+					}
+				}
+			}
+
+			// Find the notifications for this particular step
+			if (NotificationTriggerId != null)
+			{
+				INotificationTrigger? Trigger = await GetNotificationTrigger(NotificationTriggerId, bFireTrigger);
+				if (Trigger != null)
+				{
+					foreach (INotificationSubscription Subscription in Trigger.Subscriptions)
+					{
+						if (Subscription.Email)
+						{
+							// TODO?
+						}
+						if (Subscription.Slack)
+						{
+							UserIds.Add(Subscription.UserId);
+						}
+					}
+				}
+			}
+			return await UserCollection.FindUsersAsync(UserIds);
+		}
+
+		private async Task SendJobStepNotificationsAsync(IJob Job, SubResourceId BatchId, SubResourceId StepId)
+		{
+			using IDisposable Scope = Logger.BeginScope("Sending notifications for step {JobId}:{BatchId}:{StepId}", Job.Id, BatchId, StepId);
+
+			IJobStepBatch? Batch;
+			if(!Job.TryGetBatch(BatchId, out Batch))
+			{
+				Logger.LogError("Unable to find batch {BatchId} in job {JobId}", BatchId, Job.Id);
+				return;
+			}
+
+			IJobStep? Step;
+			if (!Batch.TryGetStep(StepId, out Step))
+			{
+				Logger.LogError("Unable to find step {StepId} in batch {JobId}:{BatchId}", StepId, Job.Id, BatchId);
+				return;
+			}
+
+			IGraph JobGraph = await GraphCollection.GetAsync(Job.GraphHash);
+			INode Node = JobGraph.GetNode(new NodeRef(Batch.GroupIdx, Step.NodeIdx));
+
+			// Find the notifications for this particular step
+			EventRecord EventRecord = new StepCompleteEventRecord(Job.StreamId, Job.TemplateId, Node.Name, Step.Outcome);
+
+			List<IUser> UsersToNotify = await GetUsersToNotify(EventRecord, Step.NotificationTriggerId, true);
+			if(UsersToNotify.Count == 0)
+			{
+				Logger.LogInformation("No users to notify for step {JobId}:{BatchId}:{StepId}", Job.Id, BatchId, StepId);
+				return;
+			}
+
+			IStream? JobStream = await StreamService.GetStreamAsync(Job.StreamId);
+			if (JobStream == null)
+			{
+				Logger.LogError("Unable to find stream {StreamId}", Job.StreamId);
+				return;
+			}
+
+			if(Step.LogId == null)
+			{
+				Logger.LogError("Step does not have a log file");
+				return;
+			}
+
+			ILogFile? LogFile = await LogFileService.GetLogFileAsync(Step.LogId.Value);
+			if(LogFile == null)
+			{
+				Logger.LogError("Step does not have a log file");
+				return;
+			}
+
+			List<ILogEvent> JobStepEvents = await LogFileService.FindLogEventsAsync(LogFile);
+			List<ILogEventData> JobStepEventData = new List<ILogEventData>();
+			foreach (ILogEvent Event in JobStepEvents)
+			{
+				ILogEventData EventData = await LogFileService.GetEventDataAsync(LogFile, Event.LineIndex, Event.LineCount);
+				JobStepEventData.Add(EventData);
+			}
+
+			foreach (IUser SlackUser in UsersToNotify)
+			{
+				if(Job.PreflightChange != 0)
+				{
+					if(SlackUser.Id != Job.StartedByUserId)
+					{
+						continue;
+					}
+				}
+				EnqueueTasks(Sink => Sink.NotifyJobStepCompleteAsync(SlackUser, JobStream, Job, Batch, Step, Node, JobStepEventData));
+			}
+			Logger.LogDebug("Finished sending notifications for step {JobId}:{BatchId}:{StepId}", Job.Id, BatchId, StepId);
+		}
+
+		private async Task SendAllLabelNotificationsAsync(IJob Job, IReadOnlyList<(LabelState State, LabelOutcome Outcome)> OldLabelStates, IReadOnlyList<(LabelState, LabelOutcome)> NewLabelStates)
+		{
+			IStream? Stream = await StreamService.GetStreamAsync(Job.StreamId);
+			if (Stream == null)
+			{
+				Logger.LogError("Unable to find stream {StreamId} for job {JobId}", Job.StreamId, Job.Id);
+				return;
+			}
+
+			IGraph? Graph = await GraphCollection.GetAsync(Job.GraphHash);
+			if (Graph == null)
+			{
+				Logger.LogError("Unable to find graph {GraphHash} for job {JobId}", Job.GraphHash, Job.Id);
+				return;
+			}
+
+			IReadOnlyDictionary<NodeRef, IJobStep> StepForNode = Job.GetStepForNodeMap();
+			for (int LabelIdx = 0; LabelIdx < Graph.Labels.Count; ++LabelIdx)
+			{
+				(LabelState State, LabelOutcome Outcome) OldLabel = OldLabelStates[LabelIdx];
+				(LabelState State, LabelOutcome Outcome) NewLabel = NewLabelStates[LabelIdx];
+				if (OldLabel != NewLabel)
+				{
+					// If the state transitioned from Unspecified to Running, don't update unless the outcome also changed.
+					if (OldLabel.State == LabelState.Unspecified && NewLabel.State == LabelState.Running && OldLabel.Outcome == NewLabel.Outcome)
+					{
+						continue;
+					}
+
+					// If the label isn't complete, don't report on outcome changing to success, this will be reported when the label state becomes complete.
+					if (NewLabel.State != LabelState.Complete && NewLabel.Outcome == LabelOutcome.Success)
+					{
+						return;
+					}
+
+					ILabel Label = Graph.Labels[LabelIdx];
+
+					EventRecord? EventId;
+					if (String.IsNullOrEmpty(Label.DashboardName))
+					{
+						EventId = null;
+					}
+					else
+					{
+						EventId = new LabelCompleteEventRecord(Job.StreamId, Job.TemplateId, Label.DashboardCategory, Label.DashboardName, NewLabel.Outcome);
+					}
+
+					ObjectId? TriggerId;
+					if (Job.LabelIdxToTriggerId.TryGetValue(LabelIdx, out ObjectId NewTriggerId))
+					{
+						TriggerId = NewTriggerId;
+					}
+					else
+					{
+						TriggerId = null;
+					}
+
+					bool bFireTrigger = NewLabel.State == LabelState.Complete;
+
+					List<IUser> UsersToNotify = await GetUsersToNotify(EventId, TriggerId, bFireTrigger);
+					if (UsersToNotify.Count > 0)
+					{
+						SendLabelUpdateNotifications(Job, Stream, Graph, StepForNode, Graph.Labels[LabelIdx], LabelIdx, NewLabel.Outcome, UsersToNotify);
+					}
+					else
+					{
+						Logger.LogDebug("No users to notify for label {DashboardName}/{UgsName} in job {JobId}", Graph.Labels[LabelIdx].DashboardName, Graph.Labels[LabelIdx].UgsName, Job.Id);
+					}
+				}
+			}
+		}
+
+		private void SendLabelUpdateNotifications(IJob Job, IStream Stream, IGraph Graph, IReadOnlyDictionary<NodeRef, IJobStep> StepForNode, ILabel Label, int LabelIdx, LabelOutcome Outcome, List<IUser> SlackUsers)
+		{
+			List<(string, JobStepOutcome, Uri)> StepData = new List<(string, JobStepOutcome, Uri)>();
+			if (Outcome != LabelOutcome.Success)
+			{
+				foreach (NodeRef IncludedNodeRef in Label.IncludedNodes)
+				{
+					INode IncludedNode = Graph.GetNode(IncludedNodeRef);
+					IJobStep IncludedStep = StepForNode[IncludedNodeRef];
+					StepData.Add((IncludedNode.Name, IncludedStep.Outcome, new Uri($"{Settings.CurrentValue.DashboardUrl}job/{Job.Id}?step={IncludedStep.Id}")));
+				}
+			}
+
+			foreach (IUser SlackUser in SlackUsers)
+			{
+				EnqueueTasks(Sink => Sink.NotifyLabelCompleteAsync(SlackUser, Job, Stream, Label, LabelIdx, Outcome, StepData));
+			}
+
+			Logger.LogDebug("Finished sending label notifications for label {DashboardName}/{UgsName} in job {JobId}", Label.DashboardName, Label.UgsName, Job.Id);
+		}
+	}
+}
diff --git a/Engine/Source/Programs/Horde/Horde.Build/Notifications/Impl/SlackNotificationSink.cs b/Engine/Source/Programs/Horde/Horde.Build/Notifications/Impl/SlackNotificationSink.cs
new file mode 100644
--- /dev/null
+++ b/Engine/Source/Programs/Horde/Horde.Build/Notifications/Impl/SlackNotificationSink.cs
@@ -0,0 +1,1394 @@
+// Copyright Epic Games, Inc. All Rights Reserved.
+
+using EpicGames.Core;
+using HordeCommon;
+using HordeServer.Api;
+using HordeServer.Collections;
+using HordeServer.Models;
+using HordeServer.Services;
+using HordeServer.Utilities;
+using HordeServer.Utilities.BlockKit;
+using HordeServer.Utilities.Slack.BlockKit;
+using Microsoft.Extensions.Caching.Memory;
+using Microsoft.Extensions.Hosting;
+using Microsoft.Extensions.Logging;
+using Microsoft.Extensions.Options;
+using MongoDB.Bson;
+using MongoDB.Bson.Serialization.Attributes;
+using MongoDB.Driver;
+using Newtonsoft.Json.Linq;
+using System;
+using System.Collections.Concurrent;
+using System.Collections.Generic;
+using System.Diagnostics;
+using System.Drawing;
+using System.Globalization;
+using System.Linq;
+using System.Net.Http;
+using System.Net.WebSockets;
+using System.Security.Claims;
+using System.Text;
+using System.Text.Json;
+using System.Text.Json.Serialization;
+using System.Text.RegularExpressions;
+using System.Threading;
+using System.Threading.Tasks;
+
+namespace HordeServer.Notifications.Impl
+{
+	using UserId = ObjectId<IUser>;
+
+	/// <summary>
+	/// Maintains a connection to Slack, in order to receive socket-mode notifications of user interactions
+	/// </summary>
+	public class SlackNotificationSink : BackgroundService, INotificationSink, IAvatarService
+	{
+		const string PostMessageUrl = "https://slack.com/api/chat.postMessage";
+		const string UpdateMessageUrl = "https://slack.com/api/chat.update";
+
+		class SocketResponse
+		{
+			[JsonPropertyName("ok")]
+			public bool Ok { get; set; }
+
+			[JsonPropertyName("url")]
+			public Uri? Url { get; set; }
+		}
+
+		class EventMessage
+		{
+			[JsonPropertyName("envelope_id")]
+			public string? EnvelopeId { get; set; }
+
+			[JsonPropertyName("type")]
+			public string? Type { get; set; }
+
+			[JsonPropertyName("payload")]
+			public EventPayload? Payload { get; set; }
+		}
+
+		class EventPayload
+		{
+			[JsonPropertyName("type")]
+			public string? Type { get; set; }
+
+			[JsonPropertyName("user")]
+			public UserInfo? User { get; set; }
+
+			[JsonPropertyName("response_url")]
+			public string? ResponseUrl { get; set; }
+
+			[JsonPropertyName("actions")]
+			public List<ActionInfo> Actions { get; set; } = new List<ActionInfo>();
+		}
+
+		class UserResponse
+		{
+			[JsonPropertyName("ok")]
+			public bool Ok { get; set; }
+
+			[JsonPropertyName("user")]
+			public UserInfo? User { get; set; }
+		}
+
+		class UserInfo
+		{
+			[JsonPropertyName("id")]
+			public string? UserName { get; set; }
+
+			[JsonPropertyName("profile")]
+			public UserProfile? Profile { get; set; }
+		}
+
+		class UserProfile
+		{
+			[JsonPropertyName("is_custom_image")]
+			public bool IsCustomImage { get; set; }
+
+			[JsonPropertyName("image_24")]
+			public string? Image24 { get; set; }
+
+			[JsonPropertyName("image_32")]
+			public string? Image32 { get; set; }
+
+			[JsonPropertyName("image_48")]
+			public string? Image48 { get; set; }
+
+			[JsonPropertyName("image_72")]
+			public string? Image72 { get; set; }
+		}
+
+		class ActionInfo
+		{
+			[JsonPropertyName("value")]
+			public string? Value { get; set; }
+		}
+
+		class MessageStateDocument
+		{
+			[BsonId]
+			public ObjectId Id { get; set; }
+
+			[BsonElement("uid")]
+			public string Recipient { get; set; } = String.Empty;
+
+			[BsonElement("usr")]
+			public UserId? UserId { get; set; }
+
+			[BsonElement("eid")]
+			public string EventId { get; set; } = String.Empty;
+
+			[BsonElement("ch")]
+			public string Channel { get; set; } = String.Empty;
+
+			[BsonElement("ts")]
+			public string Ts { get; set; } = String.Empty;
+
+			[BsonElement("dig")]
+			public string Digest { get; set; } = String.Empty;
+		}
+
+		class SlackUser : IAvatar
+		{
+			public const int CurrentVersion = 2;
+			
+			public UserId Id { get; set; }
+
+			[BsonElement("u")]
+			public string? SlackUserId { get; set; }
+
+			[BsonElement("i24")]
+			public string? Image24 { get; set; }
+
+			[BsonElement("i32")]
+			public string? Image32 { get; set; }
+			
+			[BsonElement("i48")]
+			public string? Image48 { get; set; }
+
+			[BsonElement("i72")]
+			public string? Image72 { get; set; }
+
+			[BsonElement("t")]
+			public DateTime Time { get; set; }
+
+			[BsonElement("v")]
+			public int Version { get; set; }
+
+			private SlackUser()
+			{
+			}
+
+			public SlackUser(UserId Id, UserInfo? Info)
+			{
+				this.Id = Id;
+				this.SlackUserId = Info?.UserName;
+				if (Info != null && Info.Profile != null && Info.Profile.IsCustomImage)
+				{
+					this.Image24 = Info.Profile.Image24;
+					this.Image32 = Info.Profile.Image32;
+					this.Image48 = Info.Profile.Image48;
+					this.Image72 = Info.Profile.Image72;
+				}
+				this.Time = DateTime.UtcNow;
+				this.Version = CurrentVersion;
+			}
+		}
+
+		IIssueService IssueService;
+		IUserCollection UserCollection;
+		ILogFileService LogFileService;
+		StreamService StreamService;
+		ServerSettings Settings;
+		IMongoCollection<MessageStateDocument> MessageStates;
+		IMongoCollection<SlackUser> SlackUsers;
+		HashSet<string>? AllowUsers;
+		ILogger Logger;
+
+		/// <summary>
+		/// Map of email address to Slack user ID.
+		/// </summary>
+		private MemoryCache UserCache = new MemoryCache(new MemoryCacheOptions());
+
+		/// <summary>
+		/// Constructor
+		/// </summary>
+		/// <param name="DatabaseService"></param>
+		/// <param name="IssueService"></param>
+		/// <param name="UserCollection">The user collection</param>
+		/// <param name="LogFileService"></param>
+		/// <param name="StreamService"></param>
+		/// <param name="Settings">The current configuration settings</param>
+		/// <param name="Logger">Logging device</param>
+		public SlackNotificationSink(DatabaseService DatabaseService, IIssueService IssueService, IUserCollection UserCollection, ILogFileService LogFileService, StreamService StreamService, IOptions<ServerSettings> Settings, ILogger<SlackNotificationSink> Logger)
+		{
+			this.IssueService = IssueService;
+			this.UserCollection = UserCollection;
+			this.LogFileService = LogFileService;
+			this.StreamService = StreamService;
+			this.Settings = Settings.Value;
+			this.MessageStates = DatabaseService.Database.GetCollection<MessageStateDocument>("Slack");
+			this.SlackUsers = DatabaseService.Database.GetCollection<SlackUser>("Slack.UsersV2");
+			this.Logger = Logger;
+
+			if (!String.IsNullOrEmpty(Settings.Value.SlackUsers))
+			{
+				AllowUsers = new HashSet<string>(Settings.Value.SlackUsers.Split(','), StringComparer.OrdinalIgnoreCase);
+			}
+		}
+
+		/// <inheritdoc/>
+		public override void Dispose()
+		{
+			base.Dispose();
+
+			UserCache.Dispose();
+
+			GC.SuppressFinalize(this);
+		}
+
+		#region Avatars
+
+		/// <inheritdoc/>
+		public async Task<IAvatar?> GetAvatarAsync(IUser User)
+		{
+			return await GetSlackUser(User);
+		}
+
+		#endregion
+
+		#region Message state 
+
+		async Task<MessageStateDocument> AddOrUpdateMessageStateAsync(string Recipient, string EventId, UserId? UserId, string Digest)
+		{
+			ObjectId NewId = ObjectId.GenerateNewId();
+
+			FilterDefinition<MessageStateDocument> Filter = Builders<MessageStateDocument>.Filter.Eq(x => x.Recipient, Recipient) & Builders<MessageStateDocument>.Filter.Eq(x => x.EventId, EventId);
+			UpdateDefinition<MessageStateDocument> Update = Builders<MessageStateDocument>.Update.SetOnInsert(x => x.Id, NewId).Set(x => x.UserId, UserId).Set(x => x.Digest, Digest);
+
+			MessageStateDocument? PrevState = await MessageStates.FindOneAndUpdateAsync(Filter, Update, new FindOneAndUpdateOptions<MessageStateDocument> { IsUpsert = true });
+			if (PrevState == null)
+			{
+				PrevState = new MessageStateDocument();
+				PrevState.Id = NewId;
+				PrevState.Recipient = Recipient;
+				PrevState.EventId = EventId;
+				PrevState.UserId = UserId;
+			}
+			return PrevState;
+		}
+
+		async Task SetMessageTimestampAsync(ObjectId MessageId, string Channel, string Ts)
+		{
+			FilterDefinition<MessageStateDocument> Filter = Builders<MessageStateDocument>.Filter.Eq(x => x.Id, MessageId) & Builders<MessageStateDocument>.Filter.Eq(x => x.Ts, String.Empty);
+			UpdateDefinition<MessageStateDocument> Update = Builders<MessageStateDocument>.Update.Set(x => x.Channel, Channel).Set(x => x.Ts, Ts);
+			await MessageStates.FindOneAndUpdateAsync(Filter, Update);
+		}
+
+		#endregion
+
+		#region Job Complete
+
+		/// <inheritdoc/>
+		public async Task NotifyJobCompleteAsync(IStream JobStream, IJob Job, IGraph Graph, LabelOutcome Outcome)
+		{
+			if (Job.NotificationChannel != null)
+			{
+				await SendJobCompleteNotificationToChannelAsync(Job.NotificationChannel, Job.NotificationChannelFilter, JobStream, Job, Graph, Outcome);
+			}
+			if (JobStream.NotificationChannel != null)
+			{
+				await SendJobCompleteNotificationToChannelAsync(JobStream.NotificationChannel, JobStream.NotificationChannelFilter, JobStream, Job, Graph, Outcome);
+			}
+		}
+
+		Task SendJobCompleteNotificationToChannelAsync(string NotificationChannel, string? NotificationFilter, IStream JobStream, IJob Job, IGraph Graph, LabelOutcome Outcome)
+		{
+			if (NotificationFilter != null)
+			{
+				List<LabelOutcome> Outcomes = new List<LabelOutcome>();
+				foreach (string FilterOption in NotificationFilter.Split('|'))
+				{
+					LabelOutcome Result;
+					if (Enum.TryParse(FilterOption, out Result))
+					{
+						Outcomes.Add(Result);
+					}
+					else
+					{
+						Logger.LogWarning("Invalid filter option {Option} specified in job filter {NotificationChannelFilter} in job {JobId} or stream {StreamId}", FilterOption, NotificationFilter, Job.Id, Job.StreamId);
+					}
+				}
+				if (!Outcomes.Contains(Outcome))
+				{
+					return Task.CompletedTask;
+				}
+			}
+			return SendJobCompleteMessageAsync(NotificationChannel, JobStream, Job, Graph);
+		}
+
+		/// <inheritdoc/>
+		public async Task NotifyJobCompleteAsync(IUser SlackUser, IStream JobStream, IJob Job, IGraph Graph, LabelOutcome Outcome)
+		{
+			string? SlackUserId = await GetSlackUserId(SlackUser);
+			if (SlackUserId != null)
+			{
+				await SendJobCompleteMessageAsync(SlackUserId, JobStream, Job, Graph);
+			}
+		}
+
+		private Task SendJobCompleteMessageAsync(string Recipient, IStream Stream, IJob Job, IGraph Graph)
+		{
+			JobStepOutcome JobOutcome = Job.Batches.SelectMany(x => x.Steps).Min(x => x.Outcome);
+			Logger.LogInformation("Sending Slack notification for job {JobId} outcome {Outcome} to {SlackUser}", Job.Id, JobOutcome, Recipient);
+
+			Uri JobLink = new Uri($"{Settings.DashboardUrl}job/{Job.Id}");
+
+			Color OutcomeColor = JobOutcome == JobStepOutcome.Failure ? BlockKitAttachmentColors.Error : JobOutcome == JobStepOutcome.Warnings ? BlockKitAttachmentColors.Warning : BlockKitAttachmentColors.Success;
+
+			BlockKitAttachment Attachment = new BlockKitAttachment();
+			Attachment.FallbackText = $"{Stream.Name} - {GetJobChangeText(Job)} - {Job.Name} - {JobOutcome}";
+			Attachment.Color = OutcomeColor;
+			Attachment.Blocks.Add(new SectionBlock($"*<{JobLink}|{Stream.Name} - {GetJobChangeText(Job)} - {Job.Name}>*"));
+			if (JobOutcome == JobStepOutcome.Success)
+			{
+				Attachment.Blocks.Add(new SectionBlock($"*Job Succeeded*"));
+			}
+			else
+			{
+				List<string> FailedStepStrings = new List<string>();
+				List<string> WarningStepStrings = new List<string>();
+
+				IReadOnlyDictionary<NodeRef, IJobStep> NodeToStep = Job.GetStepForNodeMap();
+				foreach ((NodeRef NodeRef, IJobStep Step) in NodeToStep)
+				{
+					INode StepNode = Graph.GetNode(NodeRef);
+					string StepName = $"<{JobLink}?step={Step.Id}|{StepNode.Name}>";
+					if (Step.Outcome == JobStepOutcome.Failure)
+					{
+						FailedStepStrings.Add(StepName);
+					}
+					else if (Step.Outcome == JobStepOutcome.Warnings)
+					{
+						WarningStepStrings.Add(StepName);
+					}
+				}
+
+				if (FailedStepStrings.Any())
+				{
+					string Msg = $"*Errors*\n{string.Join(", ", FailedStepStrings)}";
+					Attachment.Blocks.Add(new SectionBlock(Msg.Substring(0, Math.Min(Msg.Length, 3000))));
+				}
+				else if (WarningStepStrings.Any())
+				{
+					string Msg = $"*Warnings*\n{string.Join(", ", WarningStepStrings)}";
+					Attachment.Blocks.Add(new SectionBlock(Msg.Substring(0, Math.Min(Msg.Length, 3000))));
+				}
+			}
+
+			if (Job.AutoSubmit)
+			{
+				Attachment.Blocks.Add(new DividerBlock());
+				if (Job.AutoSubmitChange != null)
+				{
+					Attachment.Blocks.Add(new SectionBlock($"Shelved files were submitted in CL {Job.AutoSubmitChange}."));
+				}
+				else
+				{
+					Attachment.Color = BlockKitAttachmentColors.Warning;
+
+					string AutoSubmitMessage = String.Empty;
+					if (!String.IsNullOrEmpty(Job.AutoSubmitMessage))
+					{
+						AutoSubmitMessage = $"\n\n```{Job.AutoSubmitMessage}```";
+					}
+
+					Attachment.Blocks.Add(new SectionBlock($"Files in CL {Job.PreflightChange} were *not submitted*. Please resolve the following issues and submit manually.{Job.AutoSubmitMessage}"));
+				}
+			}
+
+			return SendMessageAsync(Recipient, Attachments: new[] { Attachment });
+		}
+
+		#endregion
+
+		#region Job step complete
+
+		/// <inheritdoc/>
+		public async Task NotifyJobStepCompleteAsync(IUser SlackUser, IStream JobStream, IJob Job, IJobStepBatch Batch, IJobStep Step, INode Node, List<ILogEventData> JobStepEventData)
+		{
+			Logger.LogInformation("Sending Slack notification for job {JobId}, batch {BatchId}, step {StepId}, outcome {Outcome} to {SlackUser} ({UserId})", Job.Id, Batch.Id, Step.Id, Step.Outcome, SlackUser.Name, SlackUser.Id);
+
+			string? SlackUserId = await GetSlackUserId(SlackUser);
+			if (SlackUserId != null)
+			{
+				await SendJobStepCompleteMessageAsync(SlackUserId, JobStream, Job, Step, Node, JobStepEventData);
+			}
+		}
+
+		/// <summary>
+		/// Creates a Slack message about a completed step job.
+		/// </summary>
+		/// <param name="Recipient"></param>
+		/// <param name="Stream"></param>
+		/// <param name="Job">The job that contains the step that completed.</param>
+		/// <param name="Step">The job step that completed.</param>
+		/// <param name="Node">The node for the job step.</param>
+		/// <param name="Events">Any events that occurred during the job step.</param>
+		private Task SendJobStepCompleteMessageAsync(string Recipient, IStream Stream, IJob Job, IJobStep Step, INode Node, List<ILogEventData> Events)
+		{
+			Uri JobStepLink = new Uri($"{Settings.DashboardUrl}job/{Job.Id}?step={Step.Id}");
+			Uri JobStepLogLink = new Uri($"{Settings.DashboardUrl}log/{Step.LogId}");
+
+			Color OutcomeColor = Step.Outcome == JobStepOutcome.Failure ? BlockKitAttachmentColors.Error : Step.Outcome == JobStepOutcome.Warnings ? BlockKitAttachmentColors.Warning : BlockKitAttachmentColors.Success;
+
+			BlockKitAttachment Attachment = new BlockKitAttachment();
+			Attachment.FallbackText = $"{Stream.Name} - {GetJobChangeText(Job)} - {Job.Name} - {Node.Name} - {Step.Outcome}";
+			Attachment.Color = OutcomeColor;
+			Attachment.Blocks.Add(new SectionBlock($"*<{JobStepLink}|{Stream.Name} - {GetJobChangeText(Job)} - {Job.Name} - {Node.Name}>*"));
+			if (Step.Outcome == JobStepOutcome.Success)
+			{
+				Attachment.Blocks.Add(new SectionBlock($"*Job Step Succeeded*"));
+			}
+			else
+			{
+				List<ILogEventData> Errors = Events.Where(x => x.Severity == EventSeverity.Error).ToList();
+				List<ILogEventData> Warnings = Events.Where(x => x.Severity == EventSeverity.Warning).ToList();
+				List<string> EventStrings = new List<string>();
+				if (Errors.Any())
+				{
+					string ErrorSummary = Errors.Count > MaxJobStepEvents ? $"*Errors (First {MaxJobStepEvents} shown)*" : $"*Errors*";
+					EventStrings.Add(ErrorSummary);
+					foreach (ILogEventData Error in Errors.Take(MaxJobStepEvents))
+					{
+						EventStrings.Add($"```{Error.Message}```");
+					}
+				}
+				else if (Warnings.Any())
+				{
+					string WarningSummary = Warnings.Count > MaxJobStepEvents ? $"*Warnings (First {MaxJobStepEvents} shown)*" : $"*Warnings*";
+					EventStrings.Add(WarningSummary);
+					foreach (ILogEventData Warning in Warnings.Take(MaxJobStepEvents))
+					{
+						EventStrings.Add($"```{Warning.Message}```");
+					}
+				}
+
+				Attachment.Blocks.Add(new SectionBlock(string.Join("\n", EventStrings)));
+				Attachment.Blocks.Add(new SectionBlock($"<{JobStepLogLink}|View Job Step Log>"));
+			}
+
+			return SendMessageAsync(Recipient, Attachments: new[] { Attachment });
+		}
+
+		#endregion
+
+		#region Label complete
+
+		/// <inheritdoc/>
+		public async Task NotifyLabelCompleteAsync(IUser User, IJob Job, IStream Stream, ILabel Label, int LabelIdx, LabelOutcome Outcome, List<(string, JobStepOutcome, Uri)> StepData)
+		{
+			Logger.LogInformation("Sending Slack notification for job {JobId} outcome {Outcome} to {Name} ({UserId})", Job.Id, Outcome, User.Name, User.Id);
+
+			string? SlackUserId = await GetSlackUserId(User);
+			if (SlackUserId != null)
+			{
+				await SendLabelUpdateMessageAsync(SlackUserId, Stream, Job, Label, LabelIdx, Outcome, StepData);
+			}
+		}
+
+		Task SendLabelUpdateMessageAsync(string Recipient, IStream Stream, IJob Job, ILabel Label, int LabelIdx, LabelOutcome Outcome, List<(string, JobStepOutcome, Uri)> JobStepData)
+		{
+			Uri LabelLink = new Uri($"{Settings.DashboardUrl}job/{Job.Id}?label={LabelIdx}");
+
+			Color OutcomeColor = Outcome == LabelOutcome.Failure ? BlockKitAttachmentColors.Error : Outcome == LabelOutcome.Warnings ? BlockKitAttachmentColors.Warning : BlockKitAttachmentColors.Success;
+
+			BlockKitAttachment Attachment = new BlockKitAttachment();
+			Attachment.FallbackText = $"{Stream.Name} - {GetJobChangeText(Job)} - {Job.Name} - Label {Label.DashboardName} - {Outcome}";
+			Attachment.Color = OutcomeColor;
+			Attachment.Blocks.Add(new SectionBlock($"*<{LabelLink}|{Stream.Name} - {GetJobChangeText(Job)} - {Job.Name} - Label {Label.DashboardName}>*"));
+			if (Outcome == LabelOutcome.Success)
+			{
+				Attachment.Blocks.Add(new SectionBlock($"*Label Succeeded*"));
+			}
+			else
+			{
+				List<string> FailedStepStrings = new List<string>();
+				List<string> WarningStepStrings = new List<string>();
+				foreach ((string Name, JobStepOutcome StepOutcome, Uri Link) JobStep in JobStepData)
+				{
+					string StepString = $"<{JobStep.Link}|{JobStep.Name}>";
+					if (JobStep.StepOutcome == JobStepOutcome.Failure)
+					{
+						FailedStepStrings.Add(StepString);
+					}
+					else if (JobStep.StepOutcome == JobStepOutcome.Warnings)
+					{
+						WarningStepStrings.Add(StepString);
+					}
+				}
+
+				if (FailedStepStrings.Any())
+				{
+					Attachment.Blocks.Add(new SectionBlock($"*Errors*\n{string.Join(", ", FailedStepStrings)}"));
+				}
+				else if (WarningStepStrings.Any())
+				{
+					Attachment.Blocks.Add(new SectionBlock($"*Warnings*\n{string.Join(", ", WarningStepStrings)}"));
+				}
+			}
+
+			return SendMessageAsync(Recipient, Attachments: new[] { Attachment });
+		}
+
+		#endregion
+
+		#region Issues
+
+		/// <inheritdoc/>
+		public async Task NotifyIssueUpdatedAsync(IIssue Issue)
+		{
+			IIssueDetails Details = await IssueService.GetIssueDetailsAsync(Issue);
+
+			HashSet<UserId> UserIds = new HashSet<UserId>();
+			if (Issue.Promoted)
+			{
+				UserIds.UnionWith(Details.Suspects.Select(x => x.AuthorId));
+			}
+			if (Issue.OwnerId.HasValue)
+			{
+				UserIds.Add(Issue.OwnerId.Value);
+			}
+
+			HashSet<string> Channels = new HashSet<string>(StringComparer.OrdinalIgnoreCase);
+
+			List<MessageStateDocument> ExistingMessages = await MessageStates.Find(x => x.EventId == GetIssueEventId(Issue)).ToListAsync();
+			foreach (MessageStateDocument ExistingMessage in ExistingMessages)
+			{
+				if (ExistingMessage.UserId != null)
+				{
+					UserIds.Add(ExistingMessage.UserId.Value);
+				}
+				else
+				{
+					Channels.Add(ExistingMessage.Recipient);
+				}
+			}
+
+			if (Issue.OwnerId == null && (Details.Suspects.Count == 0 || Details.Suspects.All(x => x.DeclinedAt != null) || Issue.CreatedAt < DateTime.UtcNow - TimeSpan.FromHours(1.0)))
+			{
+				foreach (IIssueSpan Span in Details.Spans)
+				{
+					IStream? Stream = await StreamService.GetCachedStream(Span.StreamId);
+					if (Stream != null)
+					{
+						TemplateRef? TemplateRef;
+						if (Stream.Templates.TryGetValue(Span.TemplateRefId, out TemplateRef) && TemplateRef.TriageChannel != null)
+						{
+							Channels.Add(TemplateRef.TriageChannel);
+						}
+						else if (Stream.TriageChannel != null)
+						{
+							Channels.Add(Stream.TriageChannel);
+						}
+					}
+				}
+			}
+
+			if (UserIds.Count > 0)
+			{
+				foreach (UserId UserId in UserIds)
+				{
+					IUser? User = await UserCollection.GetUserAsync(UserId);
+					if (User == null)
+					{
+						Logger.LogWarning("Unable to find user {UserId}", UserId);
+					}
+					else
+					{
+						await NotifyIssueUpdatedAsync(User, Issue, Details);
+					}
+				}
+			}
+
+			if (Channels.Count > 0)
+			{
+				foreach (string Channel in Channels)
+				{
+					await SendIssueMessageAsync(Channel, Issue, Details, null);
+				}
+			}
+		}
+
+		async Task NotifyIssueUpdatedAsync(IUser User, IIssue Issue, IIssueDetails Details)
+		{
+			IUserSettings UserSettings = await UserCollection.GetSettingsAsync(User.Id);
+			if (!UserSettings.EnableIssueNotifications)
+			{
+				Logger.LogInformation("Issue notifications are disabled for user {UserId} ({UserName})", User.Id, User.Name);
+				return;
+			}
+
+			string? SlackUserId = await GetSlackUserId(User);
+			if (SlackUserId == null)
+			{
+				return;
+			}
+
+			await SendIssueMessageAsync(SlackUserId, Issue, Details, User.Id);
+		}
+
+		async Task<string> GetSampleText(IIssueDetails Details)
+		{
+			StringBuilder SampleText = new StringBuilder();
+			if (Details.Spans.Count > 0)
+			{
+				IIssueSpan Span = Details.Spans[0];
+				if (Span.FirstFailure.LogId != null)
+				{
+					ILogFile? LogFile = await LogFileService.GetCachedLogFileAsync(Span.FirstFailure.LogId.Value);
+					if (LogFile != null)
+					{
+						List<ILogEvent> Events = await LogFileService.FindEventsForSpansAsync(Details.Spans.Select(x => x.Id).Take(1), null, 0, 1);
+						foreach (ILogEvent Event in Events)
+						{
+							ILogEventData Data = await LogFileService.GetEventDataAsync(LogFile, Event.LineIndex, Event.LineCount);
+							if (SampleText.Length > 0)
+							{
+								SampleText.Append("\n\n");
+							}
+							SampleText.Append($"```{Data.Message}```");
+						}
+					}
+				}
+			}
+			return SampleText.ToString();
+		}
+
+		async Task SendIssueMessageAsync(string Recipient, IIssue Issue, IIssueDetails Details, UserId? UserId)
+		{
+			using IDisposable Scope = Logger.BeginScope("SendIssueMessageAsync (User: {SlackUser}, Issue: {IssueId})", Recipient, Issue.Id);
+
+			BlockKitAttachment Attachment = new BlockKitAttachment();
+			Attachment.Color = Color.Red;
+
+			Uri IssueUrl = Settings.DashboardUrl;
+			if (Details.Steps.Count > 0)
+			{
+				IIssueStep FirstStep = Details.Steps[0];
+				IssueUrl = new Uri(IssueUrl, $"job/{FirstStep.JobId}?step={FirstStep.StepId}&issue={Issue.Id}");
+			}
+			Attachment.Blocks.Add(new SectionBlock($"*<{IssueUrl}|Issue #{Issue.Id}: {Issue.Summary}>*"));
+
+			string StreamList = StringUtils.FormatList(Details.Spans.Select(x => $"*{x.StreamName}*").Distinct(StringComparer.OrdinalIgnoreCase).OrderBy(x => x, StringComparer.OrdinalIgnoreCase));
+			Attachment.Blocks.Add(new SectionBlock($"Occurring in {StreamList}"));
+
+			if (Issue.FixChange != null)
+			{
+				IIssueStep? FixFailedStep = Issue.FindFixFailedStep(Details.Spans);
+
+				string Text;
+				if (Issue.FixChange.Value < 0)
+				{
+					Text = ":tick: Marked as a systemic issue.";
+				}
+				else if (FixFailedStep != null)
+				{
+					Uri FixFailedUrl = new Uri(Settings.DashboardUrl, $"job/{FixFailedStep.JobId}?step={FixFailedStep.StepId}&issue={Issue.Id}");
+					Text = $":cross: Marked fixed in *CL {Issue.FixChange.Value}*, but seen again at *<{FixFailedUrl}|CL {FixFailedStep.Change}>*";
+				}
+				else
+				{
+					Text = $":tick: Marked fixed in *CL {Issue.FixChange.Value}*.";
+				}
+				Attachment.Blocks.Add(new SectionBlock(Text));
+			}
+			else if (UserId != null && Issue.OwnerId == UserId)
+			{
+				if (Issue.AcknowledgedAt.HasValue)
+				{
+					Attachment.Blocks.Add(new SectionBlock($":+1: Acknowledged at {FormatSlackTime(Issue.AcknowledgedAt.Value)}"));
+				}
+				else
+				{
+					if (Issue.NominatedById != null)
+					{
+						IUser? NominatedByUser = await UserCollection.GetUserAsync(Issue.NominatedById.Value);
+						if (NominatedByUser != null)
+						{
+							string? NominatedBySlackUserId = await GetSlackUserId(NominatedByUser);
+							string Mention = (NominatedBySlackUserId != null) ? $"<@{NominatedBySlackUserId}>" : NominatedByUser.Login ?? $"User {NominatedByUser.Id}";
+							string Text = $"You were nominated to fix this issue by {Mention} at {FormatSlackTime(Issue.NominatedAt ?? DateTime.UtcNow)}";
+							Attachment.Blocks.Add(new SectionBlock(Text));
+						}
+					}
+					else
+					{
+						List<int> Changes = Details.Suspects.Where(x => x.AuthorId == UserId).Select(x => x.Change).OrderBy(x => x).ToList();
+						if (Changes.Count > 0)
+						{
+							string Text = $"Horde has determined that {StringUtils.FormatList(Changes.Select(x => $"CL {x}"), "or")} is the most likely cause for this issue.";
+							Attachment.Blocks.Add(new SectionBlock(Text));
+						}
+					}
+
+					ActionsBlock Actions = new ActionsBlock();
+					Actions.AddButton("Acknowledge", Value: $"issue_{Issue.Id}_ack_{UserId}", Style: ActionButton.ButtonStyle.Primary);
+					Actions.AddButton("Not Me", Value: $"issue_{Issue.Id}_decline_{UserId}", Style: ActionButton.ButtonStyle.Danger);
+					Attachment.Blocks.Add(Actions);
+				}
+			}
+			else if (Issue.OwnerId != null)
+			{
+				string OwnerMention = await FormatMentionAsync(Issue.OwnerId.Value);
+				if (Issue.AcknowledgedAt.HasValue)
+				{
+					Attachment.Blocks.Add(new SectionBlock($":+1: Acknowledged by {OwnerMention} at {FormatSlackTime(Issue.AcknowledgedAt.Value)}"));
+				}
+				else if (Issue.NominatedById == null)
+				{
+					Attachment.Blocks.Add(new SectionBlock($"Assigned to {OwnerMention}"));
+				}
+				else if (Issue.NominatedById == UserId)
+				{
+					Attachment.Blocks.Add(new SectionBlock($"You nominated {OwnerMention} to fix this issue."));
+				}
+				else
+				{
+					Attachment.Blocks.Add(new SectionBlock($"{OwnerMention} was nominated to fix this issue by {await FormatMentionAsync(Issue.NominatedById.Value)}"));
+				}
+			}
+			else if (UserId != null)
+			{
+				IIssueSuspect? Suspect = Details.Suspects.FirstOrDefault(x => x.AuthorId == UserId);
+				if (Suspect != null)
+				{
+					if (Suspect.DeclinedAt != null)
+					{
+						Attachment.Blocks.Add(new SectionBlock($":downvote: Declined at {FormatSlackTime(Suspect.DeclinedAt.Value)}"));
+					}
+					else
+					{
+						Attachment.Blocks.Add(new SectionBlock("Please check if any of your recently submitted changes have caused this issue."));
+
+						ActionsBlock Actions = new ActionsBlock();
+						Actions.AddButton("Will Fix", Value: $"issue_{Issue.Id}_accept_{UserId}", Style: ActionButton.ButtonStyle.Primary);
+						Actions.AddButton("Not Me", Value: $"issue_{Issue.Id}_decline_{UserId}", Style: ActionButton.ButtonStyle.Danger);
+						Attachment.Blocks.Add(Actions);
+					}
+				}
+			}
+			else if (Details.Suspects.Count > 0)
+			{
+				List<string> DeclinedLines = new List<string>();
+				foreach (IIssueSuspect Suspect in Details.Suspects)
+				{
+					if (Suspect.DeclinedAt == null)
+					{
+						DeclinedLines.Add($":heavy_minus_sign: Ignored by {await FormatNameAsync(Suspect.AuthorId)} (CL {Suspect.Change})");
+					}
+					else
+					{
+						DeclinedLines.Add($":downvote: Declined by {await FormatNameAsync(Suspect.AuthorId)} at {FormatSlackTime(Suspect.DeclinedAt.Value)} (CL {Suspect.Change})");
+					}
+				}
+				Attachment.Blocks.Add(new SectionBlock(String.Join("\n", DeclinedLines)));
+			}
+
+			await SendOrUpdateMessageAsync(Recipient, GetIssueEventId(Issue), UserId, Attachments: new[] { Attachment });
+		}
+
+		static string GetIssueEventId(IIssue Issue)
+		{
+			return $"issue_{Issue.Id}";
+		}
+
+		async Task<string> FormatNameAsync(UserId UserId)
+		{
+			IUser? User = await UserCollection.GetUserAsync(UserId);
+			if (User == null)
+			{
+				return $"User {UserId}";
+			}
+			return User.Name;
+		}
+
+		async Task<string> FormatMentionAsync(UserId UserId)
+		{
+			IUser? User = await UserCollection.GetUserAsync(UserId);
+			if (User == null)
+			{
+				return $"User {UserId}";
+			}
+
+			string? SlackUserId = await GetSlackUserId(User);
+			if (SlackUserId == null)
+			{
+				return User.Login;
+			}
+
+			return $"<@{SlackUserId}>";
+		}
+
+		async Task HandleIssueResponseAsync(EventPayload Payload, ActionInfo Action)
+		{
+			string? UserName = Payload.User?.UserName;
+			if(UserName == null)
+			{
+				Logger.LogWarning("No user for message payload: {Payload}", Payload);
+				return;
+			}
+			if (Payload.ResponseUrl == null)
+			{
+				Logger.LogWarning("No response url for payload: {Payload}", Payload);
+				return;
+			}
+
+			Match Match = Regex.Match(Action.Value ?? String.Empty, @"^issue_(\d+)_([a-zA-Z]+)_([a-fA-F0-9]{24})$");
+			if (!Match.Success)
+			{
+				Logger.LogWarning("Could not match format of button action: {Action}", Action.Value);
+				return;
+			}
+
+			int IssueId = int.Parse(Match.Groups[1].Value, NumberStyles.Integer, CultureInfo.InvariantCulture);
+			string Verb = Match.Groups[2].Value;
+			UserId UserId = Match.Groups[3].Value.ToObjectId<IUser>();
+			Logger.LogInformation("Issue {IssueId}: {Action} from {SlackUser} ({UserId})", IssueId, Verb, UserName, UserId);
+
+			if (String.Equals(Verb, "ack", StringComparison.Ordinal))
+			{
+				await IssueService.UpdateIssueAsync(IssueId, Acknowledged: true);
+			}
+			else if (String.Equals(Verb, "accept", StringComparison.Ordinal))
+			{
+				await IssueService.UpdateIssueAsync(IssueId, OwnerId: UserId, Acknowledged: true);
+			}
+			else if (String.Equals(Verb, "decline", StringComparison.Ordinal))
+			{
+				await IssueService.UpdateIssueAsync(IssueId, DeclinedById: UserId);
+			}
+
+			IIssue? NewIssue = await IssueService.GetIssueAsync(IssueId);
+			if (NewIssue != null)
+			{
+				IUser? User = await UserCollection.GetUserAsync(UserId);
+				if (User != null)
+				{
+					string? Recipient = await GetSlackUserId(User);
+					if (Recipient != null)
+					{
+						IIssueDetails Details = await IssueService.GetIssueDetailsAsync(NewIssue);
+						await SendIssueMessageAsync(Recipient, NewIssue, Details, UserId);
+					}
+				}
+			}
+		}
+
+		static string FormatSlackTime(DateTimeOffset Time)
+		{
+			return $"<!date^{Time.ToUnixTimeSeconds()}^{{time}}|{Time}>";
+		}
+
+		#endregion
+
+		#region Stream updates
+
+		/// <inheritdoc/>
+		public async Task NotifyConfigUpdateFailureAsync(string ErrorMessage, string FileName, int? Change = null, IUser? Author = null, string? Description = null)
+		{
+			Logger.LogInformation("Sending config update failure notification for {FileName}", FileName);
+
+			string? SlackUserId = null;
+			if (Author != null)
+			{
+				SlackUserId = await GetSlackUserId(Author);
+			}
+
+			if (SlackUserId != null)
+			{
+				await SendConfigUpdateFailureMessageAsync(SlackUserId, ErrorMessage, FileName, Change, SlackUserId, Description);
+			}
+			if (Settings.UpdateStreamsNotificationChannel != null)
+			{
+				await SendConfigUpdateFailureMessageAsync($"#{Settings.UpdateStreamsNotificationChannel}", ErrorMessage, FileName, Change, SlackUserId, Description);
+			}
+		}
+
+		private Task SendConfigUpdateFailureMessageAsync(string Recipient, string ErrorMessage, string FileName, int? Change = null, string? Author = null, string? Description = null)
+		{
+			Color OutcomeColor = BlockKitAttachmentColors.Error;
+			BlockKitAttachment Attachment = new BlockKitAttachment();
+			Attachment.Color = OutcomeColor;
+			Attachment.FallbackText = $"Update Failure: {FileName}";
+
+			Attachment.Blocks.Add(new HeaderBlock($"Config Update Failure :rip:", false, true));
+
+			Attachment.Blocks.Add(new SectionBlock($"Horde was unable to update {FileName}"));
+			Attachment.Blocks.Add(new SectionBlock($"```{ErrorMessage}```"));
+			if (Change != null)
+			{
+				if (Author != null)
+				{
+					Attachment.Blocks.Add(new SectionBlock($"Possibly due to CL: {Change.Value} by <@{Author}>"));
+				}
+				else
+				{
+					Attachment.Blocks.Add(new SectionBlock($"Possibly due to CL: {Change.Value} - (Could not determine author from P4 user)"));
+				}
+				if (Description != null)
+				{
+					Attachment.Blocks.Add(new SectionBlock($"```{Description}```"));
+				}
+			}
+
+			return SendMessageAsync(Recipient, Attachments: new[] { Attachment });
+		}
+
+		#endregion
+
+		#region Stream update (file)
+
+		/// <inheritdoc/>
+		public async Task NotifyStreamUpdateFailedAsync(FileSummary File)
+		{
+			Logger.LogDebug("Sending stream update failure notification for {File}", File.DepotPath);
+			if (Settings.UpdateStreamsNotificationChannel != null)
+			{
+				await SendStreamUpdateFailureMessage($"#{Settings.UpdateStreamsNotificationChannel}", File);
+			}
+		}
+
+		/// <summary>
+		/// Creates a stream update failure message in relation to a file
+		/// </summary>
+		/// <param name="Recipient"></param>
+		/// <param name="File">The file</param>
+		/// <returns></returns>
+		Task SendStreamUpdateFailureMessage(string Recipient, FileSummary File)
+		{
+			Color OutcomeColor = BlockKitAttachmentColors.Error;
+			BlockKitAttachment Attachment = new BlockKitAttachment();
+			Attachment.Color = OutcomeColor;
+			Attachment.FallbackText = $"{File.DepotPath} - Update Failure";
+
+			Attachment.Blocks.Add(new HeaderBlock($"Stream Update Failure :rip:", false, true));
+
+			Attachment.Blocks.Add(new SectionBlock($"<!here> Horde was unable to update {File.DepotPath}"));
+			Attachment.Blocks.Add(new SectionBlock($"```{File.Error}```"));
+
+			return SendMessageAsync(Recipient, Attachments: new[] { Attachment });
+		}
+
+		#endregion
+
+		#region Device notifications
+
+		/// <inheritdoc/>
+		public async Task NotifyDeviceServiceAsync(string Message, IDevice? Device = null, IDevicePool? Pool = null, IStream? Stream = null, IJob? Job = null, IJobStep? Step = null, INode? Node = null)
+		{
+			Logger.LogDebug("Sending device service failure notification for {DeviceName} in pool {PoolName}", Device?.Name, Pool?.Name);
+			if (Settings.DeviceServiceNotificationChannel != null)
+			{
+				await SendDeviceServiceMessage($"#{Settings.DeviceServiceNotificationChannel}", Message, Device, Pool, Stream, Job, Step, Node);
+			}
+		}
+
+		/// <summary>
+		/// Creates a Slack message about a completed step job.
+		/// </summary>
+		/// <param name="Recipient"></param>
+        /// <param name="Message"></param>
+        /// <param name="Device"></param>
+        /// <param name="Pool"></param>
+		/// <param name="Stream"></param>
+		/// <param name="Job">The job that contains the step that completed.</param>
+		/// <param name="Step">The job step that completed.</param>
+		/// <param name="Node">The node for the job step.</param>
+		private Task SendDeviceServiceMessage(string Recipient, string Message, IDevice? Device = null, IDevicePool? Pool = null, IStream? Stream = null, IJob? Job = null, IJobStep? Step = null, INode? Node = null)
+		{
+			BlockKitAttachment Attachment = new BlockKitAttachment();
+
+			// truncate message to avoid slack error on message length
+			if (Message.Length > 150)
+			{
+				Message = Message.Substring(0, 146) + "...";
+			}
+            
+            Attachment.FallbackText = $"{Message}";
+
+			if (Device != null && Pool != null)
+			{
+                Attachment.FallbackText += $" - Device: {Device.Name} Pool: {Pool.Name}";
+            }
+
+			Attachment.Blocks.Add(new HeaderBlock($"{Message}", false, false));
+
+			if (Stream != null && Job != null && Step != null && Node != null)
+			{
+				Uri JobStepLink = new Uri($"{Settings.DashboardUrl}job/{Job.Id}?step={Step.Id}");
+				Uri JobStepLogLink = new Uri($"{Settings.DashboardUrl}log/{Step.LogId}");
+				
+				Attachment.FallbackText += $" - {Stream.Name} - {GetJobChangeText(Job)} - {Job.Name} - {Node.Name}";				
+				Attachment.Blocks.Add(new SectionBlock($"*<{JobStepLink}|{Stream.Name} - {GetJobChangeText(Job)} - {Job.Name} - {Node.Name}>*"));
+				Attachment.Blocks.Add(new SectionBlock($"<{JobStepLogLink}|View Job Step Log>"));
+			}
+			else
+			{				
+				Attachment.FallbackText += " - No job information (Gauntlet might need to be updated in stream)";				
+				Attachment.Blocks.Add(new SectionBlock("*No job information (Gauntlet might need to be updated in stream)*"));
+			}
+
+			return SendMessageAsync(Recipient, Attachments: new[] { Attachment });
+		}
+		
+		#endregion
+
+
+		const int MaxJobStepEvents = 5;
+
+		static string GetJobChangeText(IJob Job)
+		{
+			if (Job.PreflightChange == 0)
+			{
+				return $"CL {Job.Change}";
+			}
+			else
+			{
+				return $"Preflight CL {Job.PreflightChange} against CL {Job.Change}";
+			}
+		}
+
+		static bool ShouldUpdateUser(SlackUser? Document)
+		{
+			if(Document == null || Document.Version < SlackUser.CurrentVersion)
+			{
+				return true;
+			}
+
+			TimeSpan ExpiryTime;
+			if (Document.SlackUserId == null)
+			{
+				ExpiryTime = TimeSpan.FromMinutes(10.0);
+			}
+			else
+			{
+				ExpiryTime = TimeSpan.FromDays(1.0);
+			}
+			return Document.Time + ExpiryTime < DateTime.UtcNow;
+		}
+
+		private async Task<string?> GetSlackUserId(IUser User)
+		{
+			return (await GetSlackUser(User))?.SlackUserId;
+		}
+
+		private async Task<SlackUser?> GetSlackUser(IUser User)
+		{
+			string? Email = User.Email;
+			if (Email == null)
+			{
+				Logger.LogWarning("Unable to find Slack user id for {UserId} ({Name}): No email address in user profile", User.Id, User.Name);
+				return null;
+			}
+
+			SlackUser? UserDocument;
+			if (!UserCache.TryGetValue(Email, out UserDocument))
+			{
+				UserDocument = await SlackUsers.Find(x => x.Id == User.Id).FirstOrDefaultAsync();
+				if (UserDocument == null || ShouldUpdateUser(UserDocument))
+				{
+					UserInfo? UserInfo = await GetSlackUserInfoByEmail(Email);
+					if (UserDocument == null || UserInfo != null)
+					{
+						UserDocument = new SlackUser(User.Id, UserInfo);
+						await SlackUsers.ReplaceOneAsync(x => x.Id == User.Id, UserDocument, new ReplaceOptions { IsUpsert = true });
+					}
+				}
+				using (ICacheEntry Entry = UserCache.CreateEntry(Email))
+				{
+					Entry.SlidingExpiration = TimeSpan.FromMinutes(10.0);
+					Entry.Value = UserDocument;
+				}
+			}
+
+			return UserDocument;
+		}
+
+		private async Task<UserInfo?> GetSlackUserInfoByEmail(string Email)
+		{
+			using HttpClient Client = new HttpClient();
+
+			using HttpRequestMessage GetUserIdRequest = new HttpRequestMessage(HttpMethod.Post, $"https://slack.com/api/users.lookupByEmail?email={Email}");
+			GetUserIdRequest.Headers.Add("Authorization", $"Bearer {Settings.SlackToken ?? ""}");
+
+			HttpResponseMessage ResponseMessage = await Client.SendAsync(GetUserIdRequest);
+			byte[] ResponseData = await ResponseMessage.Content.ReadAsByteArrayAsync();
+
+			UserResponse UserResponse = JsonSerializer.Deserialize<UserResponse>(ResponseData)!;
+			if(!UserResponse.Ok || UserResponse.User == null)
+			{
+				Logger.LogWarning("Unable to find Slack user id for {Email}: {Response}", Email, Encoding.UTF8.GetString(ResponseData));
+				return null;
+			}
+
+			return UserResponse.User;
+		}
+
+		class SlackResponse
+		{
+			[JsonPropertyName("ok")]
+			public bool Ok { get; set; }
+
+			[JsonPropertyName("error")]
+			public string? Error { get; set; }
+		}
+
+		class PostMessageResponse : SlackResponse
+		{
+			[JsonPropertyName("channel")]
+			public string? Channel { get; set; }
+
+			[JsonPropertyName("ts")]
+			public string? Ts { get; set; }
+		}
+
+		private async Task SendMessageAsync(string Recipient, string? Text = null, BlockBase[]? Blocks = null, BlockKitAttachment[]? Attachments = null)
+		{
+			if (AllowUsers != null && !AllowUsers.Contains(Recipient))
+			{
+				Logger.LogDebug("Suppressing message to {Recipient}: {Text}", Recipient, Text);
+				return;
+			}
+
+			BlockKitMessage Message = new BlockKitMessage();
+			Message.Channel = Recipient;
+			Message.Text = Text;
+			if (Blocks != null)
+			{
+				Message.Blocks.AddRange(Blocks);
+			}
+			if (Attachments != null)
+			{
+				Message.Attachments.AddRange(Attachments);
+			}
+
+			await SendRequestAsync<PostMessageResponse>(PostMessageUrl, Message);
+		}
+
+		private async Task SendOrUpdateMessageAsync(string Recipient, string EventId, UserId? UserId, string? Text = null, BlockBase[]? Blocks = null, BlockKitAttachment[]? Attachments = null)
+		{
+			if (AllowUsers != null && !AllowUsers.Contains(Recipient))
+			{
+				Logger.LogDebug("Suppressing message to {Recipient}: {Text}", Recipient, Text);
+				return;
+			}
+
+			BlockKitMessage Message = new BlockKitMessage();
+			Message.Text = Text;
+			if (Blocks != null)
+			{
+				Message.Blocks.AddRange(Blocks);
+			}
+			if (Attachments != null)
+			{
+				Message.Attachments.AddRange(Attachments);
+			}
+
+			string RequestDigest = ContentHash.MD5(JsonSerializer.Serialize(Message, new JsonSerializerOptions { IgnoreNullValues = true })).ToString();
+
+			MessageStateDocument PrevState = await AddOrUpdateMessageStateAsync(Recipient, EventId, UserId, RequestDigest);
+			if (String.IsNullOrEmpty(PrevState.Ts))
+			{
+				Logger.LogInformation("Sending new slack message to {SlackUser}", Recipient);
+				Message.Channel = Recipient;
+				Message.Ts = null;
+
+				PostMessageResponse Response = await SendRequestAsync<PostMessageResponse>(PostMessageUrl, Message);
+				if (String.IsNullOrEmpty(Response.Ts) || String.IsNullOrEmpty(Response.Channel))
+				{
+					Logger.LogWarning("Missing 'ts' or 'channel' field on slack response");
+				}
+				await SetMessageTimestampAsync(PrevState.Id, Response.Channel ?? String.Empty, Response.Ts ?? String.Empty);
+			}
+			else
+			{
+				Logger.LogInformation("Updating existing slack message for user {SlackUser} ({Channel}, {MessageTs})", Recipient, PrevState.Channel, PrevState.Ts);
+				Message.Channel = PrevState.Channel;
+				Message.Ts = PrevState.Ts;
+				await SendRequestAsync<PostMessageResponse>(UpdateMessageUrl, Message);
+			}
+		}
+
+		private async Task<TResponse> SendRequestAsync<TResponse>(string RequestUrl, object Request) where TResponse : SlackResponse
+		{
+			using (HttpClient Client = new HttpClient())
+			{
+				using (HttpRequestMessage SendMessageRequest = new HttpRequestMessage(HttpMethod.Post, RequestUrl))
+				{
+					string RequestJson = JsonSerializer.Serialize(Request, new JsonSerializerOptions { IgnoreNullValues = true });
+					using (StringContent MessageContent = new StringContent(RequestJson, Encoding.UTF8, "application/json"))
+					{
+						SendMessageRequest.Content = MessageContent;
+						SendMessageRequest.Headers.Add("Authorization", $"Bearer {Settings.SlackToken ?? ""}");
+
+						HttpResponseMessage Response = await Client.SendAsync(SendMessageRequest);
+						byte[] ResponseBytes = await Response.Content.ReadAsByteArrayAsync();
+
+						TResponse ResponseObject = JsonSerializer.Deserialize<TResponse>(ResponseBytes)!;
+						if(!ResponseObject.Ok)
+						{
+							Logger.LogError("Failed to send Slack message ({Error}). Request: {Request}. Response: {Response}", ResponseObject.Error, RequestJson, Encoding.UTF8.GetString(ResponseBytes));
+						}
+
+						return ResponseObject;
+					}
+				}
+			}
+		}
+
+		/// <inheritdoc/>
+		protected async override Task ExecuteAsync(CancellationToken StoppingToken)
+		{
+			if (!String.IsNullOrEmpty(Settings.SlackSocketToken))
+			{
+				while (!StoppingToken.IsCancellationRequested)
+				{
+					try
+					{
+						Uri? WebSocketUrl = await GetWebSocketUrlAsync(StoppingToken);
+						if (WebSocketUrl == null)
+						{
+							await Task.Delay(TimeSpan.FromSeconds(5.0), StoppingToken);
+							continue;
+						}
+						await HandleSocketAsync(WebSocketUrl, StoppingToken);
+					}
+					catch (OperationCanceledException)
+					{
+						break;
+					}
+					catch (Exception Ex)
+					{
+						Logger.LogError(Ex, "Exception while updating Slack socket");
+						await Task.Delay(TimeSpan.FromSeconds(5.0), StoppingToken);
+					}
+				}
+			}
+		}
+
+		/// <summary>
+		/// Get the url for opening a websocket to Slack
+		/// </summary>
+		/// <param name="StoppingToken"></param>
+		/// <returns></returns>
+		private async Task<Uri?> GetWebSocketUrlAsync(CancellationToken StoppingToken)
+		{
+			using HttpClient Client = new HttpClient();
+			Client.DefaultRequestHeaders.Add("Authorization", $"Bearer {Settings.SlackSocketToken}");
+
+			using FormUrlEncodedContent Content = new FormUrlEncodedContent(Array.Empty<KeyValuePair<string?, string?>>());
+			HttpResponseMessage Response = await Client.PostAsync(new Uri("https://slack.com/api/apps.connections.open"), Content, StoppingToken);
+
+			byte[] ResponseData = await Response.Content.ReadAsByteArrayAsync();
+
+			SocketResponse SocketResponse = JsonSerializer.Deserialize<SocketResponse>(ResponseData)!;
+			if (!SocketResponse.Ok)
+			{
+				Logger.LogWarning("Unable to get websocket url: {Response}", Encoding.UTF8.GetString(ResponseData));
+				return null;
+			}
+
+			return SocketResponse.Url;
+		}
+
+		/// <summary>
+		/// Handle the lifetime of a websocket connection
+		/// </summary>
+		/// <param name="SocketUrl"></param>
+		/// <param name="StoppingToken"></param>
+		/// <returns></returns>
+		private async Task HandleSocketAsync(Uri SocketUrl, CancellationToken StoppingToken)
+		{
+			using ClientWebSocket Socket = new ClientWebSocket();
+			await Socket.ConnectAsync(SocketUrl, StoppingToken);
+
+			byte[] Buffer = new byte[2048];
+			while (!StoppingToken.IsCancellationRequested)
+			{
+				// Read the next message
+				int Length = 0;
+				for (; ; )
+				{
+					if (Length == Buffer.Length)
+					{
+						Array.Resize(ref Buffer, Buffer.Length + 2048);
+					}
+
+					WebSocketReceiveResult Result = await Socket.ReceiveAsync(new ArraySegment<byte>(Buffer, Length, Buffer.Length - Length), StoppingToken);
+					if (Result.MessageType == WebSocketMessageType.Close)
+					{
+						return;
+					}
+					Length += Result.Count;
+
+					if (Result.EndOfMessage)
+					{
+						break;
+					}
+				}
+
+				// Get the message data
+				Logger.LogInformation("Slack event: {Message}", Encoding.UTF8.GetString(Buffer, 0, Length));
+				EventMessage EventMessage = JsonSerializer.Deserialize<EventMessage>(Buffer.AsSpan(0, Length))!;
+
+				// Acknowledge the message
+				if (EventMessage.EnvelopeId != null)
+				{
+					object Response = new { EventMessage.EnvelopeId };
+					await Socket.SendAsync(JsonSerializer.SerializeToUtf8Bytes(Response), WebSocketMessageType.Text, true, StoppingToken);
+				}
+
+				// Handle the message type
+				if (EventMessage.Type != null)
+				{
+					string Type = EventMessage.Type;
+					if (Type.Equals("disconnect", StringComparison.Ordinal))
+					{
+						break;
+					}
+					else if (Type.Equals("interactive", StringComparison.Ordinal))
+					{
+						await HandleInteractionMessage(EventMessage);
+					}
+					else
+					{
+						Logger.LogDebug("Unhandled event type ({Type})", Type);
+					}
+				}
+			}
+		}
+
+		/// <summary>
+		/// Handle a button being clicked
+		/// </summary>
+		/// <param name="Message">The event message</param>
+		/// <returns></returns>
+		private async Task HandleInteractionMessage(EventMessage Message)
+		{
+			if (Message.Payload != null && Message.Payload.User != null && Message.Payload.User.UserName != null && String.Equals(Message.Payload.Type, "block_actions", StringComparison.Ordinal))
+			{
+				foreach (ActionInfo Action in Message.Payload.Actions)
+				{
+					if (Action.Value != null)
+					{
+						if (Action.Value.StartsWith("issue_", StringComparison.Ordinal))
+						{
+							await HandleIssueResponseAsync(Message.Payload, Action);
+						}
+					}
+				}
+			}
+		}
+	}
+}
diff --git a/Engine/Source/Programs/Horde/Horde.Build/Perforce/Change.cs b/Engine/Source/Programs/Horde/Horde.Build/Perforce/Change.cs
new file mode 100644
--- /dev/null
+++ b/Engine/Source/Programs/Horde/Horde.Build/Perforce/Change.cs
@@ -0,0 +1,302 @@
+// Copyright Epic Games, Inc. All Rights Reserved.
+
+using EpicGames.Core;
+using HordeServer.Api;
+using HordeServer.Utilities;
+using MongoDB.Bson;
+using MongoDB.Bson.Serialization.Attributes;
+using System;
+using System.Collections.Generic;
+using System.Diagnostics;
+using System.Diagnostics.CodeAnalysis;
+using System.Globalization;
+using System.Linq;
+using System.Text.RegularExpressions;
+using System.Threading.Tasks;
+
+namespace HordeServer.Models
+{
+	using UserId = ObjectId<IUser>;
+
+	/// <summary>
+	/// Summary information for a change
+	/// </summary>
+	[DebuggerDisplay("{Number}: {Description}")]
+	public class ChangeSummary
+	{
+		/// <summary>
+		/// The changelist number
+		/// </summary>
+		public int Number { get; set; }
+
+		/// <summary>
+		/// Author of the change
+		/// </summary>
+		public IUser Author { get; set; }
+
+		/// <summary>
+		/// The base path for modified files
+		/// </summary>
+		public string Path { get; set; }
+
+		/// <summary>
+		/// Abbreviated changelist description
+		/// </summary>
+		public string Description { get; set; }
+
+		/// <summary>
+		/// Constructor
+		/// </summary>
+		/// <param name="Number">Changelist number</param>
+		/// <param name="Author">Author of the change</param>
+		/// <param name="Path">Base path for modified files</param>
+		/// <param name="Description">Changelist description</param>
+		public ChangeSummary(int Number, IUser Author, string Path, string Description)
+		{
+			this.Number = Number;
+			this.Author = Author;
+			this.Path = Path;
+			this.Description = Description;
+		}
+	}
+
+	/// <summary>
+	/// Flags identifying content of a changelist
+	/// </summary>
+	[Flags]
+	public enum ChangeContentFlags
+	{
+		/// <summary>
+		/// The change contains code
+		/// </summary>
+		ContainsCode = 1,
+
+		/// <summary>
+		/// The change contains content
+		/// </summary>
+		ContainsContent = 2,
+	}
+
+	/// <summary>
+	/// Modified file in a changelist
+	/// </summary>
+	public class ChangeFile
+	{
+		/// <summary>
+		/// Path to the file
+		/// </summary>
+		public string Path { get; set; }
+
+		/// <summary>
+		/// Path to the file within the depot
+		/// </summary>
+		public string DepotPath { get; set; }
+
+		/// <summary>
+		/// Revision of the file. A value of -1 indicates that the file was deleted.
+		/// </summary>
+		public int Revision { get; set; }
+
+		/// <summary>
+		/// Length of the file
+		/// </summary>
+		public long Length { get; set; }
+
+		/// <summary>
+		/// MD5 digest of the file
+		/// </summary>
+		public Md5Hash? Digest { get; set; }
+
+		/// <summary>
+		/// The file type
+		/// </summary>
+		public string Type { get; set; }
+
+		/// <summary>
+		/// Constructor
+		/// </summary>
+		/// <param name="Path"></param>
+		/// <param name="DepotPath"></param>
+		/// <param name="Revision"></param>
+		/// <param name="Length"></param>
+		/// <param name="Digest"></param>
+		/// <param name="Type"></param>
+		public ChangeFile(string Path, string DepotPath, int Revision, long Length, Md5Hash? Digest, string Type)
+		{
+			this.Path = Path;
+			this.DepotPath = DepotPath;
+			this.Revision = Revision;
+			this.Length = Length;
+			this.Digest = Digest;
+			this.Type = Type;
+		}
+	}
+
+	/// <summary>
+	/// Information about a commit
+	/// </summary>
+	[DebuggerDisplay("{Number}: {Description}")]
+	public class ChangeDetails
+	{
+		/// <summary>
+		/// Set of extensions to treat as code
+		/// </summary>
+		static readonly HashSet<string> CodeExtensions = new HashSet<string>
+		{
+			".c",
+			".cc",
+			".cpp",
+			".m",
+			".mm",
+			".rc",
+			".cs",
+			".csproj",
+			".h",
+			".hpp",
+			".inl",
+			".usf",
+			".ush",
+			".uproject",
+			".uplugin",
+			".sln"
+		};
+
+		/// <summary>
+		/// The source changelist number
+		/// </summary>
+		public int Number { get; set; }
+
+		/// <summary>
+		/// Name of the user that authored this change
+		/// </summary>
+		public IUser Author { get; set; }
+
+		/// <summary>
+		/// The base path for modified files
+		/// </summary>
+		public string Path { get; set; }
+
+		/// <summary>
+		/// The description text
+		/// </summary>
+		public string Description { get; set; }
+
+		/// <summary>
+		/// List of files that were modified, relative to the stream base
+		/// </summary>
+		public List<ChangeFile> Files { get; set; }
+
+		/// <summary>
+		/// Date that the change was submitted
+		/// </summary>
+		public DateTime Date { get; set; }
+
+		/// <summary>
+		/// Constructor
+		/// </summary>
+		/// <param name="Number">Changelist number</param>
+		/// <param name="Author">Author of the change</param>
+		/// <param name="Path">Base path for modified files</param>
+		/// <param name="Description">Changelist description</param>
+		/// <param name="Files">List of files modified, relative to the stream base</param>
+		/// <param name="Date">Date that the change was submitted</param>
+		public ChangeDetails(int Number, IUser Author, string Path, string Description, List<ChangeFile> Files, DateTime Date)
+		{
+			this.Number = Number;
+			this.Author = Author;
+			this.Path = Path;
+			this.Description = Description;
+			this.Files = Files;
+			this.Date = Date;
+		}
+
+		/// <summary>
+		/// Determines if this change is a code change
+		/// </summary>
+		/// <returns>True if this change is a code change</returns>
+		public ChangeContentFlags GetContentFlags()
+		{
+			ChangeContentFlags Scope = 0;
+
+			// Check whether the files are code or content
+			foreach (ChangeFile File in Files)
+			{
+				if (CodeExtensions.Any(Extension => File.Path.EndsWith(Extension, StringComparison.OrdinalIgnoreCase)))
+				{
+					Scope |= ChangeContentFlags.ContainsCode;
+				}
+				else
+				{
+					Scope |= ChangeContentFlags.ContainsContent;
+				}
+
+				if (Scope == (ChangeContentFlags.ContainsCode | ChangeContentFlags.ContainsContent))
+				{
+					break;
+				}
+			}
+			return Scope;
+		}
+
+		/// <summary>
+		/// Attempts to parse the Robomerge source from this commit information
+		/// </summary>
+		/// <param name="Source">On success, receives the source information</param>
+		/// <returns>True if the commit was merged from another stream</returns>
+		public bool TryParseRobomergeSource([NotNullWhen(true)] out (string, int)? Source)
+		{
+			Match Match = Regex.Match(Description, @"#ROBOMERGE-SOURCE: CL (\d+) in (//[^ ]*)/...", RegexOptions.Multiline);
+			if (Match.Success)
+			{
+				Source = (Match.Groups[2].Value, int.Parse(Match.Groups[1].Value, CultureInfo.InvariantCulture));
+				return true;
+			}
+			else
+			{
+				Source = null;
+				return false;
+			}
+		}
+	}
+
+	/// <summary>
+	/// Summary of a file in the depot
+	/// </summary>
+	public class FileSummary
+	{
+		/// <summary>
+		/// Depot path to the file
+		/// </summary>
+		public string DepotPath { get; set; }
+
+		/// <summary>
+		/// Whether the file exists
+		/// </summary>
+		public bool Exists { get; set; }
+
+		/// <summary>
+		/// Last changelist number that the file was modified
+		/// </summary>
+		public int Change { get; set; }
+
+		/// <summary>
+		/// Error about this file, if it exists
+		/// </summary>
+		public string? Error { get; set; }
+
+		/// <summary>
+		/// Information about a file
+		/// </summary>
+		/// <param name="DepotPath">Depot </param>
+		/// <param name="Exists"></param>
+		/// <param name="Change"></param>
+		/// <param name="Error"></param>
+		public FileSummary(string DepotPath, bool Exists, int Change, string? Error = null)
+		{
+			this.DepotPath = DepotPath;
+			this.Exists = Exists;
+			this.Change = Change;
+			this.Error = Error;
+		}
+	}
+}
diff --git a/Engine/Source/Programs/Horde/Horde.Build/Perforce/PerforceLoadBalancer.cs b/Engine/Source/Programs/Horde/Horde.Build/Perforce/PerforceLoadBalancer.cs
new file mode 100644
--- /dev/null
+++ b/Engine/Source/Programs/Horde/Horde.Build/Perforce/PerforceLoadBalancer.cs
@@ -0,0 +1,598 @@
+// Copyright Epic Games, Inc. All Rights Reserved.
+
+using Google.Protobuf.WellKnownTypes;
+using HordeCommon.Rpc.Tasks;
+using HordeServer.Collections;
+using HordeServer.Models;
+using HordeServer.Utilities;
+using Microsoft.Extensions.Hosting;
+using Microsoft.Extensions.Logging;
+using MongoDB.Bson;
+using MongoDB.Bson.Serialization.Attributes;
+using MongoDB.Driver;
+using System;
+using System.Collections.Concurrent;
+using System.Collections.Generic;
+using System.Globalization;
+using System.Linq;
+using System.Net;
+using System.Net.Http;
+using System.Text.Json;
+using System.Threading;
+using System.Threading.Tasks;
+
+namespace HordeServer.Services
+{
+	/// <summary>
+	/// Health of a particular server
+	/// </summary>
+	public enum PerforceServerStatus
+	{
+		/// <summary>
+		/// Server could not be reached
+		/// </summary>
+		Unknown,
+
+		/// <summary>
+		/// Bad. Do not use.
+		/// </summary>
+		Unhealthy,
+
+		/// <summary>
+		/// Degraded, but functioning
+		/// </summary>
+		Degraded,
+
+		/// <summary>
+		/// Good
+		/// </summary>
+		Healthy,
+	}
+
+	/// <summary>
+	/// Perforce server infomation
+	/// </summary>
+	public interface IPerforceServer
+	{
+		/// <summary>
+		/// The server and port
+		/// </summary>
+		public string ServerAndPort { get; }
+
+		/// <summary>
+		/// The server and port before doing DNS resolution
+		/// </summary>
+		public string BaseServerAndPort { get; }
+
+		/// <summary>
+		/// The cluster this server belongs to
+		/// </summary>
+		public string Cluster { get; }
+
+		/// <summary>
+		/// Current status
+		/// </summary>
+		public PerforceServerStatus Status { get; }
+
+		/// <summary>
+		/// Number of leases using this server
+		/// </summary>
+		public int NumLeases { get; }
+
+		/// <summary>
+		/// Error message related to this server
+		/// </summary>
+		public string? Detail { get; }
+
+		/// <summary>
+		/// Last update time for this server
+		/// </summary>
+		public DateTime? LastUpdateTime { get; }
+	}
+
+	/// <summary>
+	/// Load balancer for Perforce edge servers
+	/// </summary>
+	public sealed class PerforceLoadBalancer : IHostedService, IDisposable
+	{
+		/// <summary>
+		/// Information about a resolved Perforce server
+		/// </summary>
+		class PerforceServerEntry : IPerforceServer
+		{
+			public string ServerAndPort { get; set; } = String.Empty;
+			public string BaseServerAndPort { get; set; } = String.Empty;
+			public string? HealthCheckUrl { get; set; }
+			public string Cluster { get; set; } = String.Empty;
+			public PerforceServerStatus Status { get; set; }
+			public string? Detail { get; set; }
+			public int NumLeases { get; set; }
+			public DateTime? LastUpdateTime { get; set; }
+
+			[BsonConstructor]
+			private PerforceServerEntry()
+			{
+			}
+
+			public PerforceServerEntry(string ServerAndPort, string BaseServerAndPort, string? HealthCheckUrl, string Cluster, PerforceServerStatus Status, string? Detail, DateTime? LastUpdateTime)
+			{
+				this.ServerAndPort = ServerAndPort;
+				this.BaseServerAndPort = BaseServerAndPort;
+				this.HealthCheckUrl = HealthCheckUrl;
+				this.Cluster = Cluster;
+				this.Status = Status;
+				this.Detail = Detail;
+				this.LastUpdateTime = LastUpdateTime;
+			}
+		}
+
+		/// <summary>
+		/// 
+		/// </summary>
+		[SingletonDocument("6046aec374a9283100967ee7")]
+		class PerforceServerList : SingletonBase
+		{
+			public List<PerforceServerEntry> Servers { get; set; } = new List<PerforceServerEntry>();
+		}
+
+		DatabaseService DatabaseService;
+		ILeaseCollection LeaseCollection;
+		SingletonDocument<PerforceServerList> ServerListSingleton;
+		Random Random = new Random();
+		ILogger Logger;
+		ElectedTick Ticker;
+		Task InternalTickTask = Task.CompletedTask;
+
+		/// <summary>
+		/// Constructor
+		/// </summary>
+		/// <param name="DatabaseService"></param>
+		/// <param name="LeaseCollection"></param>
+		/// <param name="Logger"></param>
+		public PerforceLoadBalancer(DatabaseService DatabaseService, ILeaseCollection LeaseCollection, ILogger<PerforceLoadBalancer> Logger)
+		{
+			this.DatabaseService = DatabaseService;
+			this.LeaseCollection = LeaseCollection;
+			this.ServerListSingleton = new SingletonDocument<PerforceServerList>(DatabaseService);
+			this.Logger = Logger;
+			this.Ticker = new ElectedTick(DatabaseService, new ObjectId("603fb20536e999974fd6ff6b"), TickAsync, TimeSpan.FromMinutes(1.0), Logger);
+		}
+
+		/// <inheritdoc/>
+		public void Dispose()
+		{
+			Ticker.Dispose();
+		}
+
+		/// <inheritdoc/>
+		public Task StartAsync(CancellationToken cancellationToken)
+		{
+			Ticker.Start();
+			return Task.CompletedTask;
+		}
+
+		/// <inheritdoc/>
+		public async Task StopAsync(CancellationToken cancellationToken)
+		{
+			await Ticker.StopAsync();
+			await InternalTickTask;
+		}
+
+		/// <summary>
+		/// Get the current server list
+		/// </summary>
+		/// <returns></returns>
+		async Task<PerforceServerList> GetServerListAsync()
+		{
+			PerforceServerList ServerList = await ServerListSingleton.GetAsync();
+
+			DateTime MinLastUpdateTime = DateTime.UtcNow - TimeSpan.FromMinutes(2.5);
+			foreach (PerforceServerEntry Server in ServerList.Servers)
+			{
+				if (Server.Status == PerforceServerStatus.Healthy && Server.LastUpdateTime != null && Server.LastUpdateTime.Value < MinLastUpdateTime)
+				{
+					Server.Status = PerforceServerStatus.Degraded;
+					Server.Detail = "Server has not responded to health check";
+				}
+			}
+
+			return ServerList;
+		}
+
+		/// <summary>
+		/// Get the perforce servers
+		/// </summary>
+		/// <returns></returns>
+		public async Task<List<IPerforceServer>> GetServersAsync()
+		{
+			PerforceServerList ServerList = await GetServerListAsync();
+			return ServerList.Servers.ConvertAll<IPerforceServer>(x => x);
+		}
+
+		/// <summary>
+		/// Allocates a server for use by a lease
+		/// </summary>
+		/// <returns>The server to use. Null if there is no healthy server available.</returns>
+		public async Task<IPerforceServer?> GetServer(string Cluster)
+		{
+			PerforceServerList ServerList = await GetServerListAsync();
+
+			List<PerforceServerEntry> Candidates = ServerList.Servers.Where(x => x.Cluster == Cluster && x.Status >= PerforceServerStatus.Healthy).ToList();
+			if(Candidates.Count == 0)
+			{
+				int Idx = Random.Next(0, Candidates.Count);
+				return Candidates[Idx];
+			}
+			return null;
+		}
+
+		/// <summary>
+		/// Select a Perforce server to use by the Horde server
+		/// </summary>
+		/// <param name="Cluster"></param>
+		/// <returns></returns>
+		public Task<IPerforceServer?> SelectServerAsync(PerforceCluster Cluster)
+		{
+			List<string> Properties = new List<string>{ "HordeServer=1" };
+			return SelectServerAsync(Cluster, Properties);
+		}
+
+		/// <summary>
+		/// Select a Perforce server to use
+		/// </summary>
+		/// <param name="Cluster"></param>
+		/// <param name="Agent"></param>
+		/// <returns></returns>
+		public Task<IPerforceServer?> SelectServerAsync(PerforceCluster Cluster, IAgent Agent)
+		{
+			return SelectServerAsync(Cluster, Agent.Properties);
+		}
+
+		/// <summary>
+		/// Select a Perforce server to use
+		/// </summary>
+		/// <param name="Cluster"></param>
+		/// <param name="Properties"></param>
+		/// <returns></returns>
+		public async Task<IPerforceServer?> SelectServerAsync(PerforceCluster Cluster, IReadOnlyList<string> Properties)
+		{
+			// Find all the valid servers for this agent
+			List<PerforceServer> ValidServers = new List<PerforceServer>();
+			if (Properties != null)
+			{
+				ValidServers.AddRange(Cluster.Servers.Where(x => x.Properties != null && x.Properties.Count > 0 && x.Properties.All(y => Properties.Contains(y))));
+			}
+			if (ValidServers.Count == 0)
+			{
+				ValidServers.AddRange(Cluster.Servers.Where(x => x.Properties == null || x.Properties.Count == 0));
+			}
+
+			HashSet<string> ValidServerNames = new HashSet<string>(ValidServers.Select(x => x.ServerAndPort), StringComparer.OrdinalIgnoreCase);
+
+			// Find all the matching servers.
+			PerforceServerList ServerList = await GetServerListAsync();
+
+			List<PerforceServerEntry> Candidates = ServerList.Servers.Where(x => x.Cluster == Cluster.Name && ValidServerNames.Contains(x.BaseServerAndPort)).ToList();
+			if (Candidates.Count == 0)
+			{
+				foreach (PerforceServer ValidServer in ValidServers)
+				{
+					Logger.LogDebug("Fetching server info for {ServerAndPort}", ValidServer.ServerAndPort);
+					await UpdateServerAsync(Cluster, ValidServer, Candidates);
+				}
+				if (Candidates.Count == 0)
+				{
+					Logger.LogWarning("Unable to resolve any Perforce servers from valid list");
+					return null; 
+				}
+			}
+
+			// Remove any servers that are unhealthy
+			if (Candidates.Any(x => x.Status == PerforceServerStatus.Healthy))
+			{
+				Candidates.RemoveAll(x => x.Status != PerforceServerStatus.Healthy);
+			}
+			else
+			{
+				Candidates.RemoveAll(x => x.Status != PerforceServerStatus.Unknown);
+			}
+
+			if (Candidates.Count == 0)
+			{
+				Logger.LogWarning("Unable to find any healthy Perforce server in cluster {ClusterName}", Cluster.Name);
+				return null;
+			}
+
+			// Select which server to use with a weighted average of the number of active leases
+			int Index = 0;
+			if (Candidates.Count > 1)
+			{
+				const int BaseWeight = 20;
+
+				int TotalLeases = Candidates.Sum(x => x.NumLeases);
+				int TotalWeight = Candidates.Sum(x => (TotalLeases + BaseWeight) - x.NumLeases);
+
+				int Weight = Random.Next(TotalWeight);
+				for (; Index + 1 < Candidates.Count; Index++)
+				{
+					Weight -= (TotalLeases + BaseWeight) - Candidates[Index].NumLeases;
+					if(Weight < 0)
+					{
+						break;
+					}
+				}
+			}
+
+			return Candidates[Index];
+		}
+
+		/// <inheritdoc/>
+		Task TickAsync(CancellationToken StoppingToken)
+		{
+			if (InternalTickTask.IsCompleted)
+			{
+				InternalTickTask = Task.Run(() => TickInternalGuardedAsync());
+			}
+			return Task.CompletedTask;
+		}
+
+		/// <inheritdoc/>
+		async Task TickInternalGuardedAsync()
+		{
+			try
+			{
+				await TickInternalAsync();
+			}
+			catch (Exception Ex)
+			{
+				Logger.LogError(Ex, "Error updating server state");
+			}
+		}
+
+		/// <inheritdoc/>
+		async Task TickInternalAsync()
+		{
+			Globals Globals = await DatabaseService.GetGlobalsAsync();
+
+			// Set of new server entries
+			List<PerforceServerEntry> NewServers = new List<PerforceServerEntry>();
+
+			// Update the state of all the valid servers
+			foreach (PerforceCluster Cluster in Globals.PerforceClusters)
+			{
+				foreach (PerforceServer Server in Cluster.Servers)
+				{
+					try
+					{
+						await UpdateServerAsync(Cluster, Server, NewServers);
+					}
+					catch (Exception Ex)
+					{
+						Logger.LogError(Ex, "Exception while updating Perforce server status");
+					}
+				}
+			}
+
+			// Update the number of leases for each entry
+			List<PerforceServerEntry> NewEntries = NewServers.OrderBy(x => x.Cluster).ThenBy(x => x.BaseServerAndPort).ThenBy(x => x.ServerAndPort).ToList();
+			await UpdateLeaseCounts(NewEntries);
+			PerforceServerList List = await ServerListSingleton.UpdateAsync(List => MergeServerList(List, NewEntries));
+
+			// Now update the health of each entry
+			List<Task> Tasks = new List<Task>();
+			foreach (PerforceServerEntry Entry in List.Servers)
+			{
+				Tasks.Add(Task.Run(() => UpdateHealthAsync(Entry)));
+			}
+			await Task.WhenAll(Tasks);
+		}
+
+		static void MergeServerList(PerforceServerList ServerList, List<PerforceServerEntry> NewEntries)
+		{
+			Dictionary<string, PerforceServerEntry> ExistingEntries = ServerList.Servers.ToDictionary(x => x.ServerAndPort, x => x, StringComparer.OrdinalIgnoreCase);
+			foreach (PerforceServerEntry NewEntry in NewEntries)
+			{
+				PerforceServerEntry? ExistingEntry;
+				if (ExistingEntries.TryGetValue(NewEntry.ServerAndPort, out ExistingEntry))
+				{
+					NewEntry.Status = ExistingEntry.Status;
+					NewEntry.Detail = ExistingEntry.Detail;
+					NewEntry.LastUpdateTime = ExistingEntry.LastUpdateTime;
+				}
+			}
+			ServerList.Servers = NewEntries;
+		}
+
+		async Task UpdateLeaseCounts(IEnumerable<PerforceServerEntry> NewServerEntries)
+		{
+			Dictionary<string, Dictionary<string, PerforceServerEntry>> NewServerLookup = new Dictionary<string, Dictionary<string, PerforceServerEntry>>();
+			foreach (PerforceServerEntry NewServerEntry in NewServerEntries)
+			{
+				if (NewServerEntry.Cluster != null)
+				{
+					Dictionary<string, PerforceServerEntry>? ClusterServers;
+					if (!NewServerLookup.TryGetValue(NewServerEntry.Cluster, out ClusterServers))
+					{
+						ClusterServers = new Dictionary<string, PerforceServerEntry>();
+						NewServerLookup.Add(NewServerEntry.Cluster, ClusterServers);
+					}
+					ClusterServers[NewServerEntry.ServerAndPort] = NewServerEntry;
+				}
+			}
+
+			List<ILease> Leases = await LeaseCollection.FindActiveLeasesAsync();
+			foreach (ILease Lease in Leases)
+			{
+				Any Any = Any.Parser.ParseFrom(Lease.Payload.ToArray());
+
+				HashSet<(string, string)> Servers = new HashSet<(string, string)>();
+				if (Any.TryUnpack(out ConformTask ConformTask))
+				{
+					foreach (HordeCommon.Rpc.Messages.AgentWorkspace Workspace in ConformTask.Workspaces)
+					{
+						Servers.Add((Workspace.Cluster, Workspace.ServerAndPort));
+					}
+				}
+				if (Any.TryUnpack(out ExecuteJobTask ExecuteJobTask))
+				{
+					if (ExecuteJobTask.AutoSdkWorkspace != null)
+					{
+						Servers.Add((ExecuteJobTask.AutoSdkWorkspace.Cluster, ExecuteJobTask.AutoSdkWorkspace.ServerAndPort));
+					}
+					if (ExecuteJobTask.Workspace != null)
+					{
+						Servers.Add((ExecuteJobTask.Workspace.Cluster, ExecuteJobTask.Workspace.ServerAndPort));
+					}
+				}
+
+				foreach ((string Cluster, string ServerAndPort) in Servers)
+				{
+					IncrementLeaseCount(Cluster, ServerAndPort, NewServerLookup);
+				}
+			}
+		}
+
+		static void IncrementLeaseCount(string Cluster, string ServerAndPort, Dictionary<string, Dictionary<string, PerforceServerEntry>> NewServers)
+		{
+			Dictionary<string, PerforceServerEntry>? ClusterServers;
+			if (Cluster != null && NewServers.TryGetValue(Cluster, out ClusterServers))
+			{
+				PerforceServerEntry? Entry;
+				if (ClusterServers.TryGetValue(ServerAndPort, out Entry))
+				{
+					Entry.NumLeases++;
+				}
+			}
+		}
+
+		async Task UpdateServerAsync(PerforceCluster Cluster, PerforceServer Server, List<PerforceServerEntry> NewServers)
+		{
+			string InitialHostName = Server.ServerAndPort;
+			int Port = 1666;
+
+			int PortIdx = InitialHostName.LastIndexOf(':');
+			if (PortIdx != -1)
+			{
+				Port = int.Parse(InitialHostName.Substring(PortIdx + 1), NumberStyles.Integer, CultureInfo.InvariantCulture);
+				InitialHostName = InitialHostName.Substring(0, PortIdx);
+			}
+
+			List<string> HostNames = new List<string>();
+			if (!Server.ResolveDns || IPAddress.TryParse(InitialHostName, out IPAddress? HostAddress))
+			{
+				HostNames.Add(InitialHostName);
+			}
+			else
+			{
+				await ResolveServersAsync(InitialHostName, HostNames);
+			}
+
+			foreach (string HostName in HostNames)
+			{
+				string? HealthCheckUrl = null;
+				if (Server.HealthCheck)
+				{
+					HealthCheckUrl = $"http://{HostName}:5000/healthcheck";
+				}
+				NewServers.Add(new PerforceServerEntry($"{HostName}:{Port}", Server.ServerAndPort, HealthCheckUrl, Cluster.Name, PerforceServerStatus.Unknown, "", DateTime.UtcNow));
+			}
+		}
+
+		async Task ResolveServersAsync(string HostName, List<string> HostNames)
+		{
+			// Find all the addresses of the hosts
+			IPHostEntry Entry = await Dns.GetHostEntryAsync(HostName);
+			foreach (IPAddress Address in Entry.AddressList)
+			{
+				try
+				{
+					HostNames.Add(Address.ToString());
+				}
+				catch (Exception Ex)
+				{
+					Logger.LogDebug(Ex, "Unable to resolve host name for Perforce server {Address}", Address);
+				}
+			}
+		}
+
+		/// <summary>
+		/// Updates 
+		/// </summary>
+		/// <param name="Entry"></param>
+		/// <returns></returns>
+		async Task UpdateHealthAsync(PerforceServerEntry Entry)
+		{
+			DateTime? UpdateTime = null;
+			string Detail = "Health check disabled";
+
+			// Get the health of the server
+			PerforceServerStatus Health = PerforceServerStatus.Healthy;
+			if (Entry.HealthCheckUrl != null)
+			{
+				UpdateTime = DateTime.UtcNow;
+				Uri HealthCheckUrl = new Uri(Entry.HealthCheckUrl);
+				try
+				{
+					(Health, Detail) = await GetServerHealthAsync(HealthCheckUrl);
+				}
+				catch
+				{
+					(Health, Detail) = (PerforceServerStatus.Unhealthy, $"Failed to query status at {HealthCheckUrl}");
+				}
+			}
+
+			// Update the server record
+			if (Health != Entry.Status || Detail != Entry.Detail || UpdateTime != Entry.LastUpdateTime)
+			{
+				await ServerListSingleton.UpdateAsync(x => UpdateHealth(x, Entry.ServerAndPort, Health, Detail, UpdateTime));
+			}
+		}
+
+		static void UpdateHealth(PerforceServerList ServerList, string ServerAndPort, PerforceServerStatus Status, string Detail, DateTime? UpdateTime)
+		{
+			PerforceServerEntry? Entry = ServerList.Servers.FirstOrDefault(x => x.ServerAndPort.Equals(ServerAndPort, StringComparison.Ordinal));
+			if (Entry != null)
+			{
+				if (Entry.LastUpdateTime == null || UpdateTime == null || Entry.LastUpdateTime.Value < UpdateTime.Value)
+				{
+					Entry.Status = Status;
+					Entry.Detail = Detail;
+					Entry.LastUpdateTime = UpdateTime;
+				}
+			}
+		}
+
+		static async Task<(PerforceServerStatus, string)> GetServerHealthAsync(Uri HealthCheckUrl)
+		{
+			using HttpClient Client = new HttpClient();
+			HttpResponseMessage Response = await Client.GetAsync(HealthCheckUrl);
+
+			byte[] Data = await Response.Content.ReadAsByteArrayAsync();
+			JsonDocument Document = JsonDocument.Parse(Data);
+
+			foreach (JsonElement Element in Document.RootElement.GetProperty("results").EnumerateArray())
+			{
+				if (Element.TryGetProperty("checker", out JsonElement Checker) && Checker.ValueEquals("edge_traffic_lights"))
+				{
+					if (Element.TryGetProperty("output", out JsonElement Output))
+					{
+						string Status = Output.GetString() ?? String.Empty;
+						switch(Status)
+						{
+							case "green":
+								return (PerforceServerStatus.Healthy, "Server is healthy");
+							case "yellow":
+								return (PerforceServerStatus.Degraded, "Degraded service");
+							case "red":
+								return (PerforceServerStatus.Unhealthy, "Server is being drained");
+							default:
+								return (PerforceServerStatus.Unknown, $"Expected state for health check ({Status})");
+						}
+					}
+				}
+			}
+
+			return (PerforceServerStatus.Unknown, "Unable to parse health check output");
+		}
+	}
+}
