Horde: Rename the HordeServer project to Horde.Build.

diff --git a/Engine/Source/Programs/Horde/Horde.Build/Services/RpcService.cs b/Engine/Source/Programs/Horde/Horde.Build/Services/RpcService.cs
new file mode 100644
--- /dev/null
+++ b/Engine/Source/Programs/Horde/Horde.Build/Services/RpcService.cs
@@ -0,0 +1,1081 @@
+// Copyright Epic Games, Inc. All Rights Reserved.
+
+using EpicGames.Core;
+using Google.Protobuf.WellKnownTypes;
+using Grpc.Core;
+using HordeCommon;
+using HordeCommon.Rpc;
+using HordeServer.Collections;
+using HordeServer.Models;
+using HordeServer.Utilities;
+using Microsoft.AspNetCore.Authorization;
+using Microsoft.Extensions.Logging;
+using MongoDB.Bson;
+using MongoDB.Bson.Serialization;
+using System;
+using System.Collections.Generic;
+using System.Globalization;
+using System.Linq;
+using System.Net;
+using System.Runtime.CompilerServices;
+using System.Security.Claims;
+using System.Text;
+using System.Threading.Tasks;
+using HordeServer.Tasks.Impl;
+using System.Threading;
+using HordeServer.Jobs;
+
+[assembly: InternalsVisibleTo("HordeAgentTests")]
+
+namespace HordeServer.Services
+{
+	using AgentSoftwareChannelName = StringId<AgentSoftwareChannels>;
+	using AgentSoftwareVersion = StringId<IAgentSoftwareCollection>;
+	using IStream = HordeServer.Models.IStream;
+	using JobId = ObjectId<IJob>;
+	using LogId = ObjectId<ILogFile>;
+	using StreamId = StringId<IStream>;
+	using RpcAgentCapabilities = HordeCommon.Rpc.Messages.AgentCapabilities;
+	using RpcDeviceCapabilities = HordeCommon.Rpc.Messages.DeviceCapabilities;
+
+	/// <summary>
+	/// Implements the Horde gRPC service for bots updating their status and dequeing work
+	/// </summary>
+	[Authorize]
+	public class RpcService : HordeRpc.HordeRpcBase
+	{
+		/// <summary>
+		/// Timeout before closing a long-polling request (client will retry again) 
+		/// </summary>
+		internal TimeSpan LongPollTimeout = TimeSpan.FromMinutes(9);
+
+		/// <summary>
+		/// Instance of the DatabaseService singleton
+		/// </summary>
+		DatabaseService DatabaseService;
+
+		/// <summary>
+		/// Instance of the AclService singleton
+		/// </summary>
+		AclService AclService;
+
+		/// <summary>
+		/// Instance of the AgentService singleton
+		/// </summary>
+		AgentService AgentService;
+
+		/// <summary>
+		/// The stream service instance
+		/// </summary>
+		StreamService StreamService;
+
+		/// <summary>
+		/// The job service instance
+		/// </summary>
+		JobService JobService;
+
+		/// <summary>
+		/// The software service instance
+		/// </summary>
+		AgentSoftwareService AgentSoftwareService;
+
+		/// <summary>
+		/// The artifact service instance
+		/// </summary>
+		IArtifactCollection ArtifactCollection;
+
+		/// <summary>
+		/// The log file service instance
+		/// </summary>
+		ILogFileService LogFileService;
+
+		/// <summary>
+		/// The credential service instance
+		/// </summary>
+		CredentialService CredentialService;
+
+		/// <summary>
+		/// The pool service instance
+		/// </summary>
+		PoolService PoolService;
+
+		/// <summary>
+		/// The application lifetime interface
+		/// </summary>
+		LifetimeService LifetimeService;
+
+		/// <summary>
+		/// Collection of graph documents
+		/// </summary>
+		IGraphCollection Graphs;
+
+		/// <summary>
+		/// Collection of testdata documents
+		/// </summary>
+		ITestDataCollection TestData;
+
+		/// <summary>
+		/// The conform task source
+		/// </summary>
+		ConformTaskSource ConformTaskSource;
+
+		/// <summary>
+		/// Writer for log output
+		/// </summary>
+		ILogger<RpcService> Logger;
+
+		/// <summary>
+		/// Constructor
+		/// </summary>
+		/// <param name="DatabaseService">Instance of the DatabaseService singleton</param>
+		/// <param name="AclService">Instance of the AclService singleton</param>
+		/// <param name="AgentService">Instance of the AgentService singleton</param>
+		/// <param name="JobService">Instance of the JobService singleton</param>
+		/// <param name="StreamService">Instance of the StreamService singleton</param>
+		/// <param name="LogFileService">Instance of the LogFileService singleton</param>
+		/// <param name="AgentSoftwareService">Instance of the JobService singleton</param>
+		/// <param name="ArtifactCollection">Instance of the ArtifactService singleton</param>
+		/// <param name="CredentialService">Instance of the CredentialsService singleton</param>
+		/// <param name="PoolService">Instance of the PoolService singleton</param>
+		/// <param name="LifetimeService">The application lifetime</param>
+		/// <param name="Graphs">Collection of graph documents</param>
+		/// <param name="TestData">Collection of testdata</param>
+		/// <param name="ConformTaskSource"></param>
+		/// <param name="Logger">Log writer</param>
+		public RpcService(DatabaseService DatabaseService, AclService AclService, AgentService AgentService, StreamService StreamService, JobService JobService, AgentSoftwareService AgentSoftwareService, IArtifactCollection ArtifactCollection, ILogFileService LogFileService, CredentialService CredentialService, PoolService PoolService, LifetimeService LifetimeService, IGraphCollection Graphs, ITestDataCollection TestData, ConformTaskSource ConformTaskSource, ILogger<RpcService> Logger)
+		{
+			this.DatabaseService = DatabaseService;
+			this.AclService = AclService;
+			this.AgentService = AgentService;
+			this.StreamService = StreamService;
+			this.JobService = JobService;
+			this.AgentSoftwareService = AgentSoftwareService;
+			this.ArtifactCollection = ArtifactCollection;
+			this.LogFileService = LogFileService;
+			this.CredentialService = CredentialService;
+			this.PoolService = PoolService;
+			this.LifetimeService = LifetimeService;
+			this.Graphs = Graphs;
+			this.TestData = TestData;
+			this.ConformTaskSource = ConformTaskSource;
+			this.Logger = Logger;
+		}
+
+		/// <summary>
+		/// Waits until the server is terminating
+		/// </summary>
+		/// <param name="Reader">Request reader</param>
+		/// <param name="Writer">Response writer</param>
+		/// <param name="Context">Context for the call</param>
+		/// <returns>Response object</returns>
+		public override async Task QueryServerState(IAsyncStreamReader<QueryServerStateRequest> Reader, IServerStreamWriter<QueryServerStateResponse> Writer, ServerCallContext Context)
+		{
+			if (await Reader.MoveNext())
+			{
+				QueryServerStateRequest Request = Reader.Current;
+				Logger.LogInformation("Start server query for client {Name}", Request.Name);
+
+				// Return the current response
+				QueryServerStateResponse Response = new QueryServerStateResponse();
+				Response.Name = Dns.GetHostName();
+				await Writer.WriteAsync(Response);
+
+				// Move to the next request from the client. This should always be the end of the stream, but will not occur until the client stops requesting responses.
+				Task<bool> MoveNextTask = Reader.MoveNext();
+
+				// Wait for the client to close the stream or a shutdown to start
+				Task LongPollDelay = Task.Delay(LongPollTimeout);
+				Task WaitTask = await Task.WhenAny(MoveNextTask, LifetimeService.StoppingTask, LongPollDelay);
+
+				if (WaitTask == MoveNextTask)
+				{
+					throw new Exception("Unexpected request to QueryServerState posted from client.");
+				}
+				else if (WaitTask == LifetimeService.StoppingTask)
+				{
+					Logger.LogInformation("Notifying client {Name} of server shutdown", Request.Name);
+					await Writer.WriteAsync(Response);
+				}
+				else if (WaitTask == LongPollDelay)
+				{
+					// Send same response as server shutdown. In the agent perspective, they will be identical.
+					await Writer.WriteAsync(Response);
+				}
+			}
+		}
+
+		/// <summary>
+		/// Waits until the server is terminating
+		/// </summary>
+		/// <param name="Reader">Request reader</param>
+		/// <param name="Writer">Response writer</param>
+		/// <param name="Context">Context for the call</param>
+		/// <returns>Response object</returns>
+		public override async Task QueryServerStateV2(IAsyncStreamReader<QueryServerStateRequest> Reader, IServerStreamWriter<QueryServerStateResponse> Writer, ServerCallContext Context)
+		{
+			if (await Reader.MoveNext())
+			{
+				QueryServerStateRequest Request = Reader.Current;
+				Logger.LogDebug("Start server query for client {Name}", Request.Name);
+
+				try
+				{
+					// Return the current response
+					QueryServerStateResponse Response = new QueryServerStateResponse();
+					Response.Name = Dns.GetHostName();
+					Response.Stopping = LifetimeService.IsStopping;
+					await Writer.WriteAsync(Response);
+
+					// Move to the next request from the client. This should always be the end of the stream, but will not occur until the client stops requesting responses.
+					Task<bool> MoveNextTask = Reader.MoveNext();
+
+					// Wait for the client to close the stream or a shutdown to start
+					if (await Task.WhenAny(MoveNextTask, LifetimeService.StoppingTask) == LifetimeService.StoppingTask)
+					{
+						Response.Stopping = true;
+						await Writer.WriteAsync(Response);
+					}
+
+					// Wait until the client has finished sending
+					while (await MoveNextTask)
+					{
+						MoveNextTask = Reader.MoveNext();
+					}
+				}
+				catch (Exception Ex)
+				{
+					Logger.LogError(Ex, "Exception in QueryServerState for {Name}", Request.Name);
+					throw;
+				}
+			}
+		}
+
+		/// <summary>
+		/// Updates the workspaces synced for an agent
+		/// </summary>
+		/// <param name="Request">The request parameters</param>
+		/// <param name="Context">Context for the call</param>
+		/// <returns>Response object</returns>
+		public override async Task<UpdateAgentWorkspacesResponse> UpdateAgentWorkspaces(UpdateAgentWorkspacesRequest Request, ServerCallContext Context)
+		{
+			for (; ; )
+			{
+				// Get the current agent state
+				IAgent? Agent = await AgentService.GetAgentAsync(new AgentId(Request.AgentId));
+				if (Agent == null)
+				{
+					throw new StructuredRpcException(StatusCode.OutOfRange, "Agent {AgentId} does not exist", Request.AgentId);
+				}
+
+				// Get the new workspaces
+				List<AgentWorkspace> NewWorkspaces = Request.Workspaces.Select(x => new AgentWorkspace(x)).ToList();
+
+				// Get the set of workspaces that are currently required
+				HashSet<AgentWorkspace> ConformWorkspaces = await PoolService.GetWorkspacesAsync(Agent, DateTime.UtcNow);
+				bool bPendingConform = !ConformWorkspaces.SetEquals(NewWorkspaces);
+
+				// Update the workspaces
+				if (await AgentService.TryUpdateWorkspacesAsync(Agent, NewWorkspaces, bPendingConform))
+				{
+					UpdateAgentWorkspacesResponse Response = new UpdateAgentWorkspacesResponse();
+					if (bPendingConform)
+					{
+						Response.Retry = await ConformTaskSource.GetWorkspacesAsync(Agent, Response.PendingWorkspaces);
+					}
+					return Response;
+				}
+			}
+		}
+
+		static void CopyPropertyToResource(string Name, List<string> Properties, Dictionary<string, int> Resources)
+		{
+			foreach (string Property in Properties)
+			{
+				if (Property.Length > Name.Length && Property.StartsWith(Name, StringComparison.OrdinalIgnoreCase) && Property[Name.Length] == '=')
+				{
+					int Value;
+					if (int.TryParse(Property.AsSpan(Name.Length + 1), out Value))
+					{
+						Resources[Name] = Value;
+					}
+				}
+			}
+		}
+
+		static void GetCapabilities(RpcAgentCapabilities? Capabilities, out List<string> Properties, out Dictionary<string, int> Resources)
+		{
+			Properties = new List<string>();
+			Resources = new Dictionary<string, int>();
+
+			if (Capabilities != null && Capabilities.Devices.Count > 0)
+			{
+				RpcDeviceCapabilities Device = Capabilities.Devices[0];
+				if (Device.Properties != null)
+				{
+					Properties = new List<string>(Device.Properties);
+					CopyPropertyToResource(KnownPropertyNames.LogicalCores, Properties, Resources);
+					CopyPropertyToResource(KnownPropertyNames.RAM, Properties, Resources);
+				}
+			}
+		}
+
+		/// <summary>
+		/// Creates a new session
+		/// </summary>
+		/// <param name="Request">Request to create a new agent</param>
+		/// <param name="Context">Context for the RPC call</param>
+		/// <returns>Information about the new agent</returns>
+		public async override Task<CreateSessionResponse> CreateSession(CreateSessionRequest Request, ServerCallContext Context)
+		{
+			if (Request.Capabilities == null)
+			{
+				throw new StructuredRpcException(StatusCode.InvalidArgument, "Capabilities may not be null");
+			}
+
+			AgentId AgentId = new AgentId(Request.Name);
+			using IDisposable Scope = Logger.BeginScope("CreateSession({AgentId})", AgentId.ToString());
+
+			// Find the agent
+			IAgent? Agent = await AgentService.GetAgentAsync(AgentId);
+			if (Agent == null)
+			{
+				if (!await AclService.AuthorizeAsync(AclAction.CreateAgent, Context.GetHttpContext().User))
+				{
+					throw new StructuredRpcException(StatusCode.PermissionDenied, "User is not authenticated to create new agents");
+				}
+
+				const bool bEnabled = true;
+				Agent = await AgentService.CreateAgentAsync(Request.Name, bEnabled, null, null);
+			}
+
+			// Make sure we're allowed to create sessions on this agent
+			if (!await AgentService.AuthorizeAsync(Agent, AclAction.CreateSession, Context.GetHttpContext().User, null))
+			{
+				throw new StructuredRpcException(StatusCode.PermissionDenied, "User is not authenticated to create session for {AgentId}", Request.Name);
+			}
+
+			// Get the known properties for this agent
+			GetCapabilities(Request.Capabilities, out List<string> Properties, out Dictionary<string, int> Resources);
+
+			// Create a new session
+			Agent = await AgentService.CreateSessionAsync(Agent, Request.Status, Properties, Resources, Request.Version);
+			if (Agent == null)
+			{
+				throw new StructuredRpcException(StatusCode.NotFound, "Agent {AgentId} not found", Request.Name);
+			}
+
+			// Create the response
+			CreateSessionResponse Response = new CreateSessionResponse();
+			Response.AgentId = Agent.Id.ToString();
+			Response.SessionId = Agent.SessionId.ToString();
+			Response.ExpiryTime = Timestamp.FromDateTime(Agent.SessionExpiresAt!.Value);
+			Response.Token = AgentService.IssueSessionToken(Agent.SessionId!.Value);
+			return Response;
+		}
+
+		/// <summary>
+		/// Updates an agent session
+		/// </summary>
+		/// <param name="Reader">Request to create a new agent</param>
+		/// <param name="Writer">Writer for response objects</param>
+		/// <param name="Context">Context for the RPC call</param>
+		/// <returns>Information about the new agent</returns>
+		public async override Task UpdateSession(IAsyncStreamReader<UpdateSessionRequest> Reader, IServerStreamWriter<UpdateSessionResponse> Writer, ServerCallContext Context)
+		{
+			// Read the request object
+			Task<bool> NextRequestTask = Reader.MoveNext();
+			if (await NextRequestTask)
+			{
+				UpdateSessionRequest Request = Reader.Current;
+				using IDisposable Scope = Logger.BeginScope("UpdateSession for agent {AgentId}, session {SessionId}", Request.AgentId, Request.SessionId);
+
+				Logger.LogDebug("Updating session for {AgentId}", Request.AgentId);
+				foreach (HordeCommon.Rpc.Messages.Lease Lease in Request.Leases)
+				{
+					Logger.LogDebug("Session {SessionId}, Lease {LeaseId} - State: {LeaseState}, Outcome: {LeaseOutcome}", Request.SessionId, Lease.Id, Lease.State, Lease.Outcome);
+				}
+
+				// Get a task for moving to the next item. This will only complete once the call has closed.
+				using CancellationTokenSource CancellationSource = new CancellationTokenSource();
+				NextRequestTask = Reader.MoveNext();
+				NextRequestTask = NextRequestTask.ContinueWith(Task => { CancellationSource.Cancel(); return Task.Result; }, TaskScheduler.Current);
+
+				// Get the current agent state
+				IAgent? Agent = await AgentService.GetAgentAsync(new AgentId(Request.AgentId));
+				if(Agent != null)
+				{
+					// Check we're authorized to update it
+					if (!AgentService.AuthorizeSession(Agent, Context.GetHttpContext().User))
+					{
+						throw new StructuredRpcException(StatusCode.PermissionDenied, "Not authenticated for {AgentId}", Request.AgentId);
+					}
+
+					// Get the new capabilities of this agent
+					List<string>? Properties = null;
+					Dictionary<string, int>? Resources = null;
+					if (Request.Capabilities != null)
+					{
+						GetCapabilities(Request.Capabilities, out Properties, out Resources);
+					}
+
+					// Update the session
+					Agent = await AgentService.UpdateSessionWithWaitAsync(Agent, Request.SessionId.ToObjectId(), Request.Status, Properties, Resources, Request.Leases, CancellationSource.Token);
+				}
+
+				// Handle the invalid agent case
+				if (Agent == null)
+				{
+					throw new StructuredRpcException(StatusCode.NotFound, "Invalid agent name '{AgentId}'", Request.AgentId);
+				}
+
+				// Create the new session info
+				UpdateSessionResponse Response = new UpdateSessionResponse();
+				Response.Leases.Add(Agent.Leases.Select(x => x.ToRpcMessage()));
+				Response.ExpiryTime = (Agent.SessionExpiresAt == null) ? new Timestamp() : Timestamp.FromDateTime(Agent.SessionExpiresAt.Value);
+				await Writer.WriteAsync(Response);
+
+				// Wait for the client to close the stream
+				while (await NextRequestTask)
+				{
+					NextRequestTask = Reader.MoveNext();
+				}
+			}
+		}
+
+		/// <summary>
+		/// Gets information about a stream
+		/// </summary>
+		/// <param name="Request">Request arguments</param>
+		/// <param name="Context">Context for the RPC call</param>
+		/// <returns>Information about the new agent</returns>
+		public override async Task<GetStreamResponse> GetStream(GetStreamRequest Request, ServerCallContext Context)
+		{
+			StreamId StreamIdValue = new StreamId(Request.StreamId);
+
+			IStream? Stream = await StreamService.GetStreamAsync(StreamIdValue);
+			if (Stream == null)
+			{
+				throw new StructuredRpcException(StatusCode.NotFound, "Stream {StreamId} does not exist", Request.StreamId);
+			}
+			if (!await StreamService.AuthorizeAsync(Stream, AclAction.ViewStream, Context.GetHttpContext().User, null))
+			{
+				throw new StructuredRpcException(StatusCode.PermissionDenied, "Not authenticated to access stream {StreamId}", Request.StreamId);
+			}
+
+			return Stream.ToRpcResponse();
+		}
+
+		/// <summary>
+		/// Gets information about a job
+		/// </summary>
+		/// <param name="Request">Request arguments</param>
+		/// <param name="Context">Context for the RPC call</param>
+		/// <returns>Information about the new agent</returns>
+		public async override Task<GetJobResponse> GetJob(GetJobRequest Request, ServerCallContext Context)
+		{
+			JobId JobIdValue = new JobId(Request.JobId.ToObjectId());
+
+			IJob? Job = await JobService.GetJobAsync(JobIdValue);
+			if (Job == null)
+			{
+				throw new StructuredRpcException(StatusCode.NotFound, "Job {JobId} does not exist", Request.JobId);
+			}
+			if (!JobService.AuthorizeSession(Job, Context.GetHttpContext().User))
+			{
+				throw new StructuredRpcException(StatusCode.PermissionDenied, "Not authenticated to access job {JobId}", Request.JobId);
+			}
+
+			return Job.ToRpcResponse();
+		}
+
+		/// <summary>
+		/// Updates properties on a job
+		/// </summary>
+		/// <param name="Request">Request arguments</param>
+		/// <param name="Context">Context for the RPC call</param>
+		/// <returns>Information about the new agent</returns>
+		public async override Task<Empty> UpdateJob(UpdateJobRequest Request, ServerCallContext Context)
+		{
+			JobId JobIdValue = new JobId(Request.JobId);
+
+			IJob? Job = await JobService.GetJobAsync(JobIdValue);
+			if (Job == null)
+			{
+				throw new StructuredRpcException(StatusCode.NotFound, "Job {JobId} does not exist", Request.JobId);
+			}
+			if (!JobService.AuthorizeSession(Job, Context.GetHttpContext().User))
+			{
+				throw new StructuredRpcException(StatusCode.PermissionDenied, "Not authenticated to modify job {JobId}", Request.JobId);
+			}
+
+			await JobService.UpdateJobAsync(Job, Name: Request.Name);
+			return new Empty();
+		}
+
+		/// <summary>
+		/// Starts executing a batch
+		/// </summary>
+		/// <param name="Request">Request arguments</param>
+		/// <param name="Context">Context for the RPC call</param>
+		/// <returns>Information about the new agent</returns>
+		public async override Task<BeginBatchResponse> BeginBatch(BeginBatchRequest Request, ServerCallContext Context)
+		{
+			ObjectId JobId = Request.JobId.ToObjectId();
+			SubResourceId BatchId = Request.BatchId.ToSubResourceId();
+
+			IJob? Job = await JobService.GetJobAsync(new JobId(Request.JobId));
+			if (Job == null)
+			{
+				throw new StructuredRpcException(StatusCode.NotFound, "Job {JobId} not found", Request.JobId);
+			}
+
+			IJobStepBatch Batch = AuthorizeBatch(Job, Request.BatchId.ToSubResourceId(), Context);
+			if (!await JobService.UpdateBatchAsync(Job, BatchId, null, Api.JobStepBatchState.Starting))
+			{
+				throw new StructuredRpcException(StatusCode.NotFound, "Batch {JobId}:{BatchId} not found for updating", Request.JobId, Request.BatchId);
+			}
+
+			IGraph Graph = await JobService.GetGraphAsync(Job);
+
+			BeginBatchResponse Response = new BeginBatchResponse();
+			Response.LogId = Batch.LogId.ToString();
+			Response.AgentType = Graph.Groups[Batch.GroupIdx].AgentType;
+			return Response;
+		}
+
+		/// <summary>
+		/// Finishes executing a batch
+		/// </summary>
+		/// <param name="Request">Request arguments</param>
+		/// <param name="Context">Context for the RPC call</param>
+		/// <returns>Information about the new agent</returns>
+		public async override Task<Empty> FinishBatch(FinishBatchRequest Request, ServerCallContext Context)
+		{
+			IJob? Job = await JobService.GetJobAsync(new JobId(Request.JobId));
+			if (Job == null)
+			{
+				throw new StructuredRpcException(StatusCode.NotFound, "Job {JobId} not found", Request.JobId);
+			}
+
+			IJobStepBatch Batch = AuthorizeBatch(Job, Request.BatchId.ToSubResourceId(), Context);
+			await JobService.UpdateBatchAsync(Job, Batch.Id, null, Api.JobStepBatchState.Complete);
+			return new Empty();
+		}
+
+		/// <summary>
+		/// Starts executing a step
+		/// </summary>
+		/// <param name="Request">Request arguments</param>
+		/// <param name="Context">Context for the RPC call</param>
+		/// <returns>Information about the new agent</returns>
+		public async override Task<BeginStepResponse> BeginStep(BeginStepRequest Request, ServerCallContext Context)
+		{
+			Boxed<ILogFile?> Log = new Boxed<ILogFile?>(null);
+			for (; ; )
+			{
+				BeginStepResponse? Response = await TryBeginStep(Request, Log, Context);
+				if (Response != null)
+				{
+					return Response;
+				}
+			}
+		}
+
+		async Task<BeginStepResponse?> TryBeginStep(BeginStepRequest Request, Boxed<ILogFile?> Log, ServerCallContext Context)
+		{
+			// Check the job exists and we can access it
+			IJob? Job = await JobService.GetJobAsync(new JobId(Request.JobId));
+			if (Job == null)
+			{
+				throw new StructuredRpcException(StatusCode.NotFound, "Job {JobId} not found", Request.JobId);
+			}
+
+			// Find the batch being executed
+			IJobStepBatch Batch = AuthorizeBatch(Job, Request.BatchId.ToSubResourceId(), Context);
+			if (Batch.State != Api.JobStepBatchState.Starting && Batch.State != Api.JobStepBatchState.Running)
+			{
+				return new BeginStepResponse { State = BeginStepResponse.Types.Result.Complete };
+			}
+
+			// Figure out which step to execute next
+			IJobStep? Step;
+			for (int StepIdx = 0; ; StepIdx++)
+			{
+				// If there aren't any more steps, send a complete message
+				if (StepIdx == Batch.Steps.Count)
+				{
+					Logger.LogDebug("Job {JobId} batch {BatchId} is complete", Job.Id, Batch.Id);
+					if (!await JobService.TryUpdateBatchAsync(Job, Batch.Id, NewState: Api.JobStepBatchState.Stopping))
+					{
+						return null;
+					}
+					return new BeginStepResponse { State = BeginStepResponse.Types.Result.Complete };
+				}
+
+				// Check if this step is ready to be executed
+				Step = Batch.Steps[StepIdx];
+				if (Step.State == JobStepState.Ready)
+				{
+					break;
+				}
+				if (Step.State == JobStepState.Waiting)
+				{
+					Logger.LogDebug("Waiting for job {JobId}, batch {BatchId}, step {StepId}", Job.Id, Batch.Id, Step.Id);
+					return new BeginStepResponse { State = BeginStepResponse.Types.Result.Waiting };
+				}
+			}
+
+			// Create a log file if necessary
+			if (Log.Value == null)
+			{
+				Log.Value = await LogFileService.CreateLogFileAsync(Job.Id, Batch.SessionId, Api.LogType.Json);
+			}
+
+			// Get the node for this step
+			IGraph Graph = await JobService.GetGraphAsync(Job);
+			INode Node = Graph.Groups[Batch.GroupIdx].Nodes[Step.NodeIdx];
+
+			// Figure out all the credentials for it (and check we can access them)
+			Dictionary<string, string> Credentials = new Dictionary<string, string>();
+			//				if (Node.Credentials != null)
+			//				{
+			//					ClaimsPrincipal Principal = new ClaimsPrincipal(new ClaimsIdentity(Job.Claims.Select(x => new Claim(x.Type, x.Value))));
+			//					if (!await GetCredentialsForStep(Principal, Node, Credentials, Message => FailStep(Job, Batch.Id, Step.Id, Log, Message)))
+			//					{
+			//						Log = null;
+			//						continue;
+			//					}
+			//				}
+
+			// Update the step state
+			if (await JobService.TryUpdateStepAsync(Job, Batch.Id, Step.Id, JobStepState.Running, JobStepOutcome.Unspecified, null, null, Log.Value.Id, null, null, null, null))
+			{
+				BeginStepResponse Response = new BeginStepResponse();
+				Response.State = BeginStepResponse.Types.Result.Ready;
+				Response.LogId = Log.Value.Id.ToString();
+				Response.StepId = Step.Id.ToString();
+				Response.Name = Node.Name;
+				Response.Credentials.Add(Credentials);
+				if (Node.Properties != null)
+				{
+					Response.Properties.Add(Node.Properties);
+				}
+				Response.Warnings = Node.Warnings;
+				return Response;
+			}
+
+			return null;
+		}
+
+		/// <summary>
+		/// Gets all the required credentials for the given step
+		/// </summary>
+		/// <param name="User">The user to validate</param>
+		/// <param name="Node">The node being executed</param>
+		/// <param name="Credentials">Receives a list of credentials for the step</param>
+		/// <param name="WriteError">Delegate used to write error messages</param>
+		/// <returns>Async task</returns>
+		private async Task<bool> GetCredentialsForStep(ClaimsPrincipal User, INode Node, Dictionary<string, string> Credentials, Func<string, Task> WriteError)
+		{
+			if (Node.Credentials != null)
+			{
+				Dictionary<string, Credential> Cache = new Dictionary<string, Credential>(StringComparer.OrdinalIgnoreCase);
+				foreach (KeyValuePair<string, string> Pair in Node.Credentials)
+				{
+					// Get the credential path. We expect this in the format <CredentialName>.<PropertyName>
+					string Path = Pair.Value;
+
+					// Find the separator
+					int Idx = Path.IndexOf('.', StringComparison.Ordinal);
+					if (Idx == -1)
+					{
+						await WriteError($"Invalid credential path '{Path}'. Requested credentials should be in the form '<CredentialName>.<PropertyName>'.\n");
+						return false;
+					}
+
+					// Split the path into credential and property names
+					string CredentialName = Path.Substring(0, Idx);
+					string PropertyName = Path.Substring(Idx + 1);
+
+					// Try to get the credential with this name
+					Credential? Credential;
+					if (!Cache.TryGetValue(CredentialName, out Credential))
+					{
+						Credential = await CredentialService.GetCredentialAsync(CredentialName);
+						if (Credential == null)
+						{
+							await WriteError($"No credential called '{CredentialName}' could be found");
+							return false;
+						}
+						if (!await CredentialService.AuthorizeAsync(Credential, AclAction.ViewCredential, User, null))
+						{
+							await WriteError($"User is not allowed to view credential '{CredentialName}'");
+							return false;
+						}
+						Cache[CredentialName] = Credential;
+					}
+
+					// Get the property
+					string? PropertyValue;
+					if (!Credential!.Properties.TryGetValue(PropertyName, out PropertyValue))
+					{
+						await WriteError($"No property called '{PropertyName}' found in credential '{CredentialName}'");
+						return false;
+					}
+
+					// Add it to the output list
+					Credentials.Add(Pair.Key, PropertyValue);
+				}
+			}
+			return true;
+		}
+
+		async Task<IJob> GetJobAsync(JobId JobId)
+		{
+			IJob? Job = await JobService.GetJobAsync(JobId);
+			if (Job == null)
+			{
+				throw new StructuredRpcException(StatusCode.NotFound, "Job {JobId} not found", JobId);
+			}
+			return Job;
+		}
+
+		static IJobStepBatch AuthorizeBatch(IJob Job, SubResourceId BatchId, ServerCallContext Context)
+		{
+			IJobStepBatch? Batch;
+			if (!Job.TryGetBatch(BatchId, out Batch))
+			{
+				throw new StructuredRpcException(StatusCode.NotFound, "Unable to find batch {JobId}:{BatchId}", Job.Id, BatchId);
+			}
+			if (Batch.SessionId == null)
+			{
+				throw new StructuredRpcException(StatusCode.PermissionDenied, "Batch {JobId}:{BatchId} has no session id", Job.Id, BatchId);
+			}
+
+			ClaimsPrincipal Principal = Context.GetHttpContext().User;
+			if (!Principal.HasSessionClaim(Batch.SessionId.Value))
+			{
+				throw new StructuredRpcException(StatusCode.PermissionDenied, "Session id {SessionId} not valid for batch {JobId}:{BatchId}. Expected {ExpectedSessionId}.", Principal.GetSessionClaim() ?? ObjectId.Empty, Job.Id, BatchId, Batch.SessionId.Value);
+			}
+
+			return Batch;
+		}
+
+		/// <summary>
+		/// Updates the state of a jobstep
+		/// </summary>
+		/// <param name="Request">Request arguments</param>
+		/// <param name="Context">Context for the RPC call</param>
+		/// <returns>Information about the new agent</returns>
+		public async override Task<Empty> UpdateStep(UpdateStepRequest Request, ServerCallContext Context)
+		{
+			IJob? Job = await JobService.GetJobAsync(new JobId(Request.JobId));
+			if (Job == null)
+			{
+				throw new StructuredRpcException(StatusCode.NotFound, "Job {JobId} not found", Request.JobId);
+			}
+
+			IJobStepBatch Batch = AuthorizeBatch(Job, Request.BatchId.ToSubResourceId(), Context);
+
+			await JobService.UpdateStepAsync(Job, Batch.Id, Request.StepId.ToSubResourceId(), Request.State, Request.Outcome, null, null, null, null, null);
+			return new Empty();
+		}
+
+		/// <summary>
+		/// Get the state of a jobstep
+		/// </summary>
+		/// <param name="Request">Request arguments</param>
+		/// <param name="Context">Context for the RPC call</param>
+		/// <returns>Information about the step</returns>
+		public async override Task<GetStepResponse> GetStep(GetStepRequest Request, ServerCallContext Context)
+		{
+			IJob Job = await GetJobAsync(new JobId(Request.JobId));
+			IJobStepBatch Batch = AuthorizeBatch(Job, Request.BatchId.ToSubResourceId(), Context);
+
+			SubResourceId StepId = Request.StepId.ToSubResourceId();
+			if (!Batch.TryGetStep(StepId, out IJobStep? Step))
+			{
+				throw new StructuredRpcException(StatusCode.NotFound, "Unable to find step {JobId}:{BatchId}:{StepId}", Job.Id, Batch.Id, StepId);
+			}
+
+			return new GetStepResponse(Step.Outcome, Step.State, Step.AbortRequested);
+		}
+
+		/// <summary>
+		/// Updates the state of a jobstep
+		/// </summary>
+		/// <param name="Request">Request arguments</param>
+		/// <param name="Context">Context for the RPC call</param>
+		/// <returns>Information about the new agent</returns>
+		public async override Task<UpdateGraphResponse> UpdateGraph(UpdateGraphRequest Request, ServerCallContext Context)
+		{
+			List<NewGroup> NewGroups = new List<NewGroup>();
+			foreach (CreateGroupRequest Group in Request.Groups)
+			{
+				List<NewNode> NewNodes = new List<NewNode>();
+				foreach (CreateNodeRequest Node in Group.Nodes)
+				{
+					NewNode NewNode = new NewNode(Node.Name, Node.InputDependencies.ToList(), Node.OrderDependencies.ToList(), Node.Priority, Node.AllowRetry, Node.RunEarly, Node.Warnings, new Dictionary<string, string>(Node.Credentials), new Dictionary<string, string>(Node.Properties));
+					NewNodes.Add(NewNode);
+				}
+				NewGroups.Add(new NewGroup(Group.AgentType, NewNodes));
+			}
+
+			List<NewAggregate> NewAggregates = new List<NewAggregate>();
+			foreach (CreateAggregateRequest Aggregate in Request.Aggregates)
+			{
+				NewAggregate NewAggregate = new NewAggregate(Aggregate.Name, Aggregate.Nodes.ToList());
+				NewAggregates.Add(NewAggregate);
+			}
+
+			List<NewLabel> NewLabels = new List<NewLabel>();
+			foreach (CreateLabelRequest Label in Request.Labels)
+			{
+				NewLabel NewLabel = new NewLabel();
+				NewLabel.DashboardName = String.IsNullOrEmpty(Label.DashboardName) ? null : Label.DashboardName;
+				NewLabel.DashboardCategory = String.IsNullOrEmpty(Label.DashboardCategory) ? null : Label.DashboardCategory;
+				NewLabel.UgsName = String.IsNullOrEmpty(Label.UgsName) ? null : Label.UgsName;
+				NewLabel.UgsProject = String.IsNullOrEmpty(Label.UgsProject) ? null : Label.UgsProject;
+				NewLabel.Change = Label.Change;
+				NewLabel.RequiredNodes = Label.RequiredNodes.ToList();
+				NewLabel.IncludedNodes = Label.IncludedNodes.ToList();
+				NewLabels.Add(NewLabel);
+			}
+
+			JobId JobIdValue = new JobId(Request.JobId);
+			for (; ; )
+			{
+				IJob? Job = await JobService.GetJobAsync(JobIdValue);
+				if (Job == null)
+				{
+					throw new StructuredRpcException(StatusCode.NotFound, "Resource not found");
+				}
+				if (!JobService.AuthorizeSession(Job, Context.GetHttpContext().User))
+				{
+					throw new StructuredRpcException(StatusCode.PermissionDenied, "Access denied");
+				}
+
+				IGraph Graph = await JobService.GetGraphAsync(Job);
+				Graph = await Graphs.AppendAsync(Graph, NewGroups, NewAggregates, NewLabels);
+
+				if (await JobService.TryUpdateGraphAsync(Job, Graph))
+				{
+					return new UpdateGraphResponse();
+				}
+			}
+		}
+
+		/// <summary>
+		/// Creates a set of events
+		/// </summary>
+		/// <param name="Request">Request arguments</param>
+		/// <param name="Context">Context for the RPC call</param>
+		/// <returns>Information about the new agent</returns>
+		public async override Task<Empty> CreateEvents(CreateEventsRequest Request, ServerCallContext Context)
+		{
+			if (!await AclService.AuthorizeAsync(AclAction.CreateEvent, Context.GetHttpContext().User))
+			{
+				throw new StructuredRpcException(StatusCode.NotFound, "Access denied");
+			}
+
+			List<NewLogEventData> NewEvents = new List<NewLogEventData>();
+			foreach (CreateEventRequest Event in Request.Events)
+			{
+				NewLogEventData NewEvent = new NewLogEventData();
+				NewEvent.LogId = new LogId(Event.LogId);
+				NewEvent.Severity = Event.Severity;
+				NewEvent.LineIndex = Event.LineIndex;
+				NewEvent.LineCount = Event.LineCount;
+				NewEvents.Add(NewEvent);
+			}
+			await LogFileService.CreateEventsAsync(NewEvents);
+			return new Empty();
+		}
+
+		/// <summary>
+		/// Writes output to a log file
+		/// </summary>
+		/// <param name="Request">Request arguments</param>
+		/// <param name="Context">Context for the RPC call</param>
+		/// <returns>Information about the new agent</returns>
+		public async override Task<Empty> WriteOutput(WriteOutputRequest Request, ServerCallContext Context)
+		{
+			ILogFile? LogFile = await LogFileService.GetCachedLogFileAsync(new LogId(Request.LogId));
+			if (LogFile == null)
+			{
+				throw new StructuredRpcException(StatusCode.NotFound, "Resource not found");
+			}
+			if (!HordeServer.Services.LogFileService.AuthorizeForSession(LogFile, Context.GetHttpContext().User))
+			{
+				throw new StructuredRpcException(StatusCode.PermissionDenied, "Access denied");
+			}
+
+			await LogFileService.WriteLogDataAsync(LogFile, Request.Offset, Request.LineIndex, Request.Data.ToArray(), Request.Flush);
+			return new Empty();
+		}
+
+		/// <summary>
+		/// Uploads a new agent archive
+		/// </summary>
+		/// <param name="Request">Request arguments</param>
+		/// <param name="Context">Context for the RPC call</param>
+		/// <returns>Information about the new agent</returns>
+		public async override Task<UploadSoftwareResponse> UploadSoftware(UploadSoftwareRequest Request, ServerCallContext Context)
+		{
+			if (!await AclService.AuthorizeAsync(AclAction.UploadSoftware, Context.GetHttpContext().User))
+			{
+				throw new StructuredRpcException(StatusCode.PermissionDenied, "Access to software is forbidden");
+			}
+
+			string Version = await AgentSoftwareService.SetArchiveAsync(new AgentSoftwareChannelName(Request.Channel), null, Request.Data.ToArray());
+
+			UploadSoftwareResponse Response = new UploadSoftwareResponse();
+			Response.Version = Version;
+			return Response;
+		}
+
+		/// <summary>
+		/// Downloads a new agent archive
+		/// </summary>
+		/// <param name="Request">Request arguments</param>
+		/// <param name="ResponseStream">Writer for the output data</param>
+		/// <param name="Context">Context for the RPC call</param>
+		/// <returns>Information about the new agent</returns>
+		public async override Task DownloadSoftware(DownloadSoftwareRequest Request, IServerStreamWriter<DownloadSoftwareResponse> ResponseStream, ServerCallContext Context)
+		{
+			if (!await AclService.AuthorizeAsync(AclAction.DownloadSoftware, Context.GetHttpContext().User))
+			{
+				throw new StructuredRpcException(StatusCode.NotFound, "Access to software is forbidden");
+			}
+
+			byte[]? Data = await AgentSoftwareService.GetArchiveAsync(Request.Version);
+			if (Data == null)
+			{
+				throw new StructuredRpcException(StatusCode.NotFound, "Missing version {Version}");
+			}
+
+			for (int Offset = 0; Offset < Data.Length;)
+			{
+				int NextOffset = Math.Min(Offset + 128 * 1024, Data.Length);
+
+				DownloadSoftwareResponse Response = new DownloadSoftwareResponse();
+				Response.Data = Google.Protobuf.ByteString.CopyFrom(Data.AsSpan(Offset, NextOffset - Offset));
+
+				await ResponseStream.WriteAsync(Response);
+
+				Offset = NextOffset;
+			}
+		}
+
+		/// <summary>
+		/// Uploads a new artifact
+		/// </summary>
+		/// <param name="Reader">Request arguments</param>
+		/// <param name="Context">Context for the RPC call</param>
+		/// <returns>Information about the new agent</returns>
+		public async override Task<UploadArtifactResponse> UploadArtifact(IAsyncStreamReader<UploadArtifactRequest> Reader, ServerCallContext Context)
+		{
+			// Advance to the metadata object
+			if (!await Reader.MoveNext())
+			{
+				throw new StructuredRpcException(StatusCode.DataLoss, "Missing request for artifact upload");
+			}
+
+			// Read the request object
+			UploadArtifactMetadata? Metadata = Reader.Current.Metadata;
+			if (Metadata == null)
+			{
+				throw new StructuredRpcException(StatusCode.DataLoss, "Expected metadata in first artifact request");
+			}
+
+			// Get the job and step
+			IJob Job = await GetJobAsync(new JobId(Metadata.JobId));
+			IJobStepBatch Batch = AuthorizeBatch(Job, Metadata.BatchId.ToSubResourceId(), Context);
+
+			IJobStep? Step;
+			if (!Job.TryGetStep(Metadata.BatchId.ToSubResourceId(), Metadata.StepId.ToSubResourceId(), out Step))
+			{
+				throw new StructuredRpcException(StatusCode.NotFound, "Unable to find step {JobId}:{BatchId}:{StepId}", Job.Id, Metadata.BatchId, Metadata.StepId);
+			}
+
+			// Upload the stream
+			using (ArtifactChunkStream InputStream = new ArtifactChunkStream(Reader, Metadata.Length))
+			{
+				IArtifact Artifact = await ArtifactCollection.CreateArtifactAsync(Job.Id, Step.Id, Metadata.Name, Metadata.MimeType, InputStream);
+
+				UploadArtifactResponse Response = new UploadArtifactResponse();
+				Response.Id = Artifact.Id.ToString();
+				return Response;
+			}
+		}
+
+		/// <summary>
+		/// Uploads new test data
+		/// </summary>
+		/// <param name="Reader">Request arguments</param>
+		/// <param name="Context">Context for the RPC call</param>
+		/// <returns>Information about the new agent</returns>
+		public override async Task<UploadTestDataResponse> UploadTestData(IAsyncStreamReader<UploadTestDataRequest> Reader, ServerCallContext Context)
+		{
+			IJob? Job = null;
+			IJobStep? JobStep = null;
+
+			while (await Reader.MoveNext())
+			{
+				UploadTestDataRequest Request = Reader.Current;
+
+				JobId JobId = new JobId(Request.JobId);
+				if (Job == null || JobId != Job.Id)
+				{
+					Job = await JobService.GetJobAsync(JobId);
+					if (Job == null)
+					{
+						throw new StructuredRpcException(StatusCode.NotFound, "Unable to find job {JobId}", JobId);
+					}
+					JobStep = null;
+				}
+
+				SubResourceId JobStepId = Request.JobStepId.ToSubResourceId();
+				if (JobStep == null || JobStepId != JobStep.Id)
+				{
+					if (!Job.TryGetStep(JobStepId, out JobStep))
+					{
+						throw new StructuredRpcException(StatusCode.NotFound, "Unable to find step {JobStepId} on job {JobId}", JobStepId, JobId);
+					}
+				}
+
+				string Text = Encoding.UTF8.GetString(Request.Value.ToArray());
+				BsonDocument Document = BsonSerializer.Deserialize<BsonDocument>(Text);
+				await TestData.AddAsync(Job, JobStep, Request.Key, Document);
+			}
+
+			return new UploadTestDataResponse();
+		}
+
+		/// <summary>
+		/// Create a new report on a job or job step
+		/// </summary>
+		/// <param name="Request"></param>
+		/// <param name="Context"></param>
+		/// <returns></returns>
+		public override async Task<CreateReportResponse> CreateReport(CreateReportRequest Request, ServerCallContext Context)
+		{
+			IJob Job = await GetJobAsync(new JobId(Request.JobId));
+			IJobStepBatch Batch = AuthorizeBatch(Job, Request.BatchId.ToSubResourceId(), Context);
+
+			Report NewReport = new Report { Name = Request.Name, Placement = Request.Placement, ArtifactId = Request.ArtifactId.ToObjectId() };
+			if (Request.Scope == ReportScope.Job)
+			{
+				Logger.LogDebug("Adding report to job {JobId}: {Name} -> {ArtifactId}", Job.Id, Request.Name, Request.ArtifactId);
+				await JobService.UpdateJobAsync(Job, Reports: new List<Report> { NewReport });
+			}
+			else
+			{
+				Logger.LogDebug("Adding report to step {JobId}:{BatchId}:{StepId}: {Name} -> {ArtifactId}", Job.Id, Batch.Id, Request.StepId, Request.Name, Request.ArtifactId);
+				await JobService.UpdateStepAsync(Job, Batch.Id, Request.StepId.ToSubResourceId(), NewReports: new List<Report> { NewReport });
+			}
+
+			return new CreateReportResponse();
+		}
+	}
+}
diff --git a/Engine/Source/Programs/Horde/Horde.Build/Startup.cs b/Engine/Source/Programs/Horde/Horde.Build/Startup.cs
new file mode 100644
--- /dev/null
+++ b/Engine/Source/Programs/Horde/Horde.Build/Startup.cs
@@ -0,0 +1,886 @@
+// Copyright Epic Games, Inc. All Rights Reserved.
+
+using System;
+using System.Collections.Generic;
+using System.Diagnostics;
+using System.IdentityModel.Tokens.Jwt;
+using System.IO;
+using System.Linq;
+using System.Reflection;
+using System.Security.Claims;
+using System.Text.Encodings.Web;
+using System.Text.Json;
+using System.Text.Json.Serialization;
+using System.Threading.Tasks;
+using Serilog;
+using HordeServer.Models;
+using HordeServer.Services;
+using Microsoft.AspNetCore.Authentication;
+using Microsoft.AspNetCore.Authentication.Cookies;
+using Microsoft.AspNetCore.Authentication.JwtBearer;
+using Microsoft.AspNetCore.Authentication.OAuth.Claims;
+using Microsoft.AspNetCore.Authentication.OpenIdConnect;
+using Microsoft.AspNetCore.Authorization;
+using Microsoft.AspNetCore.Builder;
+using Microsoft.AspNetCore.Hosting;
+using Microsoft.Extensions.Configuration;
+using Microsoft.Extensions.DependencyInjection;
+using Microsoft.Extensions.Hosting;
+using Microsoft.Extensions.Logging;
+using Microsoft.Extensions.Options;
+using Microsoft.IdentityModel.Protocols.OpenIdConnect;
+using Microsoft.IdentityModel.Tokens;
+using Microsoft.OpenApi.Models;
+using MongoDB.Bson;
+using MongoDB.Bson.Serialization;
+using MongoDB.Bson.Serialization.Conventions;
+using Microsoft.AspNetCore.HttpOverrides;
+using HordeServer.Utilities;
+using Grpc.Core.Interceptors;
+using Grpc.Core;
+using EpicGames.Core;
+using Microsoft.Extensions.Diagnostics.HealthChecks;
+using Microsoft.AspNetCore.Diagnostics.HealthChecks;
+using HordeServer.Collections;
+using Microsoft.AspNetCore.Http;
+using System.Diagnostics.CodeAnalysis;
+using System.Globalization;
+using System.Threading;
+using HordeCommon;
+using Microsoft.AspNetCore.DataProtection;
+using StackExchange.Redis;
+using HordeServer.Storage;
+using HordeServer.Logs;
+using HordeServer.Logs.Builder;
+using HordeServer.Logs.Storage;
+using HordeServer.Logs.Readers;
+using HordeServer.Logs.Storage.Impl;
+using HordeServer.Collections.Impl;
+using HordeServer.Services.Impl;
+using Microsoft.Extensions.DependencyInjection.Extensions;
+using HordeServer.Authentication;
+using HordeServer.Tasks.Impl;
+using HordeServer.Tasks;
+using StatsdClient;
+using Status = Grpc.Core.Status;
+using HordeServer.Notifications.Impl;
+using HordeServer.Notifications;
+using Microsoft.AspNetCore.Hosting.Server.Features;
+using System.Runtime.InteropServices;
+using HordeServer.Storage.Collections;
+using Amazon.Extensions.NETCore.Setup;
+using Amazon.S3;
+using Amazon.Runtime;
+using HordeServer.Storage.Backends;
+using HordeServer.Storage.Services;
+using HordeServer.Commits.Impl;
+using HordeServer.Commits;
+using OpenTracing.Contrib.Grpc.Interceptors;
+using OpenTracing.Util;
+using EpicGames.Horde.Compute;
+using HordeServer.Compute.Impl;
+using HordeServer.Compute;
+using System.Net.Http.Headers;
+using Serilog.Events;
+using HordeServer.Jobs;
+
+namespace HordeServer
+{
+	using MessageTemplate = EpicGames.Core.MessageTemplate;
+	using JobId = ObjectId<IJob>;
+	using LeaseId = ObjectId<ILease>;
+	using LogId = ObjectId<ILogFile>;
+	using UserId = ObjectId<IUser>;
+	using ILogger = Microsoft.Extensions.Logging.ILogger;
+
+	class Startup
+	{
+		class GrpcExceptionInterceptor : Interceptor
+		{
+			ILogger<GrpcExceptionInterceptor> Logger;
+
+			public GrpcExceptionInterceptor(ILogger<GrpcExceptionInterceptor> Logger)
+			{
+				this.Logger = Logger;
+			}
+
+			[SuppressMessage("Design", "CA1031:Do not catch general exception types", Justification = "<Pending>")]
+			IDisposable? RequestContext<TRequest>(TRequest Request)
+			{
+				try
+				{
+					return Logger.BeginScope("Request: {Request}", JsonSerializer.Serialize<TRequest>(Request));
+				}
+				catch
+				{
+					return null;
+				}
+			}
+
+			[SuppressMessage("Design", "CA1031:Do not catch general exception types", Justification = "<Pending>")]
+			IDisposable? ConnectionContext(ConnectionInfo Connection)
+			{
+				try
+				{
+					return Logger.BeginScope("Client: {ClientIp}", Connection.RemoteIpAddress);
+				}
+				catch
+				{
+					return null;
+				}
+			}
+
+			public override AsyncClientStreamingCall<TRequest, TResponse> AsyncClientStreamingCall<TRequest, TResponse>(ClientInterceptorContext<TRequest, TResponse> Context, AsyncClientStreamingCallContinuation<TRequest, TResponse> Continuation)
+			{
+				try
+				{
+					return base.AsyncClientStreamingCall(Context, Continuation);
+				}
+				catch (StructuredRpcException Ex)
+				{
+					Logger.LogError(Ex, Ex.Format, Ex.Args);
+					throw;
+				}
+				catch (Exception Ex)
+				{
+					Logger.LogError(Ex, "Exception in call to {Method}", Context.Method);
+					throw new RpcException(new Status(StatusCode.Internal, $"An exception was thrown on the server: {Ex}"));
+				}
+			}
+
+			public override AsyncDuplexStreamingCall<TRequest, TResponse> AsyncDuplexStreamingCall<TRequest, TResponse>(ClientInterceptorContext<TRequest, TResponse> Context, AsyncDuplexStreamingCallContinuation<TRequest, TResponse> Continuation)
+			{
+				try
+				{
+					return base.AsyncDuplexStreamingCall(Context, Continuation);
+				}
+				catch (StructuredRpcException Ex)
+				{
+					Logger.LogError(Ex, Ex.Format, Ex.Args);
+					throw;
+				}
+				catch (Exception Ex)
+				{
+					Logger.LogError(Ex, "Exception in call to {Method}", Context.Method);
+					throw new RpcException(new Status(StatusCode.Internal, $"An exception was thrown on the server: {Ex}"));
+				}
+			}
+
+			[SuppressMessage("Design", "CA1031:Do not catch general exception types", Justification = "<Pending>")]
+			public override AsyncServerStreamingCall<TResponse> AsyncServerStreamingCall<TRequest, TResponse>(TRequest Request, ClientInterceptorContext<TRequest, TResponse> Context, AsyncServerStreamingCallContinuation<TRequest, TResponse> Continuation)
+			{
+				try
+				{
+					return base.AsyncServerStreamingCall(Request, Context, Continuation);
+				}
+				catch (StructuredRpcException Ex)
+				{
+					using IDisposable? Scope = RequestContext(Request); 
+					Logger.LogError(Ex, Ex.Format, Ex.Args);
+					throw;
+				}
+				catch (Exception Ex)
+				{
+					using IDisposable? Scope = RequestContext(Request);
+					Logger.LogError(Ex, "Exception in call to {Method}", Context.Method);
+					throw new RpcException(new Status(StatusCode.Internal, $"An exception was thrown on the server: {Ex}"));
+				}
+			}
+
+			[SuppressMessage("Design", "CA1031:Do not catch general exception types", Justification = "<Pending>")]
+			public override async Task<TResponse> UnaryServerHandler<TRequest, TResponse>(TRequest Request, ServerCallContext Context, UnaryServerMethod<TRequest, TResponse> Continuation)
+			{
+				try
+				{
+					return await base.UnaryServerHandler(Request, Context, Continuation);
+				}
+				catch (StructuredRpcException Ex)
+				{
+					using IDisposable? ConnectionScope = ConnectionContext(Context.GetHttpContext().Connection);
+					using IDisposable? Scope = RequestContext(Request);
+					Logger.LogError(Ex, Ex.Format, Ex.Args);
+					throw;
+				}
+				catch (Exception Ex)
+				{
+					using IDisposable? ConnectionScope = ConnectionContext(Context.GetHttpContext().Connection);
+					using IDisposable? Scope = RequestContext(Request);
+					Logger.LogError(Ex, "Exception in call to {Method}", Context.Method);
+					throw new RpcException(new Status(StatusCode.Internal, $"An exception was thrown on the server: {Ex}"));
+				}
+			}
+		}
+
+		class BsonSerializationProvider : IBsonSerializationProvider
+		{
+			public IBsonSerializer? GetSerializer(Type Type)
+			{
+				if (Type == typeof(ContentHash))
+				{
+					return new ContentHashSerializer();
+				}
+				if (Type == typeof(DateTimeOffset))
+				{
+					return new DateTimeOffsetStringSerializer();
+				}
+				return null;
+			}
+		}
+
+		class JsonObjectIdConverter : JsonConverter<ObjectId>
+		{
+			public override ObjectId Read(ref Utf8JsonReader Reader, Type TypeToConvert, JsonSerializerOptions Options)
+			{
+				string? String = Reader.GetString();
+				if (String == null)
+				{
+					throw new InvalidDataException("Unable to parse object id");
+				}
+				return String.ToObjectId();
+			}
+
+			public override void Write(Utf8JsonWriter Writer, ObjectId ObjectId, JsonSerializerOptions Options)
+			{
+				Writer.WriteStringValue(ObjectId.ToString());
+			}
+		}
+		
+		class JsonDateTimeConverter : JsonConverter<DateTime>
+		{
+			public override DateTime Read(ref Utf8JsonReader Reader, Type TypeToConvert, JsonSerializerOptions Options)
+			{
+				Debug.Assert(TypeToConvert == typeof(DateTime));
+
+				string? String = Reader.GetString();
+				if (String == null)
+				{
+					throw new InvalidDataException("Unable to parse DateTime");
+				}
+				return DateTime.Parse(String, CultureInfo.CurrentCulture);
+			}
+
+			public override void Write(Utf8JsonWriter Writer, DateTime DateTime, JsonSerializerOptions Options)
+			{
+				Writer.WriteStringValue(DateTime.ToUniversalTime().ToString("yyyy'-'MM'-'dd'T'HH':'mm':'ssZ", CultureInfo.CurrentCulture));
+			}
+		}
+
+		public Startup(IConfiguration Configuration)
+		{
+			this.Configuration = Configuration;
+		}
+
+		public IConfiguration Configuration { get; }
+
+		// This method gets called *multiple times* by the runtime. Use this method to add services to the container.
+		public void ConfigureServices(IServiceCollection Services)
+		{
+			// IOptionsMonitor pattern for live updating of configuration settings
+			Services.Configure<ServerSettings>(Configuration.GetSection("Horde"));
+
+			// Settings used for configuring services
+			ServerSettings Settings = new ServerSettings();
+			Configuration.GetSection("Horde").Bind(Settings);			
+
+			if (Settings.GlobalThreadPoolMinSize != null)
+			{
+				// Min thread pool size is set to combat timeouts seen with the Redis client.
+				// See comments for <see cref="ServerSettings.GlobalThreadPoolMinSize" /> and
+				// [REDACTED_PROFILE_URL]/StackExchange.Redis/issues/1680
+				int Min = Settings.GlobalThreadPoolMinSize.Value;
+				ThreadPool.SetMinThreads(Min, Min);
+			}
+
+#pragma warning disable CA2000 // Dispose objects before losing scope
+			RedisService RedisService = new RedisService(Settings);
+#pragma warning restore CA2000 // Dispose objects before losing scope
+			Services.AddSingleton<RedisService>(SP => RedisService);
+			Services.AddSingleton<IDatabase>(RedisService.Database);
+			Services.AddSingleton<ConnectionMultiplexer>(SP => RedisService.Multiplexer);
+			Services.AddDataProtection().PersistKeysToStackExchangeRedis(() => RedisService.Database, "aspnet-data-protection");
+
+			if (Settings.CorsEnabled)
+			{
+				Services.AddCors(Options =>
+				{
+					Options.AddPolicy("CorsPolicy",
+						Builder => Builder.WithOrigins(Settings.CorsOrigin.Split(";"))
+						.AllowAnyMethod()
+						.AllowAnyHeader()
+						.AllowCredentials());
+				});
+			}
+
+			Services.AddGrpc(Options =>
+			{
+				Options.EnableDetailedErrors = true;
+				Options.MaxReceiveMessageSize = 200 * 1024 * 1024; // 100 MB (packaged builds of Horde agent can be large) 
+				Options.Interceptors.Add(typeof(LifetimeGrpcInterceptor));
+				Options.Interceptors.Add(typeof(GrpcExceptionInterceptor));
+				Options.Interceptors.Add<ServerTracingInterceptor>(GlobalTracer.Instance);
+			});
+			Services.AddGrpcReflection();
+
+			Services.AddSingleton<IAgentCollection, AgentCollection>();
+			Services.AddSingleton<IAgentSoftwareCollection, AgentSoftwareCollection>();
+			Services.AddSingleton<IArtifactCollection, ArtifactCollection>();
+			Services.AddSingleton<ICommitCollection, CommitCollection>();
+			Services.AddSingleton<IGraphCollection, GraphCollection>();
+			Services.AddSingleton<IIssueCollection, IssueCollection>();
+			Services.AddSingleton<IJobCollection, JobCollection>();
+			Services.AddSingleton<IJobStepRefCollection, JobStepRefCollection>();
+			Services.AddSingleton<IJobTimingCollection, JobTimingCollection>();
+			Services.AddSingleton<ILeaseCollection, LeaseCollection>();
+			Services.AddSingleton<ILogEventCollection, LogEventCollection>();
+			Services.AddSingleton<ILogFileCollection, LogFileCollection>();
+			Services.AddSingleton<INotificationTriggerCollection, NotificationTriggerCollection>();
+			Services.AddSingleton<IPoolCollection, PoolCollection>();
+			Services.AddSingleton<IProjectCollection, ProjectCollection>();
+			Services.AddSingleton<ISessionCollection, SessionCollection>();
+			Services.AddSingleton<IServiceAccountCollection, ServiceAccountCollection>();
+			Services.AddSingleton<ISubscriptionCollection, SubscriptionCollection>();
+			Services.AddSingleton<IStreamCollection, StreamCollection>();
+			Services.AddSingleton<ITemplateCollection, TemplateCollection>();
+			Services.AddSingleton<ITestDataCollection, TestDataCollection>();
+			Services.AddSingleton<ITelemetryCollection, TelemetryCollection>();
+			Services.AddSingleton<ITemplateCollection, TemplateCollection>();
+			Services.AddSingleton<IUgsMetadataCollection, UgsMetadataCollection>();
+			Services.AddSingleton<IUserCollection, UserCollectionV2>();
+			Services.AddSingleton<IDeviceCollection, DeviceCollection>();
+
+			// Auditing
+			Services.AddSingleton<IAuditLog<AgentId>>(SP => SP.GetRequiredService<IAuditLogFactory<AgentId>>().Create("Agents.Log", "AgentId"));
+
+			// Storage
+			Services.AddSingleton<IBlobCollection>(SP => new CachingBlobCollection(new BlobCollection(SP.GetRequiredService<IStorageBackend<BlobCollection>>()), 256 * 1024 * 1024));
+			Services.AddSingleton<IBucketCollection, BucketCollection>();
+			Services.AddSingleton<INamespaceCollection, NamespaceCollection>();
+			Services.AddSingleton<IObjectCollection, ObjectCollection>();
+			Services.AddSingleton<IRefCollection, RefCollection>();
+
+			Services.AddSingleton(typeof(IAuditLogFactory<>), typeof(AuditLogFactory<>));
+			Services.AddSingleton(typeof(ISingletonDocument<>), typeof(SingletonDocument<>));
+
+			Services.AddSingleton<AutoscaleService>();
+			switch (Settings.FleetManager)
+			{
+				case FleetManagerType.Aws:
+					Services.AddSingleton<IFleetManager, AwsFleetManager>();
+					break;
+				default:
+					Services.AddSingleton<IFleetManager, DefaultFleetManager>();
+					break;
+			}
+
+			Services.AddSingleton<AclService>();
+			Services.AddSingleton<AgentService>();			
+			Services.AddSingleton<AgentSoftwareService>();
+			Services.AddSingleton<ConsistencyService>();
+			Services.AddSingleton<RequestTrackerService>();
+			Services.AddSingleton<CredentialService>();
+			Services.AddSingleton<DatabaseService>();
+			Services.AddSingleton<ConfigService>();
+			Services.AddSingleton<IDogStatsd>(ctx =>
+			{
+				string? DatadogAgentHost = Environment.GetEnvironmentVariable("DD_AGENT_HOST");
+				if (DatadogAgentHost != null)
+				{
+					// Datadog agent is configured, enable DogStatsD for metric collection
+					StatsdConfig Config = new StatsdConfig
+					{
+						StatsdServerName = DatadogAgentHost,
+						StatsdPort = 8125,
+					};
+
+					DogStatsdService DogStatsdService = new DogStatsdService();
+					DogStatsdService.Configure(Config);
+					return DogStatsdService;
+				}
+				return new NoOpDogStatsd();
+			});
+			Services.AddSingleton<CommitService>();
+			Services.AddSingleton<IClock, Clock>();
+			Services.AddSingleton<IDowntimeService, DowntimeService>();
+			Services.AddSingleton<IIssueService, IssueService>();
+			Services.AddSingleton<JobService>();
+			Services.AddSingleton<LifetimeService>();
+			Services.AddSingleton<ILogFileService, LogFileService>();
+			Services.AddSingleton<INotificationService, NotificationService>();
+			Services.AddSingleton<IPerforceService, PerforceService>();
+
+			Services.AddSingleton<PerforceLoadBalancer>();
+			Services.AddSingleton<PoolService>();
+			Services.AddSingleton<ProjectService>();
+			Services.AddSingleton<ScheduleService>();
+			Services.AddSingleton<SlackNotificationSink>();
+			Services.AddSingleton<IAvatarService, SlackNotificationSink>(SP => SP.GetRequiredService<SlackNotificationSink>());
+			Services.AddSingleton<INotificationSink, SlackNotificationSink>(SP => SP.GetRequiredService<SlackNotificationSink>());
+			Services.AddSingleton<StreamService>();
+			Services.AddSingleton<UpgradeService>();
+
+			Services.AddSingleton<DeviceService>();
+
+			AWSOptions AwsOptions = Configuration.GetAWSOptions();
+			if (Settings.S3CredentialType == "AssumeRole" && Settings.S3AssumeArn != null)
+			{
+				AwsOptions.Credentials = new AssumeRoleAWSCredentials(FallbackCredentialsFactory.GetCredentials(), Settings.S3AssumeArn, "Horde");
+			}
+			else if(Settings.S3CredentialType == "AssumeRoleWebIdentity")
+			{
+				AwsOptions.Credentials = AssumeRoleWithWebIdentityCredentials.FromEnvironmentVariables();
+			}
+			else if(Settings.S3CredentialType == "Fallback")
+			{
+				// Using the fallback credentials from the AWS SDK, it will pick up credentials through a number of default mechanisms.
+				AwsOptions.Credentials = FallbackCredentialsFactory.GetCredentials();
+			}
+			
+			Services.AddSingleton(AwsOptions);
+
+			Services.AddSingleton(new StorageBackendSettings<PersistentLogStorage> { Type = Settings.ExternalStorageProviderType, BaseDir = Settings.LocalLogsDir, BucketName = Settings.S3LogBucketName });
+			Services.AddSingleton(new StorageBackendSettings<ArtifactCollection> { Type = Settings.ExternalStorageProviderType, BaseDir = Settings.LocalArtifactsDir, BucketName = Settings.S3ArtifactBucketName });
+			Services.AddSingleton(new StorageBackendSettings<BlobCollection> { Type = Settings.ExternalStorageProviderType, BaseDir = Settings.LocalBlobsDir, BucketName = Settings.S3LogBucketName, BucketPath = "blobs/" });
+			Services.AddSingleton(typeof(IStorageBackend<>), typeof(StorageBackendFactory<>));
+
+			ConfigureLogStorage(Services);
+			Services.AddSingleton<IStorageService, SimpleStorageService>();
+			//			ConfigureLogFileWriteCache(Services, Settings);
+
+			AuthenticationBuilder AuthBuilder = Services.AddAuthentication(Options =>
+				{
+					if (Settings.DisableAuth)
+					{
+						Options.DefaultAuthenticateScheme = AnonymousAuthenticationHandler.AuthenticationScheme;
+						Options.DefaultSignInScheme = AnonymousAuthenticationHandler.AuthenticationScheme;
+						Options.DefaultChallengeScheme = AnonymousAuthenticationHandler.AuthenticationScheme;
+					}
+					else
+					{
+						// If an authentication cookie is present, use it to get authentication information
+						Options.DefaultAuthenticateScheme = CookieAuthenticationDefaults.AuthenticationScheme;
+						Options.DefaultSignInScheme = CookieAuthenticationDefaults.AuthenticationScheme;
+
+						// If authentication is required, and no cookie is present, use OIDC to sign in
+						Options.DefaultChallengeScheme = OktaDefaults.AuthenticationScheme;
+					}
+				});
+
+			List<string> Schemes = new List<string>();
+
+			AuthBuilder.AddCookie(Options =>
+				 {
+					 Options.Events.OnValidatePrincipal = Context =>
+					 {
+						 if (Context.Principal?.FindFirst(HordeClaimTypes.UserId) == null)
+						 {
+							 Context.RejectPrincipal();
+						 }
+						 return Task.CompletedTask;
+					 };
+
+					 Options.Events.OnRedirectToAccessDenied = Context =>
+					 {
+						 Context.Response.StatusCode = StatusCodes.Status403Forbidden;
+						 return Context.Response.CompleteAsync();
+					 };
+				 });
+			Schemes.Add(CookieAuthenticationDefaults.AuthenticationScheme);
+
+			AuthBuilder.AddServiceAccount(Options => { });
+			Schemes.Add(ServiceAccountAuthHandler.AuthenticationScheme);
+
+			if (Settings.DisableAuth)
+			{
+				AuthBuilder.AddAnonymous(Options =>
+				{
+					Options.AdminClaimType = Settings.AdminClaimType;
+					Options.AdminClaimValue = Settings.AdminClaimValue;
+				});
+				Schemes.Add(AnonymousAuthenticationHandler.AuthenticationScheme);
+			}
+			else if (Settings.OidcClientId != null && Settings.OidcAuthority != null)
+			{
+				AuthBuilder.AddOkta(OktaDefaults.AuthenticationScheme, OpenIdConnectDefaults.DisplayName, 
+					Options =>
+					{
+						Options.Authority = Settings.OidcAuthority;
+						Options.ClientId = Settings.OidcClientId;
+
+						if (!String.IsNullOrEmpty(Settings.OidcSigninRedirect))
+						{
+							Options.Events = new OpenIdConnectEvents
+							{
+								OnRedirectToIdentityProvider = async RedirectContext =>
+								{
+									RedirectContext.ProtocolMessage.RedirectUri = Settings.OidcSigninRedirect;
+									await Task.CompletedTask;
+								}
+							};
+						}
+					});
+				Schemes.Add(OktaDefaults.AuthenticationScheme);
+			}
+
+			AuthBuilder.AddScheme<JwtBearerOptions, HordeJwtBearerHandler>(HordeJwtBearerHandler.AuthenticationScheme, Options => { });
+			Schemes.Add(HordeJwtBearerHandler.AuthenticationScheme);
+
+			Services.AddAuthorization(Options =>
+				{
+					Options.DefaultPolicy = new AuthorizationPolicyBuilder(Schemes.ToArray())
+						.RequireAuthenticatedUser()							
+						.Build();
+				});
+
+			Services.AddHostedService(Provider => Provider.GetRequiredService<AgentService>());
+			Services.AddHostedService(Provider => Provider.GetRequiredService<AutoscaleService>());
+			Services.AddHostedService(Provider => Provider.GetRequiredService<CommitService>());
+			Services.AddHostedService(Provider => Provider.GetRequiredService<ConsistencyService>());
+			Services.AddHostedService(Provider => (DowntimeService)Provider.GetRequiredService<IDowntimeService>());
+			Services.AddHostedService(Provider => Provider.GetRequiredService<IIssueService>());
+			Services.AddHostedService(Provider => (LogFileService)Provider.GetRequiredService<ILogFileService>());
+			Services.AddHostedService(Provider => (NotificationService)Provider.GetRequiredService<INotificationService>());
+			if (!Settings.DisableSchedules)
+			{
+				Services.AddHostedService(Provider => Provider.GetRequiredService<ScheduleService>());
+			}
+
+			Services.AddHostedService<MetricService>();
+			Services.AddHostedService(Provider => Provider.GetRequiredService<PerforceLoadBalancer>());
+			Services.AddHostedService<PoolUpdateService>();
+			Services.AddHostedService(Provider => Provider.GetRequiredService<SlackNotificationSink>());
+			Services.AddHostedService<ConfigService>();
+			Services.AddHostedService<TelemetryService>();
+			Services.AddHostedService(Provider => Provider.GetRequiredService<DeviceService>());
+
+			// Task sources. Order of registration is important here; it dictates the priority in which sources are served.
+			Services.AddSingleton<JobTaskSource>();
+			Services.AddSingleton<ConformTaskSource>();
+			Services.AddSingleton<IComputeService, ComputeService>();
+
+			Services.AddSingleton<ITaskSource, UpgradeTaskSource>();
+			Services.AddSingleton<ITaskSource, ShutdownTaskSource>();
+			Services.AddSingleton<ITaskSource, RestartTaskSource>();
+			Services.AddSingleton<ITaskSource, ConformTaskSource>(Provider => Provider.GetRequiredService<ConformTaskSource>());
+			Services.AddSingleton<ITaskSource, JobTaskSource>(Provider => Provider.GetRequiredService<JobTaskSource>());
+			Services.AddSingleton<ITaskSource, ComputeService>();
+
+			Services.AddHostedService(Provider => Provider.GetRequiredService<ConformTaskSource>());
+
+			// Allow longer to shutdown so we can debug missing cancellation tokens
+			Services.Configure<HostOptions>(Options =>
+			{
+				Options.ShutdownTimeout = TimeSpan.FromSeconds(30.0);
+			});
+
+			// Allow forwarded headers
+			Services.Configure<ForwardedHeadersOptions>(Options =>
+			{
+				Options.ForwardedHeaders = ForwardedHeaders.XForwardedFor | ForwardedHeaders.XForwardedProto;
+				Options.KnownProxies.Clear();
+				Options.KnownNetworks.Clear();
+			});
+
+			Services.AddMvc().AddJsonOptions(Options => ConfigureJsonSerializer(Options.JsonSerializerOptions));
+
+			Services.AddControllers(Options =>
+			{
+				Options.InputFormatters.Add(new CbInputFormatter());
+				Options.OutputFormatters.Add(new CbOutputFormatter());
+				Options.OutputFormatters.Insert(0, new CbPreferredOutputFormatter());
+				Options.FormatterMappings.SetMediaTypeMappingForFormat("uecb", CustomMediaTypeNames.UnrealCompactBinary);
+			});
+
+			Services.AddSwaggerGen(Config =>
+			{
+				Config.SwaggerDoc("v1", new OpenApiInfo { Title = "Horde Server API", Version = "v1" });
+				Config.IncludeXmlComments(Path.Combine(AppContext.BaseDirectory, $"{Assembly.GetExecutingAssembly().GetName().Name}.xml"));
+			});
+
+			DirectoryReference DashboardDir = DirectoryReference.Combine(Program.AppDir, "DashboardApp");
+			if (DirectoryReference.Exists(DashboardDir)) 
+			{
+				Services.AddSpaStaticFiles(Config => {Config.RootPath = "DashboardApp";});
+			}
+
+			ConfigureMongoDbClient();
+			ConfigureFormatters();
+
+			OnAddHealthChecks(Services);
+		}
+
+		public static void ConfigureFormatters()
+		{
+			MessageTemplate.RegisterFormatter(typeof(AgentId), new MessageTemplateTypeNameFormatter("AgentId"));
+			MessageTemplate.RegisterFormatter(typeof(JobId), new MessageTemplateTypeNameFormatter("JobId"));
+			MessageTemplate.RegisterFormatter(typeof(LeaseId), new MessageTemplateTypeNameFormatter("LeaseId"));
+			MessageTemplate.RegisterFormatter(typeof(LogId), new MessageTemplateTypeNameFormatter("LogId"));
+			MessageTemplate.RegisterFormatter(typeof(UserId), new MessageTemplateTypeNameFormatter("UserId"));
+		}
+
+		public static void ConfigureJsonSerializer(JsonSerializerOptions Options)
+		{
+			Options.IgnoreNullValues = true;
+			Options.PropertyNamingPolicy = JsonNamingPolicy.CamelCase;
+			Options.PropertyNameCaseInsensitive = true;
+			Options.Converters.Add(new JsonObjectIdConverter());
+			Options.Converters.Add(new JsonStringEnumConverter());
+			Options.Converters.Add(new JsonKnownTypesConverterFactory());
+			Options.Converters.Add(new JsonStringIdConverterFactory());
+			Options.Converters.Add(new JsonDateTimeConverter());
+		}
+
+		public class StorageBackendSettings<T> : IFileSystemStorageOptions, IAwsStorageOptions
+		{
+			/// <summary>
+			/// The type of storage backend to use
+			/// </summary>
+			public StorageProviderType Type { get; set; }
+
+			/// <inheritdoc/>
+			public string? BaseDir { get; set; }
+
+			/// <inheritdoc/>
+			public string? BucketName { get; set; }
+
+			/// <inheritdoc/>
+			public string? BucketPath { get; set; }
+		}
+
+		public class StorageBackendFactory<T> : IStorageBackend<T>
+		{
+			IStorageBackend Inner;
+
+			public StorageBackendFactory(IServiceProvider ServiceProvider, StorageBackendSettings<T> Options)
+			{
+				switch (Options.Type)
+				{
+					case StorageProviderType.S3:
+						Inner = new AwsStorageBackend(ServiceProvider.GetRequiredService<AWSOptions>(), Options, ServiceProvider.GetRequiredService<ILogger<AwsStorageBackend>>());
+						break;
+					case StorageProviderType.FileSystem:
+						Inner = new FileSystemStorageBackend(Options);
+						break;
+					case StorageProviderType.Transient:
+						Inner = new TransientStorageBackend();
+						break;
+					case StorageProviderType.Relay:
+						Inner = new RelayStorageBackend(ServiceProvider.GetRequiredService<IOptions<ServerSettings>>());
+						break;
+					default:
+						throw new NotImplementedException();
+				}
+			}
+
+			/// <inheritdoc/>
+			public Task<Stream?> ReadAsync(string Path) => Inner.ReadAsync(Path);
+
+			/// <inheritdoc/>
+			public Task<ReadOnlyMemory<byte>?> ReadBytesAsync(string Path) => Inner.ReadBytesAsync(Path);
+
+			/// <inheritdoc/>
+			public Task WriteAsync(string Path, Stream Stream) => Inner.WriteAsync(Path, Stream);
+
+			/// <inheritdoc/>
+			public Task WriteBytesAsync(string Path, ReadOnlyMemory<byte> Data) => Inner.WriteBytesAsync(Path, Data);
+
+			/// <inheritdoc/>
+			public Task DeleteAsync(string Path) => Inner.DeleteAsync(Path);
+
+			/// <inheritdoc/>
+			public Task<bool> ExistsAsync(string Path) => Inner.ExistsAsync(Path);
+		}
+
+		private static void ConfigureLogStorage(IServiceCollection Services)
+		{
+			Services.AddSingleton<ILogBuilder>(Provider =>
+			{
+				RedisService? RedisService = Provider.GetService<RedisService>();
+				if(RedisService == null)
+				{
+					return new LocalLogBuilder();
+				}
+				else
+				{
+					return new RedisLogBuilder(RedisService.ConnectionPool, Provider.GetRequiredService<ILogger<RedisLogBuilder>>());
+				}
+			});
+
+			Services.AddSingleton<PersistentLogStorage>();
+
+			Services.AddSingleton<ILogStorage>(Provider =>
+			{
+				ILogStorage Storage = Provider.GetRequiredService<PersistentLogStorage>();
+
+//				IDatabase? RedisDb = Provider.GetService<IDatabase>();
+//				if (RedisDb != null)
+//				{
+//					Storage = new RedisLogStorage(RedisDb, Provider.GetService<ILogger<RedisLogStorage>>(), Storage);
+//				}
+
+				Storage = new SequencedLogStorage(Storage);
+				Storage = new LocalLogStorage(50, Storage);
+				return Storage;
+			});
+		}
+/*
+		private static void ConfigureLogFileWriteCache(IServiceCollection Services, ServerSettings Settings)
+		{
+			bool RedisConfigured = !String.IsNullOrEmpty(Settings.RedisConnectionConfig);
+			string CacheType = Settings.LogServiceWriteCacheType.ToLower(CultureInfo.CurrentCulture);
+
+			if (CacheType == "inmemory")
+			{
+				Services.AddSingleton<ILogFileWriteCache2>(Sp => new InMemoryLogFileWriteCache2());
+			}
+			else if (CacheType == "redis" && RedisConfigured)
+			{
+				Services.AddSingleton<ILogFileWriteCache2>(Sp => new RedisLogFileWriteCache2(Sp.GetService<ILogger<RedisLogFileWriteCache>>(), Sp.GetService<IDatabase>()));
+			}
+			else if (CacheType == "redis" && !RedisConfigured)
+			{
+				throw new Exception("Redis must be configured to use the Redis-backed log write cache");
+			}
+			else
+			{
+				throw new Exception("Unknown value set for LogServiceWriteCacheType in config: " + Settings.LogServiceWriteCacheType);
+			}
+		}
+*/
+		public static void ConfigureMongoDbClient()
+		{
+			// Ignore extra elements on deserialized documents
+			ConventionPack ConventionPack = new ConventionPack();
+			ConventionPack.Add(new IgnoreExtraElementsConvention(true));
+			ConventionPack.Add(new EnumRepresentationConvention(BsonType.String));
+			ConventionRegistry.Register("Horde", ConventionPack, Type => true);
+
+			// Register the custom serializers
+			BsonSerializer.RegisterSerializationProvider(new BsonSerializationProvider());
+			BsonSerializer.RegisterSerializationProvider(new StringIdSerializationProvider());
+			BsonSerializer.RegisterSerializationProvider(new ObjectIdSerializationProvider());
+		}
+
+		private static void OnAddHealthChecks(IServiceCollection Services)
+		{
+			IHealthChecksBuilder HealthChecks = Services.AddHealthChecks().AddCheck("self", () => HealthCheckResult.Healthy(), tags: new[] { "self" });
+		}
+
+		// This method gets called by the runtime. Use this method to configure the HTTP request pipeline.
+		public static void Configure(IApplicationBuilder App, IWebHostEnvironment Env, Microsoft.Extensions.Hosting.IHostApplicationLifetime Lifetime, IOptions<ServerSettings> Settings)
+		{
+			App.UseForwardedHeaders();
+			App.UseExceptionHandler("/api/v1/error");
+
+			if (Settings.Value.CorsEnabled)
+			{
+				App.UseCors("CorsPolicy");
+			}
+
+			// Enable middleware to serve generated Swagger as a JSON endpoint.
+			App.UseSwagger();
+
+			// Enable serilog request logging
+			App.UseSerilogRequestLogging(Options => Options.GetLevel = GetRequestLoggingLevel);
+
+			// Include the source IP address with requests
+			App.Use(async (Context, Next) => {
+				using (Serilog.Context.LogContext.PushProperty("RemoteIP", Context.Connection.RemoteIpAddress))
+				{
+					await Next();
+				}
+			});
+
+			// Enable middleware to serve swagger-ui (HTML, JS, CSS, etc.),
+			// specifying the Swagger JSON endpoint.
+			App.UseSwaggerUI(c =>
+			{
+				c.SwaggerEndpoint("/swagger/v1/swagger.json", "Horde Server API");
+				c.RoutePrefix = "swagger";
+			});
+
+			if (!Env.IsDevelopment())
+			{
+				App.UseMiddleware<RequestTrackerMiddleware>();	
+			}
+
+			App.UseDefaultFiles();
+			App.UseStaticFiles();
+
+			DirectoryReference DashboardDir = DirectoryReference.Combine(Program.AppDir, "DashboardApp");
+
+			if (DirectoryReference.Exists(DashboardDir)) 
+			{
+				App.UseSpaStaticFiles();
+			}
+						
+			App.UseRouting();
+
+			App.UseAuthentication();
+			App.UseAuthorization();
+
+			App.UseEndpoints(Endpoints =>
+			{
+				Endpoints.MapGrpcService<HealthService>();
+				Endpoints.MapGrpcService<RpcService>();
+				
+				// Google Remote Execution API
+				Endpoints.MapGrpcService<ComputeRpcServer>();
+				Endpoints.MapGrpcService<BlobStoreRpc>();
+				Endpoints.MapGrpcService<RefTableRpc>();
+
+				Endpoints.MapGrpcReflectionService();
+
+				Endpoints.MapControllers();
+			});
+
+			if (DirectoryReference.Exists(DashboardDir)) 
+			{
+				App.UseSpa(Spa =>
+				{ 
+					Spa.Options.SourcePath = "DashboardApp";        
+				});
+			}
+
+			if (Settings.Value.OpenBrowser && RuntimeInformation.IsOSPlatform(OSPlatform.Windows))
+			{
+				Lifetime.ApplicationStarted.Register(() => LaunchBrowser(App));
+			}
+		}
+
+		static void LaunchBrowser(IApplicationBuilder App)
+		{
+			IServerAddressesFeature Feature = App.ServerFeatures.Get<IServerAddressesFeature>();
+			if (Feature.Addresses.Count > 0)
+			{
+				// with a development cert, host will be set by default to localhost, otherwise there will be no host in address
+				string Address = Feature.Addresses.First().Replace("[::]", System.Net.Dns.GetHostName(), StringComparison.OrdinalIgnoreCase);
+				Process.Start(new ProcessStartInfo { FileName = Address, UseShellExecute = true });
+			}
+		}
+
+		static LogEventLevel GetRequestLoggingLevel(HttpContext Context, double ElapsedMs, Exception Ex)
+		{
+			if (Context.Request != null && Context.Request.Path.HasValue)
+			{
+				string RequestPath = Context.Request.Path;
+				if (RequestPath.Equals("/Horde.HordeRpc/QueryServerStateV2", StringComparison.OrdinalIgnoreCase))
+				{
+					return LogEventLevel.Verbose;
+				}
+				if (RequestPath.Equals("/Horde.HordeRpc/UpdateSession", StringComparison.OrdinalIgnoreCase))
+				{
+					return LogEventLevel.Verbose;
+				}
+				if (RequestPath.Equals("/Horde.HordeRpc/CreateEvents", StringComparison.OrdinalIgnoreCase))
+				{
+					return LogEventLevel.Verbose;
+				}
+				if (RequestPath.Equals("/Horde.HordeRpc/WriteOutput", StringComparison.OrdinalIgnoreCase))
+				{
+					return LogEventLevel.Verbose;
+				}
+			}
+			return LogEventLevel.Information;
+		}
+	}
+}
diff --git a/Engine/Source/Programs/Horde/Horde.Build/Storage/Backends/RelayStorageBackend.cs b/Engine/Source/Programs/Horde/Horde.Build/Storage/Backends/RelayStorageBackend.cs
new file mode 100644
--- /dev/null
+++ b/Engine/Source/Programs/Horde/Horde.Build/Storage/Backends/RelayStorageBackend.cs
@@ -0,0 +1,111 @@
+// Copyright Epic Games, Inc. All Rights Reserved.
+
+using EpicGames.Core;
+using Microsoft.Extensions.Options;
+using System;
+using System.Collections.Generic;
+using System.IO;
+using System.Linq;
+using System.Net.Http;
+using System.Net.Http.Headers;
+using System.Threading.Tasks;
+
+namespace HordeServer.Storage.Backends
+{
+	/// <summary>
+	/// Implementation of ILogFileStorage which forwards requests to another server
+	/// </summary>
+	public sealed class RelayStorageBackend : IStorageBackend, IDisposable
+	{
+		class FileSystemSettings : IFileSystemStorageOptions
+		{
+			public string? BaseDir { get; set; }
+		}
+
+		/// <summary>
+		/// The client to connect with
+		/// </summary>
+		HttpClient Client;
+
+		/// <summary>
+		/// The base server URL
+		/// </summary>
+		Uri ServerUrl;
+
+		/// <summary>
+		/// Local storage provider
+		/// </summary>
+		FileSystemStorageBackend LocalStorage;
+
+		/// <summary>
+		/// Constructor
+		/// </summary>
+		/// <param name="Settings">Settings for the server instance</param>
+		public RelayStorageBackend(IOptions<ServerSettings> Settings)
+		{
+			ServerSettings CurrentSettings = Settings.Value;
+			if (CurrentSettings.LogRelayServer == null)
+			{
+				throw new InvalidDataException("Missing LogRelayServer in server configuration");
+			}
+			if (CurrentSettings.LogRelayBearerToken == null)
+			{
+				throw new InvalidDataException("Missing LogRelayBearerToken in server configuration");
+			}
+
+			ServerUrl = new Uri(CurrentSettings.LogRelayServer);
+
+			Client = new HttpClient();
+			Client.DefaultRequestHeaders.Authorization = new AuthenticationHeaderValue("Bearer", CurrentSettings.LogRelayBearerToken);
+
+			LocalStorage = new FileSystemStorageBackend(new FileSystemSettings { BaseDir = "RelayCache" });
+		}
+
+		/// <inheritdoc/>
+		public void Dispose()
+		{
+			Client.Dispose();
+		}
+
+		/// <inheritdoc/>
+		public async Task<Stream?> ReadAsync(string Path)
+		{
+			Stream? LocalResult = await LocalStorage.ReadAsync(Path);
+			if (LocalResult != null)
+			{
+				return LocalResult;
+			}
+
+			Uri Url = new Uri(ServerUrl, $"api/v1/debug/storage?path={Path}");
+			using (HttpResponseMessage Response = await Client.GetAsync(Url))
+			{
+				if (Response.IsSuccessStatusCode)
+				{
+					byte[] ResponseData = await Response.Content.ReadAsByteArrayAsync();
+					await LocalStorage.WriteBytesAsync(Path, ResponseData);
+					return new MemoryStream(ResponseData);
+				}
+			}
+
+			return null;
+		}
+
+		/// <inheritdoc/>
+		public Task WriteAsync(string Path, Stream Stream)
+		{
+			return LocalStorage.WriteAsync(Path, Stream);
+		}
+
+		/// <inheritdoc/>
+		public Task<bool> ExistsAsync(string Path)
+		{
+			throw new NotSupportedException();
+		}
+
+		/// <inheritdoc/>
+		public Task DeleteAsync(string Path)
+		{
+			throw new NotImplementedException();
+		}
+	}
+}
diff --git a/Engine/Source/Programs/Horde/Horde.Build/Storage/Backends/TransientStorageBackend.cs b/Engine/Source/Programs/Horde/Horde.Build/Storage/Backends/TransientStorageBackend.cs
new file mode 100644
--- /dev/null
+++ b/Engine/Source/Programs/Horde/Horde.Build/Storage/Backends/TransientStorageBackend.cs
@@ -0,0 +1,59 @@
+// Copyright Epic Games, Inc. All Rights Reserved.
+
+using System;
+using System.Collections.Concurrent;
+using System.Collections.Generic;
+using System.IO;
+using System.Linq;
+using System.Threading.Tasks;
+
+namespace HordeServer.Storage.Backends
+{
+	/// <summary>
+	/// In-memory implementation of ILogFileStorage
+	/// </summary>
+	public class TransientStorageBackend : IStorageBackend
+	{
+		/// <summary>
+		/// Data storage
+		/// </summary>
+		ConcurrentDictionary<string, byte[]> PathToData = new ConcurrentDictionary<string, byte[]>();
+
+		/// <inheritdoc/>
+		public Task<Stream?> ReadAsync(string Path)
+		{
+			byte[]? Data;
+			if (PathToData.TryGetValue(Path, out Data))
+			{
+				return Task.FromResult<Stream?>(new MemoryStream(Data, false));
+			}
+			else
+			{
+				return Task.FromResult<Stream?>(null);
+			}
+		}
+
+		/// <inheritdoc/>
+		public async Task WriteAsync(string Path, Stream Stream)
+		{
+			using (MemoryStream Buffer = new MemoryStream())
+			{
+				await Stream.CopyToAsync(Buffer);
+				PathToData[Path] = Buffer.ToArray();
+			}
+		}
+
+		/// <inheritdoc/>
+		public Task<bool> ExistsAsync(string Path)
+		{
+			return Task.FromResult(PathToData.ContainsKey(Path));
+		}
+
+		/// <inheritdoc/>
+		public Task DeleteAsync(string Path)
+		{
+			PathToData.TryRemove(Path, out _);
+			return Task.CompletedTask;
+		}
+	}
+}
diff --git a/Engine/Source/Programs/Horde/Horde.Build/Storage/IStorageService.cs b/Engine/Source/Programs/Horde/Horde.Build/Storage/IStorageService.cs
new file mode 100644
--- /dev/null
+++ b/Engine/Source/Programs/Horde/Horde.Build/Storage/IStorageService.cs
@@ -0,0 +1,109 @@
+// Copyright Epic Games, Inc. All Rights Reserved.
+
+using EpicGames.Core;
+using HordeServer.Storage.Primitives;
+using System;
+using System.Collections.Generic;
+using System.Linq;
+using System.Threading.Tasks;
+
+namespace HordeServer.Storage
+{
+	/// <summary>
+	/// Combined interface for replicated content addressable storage and ref table.
+	/// </summary>
+	public interface IStorageService
+	{
+		/// <summary>
+		/// Attempts to retrieve a blob from the store. 
+		/// 
+		/// Note that this call may succeed even though the object is not in persistent storage; it may be retrieved
+		/// from a volatile cache. Call <see cref="ShouldPutBlobAsync(IoHash)"/> to determine if a blob needs to
+		/// be put to the store.
+		/// 
+		/// This method takes a timestamp indicating the required replication consistency. The method may return null
+		/// if an object was put to the store after this time but it has not replicated down yet.
+		/// </summary>
+		/// <param name="Hash">Hash of the blob</param>
+		/// <param name="Deadline">Deadline for returning a result. If zero, the storage provider should use a default value.</param>
+		/// <returns>The blob data. May return ReadOnlyMemory<byte>.Empty if it does not exist.</byte></returns>
+		Task<ReadOnlyMemory<byte>?> TryGetBlobAsync(IoHash Hash, DateTime Deadline = default(DateTime));
+
+		/// <summary>
+		/// Adds a blob to the store.
+		/// </summary>
+		/// <param name="Hash">Hash of the value. Used as a hint for possible sharding; the server will verify this hash and fail if it does not match.</param>
+		/// <param name="Value">The blob data</param>
+		/// <returns>Async task</returns>
+		Task PutBlobAsync(IoHash Hash, ReadOnlyMemory<byte> Value);
+
+		/// <summary>
+		/// Determines if the blob should be written to the store. The server may return true even if it exists in a store
+		/// at its discretion (eg. perhaps the client -> store link is faster than the store -> parent store link, or a GC
+		/// operation is in progress)
+		/// </summary>
+		/// <param name="Hash">Hash of the blob to write</param>
+		/// <returns>True if the blob should be written to the store</returns>
+		Task<bool> ShouldPutBlobAsync(IoHash Hash);
+
+		/// <summary>
+		/// Gets a reference from the storage layer
+		/// </summary>
+		/// <param name="Key">Key for the reference</param>
+		/// <param name="MaxDrift">Required replication consistency</param>
+		/// <returns>Hash of the result</returns>
+		Task<IoHash?> GetRefAsync(IoHash Key, TimeSpan MaxDrift = default(TimeSpan));
+
+		/// <summary>
+		/// Sets the value associated with a key. Writes are propagated as a "last writer wins" policy.
+		/// </summary>
+		/// <param name="Key">The key to set</param>
+		/// <param name="Value">Value for the key. If null, the key is deleted.</param>
+		/// <returns>Async task</returns>
+		Task SetRefAsync(IoHash Key, IoHash? Value);
+
+		/// <summary>
+		/// Updates the last modified timestamp of the given key. Possibly equivalent to doing a SET() on the same key/value pair, but does not 
+		/// need to force a value write to propagate up the tree. Mainly used for GC.
+		/// </summary>
+		/// <param name="Key">The key to touch</param>
+		/// <returns>Async task</returns>
+		Task TouchRefAsync(IoHash Key);
+	}
+
+	/// <summary>
+	/// Extension methods for the storage service
+	/// </summary>
+	static class StorageServiceExtensions
+	{
+		/// <summary>
+		/// Gets a blob by hash, throwing an exception if it does not exist.
+		/// </summary>
+		/// <param name="StorageService">The storage provider</param>
+		/// <param name="Hash">Hash of the blob to retrieve</param>
+		/// <param name="Deadline">Deadline for retreving the data</param>
+		/// <returns>Memory for the blob</returns>
+		public static async Task<ReadOnlyMemory<byte>> GetBlobAsync(this IStorageService StorageService, IoHash Hash, DateTime Deadline = default(DateTime))
+		{
+			ReadOnlyMemory<byte>? Blob = await StorageService.TryGetBlobAsync(Hash, Deadline);
+			if (Blob == null)
+			{
+				throw new KeyNotFoundException($"Unable to read key {Hash}");
+			}
+			return Blob.Value;
+		}
+
+		/// <summary>
+		/// Hashes a blob, then adds it to the storage service. If the hash is already known, prefer <see cref="IStorageService.PutBlobAsync(IoHash, ReadOnlyMemory{byte})"/>
+		/// </summary>
+		/// <param name="StorageService">The storage service</param>
+		/// <param name="Data">The data to add</param>
+		/// <returns>Hash of the data that was added</returns>
+		public static async Task<IoHash> PutBlobAsync(this IStorageService StorageService, ReadOnlyMemory<byte> Data)
+		{
+			IoHash Hash = IoHash.Compute(Data.Span);
+			await StorageService.PutBlobAsync(Hash, Data);
+			return Hash;
+		}
+	}
+}
diff --git a/Engine/Source/Programs/Horde/Horde.Build/Storage/Primitives/BlobRef.cs b/Engine/Source/Programs/Horde/Horde.Build/Storage/Primitives/BlobRef.cs
new file mode 100644
--- /dev/null
+++ b/Engine/Source/Programs/Horde/Horde.Build/Storage/Primitives/BlobRef.cs
@@ -0,0 +1,88 @@
+// Copyright Epic Games, Inc. All Rights Reserved.
+
+using EpicGames.Core;
+using System;
+using System.Collections.Generic;
+using System.Linq;
+using System.Threading.Tasks;
+
+namespace HordeServer.Storage.Primitives
+{
+	/// <summary>
+	/// Reference to a blob object
+	/// </summary>
+	struct BlobRef
+	{
+		/// <summary>
+		/// Type of the blob
+		/// </summary>
+		public IBlobType Type { get; }
+
+		/// <summary>
+		/// The blob hash
+		/// </summary>
+		public IoHash Hash { get; }
+
+		/// <summary>
+		/// Constructor
+		/// </summary>
+		/// <param name="Type"></param>
+		/// <param name="Hash"></param>
+		public BlobRef(IBlobType Type, IoHash Hash)
+		{
+			this.Type = Type;
+			this.Hash = Hash;
+		}
+
+		/// <inheritdoc/>
+		public override string ToString()
+		{
+			return $"{Type.Name}: {Hash}";
+		}
+	}
+
+	/// <summary>
+	/// Typed reference to a blob object
+	/// </summary>
+	/// <typeparam name="TBlob">The blob type</typeparam>
+	struct BlobRef<TBlob> where TBlob : IBlobType, new()
+	{
+		/// <summary>
+		/// Static type accessor
+		/// </summary>
+		public static TBlob Type => new TBlob();
+
+		/// <summary>
+		/// Hash of the referenced blob
+		/// </summary>
+		public IoHash Hash { get; }
+
+		/// <summary>
+		/// Constructor
+		/// </summary>
+		/// <param name="Hash">Hash of the blob</param>
+		public BlobRef(IoHash Hash)
+		{
+			this.Hash = Hash;
+		}
+
+		/// <summary>
+		/// Converts from a dynamically typed blob to this blob type
+		/// </summary>
+		/// <param name="Ref">The blob reference</param>
+		public static explicit operator BlobRef<TBlob>(BlobRef Ref)
+		{
+			_ = (TBlob)Ref.Type;
+			return new BlobRef<TBlob>(Ref.Hash);
+		}
+
+		/// <summary>
+		/// Converts from this blob to a dynamically typed blob
+		/// </summary>
+		/// <param name="Ref"></param>
+		public static implicit operator BlobRef(BlobRef<TBlob> Ref)
+		{
+			return new BlobRef(BlobRef<TBlob>.Type, Ref.Hash);
+		}
+	}
+}
diff --git a/Engine/Source/Programs/Horde/Horde.Build/Storage/Primitives/BlobType.cs b/Engine/Source/Programs/Horde/Horde.Build/Storage/Primitives/BlobType.cs
new file mode 100644
--- /dev/null
+++ b/Engine/Source/Programs/Horde/Horde.Build/Storage/Primitives/BlobType.cs
@@ -0,0 +1,42 @@
+// Copyright Epic Games, Inc. All Rights Reserved.
+
+using System;
+using System.Collections.Generic;
+using System.Linq;
+using System.Threading.Tasks;
+
+namespace HordeServer.Storage.Primitives
+{
+	/// <summary>
+	/// Interface for a particular blob type
+	/// </summary>
+	interface IBlobType
+	{
+		/// <summary>
+		/// Name of the type
+		/// </summary>
+		public string Name { get; }
+
+		/// <summary>
+		/// Gets references from the given blob
+		/// </summary>
+		/// <param name="Data">The blob data</param>
+		/// <returns>References</returns>
+		public IEnumerable<BlobRef> GetRefs(ReadOnlyMemory<byte> Data);
+	}
+
+	/// <summary>
+	/// Leaf blob type definition
+	/// </summary>
+	sealed class LeafBlob : IBlobType
+	{
+		/// <inheritdoc/>
+		public string Name => "Leaf";
+
+		/// <inheritdoc/>
+		public IEnumerable<BlobRef> GetRefs(ReadOnlyMemory<byte> Data)
+		{
+			yield break;
+		}
+	}
+}
diff --git a/Engine/Source/Programs/Horde/Horde.Build/Storage/Primitives/KeyValueItem.cs b/Engine/Source/Programs/Horde/Horde.Build/Storage/Primitives/KeyValueItem.cs
new file mode 100644
--- /dev/null
+++ b/Engine/Source/Programs/Horde/Horde.Build/Storage/Primitives/KeyValueItem.cs
@@ -0,0 +1,104 @@
+// Copyright Epic Games, Inc. All Rights Reserved.
+
+using EpicGames.Core;
+using System;
+using System.Collections.Generic;
+using System.Linq;
+using System.Threading.Tasks;
+
+namespace HordeServer.Storage.Primitives
+{
+	/// <summary>
+	/// A key/value pair for insertion into a KeyValueTree
+	/// </summary>
+	struct KeyValueItem
+	{
+		/// <summary>
+		/// The key
+		/// </summary>
+		public IoHash Key { get; }
+
+		/// <summary>
+		/// Block of memory containing the value for this item
+		/// </summary>
+		public ReadOnlyMemory<byte> Value { get; }
+
+		/// <summary>
+		/// Constructor
+		/// </summary>
+		/// <param name="Key"></param>
+		/// <param name="Value"></param>
+		public KeyValueItem(IoHash Key, ReadOnlyMemory<byte> Value)
+		{
+			this.Key = Key;
+			this.Value = Value;
+		}
+	}
+
+	class KeyValueItemReader
+	{
+		ReadOnlyMemory<byte> Keys;
+		ReadOnlyMemory<byte> Values;
+		int ValueSize;
+
+		public KeyValueItemReader(KeyValueTree Tree, ReadOnlyKeyValueNode Node)
+		{
+			this.Keys = Node.Keys;
+			this.Values = Node.Values;
+			this.ValueSize = Node.IsLeafNode ? Tree.LeafValueSize : IoHash.NumBytes;
+		}
+
+		public bool IsEmpty => Keys.Length == 0;
+
+		public void MoveNext()
+		{
+			Keys = Keys.Slice(IoHash.NumBytes);
+			Values = Values.Slice(ValueSize);
+		}
+
+		public void CopyTo(KeyValueItemWriter Writer)
+		{
+			Writer.WriteItems(Keys, Values);
+
+			Keys = ReadOnlyMemory<byte>.Empty;
+			Values = ReadOnlyMemory<byte>.Empty;
+		}
+
+		public void CopyTo(KeyValueItemWriter Writer, int Count)
+		{
+			ReadOnlyMemory<byte> CopyKeys = Keys.Slice(0, Count * IoHash.NumBytes);
+			ReadOnlyMemory<byte> CopyValues = Values.Slice(0, Count * ValueSize);
+
+			Writer.WriteItems(CopyKeys, CopyValues);
+
+			Keys = Keys.Slice(CopyKeys.Length);
+			Values = Values.Slice(CopyValues.Length);
+		}
+	}
+
+	class KeyValueItemWriter
+	{
+		public Memory<byte> Keys { get; private set; }
+		public Memory<byte> Values { get; private set; }
+
+		public KeyValueItemWriter(KeyValueNode Node)
+		{
+			this.Keys = Node.Keys;
+			this.Values = Node.Values;
+		}
+
+		public void WriteItem(IoHash Key, ReadOnlyMemory<byte> Value)
+		{
+			WriteItems(Key.Memory, Value);
+		}
+
+		public void WriteItems(ReadOnlyMemory<byte> NewKeys, ReadOnlyMemory<byte> NewValues)
+		{
+			NewKeys.CopyTo(Keys);
+			Keys = Keys.Slice(NewKeys.Length);
+
+			NewValues.CopyTo(Values);
+			Values = Values.Slice(NewValues.Length);
+		}
+	}
+}
diff --git a/Engine/Source/Programs/Horde/Horde.Build/Storage/Primitives/KeyValueNode.cs b/Engine/Source/Programs/Horde/Horde.Build/Storage/Primitives/KeyValueNode.cs
new file mode 100644
--- /dev/null
+++ b/Engine/Source/Programs/Horde/Horde.Build/Storage/Primitives/KeyValueNode.cs
@@ -0,0 +1,162 @@
+// Copyright Epic Games, Inc. All Rights Reserved.
+
+using EpicGames.Core;
+using System;
+using System.Buffers.Binary;
+using System.Collections;
+using System.Collections.Generic;
+using System.Linq;
+using System.Threading.Tasks;
+
+namespace HordeServer.Storage.Primitives
+{
+	/// <summary>
+	/// A node in a key value tree
+	/// </summary>
+	class KeyValueNode
+	{
+		/// <summary>
+		/// The underlying memory storing the node data
+		/// </summary>
+		public Memory<byte> Memory { get; }
+
+		/// <summary>
+		/// Number of child items from this node
+		/// </summary>
+		public int NumItems { get; }
+
+		/// <summary>
+		/// Number of items if this node is merged with its children. This will be equal to NumItems if this is the child of a parent-child relation.
+		/// Note that this is different from the total number of leaf nodes under this tree; it is only used for splitting/merging layers.
+		/// </summary>
+		public int NumItemsIfMerged { get; }
+
+		/// <summary>
+		/// The number of valid bits in keys for this node
+		/// </summary>
+		public int NumKeyBits { get; }
+
+		/// <summary>
+		/// The number of valid bits in the key if it is merged with its children
+		/// </summary>
+		public int NumKeyBitsIfMerged { get; }
+
+		/// <summary>
+		/// Data for the array of keys
+		/// </summary>
+		public Memory<byte> Keys { get; }
+
+		/// <summary>
+		/// Data for the array of values
+		/// </summary>
+		public Memory<byte> Values { get; }
+
+		/// <summary>
+		/// Constructor
+		/// </summary>
+		/// <param name="NumItems"></param>
+		/// <param name="NumItemsIfMerged"></param>
+		/// <param name="NumKeyBits"></param>
+		/// <param name="NumKeyBitsIfMerged"></param>
+		/// <param name="ValueSize">Size of a value in this node</param>
+		public KeyValueNode(int NumItems, int NumItemsIfMerged, int NumKeyBits, int NumKeyBitsIfMerged, int ValueSize)
+		{
+			this.NumItems = NumItems;
+			this.NumItemsIfMerged = NumItemsIfMerged;
+			this.NumKeyBits = NumKeyBits;
+			this.NumKeyBitsIfMerged = NumKeyBitsIfMerged;
+
+			int KeyDataLength = IoHash.NumBytes * NumItems;
+			int ValueDataLength = ValueSize * NumItems;
+			Memory = new byte[(sizeof(int) * 4) + KeyDataLength + ValueDataLength];
+
+			Memory<byte> Output = Memory;
+
+			BinaryPrimitives.WriteInt32LittleEndian(Output.Span, NumItems);
+			Output = Output.Slice(sizeof(int));
+
+			BinaryPrimitives.WriteInt32LittleEndian(Output.Span, NumItemsIfMerged);
+			Output = Output.Slice(sizeof(int));
+
+			BinaryPrimitives.WriteInt32LittleEndian(Output.Span, NumKeyBits);
+			Output = Output.Slice(sizeof(int));
+
+			BinaryPrimitives.WriteInt32LittleEndian(Output.Span, NumKeyBitsIfMerged);
+			Output = Output.Slice(sizeof(int));
+
+			Keys = Output.Slice(0, KeyDataLength);
+			Output = Output.Slice(KeyDataLength);
+
+			Values = Output;
+		}
+	}
+
+	/// <summary>
+	/// Read-only version of <see cref="KeyValueNode"/>
+	/// </summary>
+	class ReadOnlyKeyValueNode
+	{
+		/// <inheritdoc cref="KeyValueNode.Memory"/>
+		public ReadOnlyMemory<byte> Memory { get; }
+
+		/// <inheritdoc cref="KeyValueNode.NumItems"/>
+		public int NumItems { get; }
+
+		/// <inheritdoc cref="KeyValueNode.NumItemsIfMerged"/>
+		public int NumItemsIfMerged { get; }
+
+		/// <inheritdoc cref="KeyValueNode.NumKeyBits"/>
+		public int NumKeyBits { get; }
+
+		/// <inheritdoc cref="KeyValueNode.NumKeyBitsIfMerged"/>
+		public int NumKeyBitsIfMerged { get; }
+
+		/// <inheritdoc cref="KeyValueNode.Keys"/>
+		public ReadOnlyMemory<byte> Keys { get; }
+
+		/// <inheritdoc cref="KeyValueNode.Values"/>
+		public ReadOnlyMemory<byte> Values { get; }
+
+		/// <summary>
+		/// Whether this node is a leaf node
+		/// </summary>
+		public bool IsLeafNode => NumKeyBits == IoHash.NumBits;
+
+		/// <summary>
+		/// Constructor
+		/// </summary>
+		/// <param name="Memory">Memory backing this node</param>
+		public ReadOnlyKeyValueNode(ReadOnlyMemory<byte> Memory)
+		{
+			this.Memory = Memory;
+
+			ReadOnlyMemory<byte> Input = Memory;
+
+			NumItems = BinaryPrimitives.ReadInt32LittleEndian(Input.Span);
+			Input = Input.Slice(sizeof(int));
+
+			NumItemsIfMerged = BinaryPrimitives.ReadInt32LittleEndian(Input.Span);
+			Input = Input.Slice(sizeof(int));
+
+			NumKeyBits = BinaryPrimitives.ReadInt32LittleEndian(Input.Span);
+			Input = Input.Slice(sizeof(int));
+
+			NumKeyBitsIfMerged = BinaryPrimitives.ReadInt32LittleEndian(Input.Span);
+			Input = Input.Slice(sizeof(int));
+
+			Keys = Input.Slice(0, NumItems * IoHash.NumBytes);
+			Input = Input.Slice(Keys.Length);
+
+			Values = Input;
+		}
+
+		/// <summary>
+		/// Implicit conversion operator from a KeyValueNode
+		/// </summary>
+		/// <param name="Other"></param>
+		public static implicit operator ReadOnlyKeyValueNode(KeyValueNode Other)
+		{
+			return new ReadOnlyKeyValueNode(Other.Memory);
+		}
+	}
+}
