Horde: Rename the HordeServer project to Horde.Build.

diff --git a/Engine/Source/Programs/Horde/Horde.Build/Storage/Primitives/KeyValueTree.cs b/Engine/Source/Programs/Horde/Horde.Build/Storage/Primitives/KeyValueTree.cs
new file mode 100644
--- /dev/null
+++ b/Engine/Source/Programs/Horde/Horde.Build/Storage/Primitives/KeyValueTree.cs
@@ -0,0 +1,581 @@
+// Copyright Epic Games, Inc. All Rights Reserved.
+
+using EpicGames.Core;
+using System;
+using System.Buffers.Binary;
+using System.Collections.Generic;
+using System.Diagnostics;
+using System.Linq;
+using System.Threading.Tasks;
+
+namespace HordeServer.Storage.Primitives
+{
+	/// <summary>
+	/// Base type definition for key/value trees. Handles both internal and leaf nodes.
+	/// </summary>
+	abstract class KeyValueTree : IBlobType
+	{
+		/// <summary>
+		/// Default maximum node size
+		/// </summary>
+		const int DefaultMaxNodeSize = 1024 * 1024;
+
+		/// <summary>
+		/// The key type
+		/// </summary>
+		public IBlobType KeyType;
+
+		/// <summary>
+		/// Length of each value
+		/// </summary>
+		public int LeafValueSize { get; }
+
+		/// <summary>
+		/// Offset within each value item of blob references, and their types
+		/// </summary>
+		public (int, IBlobType)[] LeafValueRefTypes { get; }
+
+		/// <summary>
+		/// Number of items for an internal node
+		/// </summary>
+		public int MaxItemsPerInternalNode { get; }
+
+		/// <summary>
+		/// Number of items for a leaf node
+		/// </summary>
+		public int MaxItemsPerLeafNode { get; }
+
+		/// <inheritdoc/>
+		public string Name => "KeyValueTree";
+
+		/// <summary>
+		/// Constructor
+		/// </summary>
+		/// <param name="KeyType">Type of key</param>
+		/// <param name="LeafValueSize">Length of each value</param>
+		/// <param name="LeafValueRefTypes">References from each value</param>
+		public KeyValueTree(IBlobType KeyType, int LeafValueSize, params (int, IBlobType)[] LeafValueRefTypes)
+			: this(KeyType, LeafValueSize, LeafValueRefTypes, DefaultMaxNodeSize / (IoHash.NumBytes * 2), DefaultMaxNodeSize / (IoHash.NumBytes + LeafValueSize))
+		{
+		}
+
+		/// <summary>
+		/// Constructor
+		/// </summary>
+		/// <param name="KeyType">Type of key</param>
+		/// <param name="LeafValueSize">Length of each value</param>
+		/// <param name="LeafValueRefTypes">References from each value</param>
+		/// <param name="MaxItemsPerInternalNode">Maximum number of items to store in each internal node</param>
+		/// <param name="MaxItemsPerLeafNode">Maximum number of items to store in each leaf node</param>
+		public KeyValueTree(IBlobType KeyType, int LeafValueSize, (int, IBlobType)[] LeafValueRefTypes, int MaxItemsPerInternalNode, int MaxItemsPerLeafNode)
+		{
+			this.KeyType = KeyType;
+			this.LeafValueSize = LeafValueSize;
+			this.LeafValueRefTypes = LeafValueRefTypes;
+			this.MaxItemsPerInternalNode = MaxItemsPerInternalNode;
+			this.MaxItemsPerLeafNode = MaxItemsPerLeafNode;
+		}
+
+		/// <summary>
+		/// Get references from the given data
+		/// </summary>
+		/// <param name="Data">Blob of data</param>
+		/// <returns>Sequence of blob references</returns>
+		public IEnumerable<BlobRef> GetRefs(ReadOnlyMemory<byte> Data)
+		{
+			ReadOnlyKeyValueNode Node = new ReadOnlyKeyValueNode(Data);
+			if (Node.NumKeyBits == IoHash.NumBits)
+			{
+				// Leaf node: keys are references to other blobs, values follow type definition
+				ReadOnlyMemory<byte> Keys = Node.Keys;
+				while (!Keys.IsEmpty)
+				{
+					yield return new BlobRef(KeyType, new IoHash(Keys.Slice(0, IoHash.NumBytes)));
+					Keys = Keys.Slice(IoHash.NumBytes);
+				}
+
+				ReadOnlyMemory<byte> Values = Node.Values;
+				while (!Values.IsEmpty)
+				{
+					foreach ((int Offset, IBlobType Type) in LeafValueRefTypes)
+					{
+						yield return new BlobRef(Type, new IoHash(Values.Slice(Offset, IoHash.NumBytes)));
+					}
+					Values = Values.Slice(LeafValueSize);
+				}
+			}
+			else
+			{
+				// Internal node: keys are just prefixes, and not concrete references. Values are always just references to other trees.
+				ReadOnlyHashArray Values = new ReadOnlyHashArray(Node.Values);
+				foreach(IoHash Value in Values)
+				{
+					yield return new BlobRef(this, Value);
+				}
+			}
+		}
+	}
+
+	/// <summary>
+	/// Base type definition for key/value trees with a static key type
+	/// </summary>
+	abstract class KeyValueTree<TKey> : KeyValueTree where TKey : IBlobType, new()
+	{
+		/// <summary>
+		/// Constructor
+		/// </summary>
+		/// <param name="ValueSize"></param>
+		/// <param name="ValueRefTypes"></param>
+		public KeyValueTree(int ValueSize, params (int, IBlobType)[] ValueRefTypes)
+			: base(BlobRef<TKey>.Type, ValueSize, ValueRefTypes)
+		{
+		}
+	}
+
+	/// <summary>
+	/// Base type definition for key/value trees with a static key type
+	/// </summary>
+	sealed class KeyValueTree<TKey, TValue> : KeyValueTree where TKey : IBlobType, new() where TValue : IBlobType, new()
+	{
+		/// <summary>
+		/// Constructor
+		/// </summary>
+		public KeyValueTree()
+			: base(BlobRef<TKey>.Type, IoHash.NumBytes, (0, BlobRef<TValue>.Type))
+		{
+		}
+	}
+
+	/// <summary>
+	/// Extension methods for manipulating key/value trees
+	/// </summary>
+	static class KeyValueTreeExtensions
+	{
+		/// <summary>
+		/// Creates a new, empty key value tree
+		/// </summary>
+		/// <typeparam name="TTree">The tree type</typeparam>
+		/// <param name="Storage">Storage service</param>
+		/// <returns>Reference to the new tree</returns>
+		public static async Task<BlobRef<TTree>> CreateKeyValueTreeAsync<TTree>(this IStorageService Storage) where TTree : KeyValueTree, new()
+		{
+			KeyValueNode Node = new KeyValueNode(0, 0, IoHash.NumBits, IoHash.NumBits, 0);
+			return new BlobRef<TTree>(await Storage.PutKeyValueNodeAsync(Node));
+		}
+
+		public static async Task<BlobRef<TValue>?> SearchKeyValueTreeAsync<TKey, TValue>(this IStorageService Storage, BlobRef<KeyValueTree<TKey, TValue>> Tree, IoHash FindHash) where TKey : IBlobType, new() where TValue : IBlobType, new()
+		{
+			ReadOnlyMemory<byte>? Result = await SearchKeyValueTreeInternalAsync(Storage, Tree, FindHash);
+			if(Result == null)
+			{
+				return null;
+			}
+			return new BlobRef<TValue>(new IoHash(Result.Value));
+		}
+
+		public static Task<ReadOnlyMemory<byte>?> SearchKeyValueTreeAsync<TTree>(this IStorageService Storage, BlobRef<TTree> Tree, IoHash FindHash) where TTree : KeyValueTree, new()
+		{
+			return SearchKeyValueTreeInternalAsync(Storage, Tree, FindHash);
+		}
+
+		static async Task<ReadOnlyMemory<byte>?> SearchKeyValueTreeInternalAsync<TTree>(this IStorageService Storage, BlobRef<TTree> Tree, IoHash FindHash) where TTree : KeyValueTree, new()
+		{
+			IoHash Hash = Tree.Hash;
+			for (; ; )
+			{
+				ReadOnlyKeyValueNode Node = await Storage.GetKeyValueNodeAsync(Hash);
+
+				IoHash MaskHash = MaskLeft(FindHash, Node.NumKeyBits);
+
+				ReadOnlyHashArray Keys = new ReadOnlyHashArray(Node.Keys);
+
+				int KeyIdx = Keys.BinarySearch(MaskHash);
+				if (KeyIdx < 0)
+				{
+					return null;
+				}
+				else if (Node.NumKeyBits == IoHash.NumBits)
+				{
+					return Node.Values.Slice(KeyIdx * BlobRef<TTree>.Type.LeafValueSize, BlobRef<TTree>.Type.LeafValueSize);
+				}
+				else
+				{
+					Hash = new ReadOnlyHashArray(Node.Values)[KeyIdx];
+				}
+			}
+		}
+
+		/// <summary>
+		/// Updates a key value tree with the given entries
+		/// </summary>
+		/// <typeparam name="TTree">The tree to update</typeparam>
+		/// <param name="Storage">Storage service instance</param>
+		/// <param name="Tree">The tree to update</param>
+		/// <param name="NewItems">New items to add</param>
+		/// <returns></returns>
+		public static async Task<BlobRef<TTree>> UpdateKeyValueTreeAsync<TTree>(this IStorageService Storage, BlobRef<TTree> Tree, IEnumerable<KeyValueItem> NewItems) where TTree : KeyValueTree, new()
+		{
+			IoHash NewRootHash = await UpdateKeyValueTreeAsync(Storage, BlobRef<TTree>.Type, Tree.Hash, NewItems);
+			return new BlobRef<TTree>(NewRootHash);
+		}
+
+		/// <summary>
+		/// Updates a key/value tree with a set of items
+		/// </summary>
+		/// <param name="Storage">The storage service</param>
+		/// <param name="Tree">The tree to update</param>
+		/// <param name="RootHash">Hash of the root node of the tree</param>
+		/// <param name="NewItems">New items to add</param>
+		/// <returns>Hash for the new tree root</returns>
+		static async Task<IoHash> UpdateKeyValueTreeAsync(this IStorageService Storage, KeyValueTree Tree, IoHash RootHash, IEnumerable<KeyValueItem> NewItems)
+		{
+			// Sort all the items and remove duplicates
+			ArraySegment<KeyValueItem> NewSortedItems = NewItems.OrderBy(x => x.Key).ToArray();
+			if (NewSortedItems.Count > 0)
+			{
+				int OutputCount = 1;
+				for (int InputCount = 1; InputCount < NewSortedItems.Count; InputCount++)
+				{
+					if (NewSortedItems[InputCount].Key != NewSortedItems[OutputCount - 1].Key)
+					{
+						NewSortedItems[OutputCount++] = NewSortedItems[InputCount];
+					}
+				}
+				NewSortedItems = NewSortedItems.Slice(0, OutputCount);
+			}
+
+			// Get the root node
+			ReadOnlyKeyValueNode RootNode;
+			if (RootHash == IoHash.Zero)
+			{
+				RootNode = new KeyValueNode(0, 0, IoHash.NumBits, IoHash.NumBits, Tree.LeafValueSize);
+			}
+			else
+			{
+				RootNode = await Storage.GetKeyValueNodeAsync(RootHash);
+			}
+
+			// Get the new node
+			KeyValueNode NewRootNode = await UpdateKeyValueNodeAsync(Storage, Tree, 0, RootNode, NewSortedItems);
+			return await Storage.PutKeyValueNodeAsync(NewRootNode);
+		}
+
+		/// <summary>
+		/// Update a key value tree node
+		/// </summary>
+		/// <param name="Storage">Storage service</param>
+		/// <param name="Tree">The tree definition</param>
+		/// <param name="ParentKeyBits">Number of bits in the parent node</param>
+		/// <param name="Node">The current node state</param>
+		/// <param name="NewItems">The items to modify</param>
+		/// <returns>Updated node</returns>
+		static async Task<KeyValueNode> UpdateKeyValueNodeAsync(IStorageService Storage, KeyValueTree Tree, int ParentKeyBits, ReadOnlyKeyValueNode Node, ArraySegment<KeyValueItem> NewItems)
+		{
+			// Create a list wrapper for the keys, so we can perform binary searches on it
+			ReadOnlyHashArray Keys = new ReadOnlyHashArray(Node.Keys);
+
+			// Keep track of the total number of child items underneath this node in order to update the merge count
+			int NewNumChildItems = Node.NumItemsIfMerged - Node.NumItems;
+
+			// Find the necessary updates to the node
+			List<(int KeyIdx, KeyValueItem? NewItem)> Updates = new List<(int, KeyValueItem?)>(NewItems.Count);
+			if (Node.NumKeyBits == IoHash.NumBits)
+			{
+				foreach (KeyValueItem NewItem in NewItems)
+				{
+					int KeyIdx = Keys.BinarySearch(NewItem.Key);
+					if (KeyIdx >= 0 || !NewItem.Value.IsEmpty)
+					{
+						Updates.Add((KeyIdx, NewItem));
+					}
+				}
+			}
+			else
+			{
+				ReadOnlyHashArray Children = new ReadOnlyHashArray(Node.Values);
+				for (int Idx = 0; Idx < NewItems.Count;)
+				{
+					IoHash KeyPrefix = MaskLeft(NewItems[Idx].Key, Node.NumKeyBits);
+
+					// Find all the operations which start with the same prefix
+					int MinIdx = Idx++;
+					while (Idx < NewItems.Count && StartsWith(NewItems[Idx].Key, KeyPrefix, Node.NumKeyBits))
+					{
+						Idx++;
+					}
+
+					// Check if it already exists within the tree
+					int KeyIdx = Keys.BinarySearch(KeyPrefix);
+
+					// Get the old child node
+					ReadOnlyKeyValueNode OldChildNode;
+					if (KeyIdx >= 0)
+					{
+						OldChildNode = await Storage.GetKeyValueNodeAsync(Children[KeyIdx]);
+					}
+					else
+					{
+						OldChildNode = CreateEmptyChildNode(Node.NumKeyBits);
+					}
+
+					// Create the new item
+					KeyValueNode NewChildNode = await UpdateKeyValueNodeAsync(Storage, Tree, Node.NumKeyBits, OldChildNode, NewItems.Slice(MinIdx, Idx - MinIdx));
+					if (NewChildNode.NumItems == 0)
+					{
+						// Only remove it if it already exists
+						if (KeyIdx >= 0)
+						{
+							Updates.Add((KeyIdx, null));
+						}
+					}
+					else
+					{
+						// Add the new item to the update list
+						IoHash NewChildHash = await Storage.PutKeyValueNodeAsync(NewChildNode);
+						KeyValueItem NewItem = new KeyValueItem(KeyPrefix, NewChildHash.Memory);
+						Updates.Add((KeyIdx, NewItem));
+
+						// Update the number of merged items
+						NewNumChildItems += NewChildNode.NumItemsIfMerged - OldChildNode.NumItemsIfMerged;
+					}
+				}
+			}
+
+			// Create a new buffer for the keys
+			int NewNumItems = Node.NumItems + Updates.Count(x => x.KeyIdx < 0) - Updates.Count(x => x.NewItem == null);
+			KeyValueNode NewNode = CreateNode(Tree, NewNumItems, NewNumItems + NewNumChildItems, Node.NumKeyBits, Node.NumKeyBitsIfMerged);
+
+			// Copy all the data across
+			int ReadIdx = 0;
+			KeyValueItemReader Reader = new KeyValueItemReader(Tree, Node);
+			KeyValueItemWriter Writer = new KeyValueItemWriter(NewNode);
+			foreach ((int KeyIdx, KeyValueItem? NewItem) in Updates)
+			{
+				// Copy all the data before the new item
+				if (KeyIdx < 0)
+				{
+					int Count = ~KeyIdx - ReadIdx;
+					Reader.CopyTo(Writer, Count);
+					ReadIdx += Count;
+				}
+				else
+				{
+					int Count = KeyIdx - ReadIdx;
+					Reader.CopyTo(Writer, Count);
+					Reader.MoveNext();
+					ReadIdx += Count + 1;
+				}
+
+				// Copy the new item
+				if (NewItem != null)
+				{
+					Writer.WriteItem(NewItem.Value.Key, NewItem.Value.Value);
+				}
+			}
+			Reader.CopyTo(Writer);
+
+			// Merge together layers of the tree
+			while (NewNode.NumKeyBitsIfMerged > NewNode.NumKeyBits && NewNode.NumItemsIfMerged < Tree.MaxItemsPerInternalNode)
+			{
+				NewNode = await MergeAsync(Storage, Tree, NewNode);
+			}
+
+			// If it's larger than a threshold, we need to split this node into two
+			while (NewNode.NumItems > ((NewNode.NumKeyBits == IoHash.NumBits)? Tree.MaxItemsPerLeafNode : Tree.MaxItemsPerInternalNode))
+			{
+				NewNode = await SplitAsync(Storage, Tree, NewNode, ParentKeyBits);
+			}
+
+			// Return the updated node
+			return NewNode;
+		}
+
+		/// <summary>
+		/// Merges a node iwth its children
+		/// </summary>
+		/// <param name="Storage"></param>
+		/// <param name="Tree"></param>
+		/// <param name="Node"></param>
+		/// <returns></returns>
+		public static async Task<KeyValueNode> MergeAsync(IStorageService Storage, KeyValueTree Tree, ReadOnlyKeyValueNode Node)
+		{
+			KeyValueNode NewNode = CreateNode(Tree, Node.NumItemsIfMerged, Node.NumItemsIfMerged, Node.NumKeyBitsIfMerged, Node.NumKeyBitsIfMerged);
+			KeyValueItemWriter Writer = new KeyValueItemWriter(NewNode);
+
+			ReadOnlyHashArray Values = new ReadOnlyHashArray(Node.Values);
+			for (int Idx = 0; Idx < Values.Count; Idx++)
+			{
+				ReadOnlyKeyValueNode ChildNode = await Storage.GetKeyValueNodeAsync(Values[Idx]);
+				Writer.WriteItems(ChildNode.Keys, ChildNode.Values);
+			}
+
+			return NewNode;
+		}
+
+		/// <summary>
+		/// Splits a node in the key value tree
+		/// </summary>
+		/// <param name="Storage"></param>
+		/// <param name="Tree">The tree definition</param>
+		/// <param name="Node"></param>
+		/// <param name="ParentKeyBits"></param>
+		/// <returns></returns>
+		public static async Task<KeyValueNode> SplitAsync(IStorageService Storage, KeyValueTree Tree, ReadOnlyKeyValueNode Node, int ParentKeyBits)
+		{
+			// Given an even distribution of keys across the hash space, keeping the tree balanced requires that we keep the 
+			// number of items in the parent node approximately equal to the average number of items in each child node.
+			// 
+			// If h is the number of bits in the hash, n is the number of bits in the parent node, and m is the number of 
+			// child nodes, it follows that:
+			// 
+			//    Items in parent node = 2 ^ n
+			//    Items in each child node = (2 ^ (h - n)) / m
+			// 
+			// Thus: n = (h - log2(m)) / 2
+			int MaxChildren = Node.IsLeafNode ? Tree.MaxItemsPerLeafNode : Tree.MaxItemsPerInternalNode;
+			int NewMaxKeyBits = ParentKeyBits + Math.Max(1, (int)(((Node.NumKeyBits - ParentKeyBits) - Math.Log2(MaxChildren)) * 0.5));
+
+			List<KeyValueItem> NewItems = new List<KeyValueItem>();
+
+			KeyValueItemReader Reader = new KeyValueItemReader(Tree, Node);
+
+			ReadOnlyHashArray Keys = new ReadOnlyHashArray(Node.Keys);
+			for (int KeyIdx = 0; KeyIdx < Keys.Count; )
+			{
+				// Mask off the first key prefix
+				IoHash KeyPrefix = MaskLeft(Keys[KeyIdx], NewMaxKeyBits);
+
+				// Count the number of keys with the same prefix
+				int Count = 1;
+				while (KeyIdx + Count < Keys.Count && StartsWith(Keys[KeyIdx + Count], KeyPrefix, NewMaxKeyBits))
+				{
+					Count++;
+				}
+
+				// Create the new child node
+				KeyValueNode NewChildNode = CreateNode(Tree, Count, Count, Node.NumKeyBits, Node.NumKeyBits);
+				Reader.CopyTo(new KeyValueItemWriter(NewChildNode), Count);
+				KeyIdx += Count;
+
+				// If the child node is still above the max number of children, split it again
+				while (NewChildNode.NumItems > MaxChildren)
+				{
+					NewChildNode = await SplitAsync(Storage, Tree, NewChildNode, NewMaxKeyBits);
+				}
+
+				// Finally, store the child node and add it to its new parent
+				IoHash NewChildHash = await Storage.PutKeyValueNodeAsync(NewChildNode);
+				NewItems.Add(new KeyValueItem(KeyPrefix, NewChildHash.Memory));
+			}
+
+			KeyValueNode NewNode = new KeyValueNode(NewItems.Count, Node.NumItems, NewMaxKeyBits, Node.NumKeyBits, IoHash.NumBytes);
+
+			KeyValueItemWriter Writer = new KeyValueItemWriter(NewNode);
+			foreach (KeyValueItem NewItem in NewItems)
+			{
+				Writer.WriteItem(NewItem.Key, NewItem.Value);
+			}
+
+			return NewNode;
+		}
+
+		public static KeyValueNode CreateNode(KeyValueTree Tree, int NumItems, int NumItemsIfMerged, int NumKeyBits, int NumKeyBitsIfMerged)
+		{
+			int NewValueSize = (NumKeyBits == IoHash.NumBits) ? Tree.LeafValueSize : IoHash.NumBytes;
+			return new KeyValueNode(NumItems, NumItemsIfMerged, NumKeyBits, NumKeyBitsIfMerged, NewValueSize);
+		}
+
+		public static ReadOnlyKeyValueNode CreateEmptyChildNode(int MaxKeyBits)
+		{
+			int ChildKeyBits = GetChildKeyLength(MaxKeyBits);
+			return new KeyValueNode(0, 0, ChildKeyBits, ChildKeyBits, 0);
+		}
+
+		public static int GetChildKeyLength(int KeyLength)
+		{
+			// Find the largest power of two boundary that the current key length is on
+			int KeyStep = 1;
+			while (KeyLength + KeyStep < IoHash.NumBits)
+			{
+				int NextKeyStep = KeyStep << 1;
+				if ((KeyLength & (NextKeyStep - 1)) != 0)
+				{
+					break;
+				}
+				KeyStep = NextKeyStep;
+			}
+			return KeyLength + KeyStep;
+		}
+
+		/// <summary>
+		/// Check if one hash starts with the same bits as another
+		/// </summary>
+		/// <param name="This"></param>
+		/// <param name="Other">The hash to test against</param>
+		/// <param name="Length">Number of bits to compare</param>
+		/// <returns>True if the hashes match, false otherwise</returns>
+		static bool StartsWith(IoHash This, IoHash Other, int Length)
+		{
+			ReadOnlySpan<byte> ThisSpan = This.Span;
+			ReadOnlySpan<byte> OtherSpan = Other.Span;
+
+			int Idx = 0;
+			int RemainingLength = Length;
+
+			for (; RemainingLength > 8; RemainingLength -= 8, Idx++)
+			{
+				if (ThisSpan[Idx] != OtherSpan[Idx])
+				{
+					return false;
+				}
+			}
+
+			int Mask = ~((1 << (8 - RemainingLength)) - 1);
+			return (ThisSpan[Idx] & Mask) == (OtherSpan[Idx] & Mask);
+		}
+
+		/// <summary>
+		/// Mask off the left number of bits from the hash
+		/// </summary>
+		/// <param name="Hash"></param>
+		/// <param name="NumBits">Number of bits to include</param>
+		/// <returns>Masked hash value</returns>
+		static IoHash MaskLeft(IoHash Hash, int NumBits)
+		{
+			ReadOnlySpan<byte> OldData = Hash.Span;
+			byte[] NewData = new byte[IoHash.NumBytes];
+
+			int Idx = 0;
+			int RemainingBits = NumBits;
+			for (; RemainingBits > 8; RemainingBits -= 8, Idx++)
+			{
+				NewData[Idx] = OldData[Idx];
+			}
+
+			NewData[Idx] = (byte)(OldData[Idx] & ~((1 << (8 - RemainingBits)) - 1));
+			return new IoHash(NewData);
+		}
+
+		/// <summary>
+		/// Adds a key/value tree node to a storage provider
+		/// </summary>
+		/// <param name="Provider"></param>
+		/// <param name="Node"></param>
+		/// <returns></returns>
+		public static async Task<IoHash> PutKeyValueNodeAsync(this IStorageService Provider, ReadOnlyKeyValueNode Node)
+		{
+			return await Provider.PutBlobAsync(Node.Memory);
+		}
+
+		/// <summary>
+		/// Gets a key/value tree node from a storage provider
+		/// </summary>
+		/// <param name="Provider"></param>
+		/// <param name="Hash"></param>
+		/// <returns></returns>
+		public static async Task<ReadOnlyKeyValueNode> GetKeyValueNodeAsync(this IStorageService Provider, IoHash Hash)
+		{
+			return new ReadOnlyKeyValueNode(await Provider.GetBlobAsync(Hash));
+		}
+	}
+}
diff --git a/Engine/Source/Programs/Horde/Horde.Build/Storage/Services/SimpleStorageService.cs b/Engine/Source/Programs/Horde/Horde.Build/Storage/Services/SimpleStorageService.cs
new file mode 100644
--- /dev/null
+++ b/Engine/Source/Programs/Horde/Horde.Build/Storage/Services/SimpleStorageService.cs
@@ -0,0 +1,120 @@
+// Copyright Epic Games, Inc. All Rights Reserved.
+
+using EpicGames.Core;
+using HordeServer.Services;
+using HordeServer.Storage.Primitives;
+using Microsoft.Extensions.Options;
+using System;
+using System.Collections.Generic;
+using System.IO;
+using System.Linq;
+using System.Reflection;
+using System.Threading.Tasks;
+
+namespace HordeServer.Storage.Services
+{
+	/// <summary>
+	/// Very basic implementation of <see cref="IStorageService"/>. Backed by a filesystem, and uses timestamps for GC. Data may be evicted at any time.
+	/// </summary>
+	class SimpleStorageService : IStorageService
+	{
+		/// <summary>
+		/// Path for storing blobs
+		/// </summary>
+		const string BlobsPath = "SimpleStorage/Blobs";
+
+		/// <summary>
+		/// Path for storing refs
+		/// </summary>
+		const string RefsPath = "SimpleStorage/Refs";
+
+		/// <summary>
+		/// The storage backend
+		/// </summary>
+		IStorageBackend Backend;
+
+		/// <summary>
+		/// Constructor
+		/// </summary>
+		/// <param name="Backend">The storage backend</param>
+		public SimpleStorageService(IStorageBackend<ArtifactCollection> Backend)
+		{
+			this.Backend = Backend;
+		}
+
+		/// <summary>
+		/// Gets the for a file with the given hash
+		/// </summary>
+		/// <param name="BaseDir">Base directory for storage</param>
+		/// <param name="Hash">Hash of the file</param>
+		/// <returns>Path to the file for this blob</returns>
+		static string GetItemPath(string BaseDir, IoHash Hash)
+		{
+			return $"{BaseDir}/{StringUtils.FormatHexString(Hash.Span.Slice(0, 1))}/{StringUtils.FormatHexString(Hash.Span.Slice(1, 1))}/{Hash}.dat";
+		}
+
+		/// <inheritdoc/>
+		public async Task<ReadOnlyMemory<byte>?> TryGetBlobAsync(IoHash Hash, DateTime Deadline = default)
+		{
+			string Path = GetItemPath(BlobsPath, Hash);
+			return await Backend.ReadBytesAsync(Path);
+		}
+
+		/// <inheritdoc/>
+		public async Task PutBlobAsync(IoHash Hash, ReadOnlyMemory<byte> Value)
+		{
+			IoHash RealHash = IoHash.Compute(Value.Span);
+			if(Hash != RealHash)
+			{
+				throw new InvalidDataException($"Hash for blob does not match. Expected {RealHash}, got {Hash}");
+			}
+
+			string Path = GetItemPath(BlobsPath, Hash);
+			if (!await Backend.ExistsAsync(Path))
+			{
+				await Backend.WriteBytesAsync(Path, Value);
+			}
+		}
+
+		/// <inheritdoc/>
+		public Task<bool> ShouldPutBlobAsync(IoHash Hash)
+		{
+			// Always return true to force timestamps to be updated
+			return Task.FromResult(true);
+		}
+
+		/// <inheritdoc/>
+		public async Task SetRefAsync(IoHash Key, IoHash? Value)
+		{
+			string Path = GetItemPath(RefsPath, Key);
+			if (Value == null)
+			{
+				await Backend.DeleteAsync(Path);
+			}
+			else
+			{
+				await Backend.WriteBytesAsync(Path, Value.Value.Memory);
+			}
+		}
+
+		/// <inheritdoc/>
+		public Task TouchRefAsync(IoHash Key)
+		{
+			return Task.CompletedTask;
+		}
+
+		/// <inheritdoc/>
+		public async Task<IoHash?> GetRefAsync(IoHash Key, TimeSpan MaxDrift = default)
+		{
+			ReadOnlyMemory<byte>? Memory = await Backend.ReadBytesAsync(GetItemPath(RefsPath, Key));
+			if (Memory == null)
+			{
+				return null;
+			}
+			else
+			{
+				return new IoHash(Memory.Value);
+			}
+		}
+	}
+}
diff --git a/Engine/Source/Programs/Horde/Horde.Build/Users/IUserClaim.cs b/Engine/Source/Programs/Horde/Horde.Build/Users/IUserClaim.cs
new file mode 100644
--- /dev/null
+++ b/Engine/Source/Programs/Horde/Horde.Build/Users/IUserClaim.cs
@@ -0,0 +1,74 @@
+// Copyright Epic Games, Inc. All Rights Reserved.
+
+using System;
+using System.Collections.Generic;
+using System.Linq;
+using System.Security.Claims;
+using System.Threading.Tasks;
+
+namespace HordeServer.Models
+{
+	/// <summary>
+	/// Claim for a user
+	/// </summary>
+	public interface IUserClaim
+	{
+		/// <summary>
+		/// Name of the claim
+		/// </summary>
+		public string Type { get; }
+
+		/// <summary>
+		/// Value of the claim
+		/// </summary>
+		public string Value { get; }
+	}
+
+	/// <summary>
+	/// New claim document
+	/// </summary>
+	public class UserClaim : IUserClaim
+	{
+		/// <inheritdoc/>
+		public string Type { get; set; }
+
+		/// <inheritdoc/>
+		public string Value { get; set; }
+
+		/// <summary>
+		/// Constructor
+		/// </summary>
+		/// <param name="Type"></param>
+		/// <param name="Value"></param>
+		public UserClaim(string Type, string Value)
+		{
+			this.Type = Type;
+			this.Value = Value;
+		}
+
+		/// <summary>
+		/// Constructor
+		/// </summary>
+		/// <param name="Other">Claim to construct from</param>
+		public UserClaim(IUserClaim Other)
+			: this(Other.Type, Other.Value)
+		{
+		}
+
+		/// <summary>
+		/// Constructs a UserClaim from a Claim object
+		/// </summary>
+		/// <param name="Claim">Claim object</param>
+		/// <returns></returns>
+		public static UserClaim FromClaim(Claim Claim)
+		{
+			return new UserClaim(Claim.Type, Claim.Value);
+		}
+
+		/// <summary>
+		/// Conversion operator from NET claims
+		/// </summary>
+		/// <param name="Claim"></param>
+		public static implicit operator UserClaim(Claim Claim) => FromClaim(Claim);
+	}
+}
diff --git a/Engine/Source/Programs/Horde/Horde.Build/Users/IUserCollection.cs b/Engine/Source/Programs/Horde/Horde.Build/Users/IUserCollection.cs
new file mode 100644
--- /dev/null
+++ b/Engine/Source/Programs/Horde/Horde.Build/Users/IUserCollection.cs
@@ -0,0 +1,122 @@
+// Copyright Epic Games, Inc. All Rights Reserved.
+
+using HordeServer.Models;
+using HordeServer.Services;
+using HordeServer.Utilities;
+using MongoDB.Bson;
+using System;
+using System.Collections.Generic;
+using System.Linq;
+using System.Security.Claims;
+using System.Threading.Tasks;
+
+namespace HordeServer.Collections
+{
+	using JobId = ObjectId<IJob>;
+	using UserId = ObjectId<IUser>;
+
+	/// <summary>
+	/// Manages user documents
+	/// </summary>
+	public interface IUserCollection
+	{
+		/// <summary>
+		/// Gets a user by unique id
+		/// </summary>
+		/// <param name="Id">Id of the user</param>
+		/// <returns>The user information</returns>
+		Task<IUser?> GetUserAsync(UserId Id);
+
+		/// <summary>
+		/// Gets a cached user by unique id
+		/// </summary>
+		/// <param name="Id">Id of the user</param>
+		/// <returns>The user information</returns>
+		Task<IUser?> GetCachedUserAsync(UserId Id);
+
+		/// <summary>
+		/// Gets a user by unique id
+		/// </summary>
+		/// <param name="Ids">Ids of the users</param>
+		/// <param name="NameRegex">Name regex to match for the users</param>
+		/// <param name="Index">Maximum number of results</param>
+		/// <param name="Count">Number of results to return</param>
+		/// <returns>The user information</returns>
+		Task<List<IUser>> FindUsersAsync(IEnumerable<UserId>? Ids = null, string? NameRegex = null, int? Index = null, int? Count = null);
+
+		/// <summary>
+		/// Gets a user by login
+		/// </summary>
+		/// <param name="Login">Login for the user</param>
+		/// <returns>The user information</returns>
+		Task<IUser?> FindUserByLoginAsync(string Login);
+
+		/// <summary>
+		/// Find or add a user with the given claims. Claims will be updated if the user exists.
+		/// </summary>
+		/// <param name="Login">Login id of the user</param>
+		/// <param name="Name">Full name of the user</param>
+		/// <param name="Email">Email address of the user</param>
+		/// <returns>The user document</returns>
+		Task<IUser> FindOrAddUserByLoginAsync(string Login, string? Name = null, string? Email = null);
+
+		/// <summary>
+		/// Gets the claims for a user
+		/// </summary>
+		/// <param name="UserId"></param>
+		/// <returns></returns>
+		Task<IUserClaims> GetClaimsAsync(UserId UserId);
+
+		/// <summary>
+		/// Update the claims for a user
+		/// </summary>
+		/// <param name="UserId"></param>
+		/// <param name="Claims"></param>
+		/// <returns></returns>
+		Task UpdateClaimsAsync(UserId UserId, IEnumerable<IUserClaim> Claims);
+
+		/// <summary>
+		/// Get settings for a user
+		/// </summary>
+		/// <param name="UserId"></param>
+		/// <returns></returns>
+		Task<IUserSettings> GetSettingsAsync(UserId UserId);
+
+		/// <summary>
+		/// Update a user
+		/// </summary>
+		/// <param name="UserId">The user to update</param>
+		/// <param name="EnableExperimentalFeatures"></param>
+		/// <param name="EnableIssueNotifications"></param>
+		/// <param name="DashboardSettings">Opaque settings object for the dashboard</param>
+		/// <param name="AddPinnedJobIds"></param>
+		/// <param name="RemovePinnedJobIds"></param>
+		/// <returns>Updated user object</returns>
+		Task UpdateSettingsAsync(UserId UserId, bool? EnableExperimentalFeatures = null, bool? EnableIssueNotifications = null, BsonValue? DashboardSettings = null, IEnumerable<JobId>? AddPinnedJobIds = null, IEnumerable<JobId>? RemovePinnedJobIds = null);
+	}
+
+	/// <summary>
+	/// Extension methods for the user collect
+	/// </summary>
+	static class UserCollectionExtensions
+	{
+		/// <summary>
+		/// Gets a particular user info from the collection
+		/// </summary>
+		/// <param name="UserCollection"></param>
+		/// <param name="Principal"></param>
+		/// <returns></returns>
+		public static Task<IUser?> GetUserAsync(this IUserCollection UserCollection, ClaimsPrincipal Principal)
+		{
+			UserId? UserId = Principal.GetUserId();
+			if (UserId != null)
+			{
+				return UserCollection.GetUserAsync(UserId.Value);
+			}
+			else
+			{
+				return Task.FromResult<IUser?>(null);
+			}
+		}
+	}
+}
diff --git a/Engine/Source/Programs/Horde/Horde.Build/Users/Impl/ServiceAccountCollection.cs b/Engine/Source/Programs/Horde/Horde.Build/Users/Impl/ServiceAccountCollection.cs
new file mode 100644
--- /dev/null
+++ b/Engine/Source/Programs/Horde/Horde.Build/Users/Impl/ServiceAccountCollection.cs
@@ -0,0 +1,149 @@
+// Copyright Epic Games, Inc. All Rights Reserved.
+
+using System;
+using HordeServer.Models;
+using HordeServer.Services;
+using MongoDB.Bson;
+using MongoDB.Bson.Serialization.Attributes;
+using MongoDB.Driver;
+using System.Collections.Generic;
+using System.Linq;
+using System.Threading.Tasks;
+
+namespace HordeServer.Collections.Impl
+{
+	/// <summary>
+	/// Collection of service account documents
+	/// </summary>
+	public class ServiceAccountCollection : IServiceAccountCollection
+	{
+		/// <summary>
+		/// Concrete implementation of IServiceAccount
+		/// </summary>
+		class ServiceAccountDocument : IServiceAccount
+		{
+			public const string ClaimSeparator = "###";
+			
+			[BsonRequired, BsonId]
+			public ObjectId Id { get; set; }
+
+			[BsonRequired]
+			public string SecretToken { get; set; } = "empty";
+
+			[BsonRequired]
+			public List<string> Claims { get; set; } = new List<string>();
+			
+			public bool Enabled { get; set; }
+			public string Description { get; set; } = "empty";
+
+			[BsonConstructor]
+			private ServiceAccountDocument()
+			{
+			}
+
+			public ServiceAccountDocument(ObjectId Id, string SecretToken, List<string> Claims, bool Enabled, string Description)
+			{
+				this.Id = Id;
+				this.SecretToken = SecretToken;
+				this.Claims = Claims;
+				this.Enabled = Enabled;
+				this.Description = Description;
+			}
+			
+			/// <inheritdoc/>
+			public void AddClaim(string Type, string Value)
+			{
+				Claims.Add(Type + ClaimSeparator + Value);
+			}
+
+			/// <inheritdoc/>
+			public IReadOnlyList<(string Type, string Value)> GetClaims()
+			{
+				return Claims.Select(x =>
+				{
+					string[] Split = x.Split(ClaimSeparator);
+					return (Split[0], Split[1]);
+				}).ToList();
+			}
+
+			protected bool Equals(ServiceAccountDocument other)
+			{
+				bool AreClaimsEqual = !Claims.Except(other.Claims).Any();
+
+                return Id.Equals(other.Id) && SecretToken == other.SecretToken && AreClaimsEqual && Enabled == other.Enabled && Description == other.Description;
+			}
+
+			public override bool Equals(object? obj)
+			{
+				if (ReferenceEquals(null, obj)) return false;
+				if (ReferenceEquals(this, obj)) return true;
+				if (obj.GetType() != this.GetType()) return false;
+				return Equals((ServiceAccountDocument) obj);
+			}
+
+			public override int GetHashCode()
+			{
+				return HashCode.Combine(Id, SecretToken, Claims, Enabled, Description);
+			}
+		}
+
+		/// <summary>
+		/// Collection of session documents
+		/// </summary>
+		private readonly IMongoCollection<ServiceAccountDocument> ServiceAccounts;
+
+		/// <summary>
+		/// Constructor
+		/// </summary>
+		/// <param name="DatabaseService">The database service</param>
+		public ServiceAccountCollection(DatabaseService DatabaseService)
+		{
+			ServiceAccounts = DatabaseService.GetCollection<ServiceAccountDocument>("ServiceAccounts");
+
+			if (!DatabaseService.ReadOnlyMode)
+			{
+				ServiceAccounts.Indexes.CreateOne(new CreateIndexModel<ServiceAccountDocument>(Builders<ServiceAccountDocument>.IndexKeys.Ascending(x => x.SecretToken)));
+			}
+		}
+
+		/// <inheritdoc/>
+		public async Task<IServiceAccount> AddAsync(string SecretToken, List<string> Claims, string Description)
+		{
+			ServiceAccountDocument NewSession = new ServiceAccountDocument(ObjectId.GenerateNewId(), SecretToken, Claims, true, Description);
+			await ServiceAccounts.InsertOneAsync(NewSession);
+			return NewSession;
+		}
+
+		/// <inheritdoc/>
+		public async Task<IServiceAccount?> GetAsync(ObjectId Id)
+		{
+			return await ServiceAccounts.Find(x => x.Id == Id).FirstOrDefaultAsync();
+		}
+
+		/// <inheritdoc/>
+		public async Task<IServiceAccount?> GetBySecretTokenAsync(string SecretToken)
+		{
+			return await ServiceAccounts.Find(x => x.SecretToken == SecretToken).FirstOrDefaultAsync();
+		}
+
+		/// <inheritdoc/>
+		public Task UpdateAsync(ObjectId Id, string? SecretToken, List<string>? Claims, bool? Enabled, string? Description)
+		{
+			UpdateDefinitionBuilder<ServiceAccountDocument> Update = Builders<ServiceAccountDocument>.Update;
+			List<UpdateDefinition<ServiceAccountDocument>> Updates = new List<UpdateDefinition<ServiceAccountDocument>>();
+			
+			if (SecretToken != null) Updates.Add(Update.Set(x => x.SecretToken, SecretToken));
+			if (Claims != null) Updates.Add(Update.Set(x => x.Claims, Claims));
+			if (Enabled != null) Updates.Add(Update.Set(x => x.Enabled, Enabled));
+			if (Description != null) Updates.Add(Update.Set(x => x.Description, Description));
+
+			return ServiceAccounts.FindOneAndUpdateAsync(x => x.Id == Id, Update.Combine(Updates));
+		}
+
+		/// <inheritdoc/>
+		public Task DeleteAsync(ObjectId SessionId)
+		{
+			return ServiceAccounts.DeleteOneAsync(x => x.Id == SessionId);
+		}
+	}
+}
diff --git a/Engine/Source/Programs/Horde/Horde.Build/Users/Impl/UserCollection.cs b/Engine/Source/Programs/Horde/Horde.Build/Users/Impl/UserCollection.cs
new file mode 100644
--- /dev/null
+++ b/Engine/Source/Programs/Horde/Horde.Build/Users/Impl/UserCollection.cs
@@ -0,0 +1,212 @@
+// Copyright Epic Games, Inc. All Rights Reserved.
+
+using HordeServer.Models;
+using HordeServer.Services;
+using HordeServer.Utilities;
+using MongoDB.Bson;
+using MongoDB.Bson.Serialization;
+using MongoDB.Bson.Serialization.Attributes;
+using MongoDB.Driver;
+using System;
+using System.Collections.Generic;
+using System.Linq;
+using System.Security.Claims;
+using System.Threading.Tasks;
+
+namespace HordeServer.Collections.Impl
+{
+	using JobId = ObjectId<IJob>;
+	using UserId = ObjectId<IUser>;
+
+	/// <summary>
+	/// Manages user documents
+	/// </summary>
+	class UserCollectionV1 : IUserCollection
+	{
+		class UserDocument : IUser, IUserClaims, IUserSettings
+		{
+			public UserId Id { get; set; }
+
+			public ClaimDocument PrimaryClaim { get; set; } = null!;
+			public List<ClaimDocument> Claims { get; set; } = new List<ClaimDocument>();
+
+			[BsonDefaultValue(false), BsonIgnoreIfDefault]
+			public bool EnableExperimentalFeatures { get; set; }
+
+			[BsonDefaultValue(false), BsonIgnoreIfDefault]
+			public bool EnableIssueNotifications { get; set; }
+
+			public BsonValue DashboardSettings { get; set; } = BsonNull.Value;
+			public List<JobId> PinnedJobIds { get; set; } = new List<JobId>();
+
+			string IUser.Name => Claims.FirstOrDefault(x => String.Equals(x.Type, "name", StringComparison.Ordinal))?.Value ?? PrimaryClaim.Value;
+			string IUser.Login => Claims.FirstOrDefault(x => String.Equals(x.Type, ClaimTypes.Name, StringComparison.Ordinal))?.Value ?? PrimaryClaim.Value;
+			string? IUser.Email => Claims.FirstOrDefault(x => String.Equals(x.Type, ClaimTypes.Email, StringComparison.Ordinal))?.Value;
+
+			UserId IUserClaims.UserId => Id;
+			IReadOnlyList<IUserClaim> IUserClaims.Claims => Claims;
+
+			UserId IUserSettings.UserId => Id;
+			IReadOnlyList<JobId> IUserSettings.PinnedJobIds => PinnedJobIds;
+		}
+
+		class ClaimDocument : IUserClaim
+		{
+			public string Type { get; set; }
+			public string Value { get; set; }
+
+			private ClaimDocument()
+			{
+				this.Type = null!;
+				this.Value = null!;
+			}
+
+			public ClaimDocument(string Type, string Value)
+			{
+				this.Type = Type;
+				this.Value = Value;
+			}
+
+			public ClaimDocument(IUserClaim Other)
+			{
+				this.Type = Other.Type;
+				this.Value = Other.Value;
+			}
+		}
+
+		IMongoCollection<UserDocument> Users;
+
+		/// <summary>
+		/// Static constructor
+		/// </summary>
+		static UserCollectionV1()
+		{
+			BsonSerializer.RegisterDiscriminatorConvention(typeof(UserDocument), NullDiscriminatorConvention.Instance);
+			BsonSerializer.RegisterDiscriminatorConvention(typeof(ClaimDocument), NullDiscriminatorConvention.Instance);
+		}
+
+		/// <summary>
+		/// Constructor
+		/// </summary>
+		/// <param name="DatabaseService"></param>
+		public UserCollectionV1(DatabaseService DatabaseService)
+		{
+			Users = DatabaseService.GetCollection<UserDocument>("Users");
+
+			if (!DatabaseService.ReadOnlyMode)
+			{
+				Users.Indexes.CreateOne(new CreateIndexModel<UserDocument>(Builders<UserDocument>.IndexKeys.Ascending(x => x.PrimaryClaim), new CreateIndexOptions { Unique = true }));
+			}
+		}
+
+		/// <inheritdoc/>
+		public async Task<IUser?> GetUserAsync(UserId Id)
+		{
+			return await Users.Find(x => x.Id == Id).FirstOrDefaultAsync();
+		}
+
+		/// <inheritdoc/>
+		public Task<IUser?> GetCachedUserAsync(UserId Id)
+		{
+			return GetUserAsync(Id);
+		}
+
+		/// <inheritdoc/>
+		public async Task<List<IUser>> FindUsersAsync(IEnumerable<UserId>? Ids, string? NameRegex, int? Index, int? Count)
+		{
+			FilterDefinition<UserDocument> Filter = Builders<UserDocument>.Filter.In(x => x.Id, Ids);
+			return await Users.Find(Filter).Range(Index, Count).ToListAsync<UserDocument, IUser>();
+		}
+
+		/// <inheritdoc/>
+		public async Task<IUser?> FindUserByLoginAsync(string Login)
+		{
+			ClaimDocument PrimaryClaim = new ClaimDocument(ClaimTypes.Name, Login);
+			return await Users.Find(x => x.PrimaryClaim == PrimaryClaim).FirstOrDefaultAsync();
+		}
+
+		/// <inheritdoc/>
+		public async Task<IUser> FindOrAddUserByLoginAsync(string Login, string? Name, string? Email)
+		{
+			ClaimDocument NewPrimaryClaim = new ClaimDocument(ClaimTypes.Name, Login);
+			UpdateDefinition<UserDocument> Update = Builders<UserDocument>.Update.SetOnInsert(x => x.Id, UserId.GenerateNewId());
+			return await Users.FindOneAndUpdateAsync<UserDocument>(x => x.PrimaryClaim == NewPrimaryClaim, Update, new FindOneAndUpdateOptions<UserDocument> { IsUpsert = true, ReturnDocument = ReturnDocument.After });
+		}
+
+		/// <inheritdoc/>
+		public async Task<IUserClaims> GetClaimsAsync(UserId UserId)
+		{
+			return await Users.Find(x => x.Id == UserId).FirstOrDefaultAsync() ?? new UserDocument { Id = UserId };
+		}
+
+		/// <inheritdoc/>
+		public async Task UpdateClaimsAsync(UserId UserId, IEnumerable<IUserClaim> Claims)
+		{
+			List<ClaimDocument> NewClaims = Claims.Select(x => new ClaimDocument(x)).ToList();
+			await Users.FindOneAndUpdateAsync(x => x.Id == UserId, Builders<UserDocument>.Update.Set(x => x.Claims, NewClaims));
+		}
+
+		/// <inheritdoc/>
+		public async Task<IUserSettings> GetSettingsAsync(UserId UserId)
+		{
+			return await Users.Find(x => x.Id == UserId).FirstOrDefaultAsync() ?? new UserDocument { Id = UserId };
+		}
+
+		/// <inheritdoc/>
+		public async Task UpdateSettingsAsync(UserId UserId, bool? EnableExperimentalFeatures, bool? EnableIssueNotifications, BsonValue? DashboardSettings = null, IEnumerable<JobId>? AddPinnedJobIds = null, IEnumerable<JobId>? RemovePinnedJobIds = null)
+		{
+			if (AddPinnedJobIds != null)
+			{
+				foreach (JobId PinnedJobId in AddPinnedJobIds)
+				{
+					FilterDefinition<UserDocument> Filter = Builders<UserDocument>.Filter.Eq(x => x.Id, UserId) & Builders<UserDocument>.Filter.AnyNin<JobId>(x => x.PinnedJobIds, new[] { PinnedJobId });
+					UpdateDefinition<UserDocument> Update = Builders<UserDocument>.Update.PushEach(x => x.PinnedJobIds, new[] { PinnedJobId }, -50);
+					await Users.UpdateOneAsync(Filter, Update);
+				}
+			}
+
+			List<UpdateDefinition<UserDocument>> Updates = new List<UpdateDefinition<UserDocument>>();
+			if (EnableExperimentalFeatures != null)
+			{
+				Updates.Add(Builders<UserDocument>.Update.SetOrUnsetNull(x => x.EnableExperimentalFeatures, EnableExperimentalFeatures));
+			}
+			if (EnableIssueNotifications != null)
+			{
+				Updates.Add(Builders<UserDocument>.Update.SetOrUnsetNull(x => x.EnableIssueNotifications, EnableIssueNotifications));
+			}
+			if (DashboardSettings != null)
+			{
+				Updates.Add(Builders<UserDocument>.Update.Set(x => x.DashboardSettings, DashboardSettings));
+			}
+			if (RemovePinnedJobIds != null && RemovePinnedJobIds.Any())
+			{
+				Updates.Add(Builders<UserDocument>.Update.PullAll(x => x.PinnedJobIds, RemovePinnedJobIds));
+			}
+			if (Updates.Count > 0)
+			{
+				await Users.UpdateOneAsync<UserDocument>(x => x.Id == UserId, Builders<UserDocument>.Update.Combine(Updates));
+			}
+		}
+
+		/// <summary>
+		/// Enumerate all the documents in this collection
+		/// </summary>
+		/// <returns></returns>
+		public async IAsyncEnumerable<(IUser, IUserClaims, IUserSettings)> EnumerateDocumentsAsync()
+		{
+			using (IAsyncCursor<UserDocument> Cursor = await Users.Find(FilterDefinition<UserDocument>.Empty).ToCursorAsync())
+			{
+				while (await Cursor.MoveNextAsync())
+				{
+					foreach (UserDocument Document in Cursor.Current)
+					{
+						if (Document.Claims.Count > 0)
+						{
+							yield return (Document, Document, Document);
+						}
+					}
+				}
+			}
+		}
+	}
+}
diff --git a/Engine/Source/Programs/Horde/Horde.Build/Users/Impl/UserCollectionV2.cs b/Engine/Source/Programs/Horde/Horde.Build/Users/Impl/UserCollectionV2.cs
new file mode 100644
--- /dev/null
+++ b/Engine/Source/Programs/Horde/Horde.Build/Users/Impl/UserCollectionV2.cs
@@ -0,0 +1,370 @@
+// Copyright Epic Games, Inc. All Rights Reserved.
+
+using HordeServer.Models;
+using HordeServer.Services;
+using HordeServer.Utilities;
+using Microsoft.Extensions.Caching.Memory;
+using Microsoft.Extensions.Logging;
+using MongoDB.Bson;
+using MongoDB.Bson.Serialization;
+using MongoDB.Bson.Serialization.Attributes;
+using MongoDB.Driver;
+using System;
+using System.Collections.Generic;
+using System.Linq;
+using System.Security.Claims;
+using System.Threading.Tasks;
+
+namespace HordeServer.Collections.Impl
+{
+	using JobId = ObjectId<IJob>;
+	using UserId = ObjectId<IUser>;
+
+	/// <summary>
+	/// Manages user documents
+	/// </summary>
+	class UserCollectionV2 : IUserCollection, IDisposable
+	{
+		class UserDocument : IUser
+		{
+			public UserId Id { get; set; }
+
+			public string Name { get; set; }
+			public string Login { get; set; }
+			public string LoginUpper { get; set; }
+			public string? Email { get; set; }
+			public string? EmailUpper { get; set; }
+
+			[BsonIgnoreIfDefault]
+			public bool? Hidden { get; set; }
+
+			[BsonConstructor]
+			private UserDocument()
+			{
+				Name = null!;
+				Login = null!;
+				LoginUpper = null!;
+			}
+
+			public UserDocument(IUser Other)
+				: this(Other.Id, Other.Name, Other.Login, Other.Email)
+			{
+			}
+
+			public UserDocument(UserId Id, string Name, string Login, string? Email)
+			{
+				this.Id = Id;
+				this.Name = Name;
+				this.Login = Login;
+				this.LoginUpper = Login.ToUpperInvariant();
+				this.Email = Email;
+				this.EmailUpper = Email?.ToUpperInvariant();
+			}
+		}
+
+		class UserClaimsDocument : IUserClaims
+		{
+			public UserId Id { get; set; }
+			public List<UserClaim> Claims { get; set; } = new List<UserClaim>();
+
+			UserId IUserClaims.UserId => Id;
+			IReadOnlyList<IUserClaim> IUserClaims.Claims => Claims;
+
+			[BsonConstructor]
+			private UserClaimsDocument()
+			{
+			}
+
+			public UserClaimsDocument(UserId Id)
+			{
+				this.Id = Id;
+			}
+
+			public UserClaimsDocument(IUserClaims Other)
+				: this(Other.UserId)
+			{
+				this.Claims.AddRange(Other.Claims.Select(x => new UserClaim(x)));
+			}
+		}
+
+		class UserSettingsDocument : IUserSettings
+		{
+			public UserId Id { get; set; }
+
+			[BsonDefaultValue(false), BsonIgnoreIfDefault]
+			public bool EnableExperimentalFeatures { get; set; }
+
+			[BsonDefaultValue(true), BsonIgnoreIfDefault]
+			public bool EnableIssueNotifications { get; set; } = true;
+
+			public BsonValue DashboardSettings { get; set; } = BsonNull.Value;
+			public List<JobId> PinnedJobIds { get; set; } = new List<JobId>();
+
+			UserId IUserSettings.UserId => Id;
+			IReadOnlyList<JobId> IUserSettings.PinnedJobIds => PinnedJobIds;
+
+			[BsonConstructor]
+			private UserSettingsDocument()
+			{
+			}
+
+			public UserSettingsDocument(UserId Id)
+			{
+				this.Id = Id;
+			}
+
+			public UserSettingsDocument(IUserSettings Other)
+				: this(Other.UserId)
+			{
+				this.EnableExperimentalFeatures = Other.EnableExperimentalFeatures;
+				this.EnableIssueNotifications = Other.EnableIssueNotifications;
+				this.DashboardSettings = Other.DashboardSettings;
+				this.PinnedJobIds = new List<JobId>(Other.PinnedJobIds);
+			}
+		}
+
+		class ClaimDocument : IUserClaim
+		{
+			public string Type { get; set; }
+			public string Value { get; set; }
+
+			private ClaimDocument()
+			{
+				this.Type = null!;
+				this.Value = null!;
+			}
+
+			public ClaimDocument(string Type, string Value)
+			{
+				this.Type = Type;
+				this.Value = Value;
+			}
+
+			public ClaimDocument(IUserClaim Other)
+			{
+				this.Type = Other.Type;
+				this.Value = Other.Value;
+			}
+		}
+
+		IMongoCollection<UserDocument> Users;
+		IMongoCollection<UserClaimsDocument> UserClaims;
+		IMongoCollection<UserSettingsDocument> UserSettings;
+
+		ILogger<UserCollectionV2> Logger;
+
+		MemoryCache UserCache;
+
+		/// <summary>
+		/// Static constructor
+		/// </summary>
+		static UserCollectionV2()
+		{
+			BsonSerializer.RegisterDiscriminatorConvention(typeof(UserDocument), NullDiscriminatorConvention.Instance);
+			BsonSerializer.RegisterDiscriminatorConvention(typeof(UserClaimsDocument), NullDiscriminatorConvention.Instance);
+			BsonSerializer.RegisterDiscriminatorConvention(typeof(UserSettingsDocument), NullDiscriminatorConvention.Instance);
+
+			BsonSerializer.RegisterDiscriminatorConvention(typeof(ClaimDocument), NullDiscriminatorConvention.Instance);
+		}
+
+		/// <summary>
+		/// Constructor
+		/// </summary>
+		/// <param name="DatabaseService"></param>
+		/// <param name="Logger"></param>
+		public UserCollectionV2(DatabaseService DatabaseService, ILogger<UserCollectionV2> Logger)
+		{
+			this.Logger = Logger;
+
+			Users = DatabaseService.GetCollection<UserDocument>("UsersV2");
+			UserClaims = DatabaseService.GetCollection<UserClaimsDocument>("UserClaimsV2");
+			UserSettings = DatabaseService.GetCollection<UserSettingsDocument>("UserSettingsV2");
+
+			if (!DatabaseService.ReadOnlyMode)
+			{
+				Users.Indexes.CreateOne(new CreateIndexModel<UserDocument>(Builders<UserDocument>.IndexKeys.Ascending(x => x.LoginUpper), new CreateIndexOptions { Unique = true }));
+				Users.Indexes.CreateOne(new CreateIndexModel<UserDocument>(Builders<UserDocument>.IndexKeys.Ascending(x => x.EmailUpper)));
+			}
+
+			MemoryCacheOptions Options = new MemoryCacheOptions();
+			UserCache = new MemoryCache(Options);
+		}
+
+		/// <inheritdoc/>
+		public void Dispose()
+		{
+			UserCache.Dispose();
+		}
+
+		/// <inheritdoc/>
+		public async Task<IUser?> GetUserAsync(UserId Id)
+		{
+			IUser? User = await Users.Find(x => x.Id == Id).FirstOrDefaultAsync();
+			using (ICacheEntry Entry = UserCache.CreateEntry(Id))
+			{
+				Entry.SetValue(User);
+				Entry.SetSlidingExpiration(TimeSpan.FromMinutes(30.0));
+			}
+			return User;
+		}
+
+		/// <inheritdoc/>
+		public Task<IUser?> GetCachedUserAsync(UserId Id)
+		{
+			IUser? User;
+			if (UserCache.TryGetValue(Id, out User))
+			{
+				return Task.FromResult(User);
+			}
+			else
+			{
+				return GetUserAsync(Id);
+			}
+		}
+
+		/// <inheritdoc/>
+		public async Task<List<IUser>> FindUsersAsync(IEnumerable<UserId>? Ids, string? NameRegex = null, int? Index = null, int? Count = null)
+		{
+			FilterDefinition<UserDocument> Filter = FilterDefinition<UserDocument>.Empty;
+			if (Ids != null)
+			{
+				Filter &= Builders<UserDocument>.Filter.In(x => x.Id, Ids);
+			}
+			
+			if (NameRegex != null)
+			{
+				BsonRegularExpression Regex = new BsonRegularExpression(NameRegex, "i");
+				Filter &= Builders<UserDocument>.Filter.Regex(x => x.Name, Regex);
+			}
+
+			Filter &= Builders<UserDocument>.Filter.Ne(x => x.Hidden, true);
+
+			return await Users.Find(Filter).Range(Index, Count ?? 100).ToListAsync<UserDocument, IUser>();
+		}
+
+		/// <inheritdoc/>
+		public async Task<IUser?> FindUserByLoginAsync(string Login)
+		{
+			string LoginUpper = Login.ToUpperInvariant();
+			return await Users.Find(x => x.LoginUpper == LoginUpper).FirstOrDefaultAsync();
+		}
+
+		/// <inheritdoc/>
+		public async Task<IUser> FindOrAddUserByLoginAsync(string Login, string? Name, string? Email)
+		{
+			UpdateDefinition<UserDocument> Update = Builders<UserDocument>.Update.SetOnInsert(x => x.Id, UserId.GenerateNewId()).SetOnInsert(x => x.Login, Login).Unset(x => x.Hidden);
+
+			if (Name == null)
+			{
+				Update = Update.SetOnInsert(x => x.Name, Name ?? Login);
+			}
+			else
+			{
+				Update = Update.Set(x => x.Name, Name);
+			}
+
+			if (Email != null)
+			{
+				Update = Update.Set(x => x.Email, Email).Set(x => x.EmailUpper, Email.ToUpperInvariant());
+			}
+
+			string LoginUpper = Login.ToUpperInvariant();
+			return await Users.FindOneAndUpdateAsync<UserDocument>(x => x.LoginUpper == LoginUpper, Update, new FindOneAndUpdateOptions<UserDocument, UserDocument> { IsUpsert = true, ReturnDocument = ReturnDocument.After });
+		}
+
+		/// <inheritdoc/>
+		public async Task<IUserClaims> GetClaimsAsync(UserId UserId)
+		{
+			IUserClaims? Claims = await UserClaims.Find(x => x.Id == UserId).FirstOrDefaultAsync();
+			if(Claims == null)
+			{
+				Claims = new UserClaimsDocument(UserId);
+			}
+			return Claims;
+		}
+
+		/// <inheritdoc/>
+		public async Task UpdateClaimsAsync(UserId UserId, IEnumerable<IUserClaim> Claims)
+		{
+			UserClaimsDocument NewDocument = new UserClaimsDocument(UserId);
+			NewDocument.Claims.AddRange(Claims.Select(x => new UserClaim(x)));
+			await UserClaims.ReplaceOneAsync(x => x.Id == UserId, NewDocument, new ReplaceOptions { IsUpsert = true });
+		}
+
+		/// <inheritdoc/>
+		public async Task<IUserSettings> GetSettingsAsync(UserId UserId)
+		{
+			IUserSettings? Settings = await UserSettings.Find(x => x.Id == UserId).FirstOrDefaultAsync();
+			if (Settings == null)
+			{
+				Settings = new UserSettingsDocument(UserId);
+			}
+			return Settings;
+		}
+
+		/// <inheritdoc/>
+		public async Task UpdateSettingsAsync(UserId UserId, bool? EnableExperimentalFeatures = null, bool? EnableIssueNotifications = null, BsonValue? DashboardSettings = null, IEnumerable<JobId>? AddPinnedJobIds = null, IEnumerable<JobId>? RemovePinnedJobIds = null)
+		{
+			List<UpdateDefinition<UserSettingsDocument>> Updates = new List<UpdateDefinition<UserSettingsDocument>>();
+			if (EnableExperimentalFeatures != null)
+			{
+				Updates.Add(Builders<UserSettingsDocument>.Update.SetOrUnsetNull(x => x.EnableExperimentalFeatures, EnableExperimentalFeatures));
+			}
+			if (EnableIssueNotifications != null)
+			{
+				if (EnableIssueNotifications.Value)
+				{
+					Updates.Add(Builders<UserSettingsDocument>.Update.Unset(x => x.EnableIssueNotifications));
+				}
+				else
+				{
+					Updates.Add(Builders<UserSettingsDocument>.Update.Set(x => x.EnableIssueNotifications, false));
+				}
+			}
+			if (DashboardSettings != null)
+			{
+				Updates.Add(Builders<UserSettingsDocument>.Update.Set(x => x.DashboardSettings, DashboardSettings));
+			}
+			if (AddPinnedJobIds != null && AddPinnedJobIds.Any())
+			{
+				Updates.Add(Builders<UserSettingsDocument>.Update.AddToSetEach(x => x.PinnedJobIds, AddPinnedJobIds));
+			}
+			if (RemovePinnedJobIds != null && RemovePinnedJobIds.Any())
+			{
+				Updates.Add(Builders<UserSettingsDocument>.Update.PullAll(x => x.PinnedJobIds, RemovePinnedJobIds));
+			}
+			if (Updates.Count > 0)
+			{
+				await UserSettings.UpdateOneAsync<UserSettingsDocument>(x => x.Id == UserId, Builders<UserSettingsDocument>.Update.Combine(Updates), new UpdateOptions { IsUpsert = true });
+			}
+		}
+
+		/// <summary>
+		/// Upgrade from V1 collection
+		/// </summary>
+		/// <param name="UserCollectionV1"></param>
+		/// <returns></returns>
+		public async Task ResaveDocumentsAsync(UserCollectionV1 UserCollectionV1)
+		{
+			await foreach ((IUser User, IUserClaims Claims, IUserSettings Settings) in UserCollectionV1.EnumerateDocumentsAsync())
+			{
+				try
+				{
+					await Users.ReplaceOneAsync(x => x.Id == User.Id, new UserDocument(User), new ReplaceOptions { IsUpsert = true });
+					await UserClaims.ReplaceOneAsync(x => x.Id == User.Id, new UserClaimsDocument(Claims), new ReplaceOptions { IsUpsert = true });
+					await UserSettings.ReplaceOneAsync(x => x.Id == User.Id, new UserSettingsDocument(Settings), new ReplaceOptions { IsUpsert = true });
+					Logger.LogDebug("Updated user {UserId}", User.Id);
+				}
+				catch (MongoWriteException Ex)
+				{
+					Logger.LogWarning(Ex, "Unable to resave user {UserId}", User.Id);
+				}
+
+				if(Settings.PinnedJobIds.Count > 0)
+				{
+					await UpdateSettingsAsync(User.Id, AddPinnedJobIds: Settings.PinnedJobIds);
+				}
+			}
+		}
+	}
+}
diff --git a/Engine/Source/Programs/Horde/Horde.Build/Utilities/ArtifactChunkStream.cs b/Engine/Source/Programs/Horde/Horde.Build/Utilities/ArtifactChunkStream.cs
new file mode 100644
--- /dev/null
+++ b/Engine/Source/Programs/Horde/Horde.Build/Utilities/ArtifactChunkStream.cs
@@ -0,0 +1,124 @@
+// Copyright Epic Games, Inc. All Rights Reserved.
+
+using Grpc.Core;
+using HordeCommon.Rpc;
+using System;
+using System.Collections.Generic;
+using System.IO;
+using System.Linq;
+using System.Threading;
+using System.Threading.Tasks;
+
+namespace HordeServer.Utilities
+{
+	/// <summary>
+	/// Stream that takes in rpc chunks from an upload and copies them?
+	/// </summary>
+	public class ArtifactChunkStream : Stream
+	{
+		/// <summary>
+		/// The grpc client reader. Should probably be templatized
+		/// </summary>
+		IAsyncStreamReader<UploadArtifactRequest> Reader;
+
+		/// <summary>
+		/// Position within the stream
+		/// </summary>
+		long StreamPosition;
+
+		/// <summary>
+		/// The length of the stream
+		/// </summary>
+		long StreamLength;
+
+		/// <summary>
+		/// The current request being read from
+		/// </summary>
+		UploadArtifactRequest? Request;
+
+		/// <summary>
+		/// Position within the current request
+		/// </summary>
+		int RequestPos;
+
+		/// <summary>
+		/// Constructor
+		/// </summary>
+		/// <param name="Reader">the grpc reader</param>
+		/// <param name="Length">filesize reported by the client</param>
+		public ArtifactChunkStream(IAsyncStreamReader<UploadArtifactRequest> Reader, long Length)
+		{
+			this.Reader = Reader;
+			this.StreamLength = Length;
+		}
+
+		/// <inheritdoc/>
+		public override long Position
+		{
+			get { return StreamPosition; }
+			set { throw new NotSupportedException(); }
+		}
+
+		/// <inheritdoc/>
+		public override long Length => StreamLength;
+
+		/// <inheritdoc/>
+		public override void SetLength(long value) => throw new NotSupportedException();
+
+		/// <inheritdoc/>
+		public override bool CanSeek => false;
+
+		/// <inheritdoc/>
+		public override long Seek(long Offset, SeekOrigin Origin) => throw new NotSupportedException();
+
+		/// <inheritdoc/>
+		public override bool CanWrite => false;
+
+		/// <inheritdoc/>
+		public override void Write(byte[] buffer, int offset, int count) => throw new NotSupportedException();
+
+		/// <inheritdoc/>
+		public override bool CanRead => true;
+
+		/// <inheritdoc/>
+		public override int Read(byte[] Buffer, int Offset, int Count)
+		{
+			return ReadAsync(Buffer, Offset, Count, CancellationToken.None).Result;
+		}
+
+		/// <inheritdoc/>
+		public override async Task<int> ReadAsync(byte[] Buffer, int Offset, int Count, CancellationToken CancellationToken)
+		{
+			int BytesRead = 0;
+			while (BytesRead < Count && StreamPosition < StreamLength)
+			{
+				if (Request == null || RequestPos == Request.Data.Length)
+				{
+					// Read the next request object
+					if (!await Reader.MoveNext())
+					{
+						throw new EndOfStreamException("Unexpected end of stream while reading artifact");
+					}
+
+					Request = Reader.Current;
+					RequestPos = 0;
+				}
+				else
+				{
+					// Copy data from the current request object
+					int NumBytesToCopy = Math.Min(Count - BytesRead, Request.Data.Length - RequestPos);
+					Request.Data.Span.Slice(RequestPos, NumBytesToCopy).CopyTo(Buffer.AsSpan(BytesRead, NumBytesToCopy));
+					RequestPos += NumBytesToCopy;
+					BytesRead += NumBytesToCopy;
+					StreamPosition += NumBytesToCopy;
+				}
+			}
+			return BytesRead;
+		}
+
+		/// <inheritdoc/>
+		public override void Flush()
+		{
+		}
+	}
+}
diff --git a/Engine/Source/Programs/Horde/Horde.Build/Utilities/BlockKit/ActionsBlock.cs b/Engine/Source/Programs/Horde/Horde.Build/Utilities/BlockKit/ActionsBlock.cs
new file mode 100644
--- /dev/null
+++ b/Engine/Source/Programs/Horde/Horde.Build/Utilities/BlockKit/ActionsBlock.cs
@@ -0,0 +1,159 @@
+// Copyright Epic Games, Inc. All Rights Reserved.
+
+using System;
+using System.Collections.Generic;
+using System.Globalization;
+using System.Text.Json;
+
+namespace HordeServer.Utilities.Slack.BlockKit
+{
+	/// <summary>
+	/// Wrapper for a BlockKit Actions block that contains interactive elements.
+	/// </summary>
+	public class ActionsBlock : BlockBase
+	{
+		/// <summary>
+		/// A collection of interactive elements.
+		/// </summary>
+		public List<ActionElement> Elements { get; } = new List<ActionElement>();
+
+		/// <inheritdoc/>
+		public override void Write(Utf8JsonWriter Writer, JsonSerializerOptions Options)
+		{
+			Writer.WriteStartObject();
+			Writer.WriteString("type", "actions");
+			Writer.WriteStartArray("elements");
+			foreach (ActionElement Element in Elements)
+			{
+				Element.Write(Writer, Options);
+			}
+			Writer.WriteEndArray();
+			Writer.WriteEndObject();
+		}
+
+		/// <summary>
+		/// Helper method to add a new <see cref="ActionButton"/> to the block.
+		/// </summary>
+		/// <param name="Text">The button text.</param>
+		/// <param name="URL">The URL to navigate to when clicked, if any.</param>
+		/// <param name="Value">
+		/// The value, if any, that should be passed in the interaction payload if the button is clicked.
+		/// </param>
+		/// <param name="ActionId">An optional identifier for handling interaction logic.</param>
+		/// <param name="Style">An option button style to apply to the button.</param>
+		public void AddButton(string Text, Uri? URL = null, string? Value = null, string? ActionId = null, ActionButton.ButtonStyle Style = ActionButton.ButtonStyle.Default)
+		{
+			Elements.Add(new ActionButton(Text, URL, Value, ActionId, Style));
+		}
+	}
+
+	/// <summary>
+	/// A base class for interaction Action elements to derive from.
+	/// </summary>
+	public abstract class ActionElement
+	{
+		/// <summary>
+		/// Writes the Json for an Action block element to the given <see cref="Utf8JsonWriter"/>.
+		/// </summary>
+		/// <param name="Writer"></param>
+		/// <param name="Options"></param>
+		public abstract void Write(Utf8JsonWriter Writer, JsonSerializerOptions Options);
+	}
+
+	/// <summary>
+	/// An button element that can be added to an <see cref="ActionsBlock"/>.
+	/// </summary>
+	public class ActionButton : ActionElement
+	{
+		/// <summary>
+		/// Determines the visual color scheme to use for the button.
+		/// </summary>
+		public enum ButtonStyle
+		{
+			/// <summary>
+			/// The default color scheme.
+			/// </summary>
+			Default,
+
+			/// <summary>
+			/// Used for affirmation/confirmation actions.
+			/// </summary>
+			Primary,
+
+			/// <summary>
+			/// Used when an action is destructive and cannot be undone.
+			/// </summary>
+			Danger,
+		}
+
+		/// <summary>
+		/// Construct a new Button action element.
+		/// </summary>
+		/// <param name="Text">The button text.</param>
+		/// <param name="URL">The URL to navigate to when clicked, if any.</param>
+		/// <param name="Value">
+		/// The value, if any, that should be passed in the interaction payload if the button is clicked.
+		/// </param>
+		/// <param name="ActionId">An optional identifier for handling interaction logic.</param>
+		/// <param name="Style">An option button style to apply to the button.</param>
+		public ActionButton(string Text, Uri? URL = null, string? Value = null, string? ActionId = null, ButtonStyle Style = ButtonStyle.Default)
+		{
+			this.Text = new TextObject(Text, IsMarkdown: false);
+			this.URL = URL;
+			this.Value = Value;
+			this.ActionId = ActionId;
+			this.Style = Style;
+		}
+
+		/// <summary>
+		/// The button text.
+		/// </summary>
+		public TextObject Text { get; } = new TextObject();
+
+		/// <summary>
+		/// The URL to navigate to when clicked, if any.
+		/// </summary>
+		public Uri? URL { get; set; }
+
+		/// <summary>
+		/// The value, if any, that should be passed in the interaction payload if the button is clicked.
+		/// </summary>
+		public string? Value { get; set; }
+
+		/// <summary>
+		/// An optional identifier for handling interaction logic.
+		/// </summary>
+		public string? ActionId { get; set; }
+
+		/// <summary>
+		/// Determines how the button appears in the message.
+		/// </summary>
+		public ButtonStyle Style { get; set; }
+
+		/// <inheritdoc/>
+		public override void Write(Utf8JsonWriter Writer, JsonSerializerOptions Options)
+		{
+			Writer.WriteStartObject();
+			Writer.WriteString("type", "button");
+			Text.Write(Writer);
+			if (URL != null)
+			{
+				Writer.WriteString("url", URL.AbsoluteUri);
+			}
+			if (!string.IsNullOrWhiteSpace(Value))
+			{
+				Writer.WriteString("value", Value);
+			}
+			if (!string.IsNullOrWhiteSpace(ActionId))
+			{
+				Writer.WriteString("action_id", ActionId);
+			}
+			if (Style != ButtonStyle.Default)
+			{
+				Writer.WriteString("style", Style.ToString().ToLower(CultureInfo.CurrentCulture));
+			}
+
+			Writer.WriteEndObject();
+		}
+	}
+}
diff --git a/Engine/Source/Programs/Horde/Horde.Build/Utilities/BlockKit/BlockBase.cs b/Engine/Source/Programs/Horde/Horde.Build/Utilities/BlockKit/BlockBase.cs
new file mode 100644
--- /dev/null
+++ b/Engine/Source/Programs/Horde/Horde.Build/Utilities/BlockKit/BlockBase.cs
@@ -0,0 +1,42 @@
+// Copyright Epic Games, Inc. All Rights Reserved.
+
+using System;
+using System.Text.Json;
+using System.Text.Json.Serialization;
+
+namespace HordeServer.Utilities.Slack.BlockKit
+{
+	/// <summary>
+	/// Converter for overriding default Json serialization/deserialization for BlockKit blocks.
+	/// </summary>
+	public class BlockBaseConverter : JsonConverter<BlockBase>
+	{
+		/// <summary>
+		/// Reads a BlockKit block Json blob and turns it into the appropriate BlockKit block object.
+		/// </summary>
+		/// <returns>A BlockKit block object of the appropriate type.</returns>
+		public override BlockBase Read(ref Utf8JsonReader reader, Type typeToConvert, JsonSerializerOptions options) => throw new NotImplementedException();
+
+		/// <summary>
+		/// Writes the Json for a BlockKit block to the given <see cref="Utf8JsonWriter"/>.
+		/// </summary>
+		public override void Write(Utf8JsonWriter Writer, BlockBase Value, JsonSerializerOptions Options)
+		{
+			Value.Write(Writer, Options);
+		}
+	}
+
+	/// <summary>
+	/// A base class for all BlockKit blocks to derive from.
+	/// </summary>
+	[JsonConverter(typeof(BlockBaseConverter))]
+	public abstract class BlockBase
+	{
+		/// <summary>
+		/// Write a BlockKit block to a <see cref="Utf8JsonWriter"/>
+		/// </summary>
+		/// <param name="Writer"></param>
+		/// <param name="Options"></param>
+		public abstract void Write(Utf8JsonWriter Writer, JsonSerializerOptions Options);
+	}
+}
diff --git a/Engine/Source/Programs/Horde/Horde.Build/Utilities/BlockKit/BlockKitAttachment.cs b/Engine/Source/Programs/Horde/Horde.Build/Utilities/BlockKit/BlockKitAttachment.cs
new file mode 100644
--- /dev/null
+++ b/Engine/Source/Programs/Horde/Horde.Build/Utilities/BlockKit/BlockKitAttachment.cs
@@ -0,0 +1,97 @@
+// Copyright Epic Games, Inc. All Rights Reserved.
+
+using System;
+using System.Collections.Generic;
+using System.Drawing;
+using System.Text.Json;
+using System.Text.Json.Serialization;
+
+namespace HordeServer.Utilities.Slack.BlockKit
+{
+	/// <summary>
+	/// Converter for overriding default Json serialization/deserialization for BlockKitAttachment objects.
+	/// </summary>
+	public class BlockKitAttachmentConverter : JsonConverter<BlockKitAttachment>
+	{
+		/// <summary>
+		/// Reads a BlockKitAttachment Json blob and turns it into the appropriate BlockKitAttachment object.
+		/// </summary>
+		/// <returns>A BlockKitAttachment object of the appropriate type.</returns>
+		public override BlockKitAttachment Read(ref Utf8JsonReader reader, Type typeToConvert, JsonSerializerOptions options) => throw new NotImplementedException();
+
+		/// <summary>
+		/// Writes the Json for a BlockKitAttachment to the given <see cref="Utf8JsonWriter"/>.
+		/// </summary>
+		public override void Write(Utf8JsonWriter Writer, BlockKitAttachment Value, JsonSerializerOptions Options)
+		{
+			Value.Write(Writer, Options);
+		}
+	}
+
+	/// <summary>
+	/// Static helpers for consistent color usage for <see cref="BlockKitMessage"/> messages.
+	/// </summary>
+	public static class BlockKitAttachmentColors
+	{
+		/// <summary>
+		/// The color to use for error messages.
+		/// </summary>
+		public static Color Error => Color.FromArgb(236, 76, 71);
+
+		/// <summary>
+		/// The color to use for warning messages.
+		/// </summary>
+		public static Color Warning => Color.FromArgb(247, 209, 84);
+
+		/// <summary>
+		/// The color to use for success messages.
+		/// </summary>
+		public static Color Success => Color.FromArgb(77, 181, 7);
+	}
+
+	/// <summary>
+	/// An attachment to a block kit message that is indented with a vertical, colored line on the left.
+	/// Should only be used for secondary information if possible.
+	/// </summary>
+	[JsonConverter(typeof(BlockKitAttachmentConverter))]
+	public class BlockKitAttachment
+	{
+		/// <summary>
+		/// Used for the "toast" notifications sent by Slack. If not set, and the Text property of the
+		/// base message isn't set, the notification will have a message saying it has no preview.
+		/// </summary>
+		public string? FallbackText { get; set; }
+
+		/// <summary>
+		/// The color of the line down the left side of the attachment.
+		/// </summary>
+		public Color Color { get; set; }
+
+		/// <summary>
+		/// A collection of BlockKit blocks the attachment is composed from.
+		/// </summary>
+		public List<BlockBase> Blocks { get; } = new List<BlockBase>();
+
+		/// <summary>
+		/// Write a BlockKit block to a <see cref="Utf8JsonWriter"/>
+		/// </summary>
+		/// <param name="Writer"></param>
+		/// <param name="Options"></param>
+		public void Write(Utf8JsonWriter Writer, JsonSerializerOptions Options)
+		{
+			Writer.WriteStartObject();
+			if (FallbackText != null)
+			{
+				Writer.WriteString("fallback", FallbackText);
+			}
+			Writer.WriteString("color", $"#{Color.ToArgb() & 0xFFFFFF:x6}");
+			Writer.WriteStartArray("blocks");
+			foreach (BlockBase Block in Blocks)
+			{
+				Block.Write(Writer, Options);
+			}
+			Writer.WriteEndArray();
+			Writer.WriteEndObject();
+		}
+	}
+}
diff --git a/Engine/Source/Programs/Horde/Horde.Build/Utilities/BlockKit/BlockKitMessage.cs b/Engine/Source/Programs/Horde/Horde.Build/Utilities/BlockKit/BlockKitMessage.cs
new file mode 100644
--- /dev/null
+++ b/Engine/Source/Programs/Horde/Horde.Build/Utilities/BlockKit/BlockKitMessage.cs
@@ -0,0 +1,58 @@
+// Copyright Epic Games, Inc. All Rights Reserved.
+
+using System.Collections.Generic;
+using System.Text.Json.Serialization;
+
+namespace HordeServer.Utilities.Slack.BlockKit
+{
+    /// <summary>
+    /// An entire BlockKit message to be presented in Slack.
+    /// </summary>
+    public class BlockKitMessage
+    {
+		/// <summary>
+		/// The channel ID or user ID to deliver the notification to.
+		/// </summary>
+		[JsonPropertyName("channel")]
+		public string? Channel { get; set; }
+
+		/// <summary>
+		/// Timestamp of the message to update
+		/// </summary>
+		[JsonPropertyName("ts")]
+		public string? Ts { get; set; }
+
+		/// <summary>
+		/// If no blocks are provided, this text will be displayed. This is also what is used on the pop up notification.
+		/// </summary>
+		[JsonPropertyName("text")]
+		public string? Text { get; set; }
+
+        /// <summary>
+        /// A collection of BlockKit blocks the message is composed from.
+        /// </summary>
+        [JsonPropertyName("blocks")]
+        public List<BlockBase> Blocks { get; } = new List<BlockBase>();
+
+        /// <summary>
+        /// Optional value that determines if Slack should try to replace the original message for the command that
+        /// this message is associated with.
+        /// </summary>
+        [JsonPropertyName("replace_original")]
+        public bool? ReplaceOriginal { get; set; }
+
+        /// <summary>
+        /// Optional value that determines if Slack should try to replace to delete original message for the command
+        /// that this message is associated with.
+        /// </summary>
+        [JsonPropertyName("delete_original")]
+        public bool? DeleteOriginal { get; set; }
+
+		/// <summary>
+		/// A collection of attachments.
+		/// NOTE: This feature is not supported any more but is not deprecated.
+		/// </summary>
+		[JsonPropertyName("attachments")]
+		public List<BlockKitAttachment> Attachments { get; } = new List<BlockKitAttachment>();
+    }
+}
diff --git a/Engine/Source/Programs/Horde/Horde.Build/Utilities/BlockKit/DividerBlock.cs b/Engine/Source/Programs/Horde/Horde.Build/Utilities/BlockKit/DividerBlock.cs
new file mode 100644
--- /dev/null
+++ b/Engine/Source/Programs/Horde/Horde.Build/Utilities/BlockKit/DividerBlock.cs
@@ -0,0 +1,20 @@
+// Copyright Epic Games, Inc. All Rights Reserved.
+
+using System.Text.Json;
+
+namespace HordeServer.Utilities.Slack.BlockKit
+{
+	/// <summary>
+	/// Wrapper for a BlockKit Divider block that is used to add a divider between two blocks.
+	/// </summary>
+	public class DividerBlock : BlockBase
+	{
+		/// <inheritdoc/>
+		public override void Write(Utf8JsonWriter Writer, JsonSerializerOptions Options)
+		{
+			Writer.WriteStartObject();
+			Writer.WriteString("type", "divider");
+			Writer.WriteEndObject();
+		}
+	}
+}
diff --git a/Engine/Source/Programs/Horde/Horde.Build/Utilities/BlockKit/HeaderBlock.cs b/Engine/Source/Programs/Horde/Horde.Build/Utilities/BlockKit/HeaderBlock.cs
new file mode 100644
--- /dev/null
+++ b/Engine/Source/Programs/Horde/Horde.Build/Utilities/BlockKit/HeaderBlock.cs
@@ -0,0 +1,37 @@
+// Copyright Epic Games, Inc. All Rights Reserved.
+
+using System.Text.Json;
+
+namespace HordeServer.Utilities.Slack.BlockKit
+{
+	/// <summary>
+	/// Represents a BlockKit Section block
+	/// </summary>
+	public class HeaderBlock : BlockBase
+	{
+		/// <summary>
+		/// Constructor
+		/// </summary>
+		/// <param name="Text">Any text to initiale the Section with</param>
+		/// <param name="IsMarkdown">If true, text is markdown, else text is considered plain text</param>
+		/// <param name="Emoji">If true, allows escaping of emojis with colons.</param>
+		public HeaderBlock(string Text = "", bool IsMarkdown = true, bool Emoji = false)
+		{
+			this.Text = new TextObject(Text, IsMarkdown, Emoji);
+		}
+
+		/// <summary>
+		/// The text of the Section block
+		/// </summary>
+		public TextObject Text { get; }
+
+		/// <inheritdoc/>
+		public override void Write(Utf8JsonWriter Writer, JsonSerializerOptions Options)
+		{
+			Writer.WriteStartObject();
+			Writer.WriteString("type", "header");
+			Text.Write(Writer);
+			Writer.WriteEndObject();
+		}
+	}
+}
diff --git a/Engine/Source/Programs/Horde/Horde.Build/Utilities/BlockKit/SectionBlock.cs b/Engine/Source/Programs/Horde/Horde.Build/Utilities/BlockKit/SectionBlock.cs
new file mode 100644
--- /dev/null
+++ b/Engine/Source/Programs/Horde/Horde.Build/Utilities/BlockKit/SectionBlock.cs
@@ -0,0 +1,37 @@
+// Copyright Epic Games, Inc. All Rights Reserved.
+
+using System.Text.Json;
+
+namespace HordeServer.Utilities.Slack.BlockKit
+{
+	/// <summary>
+	/// Represents a BlockKit Section block
+	/// </summary>
+	public class SectionBlock : BlockBase
+	{
+		/// <summary>
+		/// Constructor
+		/// </summary>
+		/// <param name="Text">Any text to initiale the Section with</param>
+		/// <param name="IsMarkdown">If true, text is markdown, else text is considered plain text</param>
+		/// <param name="Emoji">If true, allows escaping of emoji with colon.</param>
+		public SectionBlock(string Text = "", bool IsMarkdown = true, bool Emoji = false)
+		{
+			this.Text = new TextObject(Text, IsMarkdown, Emoji);
+		}
+
+		/// <summary>
+		/// The text of the Section block
+		/// </summary>
+		public TextObject Text { get; }
+
+		/// <inheritdoc/>
+		public override void Write(Utf8JsonWriter Writer, JsonSerializerOptions Options)
+		{
+			Writer.WriteStartObject();
+			Writer.WriteString("type", "section");
+			Text.Write(Writer);
+			Writer.WriteEndObject();
+		}
+	}
+}
diff --git a/Engine/Source/Programs/Horde/Horde.Build/Utilities/BlockKit/SlackResponses.cs b/Engine/Source/Programs/Horde/Horde.Build/Utilities/BlockKit/SlackResponses.cs
new file mode 100644
--- /dev/null
+++ b/Engine/Source/Programs/Horde/Horde.Build/Utilities/BlockKit/SlackResponses.cs
@@ -0,0 +1,45 @@
+// Copyright Epic Games, Inc. All Rights Reserved.
+
+using System.Drawing;
+using HordeServer.Utilities.Slack.BlockKit;
+
+namespace HordeServer.Utilities.BlockKit
+{
+	/// <summary>
+	/// Helper methods for creating <see cref="BlockKitMessage"/> responses to send
+	/// to users on Slack via the response_url of an interaction payload.
+	/// </summary>
+	public static class SlackResponses
+	{
+		/// <summary>
+		/// Creates a response to send to Slack about receiving bad data for a request.
+		/// </summary>
+		/// <param name="Message">A message explaining what was wrong with the data.</param>
+		/// <returns>A <see cref="BlockKitMessage"/> to be sent to Slack.</returns>
+		public static BlockKitMessage CreateBadDataResponse(string Message)
+		{
+			BlockKitMessage BlockKitMessage = new BlockKitMessage();
+			BlockKitAttachment Attachment = new BlockKitAttachment();
+			Attachment.Color = BlockKitAttachmentColors.Error;
+			Attachment.Blocks.Add(new SectionBlock("Failed to process request due to bad data from Slack:"));
+			Attachment.Blocks.Add(new SectionBlock(Message));
+			BlockKitMessage.Attachments.Add(Attachment);
+			return BlockKitMessage;
+		}
+
+		/// <summary>
+		/// Creates a response to send to Slack when failing to assign an issue due to it not
+		/// being found in the DB.
+		/// </summary>
+		/// <returns>A <see cref="BlockKitMessage"/> to be sent to Slack.</returns>
+		public static BlockKitMessage CreateIssueNotFoundResponse(int IssueId)
+		{
+			BlockKitMessage BlockKitMessage = new BlockKitMessage();
+			BlockKitAttachment Attachment = new BlockKitAttachment();
+			Attachment.Color = BlockKitAttachmentColors.Error;
+			Attachment.Blocks.Add(new SectionBlock($"Failed to assign issue as issue {IssueId} could not be found."));
+			BlockKitMessage.Attachments.Add(Attachment);
+			return BlockKitMessage;
+		}
+	}
+}
diff --git a/Engine/Source/Programs/Horde/Horde.Build/Utilities/BlockKit/TextObject.cs b/Engine/Source/Programs/Horde/Horde.Build/Utilities/BlockKit/TextObject.cs
new file mode 100644
--- /dev/null
+++ b/Engine/Source/Programs/Horde/Horde.Build/Utilities/BlockKit/TextObject.cs
@@ -0,0 +1,67 @@
+// Copyright Epic Games, Inc. All Rights Reserved.
+
+using System;
+using System.Text.Json;
+
+namespace HordeServer.Utilities.Slack.BlockKit
+{
+	/// <summary>
+	/// Represents a BlockKit TextObject which is used to present markdown or plain text
+	/// </summary>
+	public class TextObject
+	{
+		/// <summary>
+		/// Constructor
+		/// </summary>
+		/// <param name="Text">Any text to initialize with</param>
+		/// <param name="IsMarkdown">If true, text is markdown, else text is considered plain text</param>
+		/// <param name="Emoji">If true, allows escaping of emojis with a colon</param>
+		public TextObject(string Text = "", bool IsMarkdown = true, bool Emoji = false)
+		{
+			this.Text = Text;
+			this.IsMarkdown = IsMarkdown;
+			// force emoji to false if markdown is true as this invalidates the schema
+			this.Emoji = this.IsMarkdown ? false : Emoji;
+		}
+
+		/// <summary>
+		/// If true, text is markdown, else text is considered plain text
+		/// </summary>
+		public bool IsMarkdown { get; set; }
+
+		/// <summary>
+		/// If true, allows escaping of emojis with a colon
+		/// </summary>
+		public bool Emoji { get; set; }
+
+		/// <summary>
+		/// The plain text or markdown for this TextObject
+		/// </summary>
+		public string Text { get; set; }
+
+		/// <summary>
+		/// Writes the TextObject JSON to a <see cref="Utf8JsonWriter"/>
+		/// </summary>
+		/// <param name="Writer"></param>
+		public void Write(Utf8JsonWriter Writer)
+		{
+			Writer.WriteStartObject("text");
+			Writer.WriteString("type", IsMarkdown ? "mrkdwn" : "plain_text");
+
+			const int MaxLength = 2048;
+
+			string TruncatedText = Text;
+			if (TruncatedText.Length > MaxLength)
+			{
+				TruncatedText = TruncatedText.Substring(0, MaxLength) + "...\n(message truncated)";
+			}
+
+			Writer.WriteString("text", TruncatedText);
+			if (Emoji)
+			{
+				Writer.WriteBoolean("emoji", Emoji);
+			}
+			Writer.WriteEndObject();
+		}
+	}
+}
diff --git a/Engine/Source/Programs/Horde/Horde.Build/Utilities/BloomFilter.cs b/Engine/Source/Programs/Horde/Horde.Build/Utilities/BloomFilter.cs
new file mode 100644
--- /dev/null
+++ b/Engine/Source/Programs/Horde/Horde.Build/Utilities/BloomFilter.cs
@@ -0,0 +1,230 @@
+// Copyright Epic Games, Inc. All Rights Reserved.
+
+using System;
+using System.Collections.Generic;
+using System.Linq;
+using System.Threading.Tasks;
+
+namespace HordeServer.Utilities
+{
+	struct LogIndex
+	{
+		const int NumHashes = 3;
+		const int NgramLength = 3;
+		const int NgramMask = (1 << (NgramLength * 8)) - 1;
+
+		/// <summary>
+		/// Enumerates the hash codes for a given block of text
+		/// </summary>
+		/// <param name="Text">The text to parse</param>
+		/// <returns></returns>
+		public static IEnumerable<int> GetHashCodes(ReadOnlyMemory<byte> Text)
+		{
+			if(Text.Length >= NgramLength)
+			{
+				int Value = 0;
+				for(int Idx = 0; Idx < NgramLength - 1; Idx++)
+				{
+					Value = (Value << 8) | ToLowerUtf8(Text.Span[Idx]);
+				}
+				for(int Idx = NgramLength - 1; Idx < Text.Length; Idx++)
+				{
+					Value = ((Value << 8) | ToLowerUtf8(Text.Span[Idx])) & NgramMask;
+
+					int HashValue = Value;
+					yield return HashValue;
+
+					for(int HashIdx = 1; HashIdx < NumHashes; HashIdx++)
+					{
+						HashValue = Scramble(HashValue);
+						yield return HashValue;
+					}
+				}
+			}
+		}
+
+		/// <summary>
+		/// Converts a character into a format for hashing
+		/// </summary>
+		/// <param name="Value">The input byte</param>
+		/// <returns>The value to include in the trigram</returns>
+		static byte ToLowerUtf8(byte Value)
+		{
+			if (Value >= 'A' && Value <= 'Z')
+			{
+				return (byte)(Value + 'a' - 'A');
+			}
+			else
+			{
+				return Value;
+			}
+		}
+
+		/// <summary>
+		/// Scramble the given number using a pseudo-RNG
+		/// </summary>
+		/// <param name="Value">Initial value</param>
+		/// <returns>The scrambled value</returns>
+		private static int Scramble(int Value)
+		{
+			Value ^= Value << 13;
+			Value ^= Value >> 17;
+			Value ^= Value << 5;
+			return Value;
+		}
+	}
+
+	/// <summary>
+	/// Read only implementation of a bloom filter
+	/// </summary>
+	public class ReadOnlyBloomFilter
+	{
+		/// <summary>
+		/// Data for the bloom filter
+		/// </summary>
+		public ReadOnlyMemory<byte> Data { get; }
+
+		/// <summary>
+		/// Constructor
+		/// </summary>
+		/// <param name="Data">Data to construct from</param>
+		public ReadOnlyBloomFilter(ReadOnlyMemory<byte> Data)
+		{
+			this.Data = Data;
+		}
+
+		/// <summary>
+		/// Determines if the filter may contain a value
+		/// </summary>
+		/// <param name="HashCodes">Sequence of hash codes</param>
+		/// <returns>True if the value is in the set</returns>
+		public bool Contains(IEnumerable<int> HashCodes)
+		{
+			ReadOnlySpan<byte> Span = Data.Span;
+			foreach (int HashCode in HashCodes)
+			{
+				int Index = HashCode & ((Span.Length << 3) - 1);
+				if ((Span[Index >> 3] & (1 << (Index & 7))) == 0)
+				{
+					return false;
+				}
+			}
+			return true;
+		}
+
+		/// <summary>
+		/// Calculates the load of this filter (ie. the proportion of bits which are set)
+		/// </summary>
+		/// <returns>Value between 0 and 1</returns>
+		public double CalculateLoadFactor()
+		{
+			int SetCount = 0;
+
+			ReadOnlySpan<byte> Span = Data.Span;
+			for (int Idx = 0; Idx < Span.Length; Idx++)
+			{
+				for (int Bit = 0; Bit < 8; Bit++)
+				{
+					if ((Span[Idx] & (1 << Bit)) != 0)
+					{
+						SetCount++;
+					}
+				}
+			}
+
+			return (double)SetCount / (Span.Length * 8);
+		}
+	}
+
+	/// <summary>
+	/// Implementation of a bloom filter
+	/// </summary>
+	public class BloomFilter
+	{
+		/// <summary>
+		/// The array of bits in the filter
+		/// </summary>
+		public byte[] Data { get; }
+
+		/// <summary>
+		/// Constructs a new filter of the given size
+		/// </summary>
+		/// <param name="MinSize">The minimum size of the filter. Will be rounded up to the next power of two.</param>
+		public BloomFilter(int MinSize)
+		{
+			int Size = MinSize;
+			while ((Size & (Size - 1)) != 0)
+			{
+				Size |= (Size - 1);
+				Size++;
+			}
+			Data = new byte[Size];
+		}
+
+		/// <summary>
+		/// Constructs a filter from raw data
+		/// </summary>
+		/// <param name="Data">The data to construct from</param>
+		public BloomFilter(byte[] Data)
+		{
+			this.Data = Data;
+			if ((Data.Length & (Data.Length - 1)) != 0)
+			{
+				throw new ArgumentException("Array for bloom filter must be a power of 2 in size.");
+			}
+		}
+
+		/// <summary>
+		/// Add a value to the filter
+		/// </summary>
+		/// <param name="HashCodes">Sequence of hash codes</param>
+		public void Add(IEnumerable<int> HashCodes)
+		{
+			foreach(int HashCode in HashCodes)
+			{
+				int Index = HashCode & ((Data.Length << 3) - 1);
+				Data[Index >> 3] = (byte)(Data[Index >> 3] | (1 << (Index & 7)));
+			}
+		}
+
+		/// <summary>
+		/// Determines if the filter may contain a value
+		/// </summary>
+		/// <param name="HashCodes">Sequence of hash codes</param>
+		/// <returns>True if the value is in the set</returns>
+		public bool Contains(IEnumerable<int> HashCodes)
+		{
+			foreach (int HashCode in HashCodes)
+			{
+				int Index = HashCode & ((Data.Length << 3) - 1);
+				if ((Data[Index >> 3] & (1 << (Index & 7))) == 0)
+				{
+					return false;
+				}
+			}
+			return true;
+		}
+
+		/// <summary>
+		/// Calculates the load of this filter (ie. the proportion of bits which are set)
+		/// </summary>
+		/// <returns>Value between 0 and 1</returns>
+		public double CalculateLoadFactor()
+		{
+			int SetCount = 0;
+			int AllCount = 0;
+			for (int Idx = 0; Idx < Data.Length; Idx++)
+			{
+				for (int Bit = 0; Bit < 8; Bit++)
+				{
+					if ((Data[Idx] & (1 << Bit)) != 0)
+					{
+						SetCount++;
+					}
+					AllCount++;
+				}
+			}
+			return (double)SetCount / AllCount;
+		}
+	}
+}
diff --git a/Engine/Source/Programs/Horde/Horde.Build/Utilities/BsonDocumentExtensions.cs b/Engine/Source/Programs/Horde/Horde.Build/Utilities/BsonDocumentExtensions.cs
new file mode 100644
--- /dev/null
+++ b/Engine/Source/Programs/Horde/Horde.Build/Utilities/BsonDocumentExtensions.cs
@@ -0,0 +1,110 @@
+// Copyright Epic Games, Inc. All Rights Reserved.
+
+using MongoDB.Bson;
+using System;
+using System.Collections.Generic;
+using System.Diagnostics.CodeAnalysis;
+using System.Linq;
+using System.Threading.Tasks;
+
+namespace HordeServer.Utilities
+{
+	/// <summary>
+	/// Extension methods for BsonDocument
+	/// </summary>
+	static class BsonDocumentExtensions
+	{
+		/// <summary>
+		/// Gets a property value from a document or subdocument, indicated with dotted notation
+		/// </summary>
+		/// <param name="Document">Document to get a property for</param>
+		/// <param name="Name">Name of the property</param>
+		/// <param name="Type">Expected type of the property</param>
+		/// <param name="OutValue">Receives the property value</param>
+		/// <returns>True if the property exists and was of the correct type</returns>
+		public static bool TryGetPropertyValue(this BsonDocument Document, string Name, BsonType Type, [NotNullWhen(true)] out BsonValue? OutValue)
+		{
+			int DotIdx = Name.IndexOf('.', StringComparison.Ordinal);
+			if (DotIdx == -1)
+			{
+				return TryGetDirectPropertyValue(Document, Name, Type, out OutValue);
+			}
+
+			BsonValue? DocValue;
+			if (TryGetDirectPropertyValue(Document, Name.Substring(0, DotIdx), BsonType.Document, out DocValue))
+			{
+				return TryGetPropertyValue(DocValue.AsBsonDocument, Name.Substring(DotIdx + 1), Type, out OutValue);
+			}
+
+			OutValue = null;
+			return false;
+		}
+
+		/// <summary>
+		/// Gets a property value that's an immediate child of the document
+		/// </summary>
+		/// <param name="Document">Document to get a property for</param>
+		/// <param name="Name">Name of the property</param>
+		/// <param name="Type">Expected type of the property</param>
+		/// <param name="OutValue">Receives the property value</param>
+		/// <returns>True if the property exists and was of the correct type</returns>
+		private static bool TryGetDirectPropertyValue(this BsonDocument Document, string Name, BsonType Type, [NotNullWhen(true)] out BsonValue? OutValue)
+		{
+			BsonValue Value;
+			if (Document.TryGetValue(Name, out Value) && Value.BsonType == Type)
+			{
+				OutValue = Value;
+				return true;
+			}
+			else
+			{
+				OutValue = null;
+				return false;
+			}
+		}
+
+		/// <summary>
+		/// Gets an int32 value from the document
+		/// </summary>
+		/// <param name="Document">Document to get a property for</param>
+		/// <param name="Name">Name of the property</param>
+		/// <param name="OutValue">Receives the property value</param>
+		/// <returns>True if the property was retrieved</returns>
+		public static bool TryGetInt32(this BsonDocument Document, string Name, out int OutValue)
+		{
+			BsonValue? Value;
+			if (Document.TryGetPropertyValue(Name, BsonType.Int32, out Value))
+			{
+				OutValue = Value.AsInt32;
+				return true;
+			}
+			else
+			{
+				OutValue = 0;
+				return false;
+			}
+		}
+
+		/// <summary>
+		/// Gets a string value from the document
+		/// </summary>
+		/// <param name="Document">Document to get a property for</param>
+		/// <param name="Name">Name of the property</param>
+		/// <param name="OutValue">Receives the property value</param>
+		/// <returns>True if the property was retrieved</returns>
+		public static bool TryGetString(this BsonDocument Document, string Name, [NotNullWhen(true)] out string? OutValue)
+		{
+			BsonValue? Value;
+			if (Document.TryGetPropertyValue(Name, BsonType.String, out Value))
+			{
+				OutValue = Value.AsString;
+				return true;
+			}
+			else
+			{
+				OutValue = null;
+				return false;
+			}
+		}
+	}
+}
diff --git a/Engine/Source/Programs/Horde/Horde.Build/Utilities/CaseInsensitiveStringSet.cs b/Engine/Source/Programs/Horde/Horde.Build/Utilities/CaseInsensitiveStringSet.cs
new file mode 100644
--- /dev/null
+++ b/Engine/Source/Programs/Horde/Horde.Build/Utilities/CaseInsensitiveStringSet.cs
@@ -0,0 +1,33 @@
+// Copyright Epic Games, Inc. All Rights Reserved.
+
+using MongoDB.Bson.Serialization;
+using MongoDB.Bson.Serialization.Serializers;
+using System;
+using System.Collections.Generic;
+using System.Linq;
+using System.Threading.Tasks;
+
+namespace HordeServer.Utilities
+{
+	/// <summary>
+	/// Case insensitive set of strings
+	/// </summary>
+	public class CaseInsensitiveStringSet : HashSet<string>
+	{
+		/// <summary>
+		/// Constructor
+		/// </summary>
+		public CaseInsensitiveStringSet()
+			: base(StringComparer.OrdinalIgnoreCase)
+		{
+		}
+
+		/// <summary>
+		/// Constructor
+		/// </summary>
+		public CaseInsensitiveStringSet(IEnumerable<string> Items)
+			: base(Items, StringComparer.OrdinalIgnoreCase)
+		{
+		}
+	}
+}
diff --git a/Engine/Source/Programs/Horde/Horde.Build/Utilities/CollectionExtensions.cs b/Engine/Source/Programs/Horde/Horde.Build/Utilities/CollectionExtensions.cs
new file mode 100644
--- /dev/null
+++ b/Engine/Source/Programs/Horde/Horde.Build/Utilities/CollectionExtensions.cs
@@ -0,0 +1,33 @@
+// Copyright Epic Games, Inc. All Rights Reserved.
+
+using Microsoft.AspNetCore.Http.Features;
+using Microsoft.AspNetCore.Mvc.Formatters;
+using Microsoft.IdentityModel.Tokens;
+using System;
+using System.Collections.Generic;
+using System.Linq;
+using System.Threading.Tasks;
+
+namespace HordeServer.Utilities
+{
+	/// <summary>
+	/// Extension methods for collections
+	/// </summary>
+	public static class CollectionExtensions
+	{
+		/// <summary>
+		/// Adds an arbitrary sequence of items to a protobuf map field
+		/// </summary>
+		/// <typeparam name="TKey">The key type</typeparam>
+		/// <typeparam name="TValue">The value type</typeparam>
+		/// <param name="Map">The map to update</param>
+		/// <param name="Sequence">Sequence of items to add</param>
+		public static void Add<TKey, TValue>(this Google.Protobuf.Collections.MapField<TKey, TValue> Map, IEnumerable<KeyValuePair<TKey, TValue>> Sequence)
+		{
+			foreach(KeyValuePair<TKey, TValue> Pair in Sequence)
+			{
+				Map.Add(Pair.Key, Pair.Value);
+			}
+		}
+	}
+}
diff --git a/Engine/Source/Programs/Horde/Horde.Build/Utilities/CompressionExtensions.cs b/Engine/Source/Programs/Horde/Horde.Build/Utilities/CompressionExtensions.cs
new file mode 100644
--- /dev/null
+++ b/Engine/Source/Programs/Horde/Horde.Build/Utilities/CompressionExtensions.cs
@@ -0,0 +1,77 @@
+// Copyright Epic Games, Inc. All Rights Reserved.
+
+using EpicGames.Core;
+using ICSharpCode.SharpZipLib.BZip2;
+using OpenTracing;
+using OpenTracing.Util;
+using System;
+using System.Buffers.Binary;
+using System.Collections.Generic;
+using System.Globalization;
+using System.IO;
+using System.Linq;
+using System.Threading.Tasks;
+
+namespace HordeServer.Utilities
+{
+	/// <summary>
+	/// Extension methods for compression
+	/// </summary>
+	static class CompressionExtensions
+	{
+		/// <summary>
+		/// Compress a block of data with bzip2
+		/// </summary>
+		/// <returns>The compressed data</returns>
+		public static byte[] CompressBzip2(this ReadOnlyMemory<byte> Memory)
+		{
+			using IScope Scope = GlobalTracer.Instance.BuildSpan("CompressBzip2").StartActive();
+			Scope.Span.SetTag("DecompressedSize", Memory.Length.ToString(CultureInfo.InvariantCulture));
+
+			byte[] CompressedData;
+			using (MemoryStream Stream = new MemoryStream())
+			{
+				byte[] DecompressedSize = new byte[4];
+				BinaryPrimitives.WriteInt32LittleEndian(DecompressedSize.AsSpan(), Memory.Length);
+				Stream.Write(DecompressedSize.AsSpan());
+
+				using (BZip2OutputStream CompressedStream = new BZip2OutputStream(Stream))
+				{
+					CompressedStream.Write(Memory.Span);
+				}
+
+				CompressedData = Stream.ToArray();
+			}
+
+			Scope.Span.SetTag("CompressedSize", CompressedData.Length.ToString(CultureInfo.InvariantCulture));
+			return CompressedData;
+		}
+
+		/// <summary>
+		/// Decompress the data
+		/// </summary>
+		/// <returns>The decompressed data</returns>
+		public static byte[] DecompressBzip2(this ReadOnlyMemory<byte> Memory)
+		{
+			int DecompressedSize = BinaryPrimitives.ReadInt32LittleEndian(Memory.Span);
+
+			using IScope Scope = GlobalTracer.Instance.BuildSpan("DecompressBzip2").StartActive();
+			Scope.Span.SetTag("CompressedSize", Memory.Length.ToString(CultureInfo.InvariantCulture));
+			Scope.Span.SetTag("DecompressedSize", DecompressedSize.ToString(CultureInfo.InvariantCulture));
+
+			byte[] Data = new byte[DecompressedSize];
+			using (ReadOnlyMemoryStream Stream = new ReadOnlyMemoryStream(Memory.Slice(4)))
+			{
+				using (BZip2InputStream DecompressedStream = new BZip2InputStream(Stream))
+				{
+					int ReadSize = DecompressedStream.Read(Data.AsSpan());
+					if (ReadSize != Data.Length)
+					{
+						throw new InvalidDataException($"Compressed data is too short (expected {Data.Length} bytes, got {ReadSize} bytes)");
+					}
+				}
+			}
+			return Data;
+		}
+	}
+}
diff --git a/Engine/Source/Programs/Horde/Horde.Build/Utilities/ContentHashSerializer.cs b/Engine/Source/Programs/Horde/Horde.Build/Utilities/ContentHashSerializer.cs
new file mode 100644
--- /dev/null
+++ b/Engine/Source/Programs/Horde/Horde.Build/Utilities/ContentHashSerializer.cs
@@ -0,0 +1,54 @@
+// Copyright Epic Games, Inc. All Rights Reserved.
+
+using EpicGames.Core;
+using MongoDB.Bson.Serialization;
+using System;
+using System.Collections.Generic;
+using System.Linq;
+using System.Threading.Tasks;
+
+namespace HordeServer.Utilities
+{
+	/// <summary>
+	/// Serializer for JobStepRefId objects
+	/// </summary>
+	public sealed class ContentHashSerializer : IBsonSerializer<ContentHash>
+	{
+		/// <inheritdoc/>
+		public Type ValueType
+		{
+			get { return typeof(ContentHash); }
+		}
+
+		/// <inheritdoc/>
+		void IBsonSerializer.Serialize(BsonSerializationContext Context, BsonSerializationArgs Args, object Value)
+		{
+			Serialize(Context, Args, (ContentHash)Value);
+		}
+
+		/// <inheritdoc/>
+		object IBsonSerializer.Deserialize(BsonDeserializationContext Context, BsonDeserializationArgs Args)
+		{
+			return ((IBsonSerializer<ContentHash>)this).Deserialize(Context, Args);
+		}
+
+		/// <inheritdoc/>
+		public ContentHash Deserialize(BsonDeserializationContext Context, BsonDeserializationArgs Args)
+		{
+			if (Context.Reader.CurrentBsonType == MongoDB.Bson.BsonType.ObjectId)
+			{
+				return new ContentHash(Context.Reader.ReadObjectId().ToByteArray());
+			}
+			else
+			{
+				return ContentHash.Parse(Context.Reader.ReadString());
+			}
+		}
+
+		/// <inheritdoc/>
+		public void Serialize(BsonSerializationContext Context, BsonSerializationArgs Args, ContentHash Value)
+		{
+			Context.Writer.WriteString(Value.ToString());
+		}
+	}
+}
