Horde: Fixing up some names after moving files.

diff --git a/Engine/Source/Programs/Horde/BuildHorde.xml b/Engine/Source/Programs/Horde/BuildHorde.xml
--- a/Engine/Source/Programs/Horde/BuildHorde.xml
+++ b/Engine/Source/Programs/Horde/BuildHorde.xml
@@ -12,7 +12,7 @@
 	<Property Name="VersionArguments" Value="/p:Version=$(Version).0 /p:InformationalVersion=$(InformationalVersion)"/>
 
 	<Agent Name="HordeServer" Type="Win64_Docker">
-		<Property Name="StagingDir" Value="$(RootDir)/Engine/Saved/HordeServer"/>
+		<Property Name="StagingDir" Value="$(RootDir)/Engine/Saved/Horde.Build"/>
 		
 		<Node Name="Build HordeServer">
 			<!-- Tag all the files that need to be staged to build -->
@@ -25,7 +25,7 @@
 			<Tag Files="Engine/Source/Programs/Shared/EpicGames.Perforce.Native/bin/..." Filter="*.pdb;*.dll;*.so" With="#InputFiles"/>
 
 			<!-- Build the Docker image and publish it -->
-			<Docker-Build BaseDir="Engine/Source" Files="#InputFiles" Tag="hordeserver-public" DockerFile="Engine/Source/Programs/Horde/HordeServer/Dockerfile" Arguments="--build-arg msbuild_args=&quot;$(VersionArguments)&quot;"/>
+			<Docker-Build BaseDir="Engine/Source" Files="#InputFiles" Tag="hordeserver-public" DockerFile="Engine/Source/Programs/Horde/Horde.Build/Dockerfile" Arguments="--build-arg msbuild_args=&quot;$(VersionArguments)&quot;"/>
 		</Node>
 	</Agent>
 </BuildGraph>
diff --git a/Engine/Source/Programs/Horde/Horde.Build.Tests/Horde.Build.Tests.csproj b/Engine/Source/Programs/Horde/Horde.Build.Tests/Horde.Build.Tests.csproj
--- a/Engine/Source/Programs/Horde/Horde.Build.Tests/Horde.Build.Tests.csproj
+++ b/Engine/Source/Programs/Horde/Horde.Build.Tests/Horde.Build.Tests.csproj
@@ -30,7 +30,7 @@
     <ProjectReference Include="..\HordeAgent\HordeAgent.csproj" Label="HordeAgent">
       <Aliases>HordeAgent</Aliases>
     </ProjectReference>
-    <ProjectReference Include="..\HordeServer\HordeServer.csproj" />
+    <ProjectReference Include="..\Horde.Build\Horde.Build.csproj" />
   </ItemGroup>
 
   <ItemGroup>
diff --git a/Engine/Source/Programs/Horde/Horde.Build/Fleet/AutoscaleService.cs b/Engine/Source/Programs/Horde/Horde.Build/Fleet/AutoscaleService.cs
--- a/Engine/Source/Programs/Horde/Horde.Build/Fleet/AutoscaleService.cs
+++ b/Engine/Source/Programs/Horde/Horde.Build/Fleet/AutoscaleService.cs
@@ -23,8 +23,6 @@ using HordeCommon;
 using Microsoft.Extensions.Options;
 using StatsdClient;
 
-[assembly: InternalsVisibleTo("HordeServerTests")]
-
 namespace HordeServer.Services
 {
 	using PoolId = StringId<IPool>;
diff --git a/Engine/Source/Programs/Horde/Horde.Build/Logs/Impl/LogFileService.cs b/Engine/Source/Programs/Horde/Horde.Build/Logs/Impl/LogFileService.cs
--- a/Engine/Source/Programs/Horde/Horde.Build/Logs/Impl/LogFileService.cs
+++ b/Engine/Source/Programs/Horde/Horde.Build/Logs/Impl/LogFileService.cs
@@ -33,8 +33,6 @@ using Stream = System.IO.Stream;
 using OpenTracing.Util;
 using OpenTracing;
 
-[assembly: InternalsVisibleTo("HordeServerTests")]
-
 namespace HordeServer.Services
 {
 	using JobId = ObjectId<IJob>;
diff --git a/Engine/Source/Programs/Horde/Horde.Build/Properties/AssemblyInfo.cs b/Engine/Source/Programs/Horde/Horde.Build/Properties/AssemblyInfo.cs
--- a/Engine/Source/Programs/Horde/Horde.Build/Properties/AssemblyInfo.cs
+++ b/Engine/Source/Programs/Horde/Horde.Build/Properties/AssemblyInfo.cs
@@ -2,4 +2,4 @@
 
 using System.Runtime.CompilerServices;
 
-[assembly: InternalsVisibleTo("HordeServer.Compute.Tests")]
+[assembly: InternalsVisibleTo("Horde.Build.Tests")]
diff --git a/Engine/Source/Programs/Horde/HordeCommon/LruCache.cs b/Engine/Source/Programs/Horde/HordeCommon/LruCache.cs
--- a/Engine/Source/Programs/Horde/HordeCommon/LruCache.cs
+++ b/Engine/Source/Programs/Horde/HordeCommon/LruCache.cs
@@ -18,7 +18,7 @@ using System.Text;
 using System.Threading;
 using System.Threading.Tasks;
 
-[assembly: InternalsVisibleTo("HordeServerTests")]
+[assembly: InternalsVisibleTo("Horde.Build.Tests")]
 
 namespace HordeCommon
 {
