Horde: Rename HordeAgent to Horde.Agent.

diff --git a/Engine/Source/Programs/Horde/Horde.Agent/Utility/JsonLogger.cs b/Engine/Source/Programs/Horde/Horde.Agent/Utility/JsonLogger.cs
new file mode 100644
--- /dev/null
+++ b/Engine/Source/Programs/Horde/Horde.Agent/Utility/JsonLogger.cs
@@ -0,0 +1,355 @@
+// Copyright Epic Games, Inc. All Rights Reserved.
+
+using EpicGames.Core;
+using HordeAgent.Parser;
+using HordeCommon;
+using Microsoft.Extensions.Logging;
+using System;
+using System.Buffers;
+using System.Collections.Generic;
+using System.Globalization;
+using System.Linq;
+using System.Text;
+using System.Text.Json;
+
+namespace HordeAgent.Utility
+{
+	/// <summary>
+	/// Class to handle uploading log data to the server in the background
+	/// </summary>
+	public abstract class JsonLogger : ILogger
+	{
+		/// <summary>
+		/// Whether the logger should include warnings, or downgrade them to information
+		/// </summary>
+		protected bool? Warnings { get; }
+
+		/// <summary>
+		/// The inner logger device
+		/// </summary>
+		protected ILogger Inner { get; }
+
+		/// <summary>
+		/// Constructor
+		/// </summary>
+		/// <param name="Warnings">Whether to include warnings in the output</param>
+		/// <param name="Inner">Additional logger to write to</param>
+		public JsonLogger(bool? Warnings, ILogger Inner)
+		{
+			this.Warnings = Warnings;
+			this.Inner = Inner;
+		}
+
+		/// <inheritdoc/>
+		public IDisposable? BeginScope<TState>(TState State)
+		{
+			return null;
+		}
+
+		/// <inheritdoc/>
+		public bool IsEnabled(LogLevel LogLevel)
+		{
+			return true;
+		}
+
+		/// <inheritdoc/>
+		public void Log<TState>(LogLevel LogLevel, EventId EventId, TState State, Exception? Exception, Func<TState, Exception?, string> Formatter)
+		{
+			// Downgrade warnings to information if not required
+			if (LogLevel == LogLevel.Warning && !(Warnings ?? true))
+			{
+				LogLevel = LogLevel.Information;
+			}
+
+			// Render the message
+			string Text = Formatter(State, Exception);
+
+			// Try to log the event
+			IEnumerable<KeyValuePair<string, object>>? Values = State as IEnumerable<KeyValuePair<string, object>>;
+			if (Values != null)
+			{
+				KeyValuePair<string, object>? Format = Values.FirstOrDefault(x => x.Key.Equals(MessageTemplate.FormatPropertyName, StringComparison.Ordinal));
+				if (Format != null)
+				{
+					// Format all the other values
+					QueueFormattedEvent(LogLevel, EventId, Text, Format.Value.Value?.ToString() ?? String.Empty, Values.Where(x => !x.Key.Equals(MessageTemplate.FormatPropertyName, StringComparison.Ordinal)), Exception);
+				}
+				else
+				{
+					// Include all the data, but don't use the format string
+					QueueFormattedEvent(LogLevel, EventId, Text, null, Values, Exception);
+				}
+			}
+			else
+			{
+				// Format as a string
+				QueueFormattedEvent(LogLevel, EventId, Text, null, null, Exception);
+			}
+		}
+
+		/// <summary>
+		/// Queues up a log event for the server
+		/// </summary>
+		/// <param name="LogLevel">Level for this log event</param>
+		/// <param name="EventId">Corresponding event id</param>
+		/// <param name="Text">Rendered text for the message</param>
+		/// <param name="Format">Format string</param>
+		/// <param name="Properties">Properties for the message</param>
+		/// <param name="Exception">Exception information</param>
+		void QueueFormattedEvent(LogLevel LogLevel, EventId EventId, string Text, string? Format, IEnumerable<KeyValuePair<string, object>>? Properties, Exception? Exception)
+		{
+			DateTime Time = DateTime.UtcNow;
+
+			// Split it into individual lines
+			string[] Lines = Text.Split('\n');
+
+			// Split the format string into lines too
+			string[]? FormatLines = null;
+			if (Format != null)
+			{
+				FormatLines = Format.Split('\n');
+			}
+
+			// If it's a systemic issue, log it to the regular log sink so we can catch it in Datadog
+			if (EventId.Id >= KnownLogEvents.Systemic.Id && EventId.Id < KnownLogEvents.Systemic.Id + 100)
+			{
+				WriteSystemicEvent(EventId, Text);
+			}
+
+			// In order to support multi-line log events while seeking through the log buffer, we re-output the same json event for each
+			// line in the multi-line output. The 'line' parameter for the second and successive lines gives the number of this line.
+			byte[][] OutputLines = new byte[Lines.Length][];
+			for (int LineIdx = 0; LineIdx < Lines.Length; LineIdx++)
+			{
+				// Encode the data for this line
+				string? LineFormat = null;
+				if (FormatLines != null)
+				{
+					if (LineIdx < FormatLines.Length)
+					{
+						LineFormat = FormatLines[LineIdx];
+					}
+					else
+					{
+						LineFormat = "Missing format string";
+					}
+				}
+
+				OutputLines[LineIdx] = EncodeMessage(Time, LogLevel, Lines[LineIdx], EventId, LineIdx, Lines.Length, LineFormat, Properties, Exception);
+			}
+			WriteFormattedEvent(LogLevel, OutputLines);
+		}
+
+		/// <summary>
+		/// Writes a formatted event
+		/// </summary>
+		/// <param name="Level">The log level</param>
+		/// <param name="Lines">Utf-8 encoded JSON line data</param>
+		protected abstract void WriteFormattedEvent(LogLevel Level, byte[][] Lines);
+
+		/// <summary>
+		/// Callback to write a systemic event
+		/// </summary>
+		/// <param name="EventId">The event id</param>
+		/// <param name="Text">The event text</param>
+		protected virtual void WriteSystemicEvent(EventId EventId, string Text)
+		{
+			Inner.LogWarning("Systemic event {KnownLogEventId}: {Text}", EventId.Id, Text);
+		}
+
+		/// <summary>
+		/// Queues up a log event for the server
+		/// </summary>
+		/// <param name="Time">Time of the event</param>
+		/// <param name="LogLevel">Level for this log event</param>
+		/// <param name="Text">Rendered text for the message</param>
+		/// <param name="EventId">Corresponding event id</param>
+		/// <param name="LineIndex">Line index within the event</param>
+		/// <param name="LineCount">Number of lines in the event</param>
+		/// <param name="Format">Format string</param>
+		/// <param name="Properties">Properties for the message</param>
+		/// <param name="Exception">Exception information</param>
+		byte[] EncodeMessage(DateTime Time, LogLevel LogLevel, string Text, EventId EventId, int LineIndex, int LineCount, string? Format, IEnumerable<KeyValuePair<string, object>>? Properties, Exception? Exception)
+		{
+			ArrayBufferWriter<byte> Writer = new ArrayBufferWriter<byte>();
+			using (Utf8JsonWriter JsonWriter = new Utf8JsonWriter(Writer))
+			{
+				JsonWriter.WriteStartObject();
+				JsonWriter.WriteString("time", Time.ToString("s", CultureInfo.InvariantCulture));
+				JsonWriter.WriteString("level", LogLevel.ToString());
+				JsonWriter.WriteString("message", Text);
+
+				if (EventId != KnownLogEvents.None)
+				{
+					JsonWriter.WriteNumber("id", EventId.Id);
+				}
+
+				if (LineCount > 1)
+				{
+					JsonWriter.WriteNumber("line", LineIndex);
+					JsonWriter.WriteNumber("lineCount", LineCount);
+				}
+
+				if (Format != null)
+				{
+					JsonWriter.WriteString("format", Format);
+				}
+
+				if (Properties != null && Properties.Any())
+				{
+					JsonWriter.WriteStartObject("properties");
+					foreach (KeyValuePair<string, object> Pair in Properties!)
+					{
+						JsonWriter.WritePropertyName(Pair.Key);
+						WritePropertyValue(Pair.Value, JsonWriter);
+					}
+					JsonWriter.WriteEndObject();
+				}
+
+				if (Exception != null)
+				{
+					JsonWriter.WriteStartObject("exception");
+					WriteException(JsonWriter, Exception);
+					JsonWriter.WriteEndObject();
+				}
+
+				JsonWriter.WriteEndObject();
+			}
+			return Writer.WrittenSpan.ToArray();
+		}
+
+		/// <summary>
+		/// Writes an exception to a json object
+		/// </summary>
+		/// <param name="Writer">Writer to receive the exception data</param>
+		/// <param name="Exception">The exception</param>
+		void WriteException(Utf8JsonWriter Writer, Exception Exception)
+		{
+			Writer.WriteString("message", Exception.Message);
+			Writer.WriteString("trace", Exception.StackTrace);
+
+			if (Exception.InnerException != null)
+			{
+				Writer.WriteStartObject("innerException");
+				WriteException(Writer, Exception.InnerException);
+				Writer.WriteEndObject();
+			}
+
+			AggregateException? AggregateException = Exception as AggregateException;
+			if (AggregateException != null)
+			{
+				Writer.WriteStartArray("innerExceptions");
+				for (int Idx = 0; Idx < 16 && Idx < AggregateException.InnerExceptions.Count; Idx++) // Cap number of exceptions returned to avoid huge messages
+				{
+					Exception InnerException = AggregateException.InnerExceptions[Idx];
+					Writer.WriteStartObject();
+					WriteException(Writer, InnerException);
+					Writer.WriteEndObject();
+				}
+				Writer.WriteEndArray();
+			}
+		}
+
+		/// <summary>
+		/// Converts a property value to log markup
+		/// </summary>
+		/// <param name="Value">The value to write</param>
+		/// <param name="Writer">The json writer</param>
+		static void WritePropertyValue(object Value, Utf8JsonWriter Writer)
+		{
+			if (Value == null)
+			{
+				Writer.WriteNullValue();
+			}
+			else
+			{
+				Type ValueType = Value.GetType();
+				if (ValueType.IsEnum)
+				{
+					Writer.WriteStringValue(Enum.GetName(ValueType, Value));
+				}
+				else
+				{
+					switch (Type.GetTypeCode(ValueType))
+					{
+						case TypeCode.Boolean:
+							Writer.WriteBooleanValue((bool)Value);
+							break;
+						case TypeCode.Byte:
+							Writer.WriteNumberValue((Byte)Value);
+							break;
+						case TypeCode.SByte:
+							Writer.WriteNumberValue((SByte)Value);
+							break;
+						case TypeCode.UInt16:
+							Writer.WriteNumberValue((UInt16)Value);
+							break;
+						case TypeCode.UInt32:
+							Writer.WriteNumberValue((UInt32)Value);
+							break;
+						case TypeCode.UInt64:
+							Writer.WriteNumberValue((UInt64)Value);
+							break;
+						case TypeCode.Int16:
+							Writer.WriteNumberValue((Int16)Value);
+							break;
+						case TypeCode.Int32:
+							Writer.WriteNumberValue((Int32)Value);
+							break;
+						case TypeCode.Int64:
+							Writer.WriteNumberValue((Int64)Value);
+							break;
+						case TypeCode.Decimal:
+							Writer.WriteNumberValue((Decimal)Value);
+							break;
+						case TypeCode.Double:
+							Writer.WriteNumberValue((Double)Value);
+							break;
+						case TypeCode.Single:
+							Writer.WriteNumberValue((Single)Value);
+							break;
+						case TypeCode.String:
+							Writer.WriteStringValue((String)Value);
+							break;
+						default:
+							WriteComplexPropertyValue(Value, ValueType, Writer);
+							break;
+					}
+				}
+			}
+		}
+
+		/// <summary>
+		/// Writes a complex object to the given json writer
+		/// </summary>
+		/// <param name="Value">The value to write</param>
+		/// <param name="ValueType">Type of the value</param>
+		/// <param name="JsonWriter">The writer to output to</param>
+		static void WriteComplexPropertyValue(object Value, Type ValueType, Utf8JsonWriter JsonWriter)
+		{
+			if (ValueType == typeof(LogEventSpan))
+			{
+				LogEventSpan Span = (LogEventSpan)Value;
+				if (Span.Properties.Count == 0)
+				{
+					JsonWriter.WriteStringValue(Span.Text);
+				}
+				else
+				{
+					JsonWriter.WriteStartObject();
+					JsonWriter.WriteString("text", Span.Text);
+					foreach (KeyValuePair<string, object> Pair in Span.Properties)
+					{
+						JsonWriter.WritePropertyName(Pair.Key);
+						WritePropertyValue(Pair.Value, JsonWriter);
+					}
+					JsonWriter.WriteEndObject();
+				}
+			}
+			else
+			{
+				JsonWriter.WriteStringValue(Value.ToString());
+			}
+		}
+	}
+}
diff --git a/Engine/Source/Programs/Horde/Horde.Agent/Utility/JsonRpcLogger.cs b/Engine/Source/Programs/Horde/Horde.Agent/Utility/JsonRpcLogger.cs
new file mode 100644
--- /dev/null
+++ b/Engine/Source/Programs/Horde/Horde.Agent/Utility/JsonRpcLogger.cs
@@ -0,0 +1,305 @@
+// Copyright Epic Games, Inc. All Rights Reserved.
+
+using HordeAgent.Parser.Interfaces;
+using HordeAgent.Utility;
+using HordeCommon;
+using HordeCommon.Rpc;
+using Microsoft.Extensions.Logging;
+using System;
+using System.Buffers;
+using System.Collections.Concurrent;
+using System.Collections.Generic;
+using System.Globalization;
+using System.Linq;
+using System.Net.Http;
+using System.Text;
+using System.Text.Json;
+using System.Threading;
+using System.Threading.Channels;
+using System.Threading.Tasks;
+
+namespace HordeAgent.Parser
+{
+	/// <summary>
+	/// Class to handle uploading log data to the server in the background
+	/// </summary>
+	class JsonRpcLogger : JsonLogger, IAsyncDisposable
+	{
+		class QueueItem
+		{
+			public byte[] Data;
+			public CreateEventRequest? Event;
+
+			public QueueItem(byte[] Data, CreateEventRequest? Event)
+			{
+				this.Data = Data;
+				this.Event = Event;
+			}
+
+			public override string ToString()
+			{
+				return Encoding.UTF8.GetString(Data);
+			}
+		}
+
+		readonly IRpcConnection RpcClient;
+		readonly string LogId;
+		readonly string? JobId;
+		readonly string? JobBatchId;
+		readonly string? JobStepId;
+		Channel<QueueItem> DataChannel;
+		Task? DataWriter;
+
+		/// <summary>
+		/// The current outcome for this step. Updated to reflect any errors and warnings that occurred.
+		/// </summary>
+		public JobStepOutcome Outcome
+		{
+			get;
+			private set;
+		}
+
+		/// <summary>
+		/// Constructor
+		/// </summary>
+		/// <param name="RpcClient">RPC client to use for server requests</param>
+		/// <param name="LogId">The log id to write to</param>
+		/// <param name="JobId">Id of the job being executed</param>
+		/// <param name="JobBatchId">Batch being executed</param>
+		/// <param name="JobStepId">Id of the step being executed</param>
+		/// <param name="Warnings">Whether to include warnings in the output</param>
+		/// <param name="Inner">Additional logger to write to</param>
+		public JsonRpcLogger(IRpcConnection RpcClient, string LogId, string? JobId, string? JobBatchId, string? JobStepId, bool? Warnings, ILogger Inner)
+			: base(Warnings, Inner)
+		{
+			this.RpcClient = RpcClient;
+			this.LogId = LogId;
+			this.JobId = JobId;
+			this.JobBatchId = JobBatchId;
+			this.JobStepId = JobStepId;
+			this.DataChannel = Channel.CreateUnbounded<QueueItem>();
+			this.DataWriter = Task.Run(() => RunDataWriter());
+			this.Outcome = JobStepOutcome.Success;
+		}
+
+		protected override void WriteFormattedEvent(LogLevel Level, byte[][] Lines)
+		{
+			// Update the state of this job if this is an error status
+			if (Level == LogLevel.Error || Level == LogLevel.Critical)
+			{
+				Outcome = JobStepOutcome.Failure;
+			}
+			else if (Level == LogLevel.Warning && Outcome != JobStepOutcome.Failure)
+			{
+				Outcome = JobStepOutcome.Warnings;
+			}
+
+			for (int Idx = 0; Idx < Lines.Length; Idx++)
+			{
+				// If we want an event for this log event, create one now
+				CreateEventRequest? Event = null;
+				if (Idx == 0)
+				{
+					if (Level == LogLevel.Warning || Level == LogLevel.Error || Level == LogLevel.Critical)
+					{
+						Event = CreateEvent(Level, Lines.Length);
+					}
+				}
+
+				// Write the data to the output channel
+				QueueItem QueueItem = new QueueItem(Lines[Idx], Event);
+				if (!DataChannel.Writer.TryWrite(QueueItem))
+				{
+					throw new InvalidOperationException("Expected unbounded writer to complete immediately");
+				}
+			}
+		}
+
+		/// <summary>
+		/// Callback to write a systemic event
+		/// </summary>
+		/// <param name="EventId">The event id</param>
+		/// <param name="Text">The event text</param>
+		protected override void WriteSystemicEvent(EventId EventId, string Text)
+		{
+			if (JobId == null)
+			{
+				Inner.LogWarning("Systemic event {KnownLogEventId} in log {LogId}: {Text}", EventId.Id, Text);
+			}
+			else if (JobBatchId == null)
+			{
+				Inner.LogWarning("Systemic event {KnownLogEventId} in log {LogId} (job {JobId}): {Text})", EventId.Id, JobId, Text);
+			}
+			else if (JobStepId == null)
+			{
+				Inner.LogWarning("Systemic event {KnownLogEventId} in log {LogId} (job batch {JobId}:{BatchId}): {Text})", EventId.Id, JobId, JobStepId, Text);
+			}
+			else
+			{
+				Inner.LogWarning("Systemic event {KnownLogEventId} in log {LogId} (job step {JobId}:{BatchId}:{StepId}): {Text}", EventId.Id, LogId, JobId, JobBatchId, JobStepId, Text);
+			}
+		}
+
+		/// <summary>
+		/// Makes a <see cref="CreateEventRequest"/> for the given parameters
+		/// </summary>
+		/// <param name="LogLevel">Level for this log event</param>
+		/// <param name="LineCount">Number of lines in the event</param>
+		CreateEventRequest CreateEvent(LogLevel LogLevel, int LineCount)
+		{
+			EventSeverity Severity = (LogLevel == LogLevel.Warning) ? EventSeverity.Warning : EventSeverity.Error;
+			return new CreateEventRequest(Severity, LogId, 0, LineCount);
+		}
+
+		/// <summary>
+		/// Stops the log writer's background task
+		/// </summary>
+		/// <returns>Async task</returns>
+		public async Task StopAsync()
+		{
+			if (DataWriter != null)
+			{
+				DataChannel.Writer.TryComplete();
+				await DataWriter;
+				DataWriter = null;
+			}
+		}
+
+		/// <summary>
+		/// Dispose of this object. Call StopAsync() to stop asynchronously.
+		/// </summary>
+		public async ValueTask DisposeAsync()
+		{
+			await StopAsync();
+		}
+
+		/// <summary>
+		/// Upload the log data to the server in the background
+		/// </summary>
+		/// <returns>Async task</returns>
+		async Task RunDataWriter()
+		{
+			// Current position and line number in the log file
+			long Offset = 0;
+			int LineIndex = 0;
+
+			// Total number of errors and warnings
+			const int MaxErrors = 50;
+			int NumErrors = 0;
+			const int MaxWarnings = 50;
+			int NumWarnings = 0;
+
+			// Buffers for chunks and events read in a single iteration
+			ArrayBufferWriter<byte> Writer = new ArrayBufferWriter<byte>();
+			List<CreateEventRequest> Events = new List<CreateEventRequest>();
+
+			// Line separator for JSON events
+			byte[] Newline = { (byte)'\n' };
+			JsonEncodedText Timestamp = JsonEncodedText.Encode("");
+
+			// The current jobstep outcome
+			JobStepOutcome PostedOutcome = JobStepOutcome.Success;
+
+			// Whether we've written the flush command
+			for (; ; )
+			{
+				Writer.Clear();
+				Events.Clear();
+
+				// Save off the current line number for sending to the server
+				int InitialLineIndex = LineIndex;
+
+				// Get the next data
+				Task WaitTask = Task.Delay(TimeSpan.FromSeconds(2.0));
+				while (Writer.WrittenCount < 256 * 1024)
+				{
+					QueueItem? Data;
+					if (DataChannel.Reader.TryRead(out Data))
+					{
+						if (Data.Event != null)
+						{
+							if ((Data.Event.Severity == EventSeverity.Warning && ++NumWarnings < MaxWarnings) || (Data.Event.Severity == EventSeverity.Error && ++NumErrors < MaxErrors))
+							{
+								Data.Event.LineIndex = LineIndex;
+								Events.Add(Data.Event);
+							}
+						}
+
+						Writer.Write(Data.Data);
+						Writer.Write(Newline);
+
+						LineIndex++;
+					}
+					else
+					{
+						Task<bool> ReadTask = DataChannel.Reader.WaitToReadAsync().AsTask();
+						if (await Task.WhenAny(ReadTask, WaitTask) == WaitTask)
+						{
+							break;
+						}
+						if (!await ReadTask)
+						{
+							break;
+						}
+					}
+				}
+
+				// Upload it to the server
+				if (Writer.WrittenCount > 0)
+				{
+					byte[] Data = Writer.WrittenSpan.ToArray();
+					try
+					{
+						await RpcClient.InvokeAsync(x => x.WriteOutputAsync(new WriteOutputRequest(LogId, Offset, InitialLineIndex, Data, false)), new RpcContext(), CancellationToken.None);
+					}
+					catch (Exception Ex)
+					{
+						Inner.LogWarning(Ex, "Unable to write data to server (log {LogId}, offset {Offset})", LogId, Offset);
+					}
+					Offset += Data.Length;
+				}
+
+				// Write all the events
+				if (Events.Count > 0)
+				{
+					try
+					{
+						await RpcClient.InvokeAsync(x => x.CreateEventsAsync(new CreateEventsRequest(Events)), new RpcContext(), CancellationToken.None);
+					}
+					catch (Exception Ex)
+					{
+						Inner.LogWarning(Ex, "Unable to create events");
+					}
+				}
+
+				// Update the outcome of this jobstep
+				if (JobId != null && JobBatchId != null && JobStepId != null && Outcome != PostedOutcome)
+				{
+					try
+					{
+						await RpcClient.InvokeAsync(x => x.UpdateStepAsync(new UpdateStepRequest(JobId, JobBatchId, JobStepId, JobStepState.Unspecified, Outcome)), new RpcContext(), CancellationToken.None);
+					}
+					catch (Exception Ex)
+					{
+						Inner.LogWarning(Ex, "Unable to update step outcome to {NewOutcome}", Outcome);
+					}
+					PostedOutcome = Outcome;
+				}
+
+				// Wait for more data to be available
+				if (!await DataChannel.Reader.WaitToReadAsync())
+				{
+					try
+					{
+						await RpcClient.InvokeAsync(x => x.WriteOutputAsync(new WriteOutputRequest(LogId, Offset, LineIndex, Array.Empty<byte>(), true)), new RpcContext(), CancellationToken.None);
+					}
+					catch (Exception Ex)
+					{
+						Inner.LogWarning(Ex, "Unable to flush data to server (log {LogId}, offset {Offset})", LogId, Offset);
+					}
+					break;
+				}
+			}
+		}
+	}
+}
diff --git a/Engine/Source/Programs/Horde/Horde.Agent/Utility/NetworkShare.cs b/Engine/Source/Programs/Horde/Horde.Agent/Utility/NetworkShare.cs
new file mode 100644
--- /dev/null
+++ b/Engine/Source/Programs/Horde/Horde.Agent/Utility/NetworkShare.cs
@@ -0,0 +1,131 @@
+// Copyright Epic Games, Inc. All Rights Reserved.
+
+using System;
+using System.Collections.Generic;
+using System.ComponentModel;
+using System.Diagnostics.CodeAnalysis;
+using System.Runtime.InteropServices;
+using System.Text;
+
+namespace HordeAgent.Utility
+{
+	/// <summary>
+	/// Utility functions for mapping network shares
+	/// </summary>
+	static class NetworkShare
+	{
+		[StructLayout(LayoutKind.Sequential, CharSet = CharSet.Unicode)]
+		private struct NETRESOURCE
+		{
+			public uint dwScope;
+			public uint dwType;
+			public uint dwDisplayType;
+			public uint dwUsage;
+			public string? lpLocalName;
+			public string? lpRemoteName;
+			public string? lpComment;
+			public string? lpProvider;
+		};
+
+		const uint RESOURCETYPE_DISK = 1;
+		const uint CONNECT_TEMPORARY = 4;
+
+
+		const int ERROR_ALREADY_ASSIGNED = 85;
+		const int ERROR_MORE_DATA = 234;
+
+		[DllImport("Mpr.dll", CharSet = CharSet.Unicode)]
+		private static extern int WNetAddConnection2W(ref NETRESOURCE lpNetResource, string? lpPassword, string? lpUsername, System.UInt32 dwFlags );
+
+		[DllImport("Mpr.dll", CharSet = CharSet.Unicode)]
+		private static extern int WNetCancelConnectionW(string lpName, [MarshalAs(UnmanagedType.Bool)] bool fForce);
+
+		[DllImport("Mpr.dll", CharSet = CharSet.Unicode)]
+		private static extern int WNetGetConnectionW(string lpLocalName, IntPtr lpRemoteName, ref int lpnLength);
+
+		/// <summary>
+		/// Mounts a network share at the given location
+		/// </summary>
+		/// <param name="MountPoint">Mount point for the share</param>
+		/// <param name="RemotePath">Path to the remote resource</param>
+		public static void Mount(string MountPoint, string RemotePath)
+		{
+			NETRESOURCE NetResource = new NETRESOURCE();
+			NetResource.dwType = RESOURCETYPE_DISK;
+			NetResource.lpLocalName = MountPoint;
+			NetResource.lpRemoteName = RemotePath;
+
+			int Result = WNetAddConnection2W(ref NetResource, null, null, CONNECT_TEMPORARY);
+			if(Result != 0)
+			{
+				if (Result == ERROR_ALREADY_ASSIGNED)
+				{
+					string? CurRemotePath;
+					if (TryGetRemotePath(MountPoint, out CurRemotePath))
+					{
+						if (CurRemotePath.Equals(RemotePath, StringComparison.OrdinalIgnoreCase))
+						{
+							return;
+						}
+						else
+						{
+							throw new Win32Exception(Result, $"Unable to mount network share {RemotePath} as {MountPoint} ({Result}). Currently connected to {CurRemotePath}.");
+						}
+					}
+				}
+				throw new Win32Exception(Result, $"Unable to mount network share {RemotePath} as {MountPoint} ({Result}: {new Win32Exception(Result).Message})");
+			}
+		}
+
+		/// <summary>
+		/// Unmounts a network share
+		/// </summary>
+		/// <param name="MountPoint">The mount point to remove</param>
+		public static void Unmount(string MountPoint)
+		{
+			int Result = WNetCancelConnectionW(MountPoint, true);
+			if(Result != 0)
+			{
+				throw new Win32Exception(Result, $"Unable to unmount {MountPoint} ({Result})");
+			}
+		}
+
+		/// <summary>
+		/// Gets the currently mounted path for a network share
+		/// </summary>
+		/// <param name="MountPoint">The mount point</param>
+		/// <param name="OutRemotePath">Receives the remote path</param>
+		/// <returns>True if the remote path is returned</returns>
+		public static bool TryGetRemotePath(string MountPoint, [NotNullWhen(true)] out string? OutRemotePath)
+		{
+			int Length = 260 * sizeof(char);
+			for (; ; )
+			{
+				IntPtr RemotePath = Marshal.AllocHGlobal(Length);
+				try
+				{
+					int PrevLength = Length;
+					int Result = WNetGetConnectionW(MountPoint, RemotePath, ref Length);
+					if (Result == 0)
+					{
+						OutRemotePath = Marshal.PtrToStringUni(RemotePath)!;
+						return OutRemotePath != null;
+					}
+					else if(Result == ERROR_MORE_DATA && Length > PrevLength)
+					{
+						continue;
+					}
+					else
+					{
+						OutRemotePath = null!;
+						return false;
+					}
+				}
+				finally
+				{
+					Marshal.FreeHGlobal(RemotePath);
+				}
+			}
+		}
+	}
+}
diff --git a/Engine/Source/Programs/Horde/Horde.Agent/Utility/RpcConnection.cs b/Engine/Source/Programs/Horde/Horde.Agent/Utility/RpcConnection.cs
new file mode 100644
--- /dev/null
+++ b/Engine/Source/Programs/Horde/Horde.Agent/Utility/RpcConnection.cs
@@ -0,0 +1,790 @@
+// Copyright Epic Games, Inc. All Rights Reserved.
+
+using EpicGames.Core;
+using Grpc.Core;
+using Grpc.Core.Interceptors;
+using Grpc.Net.Client;
+using HordeCommon.Rpc;
+using Microsoft.Extensions.Hosting;
+using Microsoft.Extensions.Logging;
+using OpenTracing;
+using OpenTracing.Util;
+using System;
+using System.Collections.Concurrent;
+using System.Collections.Generic;
+using System.Diagnostics;
+using System.IO;
+using System.Linq;
+using System.Net;
+using System.Net.Http;
+using System.Net.Http.Headers;
+using System.Runtime.CompilerServices;
+using System.Runtime.InteropServices;
+using System.Threading;
+using System.Threading.Tasks;
+
+namespace HordeAgent.Utility
+{
+	/// <summary>
+	/// Ref-counted reference to a HordeRpcClient. Must be disposed after use.
+	/// </summary>
+	interface IRpcClientRef : IDisposable
+	{
+		/// <summary>
+		/// The Grpc channel instance
+		/// </summary>
+		GrpcChannel Channel { get; }
+
+		/// <summary>
+		/// The client instance
+		/// </summary>
+		HordeRpc.HordeRpcClient Client { get; }
+
+		/// <summary>
+		/// Task which completes when the client needs to be disposed of
+		/// </summary>
+		Task DisposingTask { get; }
+	}
+
+	/// <summary>
+	/// Context for an RPC call. Used for debugging reference leaks.
+	/// </summary>
+	class RpcContext
+	{
+		/// <summary>
+		/// Text to display
+		/// </summary>
+		public string Text { get; }
+
+		/// <summary>
+		/// Constructor
+		/// </summary>
+		/// <param name="Text">Text to display</param>
+		public RpcContext(string Text)
+		{
+			this.Text = Text;
+		}
+
+		/// <summary>
+		/// Constructor
+		/// </summary>
+		/// <param name="Line">Line number of the file</param>
+		/// <param name="File">Calling file</param>
+		public RpcContext([CallerLineNumber] int Line = 0, [CallerFilePath] string File = "")
+		{
+			this.Text = $"{File}({Line})";
+		}
+	}
+
+	internal interface IRpcConnection : IAsyncDisposable
+	{
+		/// <summary>
+		/// Attempts to get a client reference, returning immediately if there's not one available
+		/// </summary>
+		/// <returns></returns>
+		IRpcClientRef? TryGetClientRef(RpcContext Context);
+
+		/// <summary>
+		/// Obtains a new client reference object
+		/// </summary>
+		/// <param name="Context">Context for the call, for debugging</param>
+		/// <param name="CancellationToken">Cancellation token for this request</param>
+		/// <returns>New client reference</returns>
+		Task<IRpcClientRef> GetClientRef(RpcContext Context, CancellationToken CancellationToken);
+
+		/// <summary>
+		/// Invokes an asynchronous command with a HordeRpcClient instance
+		/// </summary>
+		/// <param name="Func">The function to execute</param>
+		/// <param name="Context">Context for the call, for debugging</param>
+		/// <param name="CancellationToken">Cancellation token for the call</param>
+		/// <returns>Response from the call</returns>
+		Task<T> InvokeOnceAsync<T>(Func<HordeRpc.HordeRpcClient, Task<T>> Func, RpcContext Context, CancellationToken CancellationToken);
+
+		/// <summary>
+		/// Obtains a client reference and executes a unary RPC
+		/// </summary>
+		/// <typeparam name="T">Return type from the call</typeparam>
+		/// <param name="Func">Method to execute with the RPC instance</param>
+		/// <param name="Context">Context for the call, for debugging</param>
+		/// <param name="CancellationToken">Cancellation token for the call</param>
+		/// <returns>Response from the call</returns>
+		Task<T> InvokeOnceAsync<T>(Func<HordeRpc.HordeRpcClient, AsyncUnaryCall<T>> Func, RpcContext Context, CancellationToken CancellationToken);
+
+		/// <summary>
+		/// Invokes an asynchronous command with a HordeRpcClient instance
+		/// </summary>
+		/// <param name="Func">The function to execute</param>
+		/// <param name="Context">Context for the call, for debugging</param>
+		/// <param name="CancellationToken">Cancellation token for the call</param>
+		/// <returns>Response from the call</returns>
+		Task<T> InvokeAsync<T>(Func<HordeRpc.HordeRpcClient, Task<T>> Func, RpcContext Context, CancellationToken CancellationToken);
+
+		/// <summary>
+		/// Obtains a client reference and executes a unary RPC
+		/// </summary>
+		/// <typeparam name="T">Return type from the call</typeparam>
+		/// <param name="Func">Method to execute with the RPC instance</param>
+		/// <param name="Context">Context for the call, for debugging</param>
+		/// <param name="CancellationToken">Cancellation token for the call</param>
+		/// <returns>Response from the call</returns>
+		Task<T> InvokeAsync<T>(Func<HordeRpc.HordeRpcClient, AsyncUnaryCall<T>> Func, RpcContext Context, CancellationToken CancellationToken);
+	}
+
+	/// <summary>
+	/// Represents a connection to a GRPC server, which responds to server shutting down events and attempts to reconnect to another instance.
+	/// </summary>
+	class RpcConnection : IRpcConnection
+	{
+		/// <summary>
+		/// An instance of a connection to a particular server
+		/// </summary>
+		class RpcSubConnection : IAsyncDisposable
+		{
+			/// <summary>
+			/// The connection id
+			/// </summary>
+			public int ConnectionId { get; }
+
+			/// <summary>
+			/// Name of the server
+			/// </summary>
+			public string Name
+			{
+				get;
+			}
+
+			/// <summary>
+			/// The Grpc channel
+			/// </summary>
+			public GrpcChannel Channel { get; }
+
+			/// <summary>
+			/// The RPC interface
+			/// </summary>
+			public HordeRpc.HordeRpcClient Client
+			{
+				get;
+			}
+
+			/// <summary>
+			/// Task which is set to indicate the connection is being disposed
+			/// </summary>
+			public Task DisposingTask => DisposingTaskSource.Task;
+
+			/// <summary>
+			/// Logger for messages about refcount changes
+			/// </summary>
+			ILogger Logger;
+
+			/// <summary>
+			/// The current refcount
+			/// </summary>
+			int RefCount;
+
+			/// <summary>
+			/// Whether the connection is terminating
+			/// </summary>
+			bool bDisposing;
+
+			/// <summary>
+			/// Cancellation token set 
+			/// </summary>
+			TaskCompletionSource<bool> DisposingTaskSource = new TaskCompletionSource<bool>();
+
+			/// <summary>
+			/// Wraps a task allowing the disposer to wait for clients to finish using this connection
+			/// </summary>
+			TaskCompletionSource<bool> DisposedTaskSource = new TaskCompletionSource<bool>();
+
+			/// <summary>
+			/// Constructor
+			/// </summary>
+			/// <param name="ConnectionId">The connection id</param>
+			/// <param name="Name">Name of the server</param>
+			/// <param name="Channel">The channel instance</param>
+			/// <param name="Client">The client instance</param>
+			/// <param name="Logger">Logger for debug messages</param>
+			public RpcSubConnection(int ConnectionId, string Name, GrpcChannel Channel, HordeRpc.HordeRpcClient Client, ILogger Logger)
+			{
+				this.ConnectionId = ConnectionId;
+				this.Name = Name;
+				this.Channel = Channel;
+				this.Client = Client;
+				this.Logger = Logger;
+				this.RefCount = 1;
+			}
+
+			/// <summary>
+			/// Attempts to increment the reference count for this subconnection. Fails if it's already zero.
+			/// </summary>
+			/// <returns>True if a reference was added</returns>
+			public bool TryAddRef()
+			{
+				for (; ; )
+				{
+					int InitialRefCount = RefCount;
+					if (InitialRefCount == 0)
+					{
+						return false;
+					}
+					if (Interlocked.CompareExchange(ref RefCount, InitialRefCount + 1, InitialRefCount) == InitialRefCount)
+					{
+						return true;
+					}
+				}
+			}
+
+			/// <summary>
+			/// Release a reference to this connection
+			/// </summary>
+			public void Release()
+			{
+				int NewRefCount = Interlocked.Decrement(ref RefCount);
+				if (bDisposing)
+				{
+					Logger.LogInformation("Connection {ConnectionId}: Decrementing refcount to rpc server {ServerName} ({RefCount} rpc(s) in flight)", ConnectionId, Name, NewRefCount);
+				}
+				if (NewRefCount == 0)
+				{
+					DisposedTaskSource.SetResult(true);
+				}
+			}
+
+			/// <summary>
+			/// Dispose of this connection, waiting for all references to be released first
+			/// </summary>
+			/// <returns>Async task</returns>
+			public ValueTask DisposeAsync()
+			{
+				Logger.LogInformation("Connection {ConnectionId}: Disposing sub-connection", ConnectionId);
+				DisposingTaskSource.TrySetResult(true);
+				bDisposing = true;
+				Release();
+				return new ValueTask(DisposedTaskSource.Task);
+			}
+		}
+
+		/// <summary>
+		/// Reference to a connection instance
+		/// </summary>
+		class RpcClientRef : IRpcClientRef
+		{
+			/// <summary>
+			/// The subconnection containing the client
+			/// </summary>
+			public RpcSubConnection SubConnection
+			{
+				get;
+			}
+
+			/// <summary>
+			/// Context for the reference
+			/// </summary>
+			public RpcContext Context
+			{
+				get;
+			}
+
+			/// <summary>
+			/// The channel instance
+			/// </summary>
+			public GrpcChannel Channel
+			{
+				get { return SubConnection.Channel; }
+			}
+
+			/// <summary>
+			/// The client instance
+			/// </summary>
+			public HordeRpc.HordeRpcClient Client
+			{
+				get { return SubConnection.Client; }
+			}
+
+			/// <summary>
+			/// Task which completes when the connection should be disposed
+			/// </summary>
+			public Task DisposingTask
+			{
+				get { return SubConnection.DisposingTask; }
+			}
+
+			/// <summary>
+			/// Tracks all the active refs
+			/// </summary>
+			static List<RpcClientRef> CurrentRefs = new List<RpcClientRef>();
+
+			/// <summary>
+			/// Constructor
+			/// </summary>
+			/// <param name="SharedClient">The subconnection containing the client</param>
+			/// <param name="Context">Description of the reference, for debugging</param>
+			public RpcClientRef(RpcSubConnection SharedClient, RpcContext Context)
+			{
+				this.SubConnection = SharedClient;
+				this.Context = Context;
+
+				lock (CurrentRefs)
+				{
+					CurrentRefs.Add(this);
+				}
+			}
+
+			/// <summary>
+			/// Dispose of this reference
+			/// </summary>
+			public void Dispose()
+			{
+				SubConnection.Release();
+
+				lock(CurrentRefs)
+				{
+					CurrentRefs.Remove(this);
+				}
+			}
+
+			/// <summary>
+			/// Gets a list of all the active refs
+			/// </summary>
+			/// <returns>List of ref descriptions</returns>
+			public static List<string> GetActiveRefs()
+			{
+				lock(CurrentRefs)
+				{
+					return CurrentRefs.Select(x => x.Context.Text).ToList();
+				}
+			}
+		}
+
+		/// <summary>
+		/// Delays before retrying RPC calls
+		/// </summary>
+		static TimeSpan[] RetryTimes =
+		{
+			TimeSpan.FromSeconds(1.0),
+			TimeSpan.FromSeconds(10.0),
+			TimeSpan.FromSeconds(30.0),
+		};
+
+		/// <summary>
+		/// Delegate to create a new GrpcChannel
+		/// </summary>
+		Func<GrpcChannel> CreateGrpcChannel;
+
+		/// <summary>
+		/// Source for tokens asking the background thread to terminate
+		/// </summary>
+		TaskCompletionSource<bool> StoppingTaskSource = new TaskCompletionSource<bool>();
+
+		/// <summary>
+		/// Task source for waiting for a new subconnection to be made
+		/// </summary>
+		TaskCompletionSource<RpcSubConnection> SubConnectionTaskSource = new TaskCompletionSource<RpcSubConnection>();
+
+		/// <summary>
+		/// Background task to maintain the connection
+		/// </summary>
+		Task? BackgroundTask;
+
+		/// <summary>
+		/// Logger instance
+		/// </summary>
+		ILogger Logger;
+
+		/// <summary>
+		/// Constructor
+		/// </summary>
+		/// <param name="CreateGrpcChannel">Factory method for creating new GRPC channels</param>
+		/// <param name="Logger">Logger instance</param>
+		public RpcConnection(Func<GrpcChannel> CreateGrpcChannel, ILogger Logger)
+		{
+			this.CreateGrpcChannel = CreateGrpcChannel;
+			this.Logger = Logger;
+
+			BackgroundTask = Task.Run(() => ExecuteAsync());
+		}
+
+		/// <summary>
+		/// Attempts to get a client reference, returning immediately if there's not one available
+		/// </summary>
+		/// <returns></returns>
+		public IRpcClientRef? TryGetClientRef(RpcContext Context)
+		{
+			Task<RpcSubConnection> SubConnectionTask = SubConnectionTaskSource.Task;
+			if (SubConnectionTask.IsCompleted)
+			{
+				RpcSubConnection DefaultClient = SubConnectionTask.Result;
+				if(DefaultClient.TryAddRef())
+				{
+					return new RpcClientRef(DefaultClient, Context);
+				}
+			}
+			return null;
+		}
+
+		/// <summary>
+		/// Obtains a new client reference object
+		/// </summary>
+		/// <param name="Context">Context for the call, for debugging</param>
+		/// <param name="CancellationToken">Cancellation token for this request</param>
+		/// <returns>New client reference</returns>
+		public async Task<IRpcClientRef> GetClientRef(RpcContext Context, CancellationToken CancellationToken)
+		{
+			for (; ; )
+			{
+				// Print a message while we're waiting for a client to be available
+				if (!SubConnectionTaskSource.Task.IsCompleted)
+				{
+					Stopwatch Timer = Stopwatch.StartNew();
+					while (!SubConnectionTaskSource.Task.IsCompleted)
+					{
+						Task Delay = Task.Delay(TimeSpan.FromSeconds(30.0), CancellationToken);
+						if (await Task.WhenAny(Delay, SubConnectionTaskSource.Task, StoppingTaskSource.Task) == Delay)
+						{
+							await Delay; // Allow the cancellation token to throw
+							Logger.LogInformation("Thread stalled for {Time:0.0}s while waiting for IRpcClientRef. Source:\n{Context}\nActive references:\n{References}\nStack trace:\n{StackTrace}", Timer.Elapsed.TotalSeconds, Context.Text, String.Join("\n", RpcClientRef.GetActiveRefs()), Environment.StackTrace);
+						}
+						if (StoppingTaskSource.Task.IsCompleted)
+						{
+							throw new OperationCanceledException("RpcConnection is being terminated");
+						}
+					}
+				}
+
+				// Get the new task and increase the refcount
+				RpcSubConnection DefaultClient = await SubConnectionTaskSource.Task;
+				if (DefaultClient.TryAddRef())
+				{
+					return new RpcClientRef(DefaultClient, Context);
+				}
+			}
+		}
+
+		/// <summary>
+		/// Invokes an asynchronous command with a HordeRpcClient instance
+		/// </summary>
+		/// <param name="Func">The function to execute</param>
+		/// <param name="Context">Context for the call, for debugging</param>
+		/// <param name="CancellationToken">Cancellation token for the call</param>
+		/// <returns>Response from the call</returns>
+		public async Task<T> InvokeOnceAsync<T>(Func<HordeRpc.HordeRpcClient, Task<T>> Func, RpcContext Context, CancellationToken CancellationToken)
+		{
+			using (IRpcClientRef ClientRef = await GetClientRef(Context, CancellationToken))
+			{
+				return await Func(ClientRef.Client);
+			}
+		}
+
+		/// <summary>
+		/// Obtains a client reference and executes a unary RPC
+		/// </summary>
+		/// <typeparam name="T">Return type from the call</typeparam>
+		/// <param name="Func">Method to execute with the RPC instance</param>
+		/// <param name="Context">Context for the call, for debugging</param>
+		/// <param name="CancellationToken">Cancellation token for the call</param>
+		/// <returns>Response from the call</returns>
+		public Task<T> InvokeOnceAsync<T>(Func<HordeRpc.HordeRpcClient, AsyncUnaryCall<T>> Func, RpcContext Context, CancellationToken CancellationToken)
+		{
+			return InvokeOnceAsync(async (x) => await Func(x), Context, CancellationToken);
+		}
+
+		/// <summary>
+		/// Invokes an asynchronous command with a HordeRpcClient instance
+		/// </summary>
+		/// <param name="Func">The function to execute</param>
+		/// <param name="Context">Context for the call, for debugging</param>
+		/// <param name="CancellationToken">Cancellation token for the call</param>
+		/// <returns>Response from the call</returns>
+		public async Task<T> InvokeAsync<T>(Func<HordeRpc.HordeRpcClient, Task<T>> Func, RpcContext Context, CancellationToken CancellationToken)
+		{
+			for (int Attempt = 0; ;Attempt++)
+			{
+				// Attempt the RPC call
+				using (IRpcClientRef ClientRef = await GetClientRef(Context, CancellationToken))
+				{
+					try
+					{
+						return await Func(ClientRef.Client);
+					}
+					catch (Exception Ex)
+					{
+						// Otherwise check if we can try again
+						if (Attempt >= RetryTimes.Length)
+						{
+							Logger.LogError("Max number of retry attempts reached");
+							throw;
+						}
+
+						// Check if the user requested cancellation
+						if (CancellationToken.IsCancellationRequested)
+						{
+							Logger.LogInformation("Cancellation of gRPC call requested");
+							throw;
+						}
+
+						// Check we can retry it
+						if (!CanRetry(Ex))
+						{
+							Logger.LogError(Ex, "Exception during RPC: {Message}", Ex.Message);
+							throw;
+						}
+
+						RpcException? RpcEx = Ex as RpcException;
+						if (RpcEx != null && RpcEx.StatusCode == StatusCode.Unavailable)
+						{
+							Logger.LogInformation(Ex, "Failure #{Attempt} during gRPC call (service unavailable). Retrying...", Attempt + 1);
+						}
+						else
+						{
+							Logger.LogError(Ex, "Failure #{Attempt} during gRPC call.", Attempt + 1);
+						}
+					}
+				}
+
+				// Wait before retrying
+				await Task.Delay(RetryTimes[Attempt]);
+			}
+		}
+
+		/// <summary>
+		/// Determines if the given exception should be retried
+		/// </summary>
+		/// <param name="Ex"></param>
+		/// <returns></returns>
+		public static bool CanRetry(Exception Ex)
+		{
+			RpcException? RpcEx = Ex as RpcException;
+			if(RpcEx != null)
+			{
+				return RpcEx.StatusCode == StatusCode.Cancelled || RpcEx.StatusCode == StatusCode.Unavailable || RpcEx.StatusCode == StatusCode.DataLoss /* Interrupted streaming call */;
+			}
+			else
+			{
+				return false;
+			}
+		}
+
+		/// <summary>
+		/// Obtains a client reference and executes a unary RPC
+		/// </summary>
+		/// <typeparam name="T">Return type from the call</typeparam>
+		/// <param name="Func">Method to execute with the RPC instance</param>
+		/// <param name="Context">Context for the call, for debugging</param>
+		/// <param name="CancellationToken">Cancellation token for the call</param>
+		/// <returns>Response from the call</returns>
+		public Task<T> InvokeAsync<T>(Func<HordeRpc.HordeRpcClient, AsyncUnaryCall<T>> Func, RpcContext Context, CancellationToken CancellationToken)
+		{
+			return InvokeAsync(async (x) => await Func(x), Context, CancellationToken);
+		}
+
+		/// <summary>
+		/// Background task to monitor the server state
+		/// </summary>
+		/// <returns>Async task</returns>
+		async Task ExecuteAsync()
+		{
+			// Counter for the connection number. Used to track log messages through async threads.
+			int ConnectionId = 0;
+
+			// This task source is marked as completed once a RPC connection receives a notification from the server that it's shutting down.
+			// Another connection will immediately be allocated for other connection requests to use, and the load balancer should route us to
+			// the new server.
+			TaskCompletionSource<bool>? ReconnectTaskSource = null;
+
+			// List of connection handling async tasks we're tracking
+			List<Task> Tasks = new List<Task>();
+			while (Tasks.Count > 0 || !StoppingTaskSource.Task.IsCompleted)
+			{
+				// Get the task indicating we want a new connection (or create a new one)
+				Task? ReconnectTask = null;
+				if (!StoppingTaskSource.Task.IsCompleted)
+				{
+					if (ReconnectTaskSource == null || ReconnectTaskSource.Task.IsCompleted)
+					{
+						// Need to avoid lambda capture of the reconnect task source
+						int NewConnectionId = ++ConnectionId;
+						TaskCompletionSource<bool> NewReconnectTaskSource = new TaskCompletionSource<bool>();
+						Tasks.Add(Task.Run(() => HandleConnectionAsync(NewConnectionId, NewReconnectTaskSource)));
+						ReconnectTaskSource = NewReconnectTaskSource;
+					}
+					ReconnectTask = ReconnectTaskSource.Task;
+				}
+
+				// Build a list of tasks to wait for
+				List<Task> WaitTasks = new List<Task>(Tasks);
+				if (ReconnectTask != null)
+				{
+					WaitTasks.Add(ReconnectTask);
+				}
+				await Task.WhenAny(WaitTasks);
+
+				// Remove any complete tasks, logging any exceptions
+				for (int Idx = 0; Idx < Tasks.Count; Idx++)
+				{
+					Task Task = Tasks[Idx];
+					if(Task.IsCompleted)
+					{
+						await Task;
+						Tasks.RemoveAt(Idx--);
+					}
+				}
+			}
+		}
+
+		/// <summary>
+		/// Background task to monitor the server state
+		/// </summary>
+		/// <param name="ConnectionId">The connection id</param>
+		/// <param name="ReconnectTaskSource">Task source for reconnecting</param>
+		/// <returns>Async task</returns>
+		async Task HandleConnectionAsync(int ConnectionId, TaskCompletionSource<bool> ReconnectTaskSource)
+		{
+			using IScope Scope = GlobalTracer.Instance.BuildSpan("RpcConnection {ConnectionId}").StartActive();
+
+			Logger.LogInformation("Connection {ConnectionId}: Creating connection", ConnectionId);
+			try
+			{
+				await HandleConnectionInternalAsync(ConnectionId, ReconnectTaskSource);
+			}
+			catch (RpcException Ex) when (Ex.StatusCode == StatusCode.Unavailable)
+			{
+				Logger.LogInformation("Connection {ConnectionId}: Rpc service is unavailable. Reconnecting.", ConnectionId);
+			}
+			catch (Exception Ex)
+			{
+				Logger.LogError(Ex, "Connection {ConnectionId}: Exception handling connection: {Message}", ConnectionId, Ex.Message);
+			}
+			finally
+			{
+				TriggerReconnect(ReconnectTaskSource);
+			}
+		}
+
+		/// <summary>
+		/// Background task to monitor the server state
+		/// </summary>
+		/// <returns>Async task</returns>
+		async Task HandleConnectionInternalAsync(int ConnectionId, TaskCompletionSource<bool> ReconnectTaskSource)
+		{
+			using (GrpcChannel Channel = CreateGrpcChannel())
+			{
+				HordeRpc.HordeRpcClient Client = new HordeRpc.HordeRpcClient(Channel);//.Intercept(new ClientTracingInterceptor(GlobalTracer.Instance)));
+
+				RpcSubConnection? SubConnection = null;
+				try
+				{
+					while (!ReconnectTaskSource.Task.IsCompleted && !StoppingTaskSource.Task.IsCompleted)
+					{
+						// Client HAS TO BE the one to end the call, by closing the stream. The server must keep receiving requests from us until it does.
+						using (AsyncDuplexStreamingCall<QueryServerStateRequest, QueryServerStateResponse> Call = Client.QueryServerStateV2())
+						{
+							// Send the name of this agent
+							QueryServerStateRequest Request = new QueryServerStateRequest();
+							Request.Name = Dns.GetHostName();
+							await Call.RequestStream.WriteAsync(Request);
+
+							// Read the server info
+							Task<bool> MoveNextTask = Call.ResponseStream.MoveNext();
+							if (!await MoveNextTask)
+							{
+								Logger.LogError("Connection {ConnectionId}: No response from server", ConnectionId);
+								break;
+							}
+
+							// Read the server response
+							QueryServerStateResponse Response = Call.ResponseStream.Current;
+							if (!Response.Stopping)
+							{
+								// The first time we connect, log the server name and create the subconnection
+								if (SubConnection == null)
+								{
+									Logger.LogInformation("Connection {ConnectionId}: Connected to rpc server {ServerName}", ConnectionId, Response.Name);
+									SubConnection = new RpcSubConnection(ConnectionId, Response.Name, Channel, Client, Logger);
+									SubConnectionTaskSource.SetResult(SubConnection);
+								}
+
+								// Wait for the StoppingTask token to be set, or the server to inform that it's shutting down
+								MoveNextTask = Call.ResponseStream.MoveNext();
+								await Task.WhenAny(StoppingTaskSource.Task, Task.Delay(TimeSpan.FromSeconds(45.0)), MoveNextTask);
+
+								// Update the response
+								if (MoveNextTask.IsCompleted && MoveNextTask.Result)
+								{
+									Response = Call.ResponseStream.Current;
+								}
+							}
+
+							// If the server is stopping, start reconnecting
+							if (Response.Stopping)
+							{
+								Logger.LogInformation("Connection {ConnectionId}: Server is stopping. Triggering reconnect.", ConnectionId);
+								TriggerReconnect(ReconnectTaskSource);
+							}
+
+							// Close the request stream
+							await Call.RequestStream.CompleteAsync();
+
+							// Wait for the server to finish posting responses
+							while (await MoveNextTask)
+							{
+								Response = Call.ResponseStream.Current;
+								MoveNextTask = Call.ResponseStream.MoveNext();
+							}
+						}
+					}
+					Logger.LogInformation("Connection {ConnectionId}: Closing connection to rpc server", ConnectionId);
+				}
+				finally
+				{
+					if (SubConnection != null)
+					{
+						await SubConnection.DisposeAsync();
+					}
+				}
+
+				// Wait a few seconds before finishing
+				await Task.Delay(TimeSpan.FromSeconds(5.0));
+			}
+		}
+
+		/// <summary>
+		/// Clears the current subconnection and triggers a reconnect
+		/// </summary>
+		/// <param name="ReconnectTaskSource"></param>
+		void TriggerReconnect(TaskCompletionSource<bool> ReconnectTaskSource)
+		{
+			if (!ReconnectTaskSource.Task.IsCompleted)
+			{
+				// First reset the task source, so that nothing new starts to use the current connection
+				if (SubConnectionTaskSource.Task.IsCompleted)
+				{
+					SubConnectionTaskSource = new TaskCompletionSource<RpcSubConnection>();
+				}
+
+				// Now trigger another connection
+				ReconnectTaskSource.SetResult(true);
+			}
+		}
+
+		/// <summary>
+		/// Dispose of this connection
+		/// </summary>
+		/// <returns>Async task</returns>
+		public async ValueTask DisposeAsync()
+		{
+			if (BackgroundTask != null)
+			{
+				StoppingTaskSource.TrySetResult(true);
+
+				await BackgroundTask;
+				BackgroundTask = null!;
+			}
+		}
+
+		public static IRpcConnection Create(Func<GrpcChannel> CreateGrpcChannel, ILogger Logger)
+		{
+			return new RpcConnection(CreateGrpcChannel, Logger);
+		}
+	}
+}
diff --git a/Engine/Source/Programs/Horde/Horde.Agent/Utility/Shutdown.cs b/Engine/Source/Programs/Horde/Horde.Agent/Utility/Shutdown.cs
new file mode 100644
--- /dev/null
+++ b/Engine/Source/Programs/Horde/Horde.Agent/Utility/Shutdown.cs
@@ -0,0 +1,160 @@
+// Copyright Epic Games, Inc. All Rights Reserved.
+
+using Microsoft.Extensions.Logging;
+using Microsoft.Win32.SafeHandles;
+using System;
+using System.Collections.Generic;
+using System.Diagnostics;
+using System.Net.NetworkInformation;
+using System.Runtime.InteropServices;
+using System.Text;
+
+namespace HordeAgent.Utility
+{
+	static class Shutdown
+	{
+		public const UInt32 TOKEN_QUERY = 0x0008;
+		public const UInt32 TOKEN_ADJUST_PRIVILEGES = 0x0020;
+
+		[DllImport("advapi32.dll", SetLastError = true)]
+		[return: MarshalAs(UnmanagedType.Bool)]
+		static extern bool OpenProcessToken(IntPtr ProcessHandle, UInt32 DesiredAccess, out IntPtr TokenHandle);
+
+		[StructLayout(LayoutKind.Sequential)]
+		public struct LUID
+		{
+			public UInt32 LowPart;
+			public Int32 HighPart;
+		}
+
+		[DllImport("advapi32.dll")]
+		static extern bool LookupPrivilegeValue(string? lpSystemName, string lpName, ref LUID lpLuid);
+
+		[StructLayout(LayoutKind.Sequential, Pack = 4)]
+		public struct LUID_AND_ATTRIBUTES
+		{
+			public LUID Luid;
+			public UInt32 Attributes;
+		}
+
+		struct TOKEN_PRIVILEGES
+		{
+			public int PrivilegeCount;
+			[MarshalAs(UnmanagedType.ByValArray, SizeConst = 1)]
+			public LUID_AND_ATTRIBUTES[] Privileges;
+		}
+
+		// Use this signature if you do not want the previous state
+		[DllImport("advapi32.dll", SetLastError = true)]
+		[return: MarshalAs(UnmanagedType.Bool)]
+		static extern bool AdjustTokenPrivileges(IntPtr TokenHandle,
+		   [MarshalAs(UnmanagedType.Bool)]bool DisableAllPrivileges,
+		   ref TOKEN_PRIVILEGES NewState,
+		   UInt32 Zero,
+		   IntPtr Null1,
+		   IntPtr Null2);
+
+		const uint SHTDN_REASON_MAJOR_APPLICATION = 0x00040000;
+		const uint SHTDN_REASON_MINOR_MAINTENANCE = 0x00000001;
+
+		[DllImport("advapi32.dll", CharSet = CharSet.Auto, SetLastError = true)]
+		static extern bool InitiateSystemShutdownEx(string? lpMachineName, string? lpMessage, uint dwTimeout, bool bForceAppsClosed, bool bRebootAfterShutdown, uint dwReason);
+
+		const string SE_SHUTDOWN_NAME = "SeShutdownPrivilege";
+
+		const int SE_PRIVILEGE_ENABLED = 0x00000002;
+
+		/// <summary>
+		/// Initiate a shutdown operation
+		/// </summary>
+		/// <param name="bRestartAfterShutdown">Whether to restart after the shutdown</param>
+		/// <param name="Logger">Logger for the operation</param>
+		/// <returns></returns>
+		public static bool InitiateShutdown(bool bRestartAfterShutdown, ILogger Logger)
+		{
+			if (RuntimeInformation.IsOSPlatform(OSPlatform.Windows))
+			{
+				if (bRestartAfterShutdown)
+				{
+					Logger.LogInformation("Triggering restart");
+				}
+				else
+				{
+					Logger.LogInformation("Triggering shutdown");
+				}
+
+				IntPtr TokenHandle;
+				if (!OpenProcessToken(Process.GetCurrentProcess().Handle, TOKEN_ADJUST_PRIVILEGES | TOKEN_QUERY, out TokenHandle))
+				{
+					Logger.LogError("OpenProcessToken() failed (code 0x{Code:x8})", Marshal.GetLastWin32Error());
+					return false;
+				}
+
+				// Get the LUID for the shutdown privilege. 
+				LUID Luid = new LUID();
+				if (!LookupPrivilegeValue(null, SE_SHUTDOWN_NAME, ref Luid))
+				{
+					Logger.LogError("LookupPrivilegeValue() failed (code 0x{Code:x8})", Marshal.GetLastWin32Error());
+					return false;
+				}
+
+				TOKEN_PRIVILEGES Privileges = new TOKEN_PRIVILEGES();
+				Privileges.PrivilegeCount = 1;
+				Privileges.Privileges = new LUID_AND_ATTRIBUTES[1];
+				Privileges.Privileges[0].Attributes = SE_PRIVILEGE_ENABLED;
+				Privileges.Privileges[0].Luid = Luid;
+
+				if (!AdjustTokenPrivileges(TokenHandle, false, ref Privileges, 0, IntPtr.Zero, IntPtr.Zero))
+				{
+					Logger.LogError("AdjustTokenPrivileges() failed (code 0x{Code:x8})", Marshal.GetLastWin32Error());
+					return false;
+				}
+
+				if (!InitiateSystemShutdownEx(null, "HordeAgent has initiated shutdown", 10, true, bRestartAfterShutdown, SHTDN_REASON_MAJOR_APPLICATION | SHTDN_REASON_MINOR_MAINTENANCE))
+				{
+					Logger.LogError("Shutdown failed (0x{Code:x8})", Marshal.GetLastWin32Error());
+					return false;
+				}
+				return true;
+			}
+			else if (RuntimeInformation.IsOSPlatform(OSPlatform.Linux))
+			{
+				String ShutdownArgs;
+				if (bRestartAfterShutdown)
+                {
+					ShutdownArgs = "sudo shutdown -r +0 \"Horde Agent is restarting\"";
+				}
+				else
+                {
+					ShutdownArgs = "sudo shutdown +0 \"Horde Agent is shutting down\"";
+				}
+ 
+				Process ShutdownProcess = new Process()
+				{
+					StartInfo = new ProcessStartInfo
+					{
+						FileName = "/bin/sh",
+						Arguments = String.Format("-c \"{0}\"", ShutdownArgs),
+						UseShellExecute = false,
+						CreateNoWindow = true
+					}
+				};
+ 
+				ShutdownProcess.Start();
+				ShutdownProcess.WaitForExit();
+				int ExitCode = ShutdownProcess.ExitCode;
+				if (ExitCode != 0)
+				{
+					Logger.LogError("Shutdown failed ({0})", ExitCode);
+					return false;
+				}
+				return true;
+			}
+			else
+			{
+				Logger.LogError("Shutdown is not implemented on this platform");
+				return false;
+			}
+		}
+	}
+}
diff --git a/Engine/Source/Programs/Horde/Horde.Agent/Utility/WindowsServiceManager.cs b/Engine/Source/Programs/Horde/Horde.Agent/Utility/WindowsServiceManager.cs
new file mode 100644
--- /dev/null
+++ b/Engine/Source/Programs/Horde/Horde.Agent/Utility/WindowsServiceManager.cs
@@ -0,0 +1,363 @@
+// Copyright Epic Games, Inc. All Rights Reserved.
+
+using Microsoft.Win32.SafeHandles;
+using Serilog;
+using System;
+using System.Collections.Generic;
+using System.ComponentModel;
+using System.Diagnostics;
+using System.Reflection;
+using System.Runtime.InteropServices;
+using System.Text;
+using System.Threading;
+
+namespace HordeAgent.Utility
+{
+	/// <summary>
+	/// Native methods for manipulating services on Windows
+	/// </summary>
+	static partial class Native
+	{
+#pragma warning disable CS0649
+		public class ServiceHandle : SafeHandleZeroOrMinusOneIsInvalid
+		{
+			public ServiceHandle()
+				: base(true)
+			{
+			}
+
+			public ServiceHandle(IntPtr Handle)
+				: base(true)
+			{
+				SetHandle(Handle);
+			}
+
+			protected override bool ReleaseHandle()
+			{
+				return Native.CloseServiceHandle(handle);
+			}
+		}
+
+		public const uint SC_MANAGER_ALL_ACCESS = 0xF003F;
+
+		[DllImport("advapi32.dll", EntryPoint = "OpenSCManagerW", ExactSpelling = true, CharSet = CharSet.Unicode, SetLastError = true)]
+		public static extern ServiceHandle OpenSCManager(string? machineName, string? databaseName, uint dwAccess);
+
+		public const uint SERVICE_ALL_ACCESS = 0xf01ff;
+
+		public const uint SERVICE_WIN32_OWN_PROCESS = 0x00000010;
+
+		public const uint SERVICE_AUTO_START = 0x00000002;
+		public const uint SERVICE_DEMAND_START = 0x00000003;
+
+		public const uint SERVICE_ERROR_NORMAL = 0x00000001;
+
+		[DllImport("advapi32.dll", SetLastError = true, CharSet = CharSet.Auto)]
+		public static extern ServiceHandle CreateService(ServiceHandle hSCManager, string lpServiceName, string lpDisplayName, uint dwDesiredAccess, uint dwServiceType, uint dwStartType, uint dwErrorControl, string lpBinaryPathName, string? lpLoadOrderGroup, string? lpdwTagId, string? lpDependencies, string? lpServiceStartName, string? lpPassword);
+
+		[DllImport("advapi32.dll", SetLastError = true, CharSet = CharSet.Auto)]
+		public static extern ServiceHandle OpenService(ServiceHandle hSCManager, string lpServiceName, uint dwDesiredAccess);
+
+		[DllImport("advapi32", SetLastError = true)]
+		[return: MarshalAs(UnmanagedType.Bool)]
+		public static extern bool StartService(ServiceHandle hService, int dwNumServiceArgs, string[]? lpServiceArgVectors);
+
+		public const int ERROR_ACCESS_DENIED = 5;
+		public const int ERROR_SERVICE_DOES_NOT_EXIST = 1060;
+		public const int ERROR_SERVICE_NOT_ACTIVE = 1062;
+
+		[DllImport("advapi32.dll", SetLastError = true)]
+		[return: MarshalAs(UnmanagedType.Bool)]
+		public static extern bool CloseServiceHandle(IntPtr hSCObject);
+
+		[DllImport("advapi32.dll", SetLastError = true)]
+		[return: MarshalAs(UnmanagedType.Bool)]
+		public static extern bool DeleteService(ServiceHandle hService);
+
+		public const int SC_STATUS_PROCESS_INFO = 0;
+
+		public struct SERVICE_STATUS_PROCESS
+		{
+			public uint dwServiceType;
+			public uint dwCurrentState;
+			public uint dwControlsAccepted;
+			public uint dwWin32ExitCode;
+			public uint dwServiceSpecificExitCode;
+			public uint dwCheckPoint;
+			public uint dwWaitHint;
+			public uint dwProcessId;
+			public uint dwServiceFlags;
+		}
+
+		[DllImport("advapi32.dll", SetLastError = true)]
+		[return: MarshalAs(UnmanagedType.Bool)]
+		public static extern bool QueryServiceStatusEx(ServiceHandle hService, int InfoLevel, ref SERVICE_STATUS_PROCESS pBuffer, int cbBufSize, out int pcbBytesNeeded);
+
+		public struct SERVICE_STATUS
+		{
+			public uint dwServiceType;
+			public uint dwCurrentState;
+			public uint dwControlsAccepted;
+			public uint dwWin32ExitCode;
+			public uint dwServiceSpecificExitCode;
+			public uint dwCheckPoint;
+			public uint dwWaitHint;
+		}
+
+		public const int SERVICE_CONTROL_STOP = 1;
+
+		[DllImport("advapi32.dll", SetLastError = true)]
+		[return: MarshalAs(UnmanagedType.Bool)]
+		public static extern bool ControlService(ServiceHandle hService, int dwControl, ref SERVICE_STATUS lpServiceStatus);
+
+		public const int SERVICE_CONFIG_DESCRIPTION = 1;
+
+		[StructLayout(LayoutKind.Sequential)]
+		public struct SERVICE_DESCRIPTION
+		{
+			[MarshalAs(UnmanagedType.LPWStr)]
+			public String lpDescription;
+		}
+
+		[DllImport("advapi32.dll", SetLastError = true, CharSet = CharSet.Auto)]
+		[return: MarshalAs(UnmanagedType.Bool)]
+		public static extern bool ChangeServiceConfig2(ServiceHandle hService, int dwInfoLevel, ref SERVICE_DESCRIPTION lpInfo);
+#pragma warning restore CS0649
+	}
+
+	enum WindowsServiceStatus : uint
+	{
+		/// <summary>
+		/// The service has stopped.
+		/// </summary>
+		Stopped = 1,
+
+		/// <summary>
+		/// The service is starting.
+		/// </summary>
+		Starting = 2,
+
+		/// <summary>
+		/// The service is stopping.
+		/// </summary>
+		Stopping = 3,
+
+		/// <summary>
+		/// The service is running
+		/// </summary>
+		Running = 4,
+
+		/// <summary>
+		/// The service is about to continue.
+		/// </summary>
+		Continuing = 5,
+
+		/// <summary>
+		/// The service is pausing
+		/// </summary>
+		Pausing = 6,
+
+		/// <summary>
+		/// The service is paused.
+		/// </summary>
+		Paused = 7,
+	}
+
+	/// <summary>
+	/// Wrapper around 
+	/// </summary>
+	class WindowsService : IDisposable
+	{
+		/// <summary>
+		/// Handle to the service
+		/// </summary>
+		Native.ServiceHandle Handle;
+
+		/// <summary>
+		/// Constructor
+		/// </summary>
+		/// <param name="Handle">Handle to the service</param>
+		public WindowsService(Native.ServiceHandle Handle)
+		{
+			this.Handle = Handle;
+		}
+
+		/// <summary>
+		/// Determines if the service is valid
+		/// </summary>
+		public bool IsValid
+		{
+			get { return !Handle.IsInvalid; }
+		}
+
+		/// <summary>
+		/// Gets the current service status
+		/// </summary>
+		/// <returns>Status code</returns>
+		public WindowsServiceStatus GetStatus()
+		{
+			Native.SERVICE_STATUS_PROCESS Status = new Native.SERVICE_STATUS_PROCESS();
+			if (!Native.QueryServiceStatusEx(Handle, Native.SC_STATUS_PROCESS_INFO, ref Status, Marshal.SizeOf(Status), out int BytesNeeded))
+			{
+				throw new Win32Exception(String.Format("Unable to query process status (0x{0:X8)}", Marshal.GetLastWin32Error()));
+			}
+			return (WindowsServiceStatus)Status.dwCurrentState;
+		}
+
+		/// <summary>
+		/// Waits for the status to change
+		/// </summary>
+		/// <param name="TransitionStatus">Expected status while transitioning</param>
+		/// <param name="MaxWaitTime">Maximum time to wait</param>
+		public WindowsServiceStatus WaitForStatusChange(WindowsServiceStatus TransitionStatus, TimeSpan MaxWaitTime)
+		{
+			Stopwatch Timer = Stopwatch.StartNew();
+			for(; ;)
+			{
+				WindowsServiceStatus Status = GetStatus();
+				if (Status != TransitionStatus || Timer.Elapsed > MaxWaitTime)
+				{
+					return Status;
+				}
+				else
+				{
+					Thread.Sleep(1000);
+				}
+			}
+		}
+
+		/// <summary>
+		/// Sets the service description
+		/// </summary>
+		/// <param name="Description">Description for the service</param>
+		public void SetDescription(string Description)
+		{
+			Native.SERVICE_DESCRIPTION DescriptionData = new Native.SERVICE_DESCRIPTION();
+			DescriptionData.lpDescription = Description;
+			Native.ChangeServiceConfig2(Handle, Native.SERVICE_CONFIG_DESCRIPTION, ref DescriptionData);
+		}
+
+		/// <summary>
+		/// Starts the service
+		/// </summary>
+		public void Start()
+		{
+			if (!Native.StartService(Handle, 0, null))
+			{
+				throw new Win32Exception(String.Format("Unable to start service (error 0x{0:X8})", Marshal.GetLastWin32Error()));
+			}
+		}
+
+		/// <summary>
+		/// Stops the service
+		/// </summary>
+		public void Stop()
+		{
+			Native.SERVICE_STATUS Status = new Native.SERVICE_STATUS();
+			if (!Native.ControlService(Handle, Native.SERVICE_CONTROL_STOP, ref Status))
+			{
+				int Error = Marshal.GetLastWin32Error();
+				if (Error != Native.ERROR_SERVICE_NOT_ACTIVE)
+				{
+					throw new Win32Exception(String.Format("Unable to stop service (error 0x{0:X8})", Error));
+				}
+			}
+		}
+
+		/// <summary>
+		/// Deletes the service
+		/// </summary>
+		public void Delete()
+		{
+			if (!Native.DeleteService(Handle))
+			{
+				throw new Win32Exception(String.Format("Unable to delete service (error 0x{0:X8})", Marshal.GetLastWin32Error()));
+			}
+		}
+
+		/// <summary>
+		/// Dispose of the service handle
+		/// </summary>
+		public void Dispose()
+		{
+			Handle.Close();
+		}
+	}
+
+	/// <summary>
+	/// Helper functionality for manipulating Windows services
+	/// </summary>
+	class WindowsServiceManager : IDisposable
+	{
+		/// <summary>
+		/// Native handle to the service manager
+		/// </summary>
+		Native.ServiceHandle ServiceManagerHandle;
+
+		/// <summary>
+		/// Constructor. Opens a handle to the service manager.
+		/// </summary>
+		public WindowsServiceManager()
+		{
+			ServiceManagerHandle = Native.OpenSCManager(null, null, Native.SC_MANAGER_ALL_ACCESS);
+			if (ServiceManagerHandle.IsInvalid)
+			{
+				int ErrorCode = Marshal.GetLastWin32Error();
+				if (ErrorCode == Native.ERROR_ACCESS_DENIED)
+				{
+					throw new Win32Exception("Unable to open service manager (access denied). Check you're running as administrator.");
+				}
+				else
+				{
+					throw new Win32Exception(String.Format("Unable to open service manager (0x{0:X8)).", ErrorCode));
+				}
+			}
+		}
+
+		/// <summary>
+		/// Dispose of this object
+		/// </summary>
+		public void Dispose()
+		{
+			ServiceManagerHandle.Close();
+		}
+
+		/// <summary>
+		/// Opens a service with the given name
+		/// </summary>
+		/// <param name="ServiceName">Name of the service</param>
+		/// <returns>New service wrapper</returns>
+		public WindowsService Open(string ServiceName)
+		{
+			Native.ServiceHandle ServiceHandle = Native.OpenService(ServiceManagerHandle, ServiceName, Native.SERVICE_ALL_ACCESS);
+			if (ServiceHandle.IsInvalid)
+			{
+				int ErrorCode = Marshal.GetLastWin32Error();
+				if (ErrorCode != Native.ERROR_SERVICE_DOES_NOT_EXIST)
+				{
+					throw new Win32Exception("Unable to open handle to service");
+				}
+			}
+			return new WindowsService(ServiceHandle);
+		}
+
+		/// <summary>
+		/// Creates a service with the given settings
+		/// </summary>
+		/// <param name="Name">Name of the service</param>
+		/// <param name="DisplayName">Display name</param>
+		/// <param name="CommandLine">Command line to use when starting the service</param>
+		/// <param name="UserName">Username to run this service as</param>
+		/// <param name="Password">Password for the account the service is to run under</param>
+		/// <returns>New service instance</returns>
+		public WindowsService Create(string Name, string DisplayName, string CommandLine, string? UserName, string? Password)
+		{
+			Native.ServiceHandle NewServiceHandle = Native.CreateService(ServiceManagerHandle, Name, DisplayName, Native.SERVICE_ALL_ACCESS, Native.SERVICE_WIN32_OWN_PROCESS, Native.SERVICE_AUTO_START, Native.SERVICE_ERROR_NORMAL, CommandLine, null, null, null, UserName, Password);
+			if (NewServiceHandle.IsInvalid)
+			{
+				throw new Win32Exception(String.Format("Unable to create service (0x{0:X8})", Marshal.GetLastWin32Error()));
+			}
+			return new WindowsService(NewServiceHandle);
+		}
+	}
+}
diff --git a/Engine/Source/Programs/Horde/Horde.Agent/Utility/WorkspaceInfo.cs b/Engine/Source/Programs/Horde/Horde.Agent/Utility/WorkspaceInfo.cs
new file mode 100644
--- /dev/null
+++ b/Engine/Source/Programs/Horde/Horde.Agent/Utility/WorkspaceInfo.cs
@@ -0,0 +1,396 @@
+// Copyright Epic Games, Inc. All Rights Reserved.
+
+using EpicGames.Perforce.Managed;
+using HordeCommon.Rpc.Messages;
+using Microsoft.Extensions.Logging;
+using System;
+using System.Collections.Generic;
+using System.Diagnostics;
+using System.IO;
+using System.Runtime.InteropServices;
+using System.Text;
+using System.Text.RegularExpressions;
+using System.Threading;
+using System.Threading.Tasks;
+using EpicGames.Core;
+using EpicGames.Perforce;
+using System.Linq;
+
+namespace HordeAgent.Utility
+{
+	/// <summary>
+	/// Stores information about a managed Perforce workspace
+	/// </summary>
+	class WorkspaceInfo
+	{
+		/// <summary>
+		/// The perforce connection
+		/// </summary>
+		public PerforceClientConnection PerforceClient
+		{
+			get;
+		}
+
+		/// <summary>
+		/// The Perforce server and port. This is checked to not be null in the constructor.
+		/// </summary>
+		public string ServerAndPort
+		{
+			get { return PerforceClient.ServerAndPort!; }
+		}
+
+		/// <summary>
+		/// The Perforce client name. This is checked to not be null in the constructor.
+		/// </summary>
+		public string ClientName
+		{
+			get { return PerforceClient.ClientName!; }
+		}
+
+		/// <summary>
+		/// The Perforce user name. This is checked to not be null in the constructor.
+		/// </summary>
+		public string UserName
+		{
+			get { return PerforceClient.UserName!; }
+		}
+
+		/// <summary>
+		/// The hostname
+		/// </summary>
+		public string HostName
+		{
+			get;
+		}
+
+		/// <summary>
+		/// The current stream
+		/// </summary>
+		public string StreamName
+		{
+			get;
+		}
+
+		/// <summary>
+		/// The directory containing metadata for this workspace
+		/// </summary>
+		public DirectoryReference MetadataDir
+		{
+			get;
+		}
+
+		/// <summary>
+		/// The directory containing workspace data
+		/// </summary>
+		public DirectoryReference WorkspaceDir
+		{
+			get;
+		}
+
+		/// <summary>
+		/// The view for files to sync
+		/// </summary>
+		public IReadOnlyList<string> View
+		{
+			get;
+		}			
+
+		/// <summary>
+		/// Whether untracked files should be removed from this workspace
+		/// </summary>
+		public bool bRemoveUntrackedFiles
+		{
+			get;
+		}
+
+		/// <summary>
+		/// The managed repository instance
+		/// </summary>
+		public ManagedWorkspace Repository
+		{
+			get;
+		}
+
+		/// <summary>
+		/// Constructor
+		/// </summary>
+		/// <param name="Perforce">The perforce connection</param>
+		/// <param name="HostName">Name of this host</param>
+		/// <param name="StreamName">Name of the stream to sync</param>
+		/// <param name="MetadataDir">Path to the metadata directory</param>
+		/// <param name="WorkspaceDir">Path to the workspace directory</param>
+		/// <param name="View">View for files to be synced</param>
+		/// <param name="bRemoveUntrackedFiles">Whether to remove untracked files when syncing</param>
+		/// <param name="Repository">The repository instance</param>
+		public WorkspaceInfo(PerforceClientConnection Perforce, string HostName, string StreamName, DirectoryReference MetadataDir, DirectoryReference WorkspaceDir, IList<string>? View, bool bRemoveUntrackedFiles, ManagedWorkspace Repository)
+		{
+			this.PerforceClient = Perforce;
+
+			if (Perforce.ServerAndPort == null)
+			{
+				throw new ArgumentException("Perforce connection does not have valid server and port");
+			}
+			if (Perforce.UserName == null)
+			{
+				throw new ArgumentException("PerforceConnection does not have valid username");
+			}
+			if (Perforce.ClientName == null)
+			{
+				throw new ArgumentException("PerforceConnection does not have valid client name");
+			}
+
+			this.HostName = HostName;
+			this.StreamName = StreamName;
+			this.MetadataDir = MetadataDir;
+			this.WorkspaceDir = WorkspaceDir;
+			this.View = (View == null) ? new List<string>() : new List<string>(View);
+			this.bRemoveUntrackedFiles = bRemoveUntrackedFiles;
+			this.Repository = Repository;
+		}
+
+		/// <summary>
+		/// Creates a new managed workspace
+		/// </summary>
+		/// <param name="Workspace">The workspace definition</param>
+		/// <param name="RootDir">Root directory for storing the workspace</param>
+		/// <param name="Logger">Logger output</param>
+		/// <param name="CancellationToken">Cancellation token</param>
+		/// <returns>New workspace info</returns>
+		public static async Task<WorkspaceInfo> SetupWorkspaceAsync(AgentWorkspace Workspace, DirectoryReference RootDir, ILogger Logger, CancellationToken CancellationToken)
+		{
+			// Fill in the default credentials iff they are not set
+			string? ServerAndPort = String.IsNullOrEmpty(Workspace.ServerAndPort)? null : Workspace.ServerAndPort;
+			string? UserName = String.IsNullOrEmpty(Workspace.UserName)? null : Workspace.UserName;
+			string? Password = String.IsNullOrEmpty(Workspace.Password)? null : Workspace.Password;
+
+			// Create the connection
+			PerforceConnection Perforce = new PerforceConnection(ServerAndPort, UserName, Logger);
+			if (UserName != null)
+			{
+				if (Password != null)
+				{
+					await Perforce.LoginAsync(Password, CancellationToken);
+				}
+				else
+				{
+					Logger.LogInformation($"Using locally logged in session for {UserName}");
+				}
+			}
+			return await SetupWorkspaceAsync(Perforce, Workspace.Stream, Workspace.Identifier, Workspace.View, !Workspace.Incremental, RootDir, Logger, CancellationToken);
+		}
+
+		/// <summary>
+		/// Creates a new managed workspace
+		/// </summary>
+		/// <param name="Perforce">The perforce connection to use</param>
+		/// <param name="StreamName">The stream being synced</param>
+		/// <param name="Identifier">Identifier to use to share </param>
+		/// <param name="View">View for this workspace</param>
+		/// <param name="bRemoveUntrackedFiles">Whether untracked files should be removed when cleaning this workspace</param>
+		/// <param name="RootDir">Root directory for storing the workspace</param>
+		/// <param name="Logger">Logger output</param>
+		/// <param name="CancellationToken">Cancellation token</param>
+		/// <returns>New workspace info</returns>
+		public static async Task<WorkspaceInfo> SetupWorkspaceAsync(PerforceConnection Perforce, string StreamName, string Identifier, IList<string> View, bool bRemoveUntrackedFiles, DirectoryReference RootDir, ILogger Logger, CancellationToken CancellationToken)
+		{
+			// Get the host name, and fill in any missing metadata about the connection
+			InfoRecord Info = await Perforce.GetInfoAsync(InfoOptions.ShortOutput, CancellationToken);
+
+			if (Perforce.ServerAndPort == null)
+			{
+				Perforce.ServerAndPort = await Perforce.TryGetSettingAsync("p4port", CancellationToken) ?? "perforce:1666";
+				Logger.LogInformation($"Using locally configured Perforce server: '{Perforce.ServerAndPort}'");
+			}
+
+			if (Perforce.UserName == null)
+			{
+				Perforce.UserName = Info.UserName;
+				Logger.LogInformation("Using locally configured and logged in Perforce user: '{UserName}'", Perforce.UserName);
+			}
+
+			string? HostName = Info.ClientHost;
+			if(HostName == null)
+			{
+				throw new Exception("Unable to determine Perforce host name");
+			}
+
+			// replace invalid characters in the workspace identifier with a '+' character
+
+			// append the slot index, if it's non-zero
+			//			my $slot_idx = $optional_arguments->{'slot_idx'} || 0;
+			//			$workspace_identifier .= sprintf("+%02d", $slot_idx) if $slot_idx;
+
+			// if running on an edge server, append the server id to the client name
+			string EdgeSuffix = String.Empty;
+			if(Info.Services != null && Info.ServerID != null)
+			{
+				string[] Services = Info.Services.Split(' ', StringSplitOptions.RemoveEmptyEntries);
+				if (Services.Any(x => x.Equals("edge-server", StringComparison.OrdinalIgnoreCase)))
+				{
+					EdgeSuffix = $"+{Info.ServerID}";
+				}
+			}
+
+			// get all the workspace settings
+			string ClientName = $"Horde+{GetNormalizedHostName(HostName)}+{Identifier}{EdgeSuffix}";
+
+			// Create the client Perforce connection
+			PerforceClientConnection PerforceClient = new PerforceClientConnection(Perforce.ServerAndPort, Perforce.UserName, ClientName, Perforce.Logger);
+
+			// get the workspace names
+			DirectoryReference MetadataDir = DirectoryReference.Combine(RootDir, Identifier);
+			DirectoryReference WorkspaceDir = DirectoryReference.Combine(MetadataDir, "Sync");
+
+			// Create the repository
+			ManagedWorkspace NewRepository = await ManagedWorkspace.LoadOrCreateAsync(HostName, MetadataDir, true, Logger, CancellationToken);
+			if (bRemoveUntrackedFiles)
+			{
+				await NewRepository.DeleteClientAsync(PerforceClient, CancellationToken);
+			}
+			await NewRepository.SetupAsync(PerforceClient, StreamName, CancellationToken);
+
+			// Revert any open files
+			await NewRepository.RevertAsync(PerforceClient, CancellationToken);
+
+			// Create the workspace info
+			return new WorkspaceInfo(PerforceClient, HostName, StreamName, MetadataDir, WorkspaceDir, View, bRemoveUntrackedFiles, NewRepository);
+		}
+
+		/// <summary>
+		/// Gets the latest change in the stream
+		/// </summary>
+		/// <param name="CancellationToken">Cancellation token</param>
+		/// <returns>Latest changelist number</returns>
+		public Task<int> GetLatestChangeAsync(CancellationToken CancellationToken)
+		{
+			return Repository.GetLatestChangeAsync(PerforceClient, StreamName, CancellationToken);
+		}
+
+		/// <summary>
+		/// Revert any open files in the workspace and clean it
+		/// </summary>
+		/// <param name="CancellationToken">Cancellation token</param>
+		/// <returns>Async task</returns>
+		public async Task CleanAsync(CancellationToken CancellationToken)
+		{
+			await Repository.RevertAsync(PerforceClient, CancellationToken);
+			await Repository.CleanAsync(bRemoveUntrackedFiles, CancellationToken);
+		}
+
+		/// <summary>
+		/// Removes the given cache file if it's invalid, and updates the metadata to reflect the version about to be synced
+		/// </summary>
+		/// <param name="CacheFile">Path to the cache file</param>
+		/// <param name="Change">The current change being built</param>
+		/// <param name="PreflightChange">The preflight changelist number</param>
+		/// <returns>Async task</returns>
+		public async Task UpdateLocalCacheMarker(FileReference CacheFile, int Change, int PreflightChange)
+		{
+			// Create the new cache file descriptor
+			string NewDescriptor;
+			if(PreflightChange <= 0)
+			{
+				NewDescriptor = $"CL {Change}";
+			}
+			else
+			{
+				NewDescriptor = $"CL {PreflightChange} with base CL {Change}";
+			}
+
+			// Remove the cache file if the current descriptor doesn't match
+			FileReference DescriptorFile = CacheFile.ChangeExtension(".txt");
+			if (FileReference.Exists(CacheFile))
+			{
+				if (FileReference.Exists(DescriptorFile))
+				{
+					string OldDescriptor = await FileReference.ReadAllTextAsync(DescriptorFile);
+					if (OldDescriptor == NewDescriptor)
+					{
+						return;
+					}
+					else
+					{
+						FileReference.Delete(DescriptorFile);
+					}
+				}
+				FileReference.Delete(CacheFile);
+			}
+
+			// Write the new descriptor file
+			FileReference.WriteAllText(DescriptorFile, NewDescriptor);
+		}
+
+		/// <summary>
+		/// Sync the workspace to a given changelist, and capture it so we can quickly clean in the future
+		/// </summary>
+		/// <param name="Change">The changelist to sync to</param>
+		/// <param name="PreflightChange">Change to preflight, or 0</param>
+		/// <param name="CacheFile">Path to the cache file to use</param>
+		/// <param name="Logger">The logger for output messages</param>
+		/// <param name="CancellationToken">Cancellation token</param>
+		/// <returns>Async task</returns>
+		public async Task SyncAsync(int Change, int PreflightChange, FileReference? CacheFile, ILogger Logger, CancellationToken CancellationToken)
+		{
+			await Repository.SyncAsync(PerforceClient, StreamName, Change, View, bRemoveUntrackedFiles, false, CacheFile, CancellationToken);
+
+			// Purge the cache for incremental workspaces
+			if (!bRemoveUntrackedFiles)
+			{
+				await Repository.PurgeAsync(0, CancellationToken);
+			}
+
+			// if we're running a preflight
+			if (PreflightChange > 0)
+			{
+				await Repository.UnshelveAsync(PerforceClient, StreamName, PreflightChange, CancellationToken);
+			}
+		}
+
+
+		/// <summary>
+		/// Strip the domain suffix to get the host name
+		/// </summary>
+		/// <param name="HostName">Hostname with optional domain</param>
+		/// <returns>Normalized host name</returns>
+		static string GetNormalizedHostName(string HostName)
+		{
+			return Regex.Replace(HostName, @"\..*$", "").ToUpperInvariant();
+		}
+
+		/// <summary>
+		/// Revert all files in a workspace
+		/// </summary>
+		public static async Task RevertAllChangesAsync(PerforceConnection Perforce, ILogger Logger, CancellationToken CancellationToken)
+		{
+			// Make sure the client name is set
+			if(Perforce.ClientName == null)
+			{
+				throw new ArgumentException("RevertAllChangesAsync() requires PerforceConnection with client");
+			}
+
+			// check if there are any open files, and revert them
+			List<FStatRecord> Files = await Perforce.GetOpenFilesAsync(OpenedOptions.None, -1, null, null, 1, new[] { "//..." }, CancellationToken);
+			if (Files.Count > 0)
+			{
+				await Perforce.RevertAsync(-1, null, RevertOptions.KeepWorkspaceFiles, new[] { "//..." }, CancellationToken);
+			}
+
+			// enumerate all the pending changelists
+			List<ChangesRecord> PendingChanges = await Perforce.GetChangesAsync(ChangesOptions.None, Perforce.ClientName, -1, ChangeStatus.Pending, null, FileSpecList.Empty, CancellationToken);
+			foreach (ChangesRecord PendingChange in PendingChanges)
+			{
+				// delete any shelved files if there are any
+				List<DescribeRecord> Records = await Perforce.DescribeAsync(DescribeOptions.Shelved, -1, new int[] { PendingChange.Number }, CancellationToken);
+				if (Records.Count > 0 && Records[0].Files.Count > 0)
+				{
+					Logger.LogInformation($"Deleting shelved files in changelist {PendingChange.Number}");
+					await Perforce.DeleteChangeAsync(DeleteChangeOptions.None, PendingChange.Number, CancellationToken);
+				}
+
+				// delete the changelist
+				Logger.LogInformation($"Deleting changelist {PendingChange.Number}");
+				await Perforce.DeleteChangeAsync(DeleteChangeOptions.None, PendingChange.Number, CancellationToken);
+			}
+		}
+	}
+}
diff --git a/Engine/Source/Programs/Horde/Horde.Agent/appsettings.Production.json b/Engine/Source/Programs/Horde/Horde.Agent/appsettings.Production.json
new file mode 100644
--- /dev/null
+++ b/Engine/Source/Programs/Horde/Horde.Agent/appsettings.Production.json
@@ -0,0 +1,2 @@
+{
+}
diff --git a/Engine/Source/Programs/Horde/Horde.Build/Services/RpcService.cs b/Engine/Source/Programs/Horde/Horde.Build/Services/RpcService.cs
--- a/Engine/Source/Programs/Horde/Horde.Build/Services/RpcService.cs
+++ b/Engine/Source/Programs/Horde/Horde.Build/Services/RpcService.cs
@@ -25,8 +25,6 @@ using HordeServer.Tasks.Impl;
 using System.Threading;
 using HordeServer.Jobs;
 
-[assembly: InternalsVisibleTo("HordeAgentTests")]
-
 namespace HordeServer.Services
 {
 	using AgentSoftwareChannelName = StringId<AgentSoftwareChannels>;
