Horde: More post-rename fixup.

diff --git a/Engine/Source/Programs/Horde/Horde.Build.Tests/Horde.Build.Tests.csproj b/Engine/Source/Programs/Horde/Horde.Build.Tests/Horde.Build.Tests.csproj
--- a/Engine/Source/Programs/Horde/Horde.Build.Tests/Horde.Build.Tests.csproj
+++ b/Engine/Source/Programs/Horde/Horde.Build.Tests/Horde.Build.Tests.csproj
@@ -27,9 +27,7 @@
 
   <ItemGroup>
     <ProjectReference Include="..\..\Shared\EpicGames.Serialization\EpicGames.Serialization.csproj" />
-    <ProjectReference Include="..\HordeAgent\HordeAgent.csproj" Label="HordeAgent">
-      <Aliases>HordeAgent</Aliases>
-    </ProjectReference>
+    <ProjectReference Include="..\Horde.Agent\Horde.Agent.csproj" />
     <ProjectReference Include="..\Horde.Build\Horde.Build.csproj" />
   </ItemGroup>
 
diff --git a/Engine/Source/Programs/Horde/Horde.Build.Tests/IssueServiceTests.cs b/Engine/Source/Programs/Horde/Horde.Build.Tests/IssueServiceTests.cs
--- a/Engine/Source/Programs/Horde/Horde.Build.Tests/IssueServiceTests.cs
+++ b/Engine/Source/Programs/Horde/Horde.Build.Tests/IssueServiceTests.cs
@@ -1,9 +1,8 @@
 // Copyright Epic Games, Inc. All Rights Reserved.
-extern alias HordeAgent;
 
 using EpicGames.Core;
-using HordeAgent::HordeAgent.Parser;
-using HordeAgent::HordeAgent.Utility;
+using HordeAgent.Parser;
+using HordeAgent.Utility;
 using HordeCommon;
 using HordeServer;
 using HordeServer.Collections;
