Horde: Build fix.

diff --git a/Engine/Source/Programs/Horde/Horde.Build/Horde.Build.csproj b/Engine/Source/Programs/Horde/Horde.Build/Horde.Build.csproj
--- a/Engine/Source/Programs/Horde/Horde.Build/Horde.Build.csproj
+++ b/Engine/Source/Programs/Horde/Horde.Build/Horde.Build.csproj
@@ -1,4 +1,4 @@
-﻿<Project Sdk="Microsoft.NET.Sdk.Web">
+<Project Sdk="Microsoft.NET.Sdk.Web">
 
   <PropertyGroup>
     <TargetFramework>net5.0</TargetFramework>
@@ -81,7 +81,6 @@
     <ProjectReference Include="..\HordeCommon\HordeCommon.csproj" />
     <ProjectReference Include="..\..\..\ThirdParty\Perforce\P4Api.Net\p4api.net\p4api.net.csproj" />
     <ProjectReference Include="..\HordeServer.Common\HordeServer.Common.csproj" />
-    <ProjectReference Include="..\HordeServer.Compute\HordeServer.Compute.csproj" />
   </ItemGroup>
 
   <ItemGroup>
