Horde: Build fixes.

diff --git a/Engine/Source/Programs/Horde/Horde.Agent/Horde.Agent.csproj b/Engine/Source/Programs/Horde/Horde.Agent/Horde.Agent.csproj
--- a/Engine/Source/Programs/Horde/Horde.Agent/Horde.Agent.csproj
+++ b/Engine/Source/Programs/Horde/Horde.Agent/Horde.Agent.csproj
@@ -74,11 +74,11 @@
     <Content Include="..\..\Shared\EpicGames.Horde.Common\Protos\**\*.*">
       <Link>Protos\%(RecursiveDir)%(FileName)%(Extension)</Link>
     </Content>
-    <Compile Include="../../../../Platforms/*/Source/Programs/Horde/HordeAgent/**/*.cs">
-      <Link>$([System.Text.RegularExpressions.Regex]::Replace(%(FullPath), '^.+?[\\/]HordeAgent[\\/]', ''))</Link>
+    <Compile Include="../../../../Platforms/*/Source/Programs/Horde/Horde.Agent/**/*.cs">
+      <Link>$([System.Text.RegularExpressions.Regex]::Replace(%(FullPath), '^.+?[\\/]Horde.Agent[\\/]', ''))</Link>
     </Compile>
-    <Compile Include="../../../../Restricted/*/Source/Programs/Horde/HordeAgent/**/*.cs">
-      <Link>$([System.Text.RegularExpressions.Regex]::Replace(%(FullPath), '^.+?[\\/]HordeAgent[\\/]', ''))</Link>
+    <Compile Include="../../../../Restricted/*/Source/Programs/Horde/Horde.Agent/**/*.cs">
+      <Link>$([System.Text.RegularExpressions.Regex]::Replace(%(FullPath), '^.+?[\\/]Horde.Agent[\\/]', ''))</Link>
     </Compile>
   </ItemGroup>
 
