Horde: Rename assembly for Horde server to Horde.Build.

diff --git a/Engine/Source/Programs/Horde/Horde.Build/Horde.Build.csproj b/Engine/Source/Programs/Horde/Horde.Build/Horde.Build.csproj
--- a/Engine/Source/Programs/Horde/Horde.Build/Horde.Build.csproj
+++ b/Engine/Source/Programs/Horde/Horde.Build/Horde.Build.csproj
@@ -4,7 +4,7 @@
     <TargetFramework>net5.0</TargetFramework>
     <Nullable>enable</Nullable>
     <Deterministic>true</Deterministic>
-    <AssemblyName>HordeServer</AssemblyName>
+    <AssemblyName>Horde.Build</AssemblyName>
     <DocumentationFile>HordeServer.xml</DocumentationFile>
     <ApplicationIcon>HordeServer.ico</ApplicationIcon>
   </PropertyGroup>
