Horde: Change documentation file to match assembly name.

diff --git a/Engine/Source/Programs/Horde/Horde.Build/Horde.Build.csproj b/Engine/Source/Programs/Horde/Horde.Build/Horde.Build.csproj
--- a/Engine/Source/Programs/Horde/Horde.Build/Horde.Build.csproj
+++ b/Engine/Source/Programs/Horde/Horde.Build/Horde.Build.csproj
@@ -5,7 +5,7 @@
     <Nullable>enable</Nullable>
     <Deterministic>true</Deterministic>
     <AssemblyName>Horde.Build</AssemblyName>
-    <DocumentationFile>HordeServer.xml</DocumentationFile>
+    <DocumentationFile>Horde.Build.xml</DocumentationFile>
     <ApplicationIcon>HordeServer.ico</ApplicationIcon>
   </PropertyGroup>
 
