Horde: Leave the agent assembly name set to HordeAgent.

diff --git a/Engine/Source/Programs/Horde/Horde.Agent/Horde.Agent.csproj b/Engine/Source/Programs/Horde/Horde.Agent/Horde.Agent.csproj
--- a/Engine/Source/Programs/Horde/Horde.Agent/Horde.Agent.csproj
+++ b/Engine/Source/Programs/Horde/Horde.Agent/Horde.Agent.csproj
@@ -6,6 +6,7 @@
     <Nullable>enable</Nullable>
 	<ApplicationIcon>HordeAgent.ico</ApplicationIcon>
     <ApplicationManifest>app.manifest</ApplicationManifest>
+    <AssemblyName>HordeAgent</AssemblyName>
     <PublishRepositoryUrl>true</PublishRepositoryUrl>
     <EmbedUntrackedSources>false</EmbedUntrackedSources>
   </PropertyGroup>
