Horde: Use the standard message format for logging compute leases.

diff --git a/Engine/Source/Programs/Horde/Horde.Build/Compute/Impl/ComputeService.cs b/Engine/Source/Programs/Horde/Horde.Build/Compute/Impl/ComputeService.cs
--- a/Engine/Source/Programs/Horde/Horde.Build/Compute/Impl/ComputeService.cs
+++ b/Engine/Source/Programs/Horde/Horde.Build/Compute/Impl/ComputeService.cs
@@ -271,7 +271,7 @@ namespace HordeServer.Compute.Impl
 
 		public override Task OnLeaseStartedAsync(IAgent Agent, LeaseId LeaseId, ComputeTaskMessage ComputeTask, ILogger Logger)
 		{
-			Logger.LogInformation("Compute lease started (lease: {LeaseId}, task: {TaskHash}, agent: {AgentId}, channel: {ChannelId})", LeaseId, ComputeTask.Task.Hash, Agent.Id, ComputeTask.ChannelId);
+			base.OnLeaseStartedAsync(Agent, LeaseId, ComputeTask, Logger);
 
 			ComputeTaskStatus Status = new ComputeTaskStatus(ComputeTask.Task, ComputeTaskState.Executing, Agent.Id, LeaseId);
 			return MessageQueue.PostAsync(ComputeTask.ChannelId, Status);
