Horde: Disable debug controllers by default.

diff --git a/Engine/Source/Programs/Horde/Horde.Build/Controllers/DebugController.cs b/Engine/Source/Programs/Horde/Horde.Build/Controllers/DebugController.cs
--- a/Engine/Source/Programs/Horde/Horde.Build/Controllers/DebugController.cs
+++ b/Engine/Source/Programs/Horde/Horde.Build/Controllers/DebugController.cs
@@ -1,5 +1,8 @@
 // Copyright Epic Games, Inc. All Rights Reserved.
 
+//#define ENABLE_PUBLIC_DEBUG_CONTROLLER
+//#define ENABLE_SECURE_DEBUG_CONTROLLER
+
 using EpicGames.Core;
 using HordeServer.Api;
 using HordeCommon;
@@ -48,11 +51,11 @@ namespace HordeServer.Controllers
 	using IStream = HordeServer.Models.IStream;
 	using P4Debugging = Perforce.P4.P4Debugging;
 
+#if ENABLE_PUBLIC_DEBUG_CONTROLLER
 	/// <summary>
 	/// Public endpoints for the debug controller
 	/// </summary>
 	[ApiController]
-	[Route("[controller]")]
 	public class PublicDebugController : ControllerBase
 	{
 		/// <summary>
@@ -252,13 +255,13 @@ namespace HordeServer.Controllers
 		}
 		/**/
 	}
-
+#endif
+#if ENABLE_SECURE_DEBUG_CONTROLLER
 	/// <summary>
 	/// Controller managing account status
 	/// </summary>
 	[ApiController]
 	[Authorize]
-	[Route("[controller]")]
 	public class SecureDebugController : ControllerBase
 	{
 		/// <summary>
@@ -928,4 +931,5 @@ namespace HordeServer.Controllers
 			return PhysicalFile(SnapshotZipFile, "application/zip", Path.GetFileName(SnapshotZipFile));
 		}
 	}
+#endif
 }
