Horde: Adding agent audit view to dashboard
#rnx

diff --git a/Engine/Source/Programs/Horde/HordeDashboard/src/backend/Api.ts b/Engine/Source/Programs/Horde/HordeDashboard/src/backend/Api.ts
--- a/Engine/Source/Programs/Horde/HordeDashboard/src/backend/Api.ts
+++ b/Engine/Source/Programs/Horde/HordeDashboard/src/backend/Api.ts
@@ -382,7 +382,7 @@ export type AuditLogQuery = {
 }
 
 export enum AuditLogLevel {
-    Information = "Information"
+	Information = "Information"
 }
 
 export type AuditLogEntry = {
@@ -390,7 +390,7 @@ export type AuditLogEntry = {
 	level: AuditLogLevel;
 	message: string;
 	format: string;
-	properties?: Record<string, string>;
+	properties?: Record<string, string | number | Record<string, string | number>>;
 }
 
 
@@ -3010,7 +3010,7 @@ export type GetServerSettingsResponse = {
 	globalConfigPath: string;
 
 	/** The server settings on local storage */
-	userServerSettingsPath:string;
+	userServerSettingsPath: string;
 
 	/** MongoDB connection string */
 	databaseConnectionString?: string;
@@ -3034,7 +3034,7 @@ export type GetServerSettingsResponse = {
 	/** Optional PFX certificate to use for encryting agent SSL traffic. This can be a self-signed certificate, as long as it's trusted by agents.	*/
 	serverPrivateCert?: string;
 
-	/** Issuer for tokens from the auth provider */	
+	/** Issuer for tokens from the auth provider */
 	oidcAuthority?: string;
 
 	/** Client id for the OIDC authority */
@@ -3589,10 +3589,10 @@ export type UpdateGlobalConfigRequest = {
 	globalsJson?: string;
 
 	/// projects json
-	projectsJson?:Record<string, string>;
+	projectsJson?: Record<string, string>;
 
 	/// streams json
-	streamsJson?:Record<string, string>;
+	streamsJson?: Record<string, string>;
 
 	/// default pool name
 	defaultPoolName?: string;
