Horde: Fix exception when piping raw output to structured log parser.

diff --git a/Engine/Source/Programs/Horde/Horde.Agent/Parser/LogParser.cs b/Engine/Source/Programs/Horde/Horde.Agent/Parser/LogParser.cs
--- a/Engine/Source/Programs/Horde/Horde.Agent/Parser/LogParser.cs
+++ b/Engine/Source/Programs/Horde/Horde.Agent/Parser/LogParser.cs
@@ -126,7 +126,7 @@ namespace HordeAgent.Parser
 			{
 				Data = Data.Slice(0, Data.Length - 1);
 			}
-			if (Buffer.Length == 0 && Data[0] == '{' && WriteRawEvent(Data))
+			if (Buffer.Length == 0 && Data.Length > 0 && Data[0] == '{' && WriteRawEvent(Data))
 			{
 				return;
 			}
