Horde: Move Grpc health proto into server code.

diff --git a/Engine/Source/Programs/Horde/Horde.Build/Horde.Build.csproj b/Engine/Source/Programs/Horde/Horde.Build/Horde.Build.csproj
--- a/Engine/Source/Programs/Horde/Horde.Build/Horde.Build.csproj
+++ b/Engine/Source/Programs/Horde/Horde.Build/Horde.Build.csproj
@@ -67,7 +67,7 @@
     <Protobuf Include="..\..\Shared\EpicGames.Horde\Protos\horde\storage\blob_store.proto" ProtoRoot="..\..\Shared\EpicGames.Horde\Protos" GrpcServices="Server" />
     <Protobuf Include="..\..\Shared\EpicGames.Horde\Protos\horde\compute\compute_rpc.proto" ProtoRoot="..\..\Shared\EpicGames.Horde\Protos" GrpcServices="Server" />
     <Protobuf Include="..\HordeCommon\Protos\horde\horde_rpc.proto" ProtoRoot="..\HordeCommon\Protos" GrpcServices="Server" />
-    <Protobuf Include="..\HordeCommon\Protos\grpc\health\v1\health.proto" ProtoRoot="..\HordeCommon\Protos" GrpcServices="Server" />
+    <Protobuf Include="Protos\grpc\health\v1\health.proto" ProtoRoot="Protos" GrpcServices="Server" />
   </ItemGroup>
 
   <ItemGroup>
