Missing files.

diff --git a/Engine/Source/Programs/Horde/Horde.Agent/Horde.Agent.csproj b/Engine/Source/Programs/Horde/Horde.Agent/Horde.Agent.csproj
--- a/Engine/Source/Programs/Horde/Horde.Agent/Horde.Agent.csproj
+++ b/Engine/Source/Programs/Horde/Horde.Agent/Horde.Agent.csproj
@@ -75,6 +75,18 @@
     <Content Include="..\..\Shared\EpicGames.Horde.Common\Protos\**\*.*">
       <Link>Protos\%(RecursiveDir)%(FileName)%(Extension)</Link>
     </Content>
+    <Content Include="..\Horde.Build\Protos\horde\horde_rpc_messages.proto">
+      <Link>Protos\horde_rpc_messages.proto</Link>
+    </Content>
+    <Content Include="..\Horde.Build\Protos\horde\common\**">
+      <Link>Protos\horde\common\%(RecursiveDir)%(FileName)%(Extension)</Link>
+    </Content>
+    <Content Include="..\Horde.Build\Protos\horde\messages\**">
+      <Link>Protos\horde\messages\%(RecursiveDir)%(FileName)%(Extension)</Link>
+    </Content>
+    <Content Include="..\Horde.Build\Protos\horde\tasks\**">
+      <Link>Protos\horde\tasks\%(RecursiveDir)%(FileName)%(Extension)</Link>
+    </Content>
     <Compile Include="../../../../Platforms/*/Source/Programs/Horde/Horde.Agent/**/*.cs">
       <Link>$([System.Text.RegularExpressions.Regex]::Replace(%(FullPath), '^.+?[\\/]Horde.Agent[\\/]', ''))</Link>
     </Compile>
@@ -93,7 +105,26 @@
   </ItemGroup>
 
   <ItemGroup>
-    <Protobuf Include="..\HordeCommon\Protos\horde\horde_rpc.proto" ProtoRoot="..\HordeCommon\Protos" GrpcServices="Client" />
+    <Protobuf Include="..\Horde.Build\Protos\horde\horde_rpc.proto" ProtoRoot="..\Horde.Build\Protos" GrpcServices="Client" Access="Internal" />
+    <Protobuf Include="..\Horde.Build\Protos\horde\horde_rpc_messages.proto" ProtoRoot="..\Horde.Build\Protos" Access="Internal" />
+    <Protobuf Include="..\Horde.Build\Protos\horde\common\agent_status.proto" ProtoRoot="..\Horde.Build\Protos" Access="Internal" />
+    <Protobuf Include="..\Horde.Build\Protos\horde\common\event_severity.proto" ProtoRoot="..\Horde.Build\Protos" Access="Internal" />
+    <Protobuf Include="..\Horde.Build\Protos\horde\common\jobstep_outcome.proto" ProtoRoot="..\Horde.Build\Protos" Access="Internal" />
+    <Protobuf Include="..\Horde.Build\Protos\horde\common\jobstep_state.proto" ProtoRoot="..\Horde.Build\Protos" Access="Internal" />
+    <Protobuf Include="..\Horde.Build\Protos\horde\common\label_change.proto" ProtoRoot="..\Horde.Build\Protos" Access="Internal" />
+    <Protobuf Include="..\Horde.Build\Protos\horde\common\lease_outcome.proto" ProtoRoot="..\Horde.Build\Protos" Access="Internal" />
+    <Protobuf Include="..\Horde.Build\Protos\horde\common\lease_state.proto" ProtoRoot="..\Horde.Build\Protos" Access="Internal" />
+    <Protobuf Include="..\Horde.Build\Protos\horde\common\priority.proto" ProtoRoot="..\Horde.Build\Protos" Access="Internal" />
+    <Protobuf Include="..\Horde.Build\Protos\horde\messages\agent_capabilities.proto" ProtoRoot="..\Horde.Build\Protos" Access="Internal" />
+    <Protobuf Include="..\Horde.Build\Protos\horde\messages\agent_requirements.proto" ProtoRoot="..\Horde.Build\Protos" Access="Internal" />
+    <Protobuf Include="..\Horde.Build\Protos\horde\messages\agent_workspace.proto" ProtoRoot="..\Horde.Build\Protos" Access="Internal" />
+    <Protobuf Include="..\Horde.Build\Protos\horde\messages\lease.proto" ProtoRoot="..\Horde.Build\Protos" Access="Internal" />
+    <Protobuf Include="..\Horde.Build\Protos\horde\tasks\compute_task.proto" ProtoRoot="..\Horde.Build\Protos" AdditionalImportDirs="..\..\Shared\EpicGames.Horde\Protos" Access="Internal" />
+    <Protobuf Include="..\Horde.Build\Protos\horde\tasks\conform_task.proto" ProtoRoot="..\Horde.Build\Protos" Access="Internal" />
+    <Protobuf Include="..\Horde.Build\Protos\horde\tasks\job_task.proto" ProtoRoot="..\Horde.Build\Protos" Access="Internal" />
+    <Protobuf Include="..\Horde.Build\Protos\horde\tasks\restart_task.proto" ProtoRoot="..\Horde.Build\Protos" Access="Internal" />
+    <Protobuf Include="..\Horde.Build\Protos\horde\tasks\shutdown_task.proto" ProtoRoot="..\Horde.Build\Protos" Access="Internal" />
+    <Protobuf Include="..\Horde.Build\Protos\horde\tasks\upgrade_task.proto" ProtoRoot="..\Horde.Build\Protos" Access="Internal" />
   </ItemGroup>
 
   <ItemGroup>
diff --git a/Engine/Source/Programs/Horde/Horde.Build/Services/RpcService.cs b/Engine/Source/Programs/Horde/Horde.Build/Services/RpcService.cs
--- a/Engine/Source/Programs/Horde/Horde.Build/Services/RpcService.cs
+++ b/Engine/Source/Programs/Horde/Horde.Build/Services/RpcService.cs
@@ -796,7 +796,7 @@ namespace HordeServer.Services
 				throw new StructuredRpcException(StatusCode.NotFound, "Unable to find step {JobId}:{BatchId}:{StepId}", Job.Id, Batch.Id, StepId);
 			}
 
-			return new GetStepResponse(Step.Outcome, Step.State, Step.AbortRequested);
+			return new GetStepResponse { Outcome = Step.Outcome, State = Step.State, AbortRequested = Step.AbortRequested };
 		}
 
 		/// <summary>
diff --git a/Engine/Source/Programs/Horde/HordeCommon/HordeCommon.csproj b/Engine/Source/Programs/Horde/HordeCommon/HordeCommon.csproj
--- a/Engine/Source/Programs/Horde/HordeCommon/HordeCommon.csproj
+++ b/Engine/Source/Programs/Horde/HordeCommon/HordeCommon.csproj
@@ -19,26 +19,4 @@
   <ItemGroup>
     <ProjectReference Include="..\..\Shared\EpicGames.Horde\EpicGames.Horde.csproj" />
   </ItemGroup>
-
-  <ItemGroup>
-    <Protobuf Include="Protos\horde\horde_rpc_messages.proto" ProtoRoot="Protos" />
-    <Protobuf Include="Protos\horde\common\agent_status.proto" ProtoRoot="Protos" />
-    <Protobuf Include="Protos\horde\common\event_severity.proto" ProtoRoot="Protos" />
-    <Protobuf Include="Protos\horde\common\jobstep_outcome.proto" ProtoRoot="Protos" />
-    <Protobuf Include="Protos\horde\common\jobstep_state.proto" ProtoRoot="Protos" />
-    <Protobuf Include="Protos\horde\common\label_change.proto" ProtoRoot="Protos" />
-    <Protobuf Include="Protos\horde\common\lease_outcome.proto" ProtoRoot="Protos" />
-    <Protobuf Include="Protos\horde\common\lease_state.proto" ProtoRoot="Protos" />
-    <Protobuf Include="Protos\horde\common\priority.proto" ProtoRoot="Protos" />
-    <Protobuf Include="Protos\horde\messages\agent_capabilities.proto" ProtoRoot="Protos" />
-    <Protobuf Include="Protos\horde\messages\agent_requirements.proto" ProtoRoot="Protos" />
-    <Protobuf Include="Protos\horde\messages\agent_workspace.proto" ProtoRoot="Protos" />
-    <Protobuf Include="Protos\horde\messages\lease.proto" ProtoRoot="Protos" />
-    <Protobuf Include="Protos\horde\tasks\compute_task.proto" ProtoRoot="Protos;..\..\Shared\EpicGames.Horde\Protos" />
-    <Protobuf Include="Protos\horde\tasks\conform_task.proto" ProtoRoot="Protos" />
-    <Protobuf Include="Protos\horde\tasks\job_task.proto" ProtoRoot="Protos" />
-    <Protobuf Include="Protos\horde\tasks\restart_task.proto" ProtoRoot="Protos" />
-    <Protobuf Include="Protos\horde\tasks\shutdown_task.proto" ProtoRoot="Protos" />
-    <Protobuf Include="Protos\horde\tasks\upgrade_task.proto" ProtoRoot="Protos" />
-  </ItemGroup>
 </Project>
