Horde: Move certificate handling methods into EpicGames.Core.

diff --git a/Engine/Source/Programs/Horde/Horde.Agent/Commands/Utilities/CreateCertCommand.cs b/Engine/Source/Programs/Horde/Horde.Agent/Commands/Utilities/CreateCertCommand.cs
--- a/Engine/Source/Programs/Horde/Horde.Agent/Commands/Utilities/CreateCertCommand.cs
+++ b/Engine/Source/Programs/Horde/Horde.Agent/Commands/Utilities/CreateCertCommand.cs
@@ -77,7 +77,7 @@ namespace HordeAgent.Commands.Certs
 
 			Logger.LogInformation("Creating certificate for {DnsName}", DnsName);
 
-			byte[] PrivateCertData = AgentUtilities.CreateAgentCert(DnsName);
+			byte[] PrivateCertData = CertificateUtils.CreateSelfSignedCert(DnsName, "Horde Server");
 
 			Logger.LogInformation("Writing private cert: {PrivateCert}", new FileReference(PrivateCertFile).FullName);
 			File.WriteAllBytes(PrivateCertFile, PrivateCertData);
diff --git a/Engine/Source/Programs/Horde/Horde.Build/Program.cs b/Engine/Source/Programs/Horde/Horde.Build/Program.cs
--- a/Engine/Source/Programs/Horde/Horde.Build/Program.cs
+++ b/Engine/Source/Programs/Horde/Horde.Build/Program.cs
@@ -346,7 +346,7 @@ namespace HordeServer
 					string DnsName = System.Net.Dns.GetHostName();
 					Serilog.Log.Logger.Information("Creating certificate for {DnsName}", DnsName);
 
-					byte[] PrivateCertData = AgentUtilities.CreateAgentCert(DnsName);
+					byte[] PrivateCertData = CertificateUtils.CreateSelfSignedCert(DnsName, "Horde Server");
 					
 					Serilog.Log.Logger.Information("Writing private cert: {PrivateCert}", PrivateCertFile.FullName);
 
diff --git a/Engine/Source/Programs/Horde/Horde.Build/Utilities/AgentUtilities.cs b/Engine/Source/Programs/Horde/Horde.Build/Utilities/AgentUtilities.cs
new file mode 100644
--- /dev/null
+++ b/Engine/Source/Programs/Horde/Horde.Build/Utilities/AgentUtilities.cs
@@ -0,0 +1,105 @@
+// Copyright Epic Games, Inc. All Rights Reserved.
+
+using System;
+using System.Collections.Generic;
+using System.IO;
+using System.IO.Compression;
+using System.Security.Cryptography;
+using System.Security.Cryptography.X509Certificates;
+using System.Text.RegularExpressions;
+using System.Text.Json;
+using System.Diagnostics;
+
+#nullable disable
+
+namespace HordeCommon
+{
+	public static class AgentUtilities
+	{
+		/// <summary>
+		/// Reads the version number from an archive
+		/// </summary>
+		/// <param name="Data">The archive data</param>
+		/// <returns></returns>
+		public static string ReadVersion(byte[] Data)
+		{
+			MemoryStream InputStream = new MemoryStream(Data);
+			using (ZipArchive InputArchive = new ZipArchive(InputStream, ZipArchiveMode.Read, true))
+			{
+				foreach (ZipArchiveEntry InputEntry in InputArchive.Entries)
+				{
+					if (InputEntry.FullName.Equals("HordeAgent.dll", StringComparison.OrdinalIgnoreCase))
+					{
+						string TempFile = Path.GetTempFileName();
+						try
+						{
+							InputEntry.ExtractToFile(TempFile, true);
+							return FileVersionInfo.GetVersionInfo(TempFile).ProductVersion;
+						}
+						finally
+						{
+							File.Delete(TempFile);
+						}
+					}
+				}
+			}
+			throw new Exception("Unable to find HordeAgent.dll in archive");
+		}
+
+
+		/// <summary>
+		/// Updates the agent app settings within the archive data
+		/// </summary>
+		/// <param name="Data">Data for the zip archive</param>
+		/// <param name="Settings">The settings to update</param>
+		/// <returns>New agent app data</returns>
+		public static byte[] UpdateAppSettings(byte[] Data, Dictionary<string, object> Settings)
+		{
+			bool bWrittenClientId = false;
+
+			MemoryStream OutputStream = new MemoryStream();
+			using (ZipArchive OutputArchive = new ZipArchive(OutputStream, ZipArchiveMode.Create, true))
+			{
+				MemoryStream InputStream = new MemoryStream(Data);
+				using (ZipArchive InputArchive = new ZipArchive(InputStream, ZipArchiveMode.Read, true))
+				{
+					foreach (ZipArchiveEntry InputEntry in InputArchive.Entries)
+					{
+						ZipArchiveEntry OutputEntry = OutputArchive.CreateEntry(InputEntry.FullName);
+
+						using System.IO.Stream InputEntryStream = InputEntry.Open();
+						using System.IO.Stream OutputEntryStream = OutputEntry.Open();
+
+						if (InputEntry.FullName.Equals("appsettings.json", StringComparison.OrdinalIgnoreCase))
+						{
+							using MemoryStream MemoryStream = new MemoryStream();
+							InputEntryStream.CopyTo(MemoryStream);
+
+							Dictionary<string, Dictionary<string, object>> Document = JsonSerializer.Deserialize<Dictionary<string, Dictionary<string, object>>>(MemoryStream.ToArray());
+							foreach (KeyValuePair<string, object> Pair in Settings) {								
+								Document["Horde"][Pair.Key] = Pair.Value;
+							}
+
+							using Utf8JsonWriter Writer = new Utf8JsonWriter(OutputEntryStream, new JsonWriterOptions { Indented = true });
+							JsonSerializer.Serialize<Dictionary<string, Dictionary<string, object>>>(Writer, Document, new JsonSerializerOptions { WriteIndented = true });
+
+							bWrittenClientId = true;
+						}
+						else
+						{
+							InputEntryStream.CopyTo(OutputEntryStream);
+						}
+					}
+				}
+			}
+
+			if (!bWrittenClientId)
+			{
+				throw new InvalidDataException("Missing appsettings.json file from zip archive");
+			}
+
+			return OutputStream.ToArray();
+		}
+
+	}
+}
diff --git a/Engine/Source/Programs/Horde/HordeCommon/AgentUtilities.cs b/Engine/Source/Programs/Horde/HordeCommon/AgentUtilities.cs
deleted file mode 100644
diff --git a/Engine/Source/Programs/Shared/EpicGames.Core/CertificateUtils.cs b/Engine/Source/Programs/Shared/EpicGames.Core/CertificateUtils.cs
new file mode 100644
--- /dev/null
+++ b/Engine/Source/Programs/Shared/EpicGames.Core/CertificateUtils.cs
@@ -0,0 +1,49 @@
+using System;
+using System.Collections.Generic;
+using System.Security.Cryptography;
+using System.Security.Cryptography.X509Certificates;
+using System.Text;
+using System.Text.RegularExpressions;
+
+namespace EpicGames.Core
+{
+	/// <summary>
+	/// Utility methods for certificates
+	/// </summary>
+	public static class CertificateUtils
+	{
+		/// <summary>
+		/// Create a self signed cert for the agent
+		/// </summary>
+		/// <param name="DnsName"></param>
+		/// <param name="FriendlyName">Friendly name for the certificate</param>
+		/// <returns></returns>
+		public static byte[] CreateSelfSignedCert(string DnsName, string FriendlyName)
+		{
+			DnsName = Regex.Replace(DnsName, @"^[a-zA-Z]+://", "");
+			DnsName = Regex.Replace(DnsName, "/.*$", "");
+
+			using (RSA Algorithm = RSA.Create(2048))
+			{
+				X500DistinguishedName distinguishedName = new X500DistinguishedName($"CN={DnsName}");
+				CertificateRequest Request = new CertificateRequest(distinguishedName, Algorithm, HashAlgorithmName.SHA256, RSASignaturePadding.Pkcs1);
+
+				Request.CertificateExtensions.Add(new X509BasicConstraintsExtension(false, false, 0, true));
+				Request.CertificateExtensions.Add(new X509KeyUsageExtension(X509KeyUsageFlags.KeyEncipherment | X509KeyUsageFlags.DigitalSignature, true));
+				Request.CertificateExtensions.Add(new X509EnhancedKeyUsageExtension(new OidCollection { new Oid("1.3.6.1.5.5.7.3.1") }, true));
+
+				SubjectAlternativeNameBuilder AlternativeNameBuilder = new SubjectAlternativeNameBuilder();
+				AlternativeNameBuilder.AddDnsName(DnsName);
+				Request.CertificateExtensions.Add(AlternativeNameBuilder.Build(true));
+
+				// NB: MacOS requires 825 days or fewer (https://support.apple.com/en-us/HT210176)
+				using (X509Certificate2 Certificate = Request.CreateSelfSigned(new DateTimeOffset(DateTime.UtcNow.AddDays(-1)), new DateTimeOffset(DateTime.UtcNow.AddDays(800))))
+				{
+					Certificate.FriendlyName = FriendlyName;
+					byte[] PrivateCertData = Certificate.Export(X509ContentType.Pkcs12); // Note: Need to reimport this to use immediately, otherwise key is ephemeral
+					return PrivateCertData;
+				}
+			}
+		}
+	}
+}
