Horde: Fix build warnings.

diff --git a/Engine/Source/Programs/Horde/Horde.Build/Utilities/AgentUtilities.cs b/Engine/Source/Programs/Horde/Horde.Build/Utilities/AgentUtilities.cs
--- a/Engine/Source/Programs/Horde/Horde.Build/Utilities/AgentUtilities.cs
+++ b/Engine/Source/Programs/Horde/Horde.Build/Utilities/AgentUtilities.cs
@@ -12,9 +12,9 @@ using System.Diagnostics;
 
 #nullable disable
 
-namespace HordeCommon
+namespace HordeServer.Utilities
 {
-	public static class AgentUtilities
+	static class AgentUtilities
 	{
 		/// <summary>
 		/// Reads the version number from an archive
