Horde: Remove unused projects.

diff --git a/Engine/Source/Programs/Horde/Horde.Build/Horde.Build.csproj b/Engine/Source/Programs/Horde/Horde.Build/Horde.Build.csproj
--- a/Engine/Source/Programs/Horde/Horde.Build/Horde.Build.csproj
+++ b/Engine/Source/Programs/Horde/Horde.Build/Horde.Build.csproj
@@ -66,7 +66,26 @@
     <Protobuf Include="..\..\Shared\EpicGames.Horde\Protos\horde\storage\ref_table.proto" ProtoRoot="..\..\Shared\EpicGames.Horde\Protos" GrpcServices="Server" />
     <Protobuf Include="..\..\Shared\EpicGames.Horde\Protos\horde\storage\blob_store.proto" ProtoRoot="..\..\Shared\EpicGames.Horde\Protos" GrpcServices="Server" />
     <Protobuf Include="..\..\Shared\EpicGames.Horde\Protos\horde\compute\compute_rpc.proto" ProtoRoot="..\..\Shared\EpicGames.Horde\Protos" GrpcServices="Server" />
-    <Protobuf Include="..\HordeCommon\Protos\horde\horde_rpc.proto" ProtoRoot="..\HordeCommon\Protos" GrpcServices="Server" />
+    <Protobuf Include="Protos\horde\horde_rpc.proto" ProtoRoot="Protos" GrpcServices="Server" />
+    <Protobuf Include="Protos\horde\horde_rpc_messages.proto" ProtoRoot="Protos" />
+    <Protobuf Include="Protos\horde\common\agent_status.proto" ProtoRoot="Protos" />
+    <Protobuf Include="Protos\horde\common\event_severity.proto" ProtoRoot="Protos" />
+    <Protobuf Include="Protos\horde\common\jobstep_outcome.proto" ProtoRoot="Protos" />
+    <Protobuf Include="Protos\horde\common\jobstep_state.proto" ProtoRoot="Protos" />
+    <Protobuf Include="Protos\horde\common\label_change.proto" ProtoRoot="Protos" />
+    <Protobuf Include="Protos\horde\common\lease_outcome.proto" ProtoRoot="Protos" />
+    <Protobuf Include="Protos\horde\common\lease_state.proto" ProtoRoot="Protos" />
+    <Protobuf Include="Protos\horde\common\priority.proto" ProtoRoot="Protos" />
+    <Protobuf Include="Protos\horde\messages\agent_capabilities.proto" ProtoRoot="Protos" />
+    <Protobuf Include="Protos\horde\messages\agent_requirements.proto" ProtoRoot="Protos" />
+    <Protobuf Include="Protos\horde\messages\agent_workspace.proto" ProtoRoot="Protos" />
+    <Protobuf Include="Protos\horde\messages\lease.proto" ProtoRoot="Protos" />
+    <Protobuf Include="Protos\horde\tasks\compute_task.proto" ProtoRoot="Protos" AdditionalImportDirs="..\..\Shared\EpicGames.Horde\Protos" />
+    <Protobuf Include="Protos\horde\tasks\conform_task.proto" ProtoRoot="Protos" />
+    <Protobuf Include="Protos\horde\tasks\job_task.proto" ProtoRoot="Protos" />
+    <Protobuf Include="Protos\horde\tasks\restart_task.proto" ProtoRoot="Protos" />
+    <Protobuf Include="Protos\horde\tasks\shutdown_task.proto" ProtoRoot="Protos" />
+    <Protobuf Include="Protos\horde\tasks\upgrade_task.proto" ProtoRoot="Protos" />
     <Protobuf Include="Protos\grpc\health\v1\health.proto" ProtoRoot="Protos" GrpcServices="Server" />
   </ItemGroup>
 
@@ -78,9 +97,7 @@
     <ProjectReference Include="..\..\Shared\EpicGames.Perforce\EpicGames.Perforce.csproj" />
     <ProjectReference Include="..\..\Shared\EpicGames.Redis\EpicGames.Redis.csproj" />
     <ProjectReference Include="..\..\Shared\EpicGames.Serialization\EpicGames.Serialization.csproj" />
-    <ProjectReference Include="..\HordeCommon\HordeCommon.csproj" />
     <ProjectReference Include="..\..\..\ThirdParty\Perforce\P4Api.Net\p4api.net\p4api.net.csproj" />
-    <ProjectReference Include="..\HordeServer.Common\HordeServer.Common.csproj" />
   </ItemGroup>
 
   <ItemGroup>
diff --git a/Engine/Source/Programs/Horde/HordeServer.Common/HordeServer.Common.csproj b/Engine/Source/Programs/Horde/HordeServer.Common/HordeServer.Common.csproj
deleted file mode 100644
diff --git a/Engine/Source/Programs/Horde/HordeServer.Compute/HordeServer.Compute.csproj b/Engine/Source/Programs/Horde/HordeServer.Compute/HordeServer.Compute.csproj
deleted file mode 100644
