Horde: Flush the write buffer when switching into raw Json output mode. The batch files that launch UAT will output a few lines of regular text before switching into Json.

diff --git a/Engine/Source/Programs/Horde/Horde.Agent/Parser/LogParser.cs b/Engine/Source/Programs/Horde/Horde.Agent/Parser/LogParser.cs
--- a/Engine/Source/Programs/Horde/Horde.Agent/Parser/LogParser.cs
+++ b/Engine/Source/Programs/Horde/Horde.Agent/Parser/LogParser.cs
@@ -126,7 +126,7 @@ namespace HordeAgent.Parser
 			{
 				Data = Data.Slice(0, Data.Length - 1);
 			}
-			if (Buffer.Length == 0 && Data.Length > 0 && Data[0] == '{' && WriteRawEvent(Data))
+			if (Data.Length > 0 && Data[0] == '{' && WriteRawEvent(Data))
 			{
 				return;
 			}
@@ -193,6 +193,7 @@ namespace HordeAgent.Parser
 				}
 			}
 
+			ProcessData(true);
 			JsonLogger.WriteFormattedEvent(Level, new byte[][] { Data.ToArray() });
 			return true;
 		}
@@ -275,8 +276,8 @@ namespace HordeAgent.Parser
 		/// <summary>
 		/// Process any data in the buffer
 		/// </summary>
-		/// <param name="bEndOfStream">Whether we've reached the end of the stream</param>
-		void ProcessData(bool bEndOfStream)
+		/// <param name="bFlush">Whether we've reached the end of the stream</param>
+		void ProcessData(bool bFlush)
 		{
 			while (Buffer.Length > 0)
 			{
@@ -296,7 +297,7 @@ namespace HordeAgent.Parser
 				}
 
 				// Bail out if we need more data
-				if(Buffer.Length < 1024 && !bEndOfStream && Buffer.NeedMoreData)
+				if(Buffer.Length < 1024 && !bFlush && Buffer.NeedMoreData)
 				{
 					break;
 				}
