Horde: Use a special state type to pass pre-formatted json log data through to sinks. Allows data to pass through other logger implementations before reaching a json logger.

diff --git a/Engine/Source/Programs/Horde/Horde.Agent.Tests/LogParserTests.cs b/Engine/Source/Programs/Horde/Horde.Agent.Tests/LogParserTests.cs
--- a/Engine/Source/Programs/Horde/Horde.Agent.Tests/LogParserTests.cs
+++ b/Engine/Source/Programs/Horde/Horde.Agent.Tests/LogParserTests.cs
@@ -2,6 +2,7 @@
 
 using EpicGames.Core;
 using HordeAgent.Parser;
+using HordeAgent.Utility;
 using HordeCommon;
 using Microsoft.Extensions.Logging;
 using Microsoft.VisualStudio.TestPlatform.Common.Utilities;
@@ -57,8 +58,9 @@ namespace HordeAgentTests
 				int LineCount = Message.Count(x => x == '\n') + 1;
 
 				IEnumerable<KeyValuePair<string, object>>? Properties = State as IEnumerable<KeyValuePair<string, object>>;
-				if (Properties != null)
+				if (LogLevel != LogLevel.Information || EventId.Id != 0 || Properties != null || State is JsonLogEvent)
 				{
+					Properties ??= new Dictionary<string, object>();
 					Events.Add(new CapturedEvent(EventId, LogLevel, LineIndex, LineCount, Message, Properties));
 				}
 
@@ -66,6 +68,43 @@ namespace HordeAgentTests
 			}
 		}
 
+		[TestMethod]
+		public void StructuredOutputMatcher()
+		{
+			string[] Lines =
+			{
+				"{\"Timestamp\":\"2021-10-28T09:14:48.8561456-04:00\",\"Level\":\"Information\",\"MessageTemplate\":\"Hello {World}\",\"RenderedMessage\":\"Hello 123\",\"Properties\":{\"World\":123,\"EventId\":{\"Id\":123},\"SourceContext\":\"HordeAgent\",\"dd.env\":\"default\",\"dd.service\":\"hordeagent\",\"dd.version\":\"1.0.0\"}}",
+				"{\"Timestamp\":\"2021-10-28T05:36:08\",\"Level\":\"Information\",\"RenderedMessage\":\"Building 43 projects (see Log \u0027Engine/Programs/AutomationTool/Saved/Logs/Log.txt\u0027 for more details)\"}",
+				"{\"Timestamp\":\"2021-10-28T05:36:08\",\"Level\":\"Warning\",\"RenderedMessage\":\" Restore...\"}",
+				"{\"Timestamp\":\"2021-10-28T05:36:15\",\"Level\":\"Error\",\"RenderedMessage\":\" Build...\"}",
+			};
+
+			List<CapturedEvent> Events = Parse(Lines);
+			Assert.AreEqual(4, Events.Count);
+
+			int Idx = 0;
+
+			CapturedEvent Event = Events[Idx++];
+			Assert.AreEqual(LogLevel.Information, Event.Level);
+			Assert.AreEqual(new EventId(123), Event.Id);
+			Assert.AreEqual("Hello 123", Event.Text);
+
+			Event = Events[Idx++];
+			Assert.AreEqual(LogLevel.Information, Event.Level);
+			Assert.AreEqual(new EventId(0), Event.Id);
+			Assert.AreEqual("Building 43 projects (see Log \u0027Engine/Programs/AutomationTool/Saved/Logs/Log.txt\u0027 for more details)", Event.Text);
+
+			Event = Events[Idx++];
+			Assert.AreEqual(LogLevel.Warning, Event.Level);
+			Assert.AreEqual(new EventId(0), Event.Id);
+			Assert.AreEqual(" Restore...", Event.Text);
+
+			Event = Events[Idx++];
+			Assert.AreEqual(LogLevel.Error, Event.Level);
+			Assert.AreEqual(new EventId(0), Event.Id);
+			Assert.AreEqual(" Build...", Event.Text);
+		}
+
 		[TestMethod]
 		public void ExitCodeEventMatcher()
 		{
@@ -1060,7 +1099,7 @@ namespace HordeAgentTests
 				while(Pos < TextBytes.Length)
 				{
 					int Len = Math.Min((int)(Generator.NextDouble() * 256), TextBytes.Length - Pos);
-					Parser.WriteData(TextBytes.AsSpan(Pos, Len), Pos + Len == TextBytes.Length);
+					Parser.WriteData(TextBytes.AsMemory(Pos, Len), Pos + Len == TextBytes.Length);
 					Pos += Len;
 				}
 			}
diff --git a/Engine/Source/Programs/Horde/Horde.Agent/Commands/Parse/ParseCommand.cs b/Engine/Source/Programs/Horde/Horde.Agent/Commands/Parse/ParseCommand.cs
--- a/Engine/Source/Programs/Horde/Horde.Agent/Commands/Parse/ParseCommand.cs
+++ b/Engine/Source/Programs/Horde/Horde.Agent/Commands/Parse/ParseCommand.cs
@@ -38,7 +38,7 @@ namespace HordeAgent.Commands.Parse
 		[Description("The changelist number that has been synced")]
 		int? Change = null;
 
-		public override Task<int> ExecuteAsync(ILogger Logger)
+		public override async Task<int> ExecuteAsync(ILogger Logger)
 		{
 			// Read all the ignore patterns
 			List<string> IgnorePatterns = new List<string>();
@@ -62,7 +62,7 @@ namespace HordeAgent.Commands.Parse
 			}
 
 			// Read the file and pipe it through the event parser
-			using (StreamReader Reader = new StreamReader(InputFile.FullName))
+			using (FileStream InputStream = FileReference.Open(InputFile, FileMode.Open, FileAccess.Read))
 			{
 				LogParserContext Context = new LogParserContext();
 				Context.WorkspaceDir = WorkspaceDir;
@@ -71,18 +71,20 @@ namespace HordeAgent.Commands.Parse
 
 				using (LogParser Parser = new LogParser(Logger, Context, IgnorePatterns))
 				{
+					byte[] Data = new byte[1024];
 					for (; ; )
 					{
-						string? Line = Reader.ReadLine();
-						if(Line == null)
+						int Length = await InputStream.ReadAsync(Data);
+						if(Length == 0)
 						{
+							Parser.WriteData(ReadOnlyMemory<byte>.Empty, true);
 							break;
 						}
-						Parser.WriteLine(Line);
+						Parser.WriteData(Data.AsMemory(0, Length), false);
 					}
 				}
 			}
-			return Task.FromResult(0);
+			return 0;
 		}
 	}
 }
diff --git a/Engine/Source/Programs/Horde/Horde.Agent/Execution/BuildGraphExecutor.cs b/Engine/Source/Programs/Horde/Horde.Agent/Execution/BuildGraphExecutor.cs
--- a/Engine/Source/Programs/Horde/Horde.Agent/Execution/BuildGraphExecutor.cs
+++ b/Engine/Source/Programs/Horde/Horde.Agent/Execution/BuildGraphExecutor.cs
@@ -742,7 +742,7 @@ namespace HordeAgent.Execution
 					List<string> IgnorePatterns = await ReadIgnorePatternsAsync(WorkspaceDir, Logger);
 					using (LogParser Filter = new LogParser(Logger, Context, IgnorePatterns))
 					{
-						await Process.CopyToAsync((Buffer, Offset, Length) => Filter.WriteData(Buffer.AsSpan(Offset, Length), false), 4096, CancellationToken);
+						await Process.CopyToAsync((Buffer, Offset, Length) => Filter.WriteData(Buffer.AsMemory(Offset, Length), false), 4096, CancellationToken);
 						Filter.WriteData(Array.Empty<byte>(), true);
 					}
 					
diff --git a/Engine/Source/Programs/Horde/Horde.Agent/Parser/LogParser.cs b/Engine/Source/Programs/Horde/Horde.Agent/Parser/LogParser.cs
--- a/Engine/Source/Programs/Horde/Horde.Agent/Parser/LogParser.cs
+++ b/Engine/Source/Programs/Horde/Horde.Agent/Parser/LogParser.cs
@@ -68,11 +68,6 @@ namespace HordeAgent.Parser
 		/// </summary>
 		ILogger Logger;
 
-		/// <summary>
-		/// Special passthrough for Json log devices
-		/// </summary>
-		JsonLogger? JsonLogger;
-
 		/// <summary>
 		/// Context for matchers for this log
 		/// </summary>
@@ -93,7 +88,6 @@ namespace HordeAgent.Parser
 		{
 			this.IgnorePatterns = IgnorePatterns;
 			this.Logger = Logger;
-			this.JsonLogger = Logger as JsonLogger;
 			this.Buffer = new LineBuffer(50);
 			this.Context = Context;
 		}
@@ -120,82 +114,23 @@ namespace HordeAgent.Parser
 		/// Adds a raw utf-8 string to the buffer
 		/// </summary>
 		/// <param name="Data">The string data</param>
-		private void AddLine(ReadOnlySpan<byte> Data)
+		private void AddLine(ReadOnlyMemory<byte> Data)
 		{
-			if(Data.Length > 0 && Data[Data.Length - 1] == '\r')
+			if(Data.Length > 0 && Data.Span[Data.Length - 1] == '\r')
 			{
 				Data = Data.Slice(0, Data.Length - 1);
 			}
-			if (Data.Length > 0 && Data[0] == '{' && WriteRawEvent(Data))
-			{
-				return;
-			}
-			Buffer.AddLine(Encoding.UTF8.GetString(Data));
-		}
-
-		/// <summary>
-		/// Static constant string for the level property in a log event
-		/// </summary>
-		readonly Utf8String LevelString = new Utf8String("Level");
-
-		/// <summary>
-		/// Attempts to parse the line as a raw JSON object and pass it directly through to the JsonLogger. For this to succeed, it must be a well formatted JSON object and have a Level property.
-		/// </summary>
-		/// <param name="Data"></param>
-		/// <returns></returns>
-		private bool WriteRawEvent(ReadOnlySpan<byte> Data)
-		{
-			if(JsonLogger == null)
-			{
-				return false;
-			}
-
-			Utf8JsonReader Reader = new Utf8JsonReader(Data);
-			if (!Reader.Read() || Reader.TokenType != JsonTokenType.StartObject)
-			{
-				return false;
-			}
-
-			LogLevel Level = LogLevel.None;
-			for(; ;)
+			if (Data.Length > 0 && Data.Span[0] == '{')
 			{
-				if(!Reader.Read())
+				JsonLogEvent JsonEvent;
+				if (JsonLogEvent.TryParse(Data, out JsonEvent))
 				{
-					return false;
-				}
-
-				if (Reader.TokenType == JsonTokenType.EndObject && Level != LogLevel.None && Reader.BytesConsumed == Data.Length)
-				{
-					break;
-				}
-				else if(Reader.TokenType != JsonTokenType.PropertyName)
-				{
-					return false;
-				}
-
-				if (Utf8StringComparer.OrdinalIgnoreCase.Equals(Reader.ValueSpan, LevelString.Span))
-				{
-					if (!Reader.TrySkip() || Reader.TokenType != JsonTokenType.String)
-					{
-						return false;
-					}
-					if(!Enum.TryParse(Reader.GetString(), out Level))
-					{
-						return false;
-					}
-				}
-				else
-				{
-					if (!Reader.TrySkip())
-					{
-						return false;
-					}
+					ProcessData(true);
+					Logger.Log(JsonEvent.Level, JsonEvent.EventId, JsonEvent, null, JsonLogEvent.Format);
+					return;
 				}
 			}
-
-			ProcessData(true);
-			JsonLogger.WriteFormattedEvent(Level, new byte[][] { Data.ToArray() });
-			return true;
+			Buffer.AddLine(Encoding.UTF8.GetString(Data.Span));
 		}
 
 		/// <summary>
@@ -213,19 +148,20 @@ namespace HordeAgent.Parser
 		/// </summary>
 		/// <param name="Data">Data to write</param>
 		/// <param name="bEndOfStream">Whether this is the end of the stream</param>
-		public void WriteData(ReadOnlySpan<byte> Data, bool bEndOfStream)
+		public void WriteData(ReadOnlyMemory<byte> Data, bool bEndOfStream)
 		{
 			int BaseIdx = 0;
 			int ScanIdx = 0;
+			ReadOnlySpan<byte> Span = Data.Span;
 
 			// Handle a partially existing line
 			if (PartialLine.Length > 0)
 			{
-				for (; ScanIdx < Data.Length; ScanIdx++)
+				for (; ScanIdx < Span.Length; ScanIdx++)
 				{
-					if (Data[ScanIdx] == '\n')
+					if (Span[ScanIdx] == '\n')
 					{
-						PartialLine.Write(Data.Slice(BaseIdx, ScanIdx - BaseIdx));
+						PartialLine.Write(Span.Slice(BaseIdx, ScanIdx - BaseIdx));
 						FlushPartialLine();
 						BaseIdx = ++ScanIdx;
 						break;
@@ -234,9 +170,9 @@ namespace HordeAgent.Parser
 			}
 
 			// Handle any complete lines
-			for (; ScanIdx < Data.Length; ScanIdx++)
+			for (; ScanIdx < Span.Length; ScanIdx++)
 			{
-				if(Data[ScanIdx] == '\n')
+				if(Span[ScanIdx] == '\n')
 				{
 					AddLine(Data.Slice(BaseIdx, ScanIdx - BaseIdx));
 					BaseIdx = ScanIdx + 1;
@@ -244,7 +180,7 @@ namespace HordeAgent.Parser
 			}
 
 			// Add the rest of the text to the partial line buffer
-			PartialLine.Write(Data.Slice(BaseIdx));
+			PartialLine.Write(Span.Slice(BaseIdx));
 
 			// If it's the end of the stream, force a flush
 			if (bEndOfStream && PartialLine.Length > 0)
diff --git a/Engine/Source/Programs/Horde/Horde.Agent/Utility/JsonLogger.cs b/Engine/Source/Programs/Horde/Horde.Agent/Utility/JsonLogger.cs
--- a/Engine/Source/Programs/Horde/Horde.Agent/Utility/JsonLogger.cs
+++ b/Engine/Source/Programs/Horde/Horde.Agent/Utility/JsonLogger.cs
@@ -14,6 +14,151 @@ using System.Text.Json;
 
 namespace HordeAgent.Utility
 {
+	/// <summary>
+	/// Defines a preformatted Json log event, which can pass through raw Json data directly or format it as a regular string
+	/// </summary>
+	public struct JsonLogEvent
+	{
+		static Utf8String LevelString { get; } = new Utf8String("Level");
+		static Utf8String PropertiesString { get; } = new Utf8String("Properties");
+		static Utf8String EventIdString { get; } = new Utf8String("EventId");
+		static Utf8String IdString { get; } = new Utf8String("Id");
+		static Utf8String RenderedMessageString { get; } = new Utf8String("RenderedMessage");
+
+		/// <summary>
+		/// The 
+		/// </summary>
+		public LogLevel Level { get; }
+
+		/// <summary>
+		/// The event id, if set
+		/// </summary>
+		public EventId EventId { get; }
+
+		/// <summary>
+		/// The utf-8 encoded JSON event
+		/// </summary>
+		public ReadOnlyMemory<byte> Data { get; }
+
+		/// <summary>
+		/// Constructor
+		/// </summary>
+		public JsonLogEvent(LogLevel Level, EventId EventId, ReadOnlyMemory<byte> Data)
+		{
+			this.Level = Level;
+			this.EventId = EventId;
+			this.Data = Data;
+		}
+
+		/// <summary>
+		/// Tries to parse a Json log event from the given 
+		/// </summary>
+		/// <param name="Data"></param>
+		/// <param name="Event"></param>
+		/// <returns></returns>
+		public static bool TryParse(ReadOnlyMemory<byte> Data, out JsonLogEvent Event)
+		{
+			LogLevel Level = LogLevel.None;
+			int EventId = 0;
+
+			Utf8JsonReader Reader = new Utf8JsonReader(Data.Span);
+			if (Reader.Read() && Reader.TokenType == JsonTokenType.StartObject)
+			{
+				ReadOnlySpan<byte> PropertyName;
+				for (; ReadNextPropertyName(ref Reader, out PropertyName); Reader.Skip())
+				{
+					if (Utf8StringComparer.OrdinalIgnoreCase.Equals(PropertyName, LevelString.Span))
+					{
+						if (!Enum.TryParse(Reader.GetString(), out Level))
+						{
+							break;
+						}
+					}
+					else if (Utf8StringComparer.OrdinalIgnoreCase.Equals(PropertyName, PropertiesString.Span) && Reader.TokenType == JsonTokenType.StartObject)
+					{
+						for (; ReadNextPropertyName(ref Reader, out PropertyName); Reader.Skip())
+						{
+							if (Utf8StringComparer.OrdinalIgnoreCase.Equals(EventIdString.Span, PropertyName) && Reader.TokenType == JsonTokenType.StartObject)
+							{
+								for (; ReadNextPropertyName(ref Reader, out PropertyName); Reader.Skip())
+								{
+									if (Utf8StringComparer.OrdinalIgnoreCase.Equals(IdString.Span, PropertyName) && Reader.TryGetInt32(out int NewEventId))
+									{
+										EventId = NewEventId;
+									}
+								}
+							}
+						}
+					}
+				}
+			}
+
+			if (Reader.TokenType == JsonTokenType.EndObject && Level != LogLevel.None && Reader.BytesConsumed == Data.Length)
+			{
+				Event = new JsonLogEvent(Level, new EventId(EventId), Data);
+				return true;
+			}
+			else
+			{
+				Event = default;
+				return false;
+			}
+		}
+
+		static bool ReadNextPropertyName(ref Utf8JsonReader Reader, out ReadOnlySpan<byte> PropertyName)
+		{
+			if (Reader.Read() && Reader.TokenType == JsonTokenType.PropertyName)
+			{
+				PropertyName = Reader.ValueSpan;
+				return Reader.Read();
+			}
+			else
+			{
+				PropertyName = ReadOnlySpan<byte>.Empty;
+				return false;
+			}
+		}
+
+		/// <summary>
+		/// Formats an event as a string
+		/// </summary>
+		/// <param name="State"></param>
+		/// <param name="Ex"></param>
+		/// <returns></returns>
+		/// <exception cref="NotImplementedException"></exception>
+		public static string Format(JsonLogEvent State, Exception? Ex)
+		{
+			Utf8JsonReader Reader = new Utf8JsonReader(State.Data.Span);
+			if (Reader.Read() && Reader.TokenType == JsonTokenType.StartObject)
+			{
+				for (; ; )
+				{
+					if (!Reader.Read() || Reader.TokenType != JsonTokenType.PropertyName)
+					{
+						break;
+					}
+
+					bool IsMessage = Utf8StringComparer.OrdinalIgnoreCase.Equals(Reader.ValueSpan, RenderedMessageString.Span);
+					if (!Reader.TrySkip())
+					{
+						break;
+					}
+					if (IsMessage && Reader.TokenType == JsonTokenType.String)
+					{
+						return Reader.GetString();
+					}
+				}
+			}
+			return String.Empty;
+		}
+
+		/// <inheritdoc/>
+		public override string ToString()
+		{
+			return Encoding.UTF8.GetString(Data.ToArray());
+		}
+	}
+
 	/// <summary>
 	/// Class to handle uploading log data to the server in the background
 	/// </summary>
@@ -61,6 +206,12 @@ namespace HordeAgent.Utility
 				LogLevel = LogLevel.Information;
 			}
 
+			if(State is JsonLogEvent Event)
+			{
+				WriteFormattedEvent(LogLevel, new byte[][] { Event.Data.ToArray() });
+				return;
+			}
+
 			// Render the message
 			string Text = Formatter(State, Exception);
 
@@ -145,7 +296,7 @@ namespace HordeAgent.Utility
 		/// </summary>
 		/// <param name="Level">The log level</param>
 		/// <param name="Lines">Utf-8 encoded JSON line data</param>
-		public abstract void WriteFormattedEvent(LogLevel Level, byte[][] Lines);
+		protected abstract void WriteFormattedEvent(LogLevel Level, byte[][] Lines);
 
 		/// <summary>
 		/// Callback to write a systemic event
diff --git a/Engine/Source/Programs/Horde/Horde.Agent/Utility/JsonRpcLogger.cs b/Engine/Source/Programs/Horde/Horde.Agent/Utility/JsonRpcLogger.cs
--- a/Engine/Source/Programs/Horde/Horde.Agent/Utility/JsonRpcLogger.cs
+++ b/Engine/Source/Programs/Horde/Horde.Agent/Utility/JsonRpcLogger.cs
@@ -82,7 +82,7 @@ namespace HordeAgent.Parser
 			this.Outcome = JobStepOutcome.Success;
 		}
 
-		public override void WriteFormattedEvent(LogLevel Level, byte[][] Lines)
+		protected override void WriteFormattedEvent(LogLevel Level, byte[][] Lines)
 		{
 			// Update the state of this job if this is an error status
 			if (Level == LogLevel.Error || Level == LogLevel.Critical)
diff --git a/Engine/Source/Programs/Horde/Horde.Build.Tests/IssueServiceTests.cs b/Engine/Source/Programs/Horde/Horde.Build.Tests/IssueServiceTests.cs
--- a/Engine/Source/Programs/Horde/Horde.Build.Tests/IssueServiceTests.cs
+++ b/Engine/Source/Programs/Horde/Horde.Build.Tests/IssueServiceTests.cs
@@ -68,7 +68,7 @@ namespace HordeServerTests
 				}
 			}
 
-			public override void WriteFormattedEvent(LogLevel Level, byte[][] Lines)
+			protected override void WriteFormattedEvent(LogLevel Level, byte[][] Lines)
 			{
 				Events.Add((Level, Lines));
 			}
