Horde: Add logging statement at start of scheduled autoscaling to aid debugging

It seems the autoscaling does not run as frequent as specified.

diff --git a/Engine/Source/Programs/Horde/Horde.Build/Fleet/AutoscaleService.cs b/Engine/Source/Programs/Horde/Horde.Build/Fleet/AutoscaleService.cs
--- a/Engine/Source/Programs/Horde/Horde.Build/Fleet/AutoscaleService.cs
+++ b/Engine/Source/Programs/Horde/Horde.Build/Fleet/AutoscaleService.cs
@@ -116,6 +116,8 @@ namespace HordeServer.Services
 			DateTime UtcNow = DateTime.UtcNow;
 			DateTime NextTickTime = UtcNow + TimeSpan.FromMinutes(5.0);
 
+			Logger.LogInformation("Autoscaling pools...");
+
 			// Find all the current agents
 			List<IAgent> Agents = await AgentCollection.FindAsync(Status: AgentStatus.Ok);
 
