Horde: Log time taken for scheduled work in Autoscale and Schedule service

To help debug why they are scheduled so erratically.

diff --git a/Engine/Source/Programs/Horde/Horde.Build/Fleet/AutoscaleService.cs b/Engine/Source/Programs/Horde/Horde.Build/Fleet/AutoscaleService.cs
--- a/Engine/Source/Programs/Horde/Horde.Build/Fleet/AutoscaleService.cs
+++ b/Engine/Source/Programs/Horde/Horde.Build/Fleet/AutoscaleService.cs
@@ -14,6 +14,7 @@ using Microsoft.Extensions.Logging;
 using MongoDB.Bson;
 using System;
 using System.Collections.Generic;
+using System.Diagnostics;
 using System.Linq;
 using System.Net;
 using System.Runtime.CompilerServices;
@@ -117,6 +118,7 @@ namespace HordeServer.Services
 			DateTime NextTickTime = UtcNow + TimeSpan.FromMinutes(5.0);
 
 			Logger.LogInformation("Autoscaling pools...");
+			Stopwatch Stopwatch = Stopwatch.StartNew();
 
 			// Find all the current agents
 			List<IAgent> Agents = await AgentCollection.FindAsync(Status: AgentStatus.Ok);
@@ -224,6 +226,10 @@ namespace HordeServer.Services
 				DogStatsd.Gauge("agentpools.autoscale.target", TargetAgents, tags: new []{"pool:" + Pool.Name});
 				DogStatsd.Gauge("agentpools.autoscale.current", PoolData.Agents.Count, tags: new []{"pool:" + Pool.Name});
 			}
+			
+			Stopwatch.Stop();
+			Logger.LogInformation("Autoscaling pools took {ElapsedTime} ms", Stopwatch.ElapsedMilliseconds);
+			
 			return NextTickTime;
 		}
 
diff --git a/Engine/Source/Programs/Horde/Horde.Build/Projects/ScheduleService.cs b/Engine/Source/Programs/Horde/Horde.Build/Projects/ScheduleService.cs
--- a/Engine/Source/Programs/Horde/Horde.Build/Projects/ScheduleService.cs
+++ b/Engine/Source/Programs/Horde/Horde.Build/Projects/ScheduleService.cs
@@ -129,10 +129,13 @@ namespace HordeServer.Services
 			if (!DowntimeService.IsDowntimeActive)
 			{
 				Logger.LogInformation("Updating scheduled triggers...");
+				Stopwatch Stopwatch = Stopwatch.StartNew();
 				using (IScope Scope = GlobalTracer.Instance.BuildSpan("ScheduleService Tick").StartActive())
 				{
 					NextTickTime = await UpdateSchedulesAsync();
 				}
+				Stopwatch.Stop();
+				Logger.LogInformation("Scheduling triggers took {ElapsedTime} ms", Stopwatch.ElapsedMilliseconds);
 			}
 			return NextTickTime;
 		}
