Horde: Remove connection strings from app setings

diff --git a/Engine/Source/Programs/Horde/Horde.Build/appsettings.json b/Engine/Source/Programs/Horde/Horde.Build/appsettings.json
--- a/Engine/Source/Programs/Horde/Horde.Build/appsettings.json
+++ b/Engine/Source/Programs/Horde/Horde.Build/appsettings.json
@@ -1,12 +1,11 @@
 {
-	"Horde": {
+	"Horde": 
+	{
 		"DashboardUrl": "https://localhost:3000",
 		"DisableAuth": true,
 		"CorsEnabled": true,
 		"CorsOrigin": "https://localhost:5003;http://localhost:3000",
-		"LogFileProviderType": "Local",
-//		"DatabaseConnectionString": "mongodb://hordeadmin:[REDACTED_EMAIL]:27017/?ssl=true&replicaSet=rs0&readPreference=secondaryPreferred",
-		"SlackToken": "xoxb-2151903411-1236637263776-E4TraCodBFM2G5rKNanyxFZv"
+		"LogFileProviderType": "Local"
 	},
 	"Serilog": 
 	{
