Horde: Add template severity to affected stream responses
#rnx

diff --git a/Engine/Source/Programs/Horde/Horde.Build/Api/Issues.cs b/Engine/Source/Programs/Horde/Horde.Build/Api/Issues.cs
--- a/Engine/Source/Programs/Horde/Horde.Build/Api/Issues.cs
+++ b/Engine/Source/Programs/Horde/Horde.Build/Api/Issues.cs
@@ -4,6 +4,7 @@ using EpicGames.Core;
 using HordeServer.Models;
 using HordeServer.Services;
 using HordeServer.Utilities;
+using MongoDB.Bson;
 using System;
 using System.Collections.Generic;
 using System.Diagnostics.CodeAnalysis;
@@ -260,17 +261,23 @@ namespace HordeServer.Api
 		/// </summary>
 		public bool Resolved { get; set; }
 
+		/// <summary>
+		/// The issue severity of the affected template
+		/// </summary>
+		public IssueSeverity Severity { get; set; }
+
 		/// <summary>
 		/// Constructor
 		/// </summary>
 		/// <param name="TemplateId"></param>
 		/// <param name="TemplateName"></param>
 		/// <param name="Resolved"></param>
-		public GetIssueAffectedTemplateResponse(string TemplateId, string TemplateName, bool Resolved)
+		public GetIssueAffectedTemplateResponse(string TemplateId, string TemplateName, bool Resolved, IssueSeverity Severity = IssueSeverity.Unspecified)
 		{
 			this.TemplateId = TemplateId;
 			this.TemplateName = TemplateName;
 			this.Resolved = Resolved;
+			this.Severity = Severity;
 		}
 	}
 
@@ -317,9 +324,10 @@ namespace HordeServer.Api
 		/// <summary>
 		/// Constructor
 		/// </summary>
+		/// <param name="Details">Issue to construct from</param>
 		/// <param name="Stream"></param>
 		/// <param name="Spans">The spans to construct from</param>
-		public GetIssueAffectedStreamResponse(IStream? Stream, IEnumerable<IIssueSpan> Spans)
+		public GetIssueAffectedStreamResponse(IIssueDetails Details, IStream? Stream, IEnumerable<IIssueSpan> Spans)
 		{
 			IIssueSpan FirstSpan = Spans.First();
 			this.StreamId = FirstSpan.StreamId.ToString();
@@ -334,7 +342,12 @@ namespace HordeServer.Api
 				{
 					TemplateName = TemplateRef.Name;
 				}
-				this.AffectedTemplates.Add(new GetIssueAffectedTemplateResponse(Template.Key.ToString(), TemplateName, Template.All(x => x.NextSuccess != null)));
+
+				HashSet<ObjectId> UnresolvedTemplateSpans = new HashSet<ObjectId>(Template.Where(x => x.NextSuccess == null).Select(x => x.Id));
+
+				IIssueStep? TemplateStep = Details.Steps.Where(x => UnresolvedTemplateSpans.Contains(x.SpanId)).OrderByDescending(x => x.StepTime).FirstOrDefault();
+
+				this.AffectedTemplates.Add(new GetIssueAffectedTemplateResponse(Template.Key.ToString(), TemplateName, Template.All(x => x.NextSuccess != null), TemplateStep?.Severity ?? IssueSeverity.Unspecified));
 			}
 
 			HashSet<TemplateRefId> TemplateIdsSet = new HashSet<TemplateRefId>(Spans.Select(x => x.TemplateRefId));
diff --git a/Engine/Source/Programs/Horde/Horde.Build/Controllers/IssuesController.cs b/Engine/Source/Programs/Horde/Horde.Build/Controllers/IssuesController.cs
--- a/Engine/Source/Programs/Horde/Horde.Build/Controllers/IssuesController.cs
+++ b/Engine/Source/Programs/Horde/Horde.Build/Controllers/IssuesController.cs
@@ -190,7 +190,7 @@ namespace HordeServer.Controllers
 			foreach (IGrouping<StreamId, IIssueSpan> StreamSpans in Details.Spans.GroupBy(x => x.StreamId))
 			{
 				IStream? Stream = await StreamService.GetCachedStream(StreamSpans.Key);
-				AffectedStreams.Add(new GetIssueAffectedStreamResponse(Stream, StreamSpans));
+				AffectedStreams.Add(new GetIssueAffectedStreamResponse(Details, Stream, StreamSpans));
 			}
 			return new GetIssueResponse(Details, AffectedStreams, ShowDesktopAlerts);
 		}
