Horde: Fix warning on issue severity parameter
#rnx

diff --git a/Engine/Source/Programs/Horde/Horde.Build/Api/Issues.cs b/Engine/Source/Programs/Horde/Horde.Build/Api/Issues.cs
--- a/Engine/Source/Programs/Horde/Horde.Build/Api/Issues.cs
+++ b/Engine/Source/Programs/Horde/Horde.Build/Api/Issues.cs
@@ -272,6 +272,7 @@ namespace HordeServer.Api
 		/// <param name="TemplateId"></param>
 		/// <param name="TemplateName"></param>
 		/// <param name="Resolved"></param>
+		/// <param name="Severity"></param>
 		public GetIssueAffectedTemplateResponse(string TemplateId, string TemplateName, bool Resolved, IssueSeverity Severity = IssueSeverity.Unspecified)
 		{
 			this.TemplateId = TemplateId;
