Horde: Build health improvements, charting, promotion, markdown descriptions, modal refactoring
#rnx

diff --git a/Engine/Source/Programs/Horde/HordeDashboard/package.json b/Engine/Source/Programs/Horde/HordeDashboard/package.json
--- a/Engine/Source/Programs/Horde/HordeDashboard/package.json
+++ b/Engine/Source/Programs/Horde/HordeDashboard/package.json
@@ -8,7 +8,7 @@
     "@blueprintjs/icons": "^3.13.0",
     "@datadog/browser-logs": "^1.12.4",
     "@fluentui/react": "^8.36.5",
-    "@fluentui/react-charting": "^5.0.22",
+    "@fluentui/react-charting": "^5.2.0",
     "crypto-js": "^4.0.0",
     "formik": "^2.1.1",
     "idb": "^6.0.0",
diff --git a/Engine/Source/Programs/Horde/HordeDashboard/src/backend/Api.ts b/Engine/Source/Programs/Horde/HordeDashboard/src/backend/Api.ts
--- a/Engine/Source/Programs/Horde/HordeDashboard/src/backend/Api.ts
+++ b/Engine/Source/Programs/Horde/HordeDashboard/src/backend/Api.ts
@@ -284,6 +284,7 @@ export type IssueQuery = {
 	count?: number;
 	userId?: string;
 	resolved?: boolean;
+	promoted?: boolean;
 }
 
 export type UsersQuery = {
@@ -2462,10 +2463,16 @@ export type GetIssueAffectedTemplateResponse = {
 	/**  Whether it has been resolved or not */
 	resolved: boolean;
 
+	/** Severity of this template */
+	severity: IssueSeverity;
+
 }
 
 /**Trace of a set of node failures across multiple steps */
-export type GetIssueNodeResponse = {
+export type GetIssueSpanResponse = {
+
+	/** Unique id of this span */
+	id: string;
 
 	/** The template containing this step */
 	templateId: string;
@@ -2496,7 +2503,7 @@ export type GetIssueStreamResponse = {
 	maxChange?: number;
 
 	/**Map of steps to (event signature id -> trace id) */
-	nodes: GetIssueNodeResponse[];
+	nodes: GetIssueSpanResponse[];
 }
 
 /**Outcome of a particular build */
@@ -2583,11 +2590,14 @@ export type GetIssueResponse = {
 	/**The summary text for this issue */
 	summary: string;
 
-	/**Details text describing the issue */
-	details?: string;
+	/**Description of the issue*/
+	description?: string;
 
 	/** Severity of this issue	*/
 	severity: IssueSeverity;
+	
+	/**Whether the issue is promoted */
+	promoted: boolean;
 
 	/**Owner of the issue */
 	owner?: string;
@@ -2655,6 +2665,19 @@ export type UpdateIssueRequest = {
 
 	/**Whether the issue should be marked as resolved */
 	resolved?: boolean;
+
+	/**Description of the issue*/
+	description?: string;
+
+	/**Whether the issue is promoted */
+	promoted?: boolean;
+
+	/**  List of spans to add to this issue	 */
+	addSpans?: string[];
+
+	/** List of spans to remove from this issue */
+	removeSpans?: string[];
+
 }
 
 export type GetUtilizationTelemetryStream = {
diff --git a/Engine/Source/Programs/Horde/HordeDashboard/src/backend/Issues.ts b/Engine/Source/Programs/Horde/HordeDashboard/src/backend/Issues.ts
--- a/Engine/Source/Programs/Horde/HordeDashboard/src/backend/Issues.ts
+++ b/Engine/Source/Programs/Horde/HordeDashboard/src/backend/Issues.ts
@@ -73,7 +73,7 @@ export const getIssueSummary = (issue?: GetIssueResponse, events?: GetLogEventRe
     }
 
     if (type === IssueType.Unknown || !summary || !events || !events.length) {
-        return summary ? summary : issue.details ?? "Issue has no summary or details";
+        return summary ? summary : "Issue has no summary or details";
     }
 
     if (type === IssueType.Compiler) {
diff --git a/Engine/Source/Programs/Horde/HordeDashboard/src/base/utilities/issueUtils.ts b/Engine/Source/Programs/Horde/HordeDashboard/src/base/utilities/issueUtils.ts
--- a/Engine/Source/Programs/Horde/HordeDashboard/src/base/utilities/issueUtils.ts
+++ b/Engine/Source/Programs/Horde/HordeDashboard/src/base/utilities/issueUtils.ts
@@ -4,10 +4,43 @@ import dashboard from "../../backend/Dashboard";
 import { projectStore } from "../../backend/ProjectStore";
 import { displayTimeZone } from "./timeUtils";
 
+export type IssueJira = {
+	description: string;
+	link: string;
+}
+
+const urlRegex = /([\w+]+\:\/\/)?([\w\d-]+\.)*[\w-]+[\.\:]\w+([\/\?\=\&\#\.]?[\w-]+)*\/?/gm
+
+export function getIssueJiras(issue: GetIssueResponse): IssueJira[] {
+
+	if (!issue.description) {
+		return [];
+	}
+
+	const jiras: IssueJira[] = [];
+
+	const matches = Array.from(issue.description.matchAll(urlRegex));
+
+	matches.forEach(m => {
+		m.forEach(match => {
+			// @todo: add jira match to server config
+			if (match.toLowerCase().startsWith("https://jira.it.epicgames.com")) {
+				const path = new URL(match).pathname.split("/");
+				if (path.length) {
+					jiras.push({link: match, description: path[path.length - 1]})
+				}
+				
+			}
+		})
+	})
+	
+	return jiras;
+
+}
 
-export function getIssueStatus(issue: GetIssueResponse, showResolveTime?:boolean): string {
+export function getIssueStatus(issue: GetIssueResponse, showResolveTime?: boolean): string {
 
-	let text = "";		
+	let text = "";
 
 	if (issue.resolvedAt) {
 
@@ -17,7 +50,7 @@ export function getIssueStatus(issue: GetIssueResponse, showResolveTime?:boolean
 			resolvedText += ` in CL ${issue.fixChange}`;
 		}
 
-		resolvedText += ` by ${issue.resolvedBy ?? "Horde" }`;
+		resolvedText += ` by ${issue.resolvedBy ?? "Horde"}`;
 
 		if (!showResolveTime) {
 			return resolvedText;
