Horde: Support external links in issue markdown
#rnx

diff --git a/Engine/Source/Programs/Horde/HordeDashboard/src/base/utilities/issueUtils.ts b/Engine/Source/Programs/Horde/HordeDashboard/src/base/utilities/issueUtils.ts
--- a/Engine/Source/Programs/Horde/HordeDashboard/src/base/utilities/issueUtils.ts
+++ b/Engine/Source/Programs/Horde/HordeDashboard/src/base/utilities/issueUtils.ts
@@ -23,6 +23,10 @@ export function getIssueJiras(issue: GetIssueResponse): IssueJira[] {
 
 	matches.forEach(m => {
 		m.forEach(match => {
+
+			if (!match) {
+				return;
+			}
 			// @todo: add jira match to server config
 			if (match.toLowerCase().startsWith("https://jira.it.epicgames.com")) {
 				const path = new URL(match).pathname.split("/");
