Horde: Always return copies of job documents rather than mutating the original.

diff --git a/Engine/Source/Programs/Horde/Horde.Build.Tests/DatabaseIntegrationTest.cs b/Engine/Source/Programs/Horde/Horde.Build.Tests/DatabaseIntegrationTest.cs
--- a/Engine/Source/Programs/Horde/Horde.Build.Tests/DatabaseIntegrationTest.cs
+++ b/Engine/Source/Programs/Horde/Horde.Build.Tests/DatabaseIntegrationTest.cs
@@ -8,6 +8,7 @@ using HordeServer.Services;
 using HordeServer.Utilities;
 using Microsoft.Extensions.Logging;
 using Microsoft.Extensions.Options;
+using Microsoft.VisualStudio.TestTools.UnitTesting;
 using StackExchange.Redis;
 
 namespace HordeServerTests
@@ -139,5 +140,11 @@ namespace HordeServerTests
 		        }    
 	        }
         }
-    }
+
+		public static T Deref<T>(T? Item)
+		{
+			Assert.IsNotNull(Item);
+			return Item!;
+		}
+	}
 }
\ No newline at end of file
diff --git a/Engine/Source/Programs/Horde/Horde.Build.Tests/JobCollectionTests.cs b/Engine/Source/Programs/Horde/Horde.Build.Tests/JobCollectionTests.cs
--- a/Engine/Source/Programs/Horde/Horde.Build.Tests/JobCollectionTests.cs
+++ b/Engine/Source/Programs/Horde/Horde.Build.Tests/JobCollectionTests.cs
@@ -63,21 +63,23 @@ namespace HordeServerTests
 			return Node;
 		}
 
-		async Task StartBatch(IJob Job, IGraph Graph, int BatchIdx)
+		async Task<IJob> StartBatch(IJob Job, IGraph Graph, int BatchIdx)
 		{
 			Assert.AreEqual(JobStepBatchState.Ready, Job.Batches[BatchIdx].State);
-			Assert.IsTrue(await JobCollection.TryUpdateBatchAsync(Job, Graph, Job.Batches[BatchIdx].Id, null, JobStepBatchState.Running, null));
+			Job = Deref(await JobCollection.TryUpdateBatchAsync(Job, Graph, Job.Batches[BatchIdx].Id, null, JobStepBatchState.Running, null));
 			Assert.AreEqual(JobStepBatchState.Running, Job.Batches[BatchIdx].State);
+			return Job;
 		}
 
-		async Task RunStep(IJob Job, IGraph Graph, int BatchIdx, int StepIdx, JobStepOutcome Outcome)
+		async Task<IJob> RunStep(IJob Job, IGraph Graph, int BatchIdx, int StepIdx, JobStepOutcome Outcome)
 		{
 			Assert.AreEqual(JobStepState.Ready, Job.Batches[BatchIdx].Steps[StepIdx].State);
-			Assert.IsTrue(await JobCollection.TryUpdateStepAsync(Job, Graph, Job.Batches[BatchIdx].Id, Job.Batches[BatchIdx].Steps[StepIdx].Id, JobStepState.Running, JobStepOutcome.Success));
+			Job = Deref(await JobCollection.TryUpdateStepAsync(Job, Graph, Job.Batches[BatchIdx].Id, Job.Batches[BatchIdx].Steps[StepIdx].Id, JobStepState.Running, JobStepOutcome.Success));
 			Assert.AreEqual(JobStepState.Running, Job.Batches[BatchIdx].Steps[StepIdx].State);
-			Assert.IsTrue(await JobCollection.TryUpdateStepAsync(Job, Graph, Job.Batches[BatchIdx].Id, Job.Batches[BatchIdx].Steps[StepIdx].Id, JobStepState.Completed, Outcome));
+			Job = Deref(await JobCollection.TryUpdateStepAsync(Job, Graph, Job.Batches[BatchIdx].Id, Job.Batches[BatchIdx].Steps[StepIdx].Id, JobStepState.Completed, Outcome));
 			Assert.AreEqual(JobStepState.Completed, Job.Batches[BatchIdx].Steps[StepIdx].State);
 			Assert.AreEqual(Outcome, Job.Batches[BatchIdx].Steps[StepIdx].Outcome);
+			return Job;
 		}
 
 		[TestMethod]
@@ -94,8 +96,8 @@ namespace HordeServerTests
 
 			IJob Job = await JobCollection.AddAsync(JobId.GenerateNewId(), new StreamId("ue4-main"), new TemplateRefId("test-build"), ContentHash.SHA1("hello"), BaseGraph, "Test job", 123, 123, null, null, null, null, null, null, null, false, false, null, null, Arguments);
 
-			await StartBatch(Job, BaseGraph, 0);
-			await RunStep(Job, BaseGraph, 0, 0, JobStepOutcome.Success); // Setup Build
+			Job = await StartBatch(Job, BaseGraph, 0);
+			Job = await RunStep(Job, BaseGraph, 0, 0, JobStepOutcome.Success); // Setup Build
 
 			List<NewGroup> NewGroups = new List<NewGroup>();
 
@@ -112,17 +114,17 @@ namespace HordeServerTests
 			AddNode(PublishGroup, "Post-Publish Client", null, x => x.OrderDependencies = new List<string> { "Publish Client" });
 
 			IGraph Graph = await GraphCollection.AppendAsync(BaseGraph, NewGroups, null, null);
-			Assert.IsTrue(await JobCollection.TryUpdateGraphAsync(Job, Graph));
+			Job = Deref(await JobCollection.TryUpdateGraphAsync(Job, Graph));
 
-			await StartBatch(Job, Graph, 1);
-			await RunStep(Job, Graph, 1, 0, JobStepOutcome.Success); // Update Version Files
-			await RunStep(Job, Graph, 1, 1, JobStepOutcome.Success); // Compile Editor
+			Job = await StartBatch(Job, Graph, 1);
+			Job = await RunStep(Job, Graph, 1, 0, JobStepOutcome.Success); // Update Version Files
+			Job = await RunStep(Job, Graph, 1, 1, JobStepOutcome.Success); // Compile Editor
 
-			await StartBatch(Job, Graph, 2);
-			await RunStep(Job, Graph, 2, 0, JobStepOutcome.Success); // Compile Client
+			Job = await StartBatch(Job, Graph, 2);
+			Job = await RunStep(Job, Graph, 2, 0, JobStepOutcome.Success); // Compile Client
 
-			await StartBatch(Job, Graph, 3);
-			await RunStep(Job, Graph, 3, 0, JobStepOutcome.Failure); // Cook Client
+			Job = await StartBatch(Job, Graph, 3);
+			Job = await RunStep(Job, Graph, 3, 0, JobStepOutcome.Failure); // Cook Client
 			Assert.AreEqual(JobStepState.Skipped, Job.Batches[3].Steps[1].State); // Publish Client
 			Assert.AreEqual(JobStepState.Skipped, Job.Batches[3].Steps[2].State); // Post-Publish Client
 		}
@@ -166,8 +168,8 @@ namespace HordeServerTests
 
 			IJob Job = await JobCollection.AddAsync(JobId.GenerateNewId(), new StreamId("ue4-main"), new TemplateRefId("test-build"), ContentHash.SHA1("hello"), BaseGraph, "Test job", 123, 123, null, null, null, null, null, null, null, false, false, null, null, Arguments);
 
-			await StartBatch(Job, BaseGraph, 0);
-			await RunStep(Job, BaseGraph, 0, 0, JobStepOutcome.Success); // Setup Build
+			Job = await StartBatch(Job, BaseGraph, 0);
+			Job = await RunStep(Job, BaseGraph, 0, 0, JobStepOutcome.Success); // Setup Build
 
 			List<NewGroup> NewGroups = new List<NewGroup>();
 
@@ -177,14 +179,14 @@ namespace HordeServerTests
 			AddNode(InitialGroup, "Step 3", new[] { "Step 2" });
 
 			IGraph Graph = await GraphCollection.AppendAsync(BaseGraph, NewGroups, null, null);
-			Assert.IsTrue(await JobCollection.TryUpdateGraphAsync(Job, Graph));
+			Job = Deref(await JobCollection.TryUpdateGraphAsync(Job, Graph));
 
-			await StartBatch(Job, Graph, 1);
-			await RunStep(Job, Graph, 1, 0, JobStepOutcome.Success); // Step 1
-			await RunStep(Job, Graph, 1, 1, JobStepOutcome.Success); // Step 2
+			Job = await StartBatch(Job, Graph, 1);
+			Job = await RunStep(Job, Graph, 1, 0, JobStepOutcome.Success); // Step 1
+			Job = await RunStep(Job, Graph, 1, 1, JobStepOutcome.Success); // Step 2
 
 			// Force an error executing the batch
-			Assert.IsTrue(await JobCollection.TryUpdateBatchAsync(Job, Graph, Job.Batches[1].Id, null, JobStepBatchState.Complete, JobStepBatchError.Incomplete));
+			Job = Deref(await JobCollection.TryUpdateBatchAsync(Job, Graph, Job.Batches[1].Id, null, JobStepBatchState.Complete, JobStepBatchError.Incomplete));
 
 			// Check that it restarted all three nodes
 			IJob NewJob = (await JobCollection.GetAsync(Job.Id))!;
@@ -223,8 +225,8 @@ namespace HordeServerTests
 			IJob Job = await JobCollection.AddAsync(JobId.GenerateNewId(), new StreamId("ue4-main"), new TemplateRefId("test-build"), ContentHash.SHA1("hello"), BaseGraph, "Test job", 123, 123, null, null, null, null, null, null, null, false, false, null, null, Arguments);
 			Assert.AreEqual(1, Job.Batches.Count);
 
-			await StartBatch(Job, BaseGraph, 0);
-			Assert.IsTrue(await JobCollection.TryUpdateBatchAsync(Job, BaseGraph, Job.Batches[0].Id, null, JobStepBatchState.Complete, JobStepBatchError.Incomplete));
+			Job = await StartBatch(Job, BaseGraph, 0);
+			Job = Deref(await JobCollection.TryUpdateBatchAsync(Job, BaseGraph, Job.Batches[0].Id, null, JobStepBatchState.Complete, JobStepBatchError.Incomplete));
 
 			Job = (await JobCollection.GetAsync(Job.Id))!;
 			Assert.AreEqual(2, Job.Batches.Count);
@@ -233,7 +235,7 @@ namespace HordeServerTests
 			Assert.AreEqual(JobStepBatchState.Ready, Job.Batches[1].State);
 			Assert.AreEqual(1, Job.Batches[1].Steps.Count);
 
-			Assert.IsTrue(await JobCollection.TryUpdateBatchAsync(Job, BaseGraph, Job.Batches[1].Id, null, JobStepBatchState.Complete, JobStepBatchError.Incomplete));
+			Job = Deref(await JobCollection.TryUpdateBatchAsync(Job, BaseGraph, Job.Batches[1].Id, null, JobStepBatchState.Complete, JobStepBatchError.Incomplete));
 
 			Job = (await JobCollection.GetAsync(Job.Id))!;
 			Assert.AreEqual(3, Job.Batches.Count);
@@ -244,7 +246,7 @@ namespace HordeServerTests
 			Assert.AreEqual(JobStepBatchState.Ready, Job.Batches[2].State);
 			Assert.AreEqual(1, Job.Batches[2].Steps.Count);
 
-			Assert.IsTrue(await JobCollection.TryUpdateBatchAsync(Job, BaseGraph, Job.Batches[2].Id, null, JobStepBatchState.Complete, JobStepBatchError.Incomplete));
+			Job = Deref(await JobCollection.TryUpdateBatchAsync(Job, BaseGraph, Job.Batches[2].Id, null, JobStepBatchState.Complete, JobStepBatchError.Incomplete));
 
 			Job = (await JobCollection.GetAsync(Job.Id))!;
 			Assert.AreEqual(3, Job.Batches.Count);
diff --git a/Engine/Source/Programs/Horde/Horde.Build.Tests/JobServiceTests.cs b/Engine/Source/Programs/Horde/Horde.Build.Tests/JobServiceTests.cs
--- a/Engine/Source/Programs/Horde/Horde.Build.Tests/JobServiceTests.cs
+++ b/Engine/Source/Programs/Horde/Horde.Build.Tests/JobServiceTests.cs
@@ -50,9 +50,9 @@ namespace HordeServerTests
 			IJob Job = await TestSetup.JobService.CreateJobAsync(null, Stream!, TemplateRefId1, Template.Id, Graph, "Hello", 1234, 1233, 999, null, null, null, null, null, Stream!.Templates[TemplateRefId1].ChainedJobs, true, true, null, null, new List<string>());
 			Assert.AreEqual(1, Job.ChainedJobs.Count);
 
-			Assert.IsTrue(await TestSetup.JobService.UpdateBatchAsync(Job, Job.Batches[0].Id, LogId.GenerateNewId(), JobStepBatchState.Running));
-			Assert.IsNotNull(await TestSetup.JobService.UpdateStepAsync(Job, Job.Batches[0].Id, Job.Batches[0].Steps[0].Id, JobStepState.Running));
-			Assert.IsNotNull(await TestSetup.JobService.UpdateStepAsync(Job, Job.Batches[0].Id, Job.Batches[0].Steps[0].Id, JobStepState.Completed, JobStepOutcome.Success));
+			Job = Deref(await TestSetup.JobService.UpdateBatchAsync(Job, Job.Batches[0].Id, LogId.GenerateNewId(), JobStepBatchState.Running));
+			Job = Deref(await TestSetup.JobService.UpdateStepAsync(Job, Job.Batches[0].Id, Job.Batches[0].Steps[0].Id, JobStepState.Running));
+			Job = Deref(await TestSetup.JobService.UpdateStepAsync(Job, Job.Batches[0].Id, Job.Batches[0].Steps[0].Id, JobStepState.Completed, JobStepOutcome.Success));
 
 			Assert.IsNotNull(Job.ChainedJobs[0].JobId);
 
@@ -172,14 +172,14 @@ namespace HordeServerTests
 
 			IJob Job = await TestSetup.JobService.CreateJobAsync(null, Stream!, new TemplateRefId("temp"), Template.Id, Graph, "Hello", 1234, 1233, 999, null, null, null, null, null, null, true, true, null, null, new List<string> { "-Target=Pak" });
 
-			Assert.IsTrue(await TestSetup.JobService.UpdateBatchAsync(Job, Job.Batches[0].Id, LogId.GenerateNewId(), JobStepBatchState.Running));
+			Job = Deref(await TestSetup.JobService.UpdateBatchAsync(Job, Job.Batches[0].Id, LogId.GenerateNewId(), JobStepBatchState.Running));
 			Assert.IsNotNull(await TestSetup.JobService.UpdateStepAsync(Job, Job.Batches[0].Id, Job.Batches[0].Steps[0].Id, JobStepState.Running));
 			Assert.IsNotNull(await TestSetup.JobService.UpdateStepAsync(Job, Job.Batches[0].Id, Job.Batches[0].Steps[0].Id, JobStepState.Completed, JobStepOutcome.Success));
 
-			Assert.IsTrue(await TestSetup.JobService.UpdateBatchAsync(Job, Job.Batches[1].Id, LogId.GenerateNewId(), JobStepBatchState.Running));
+			Job = Deref(await TestSetup.JobService.UpdateBatchAsync(Job, Job.Batches[1].Id, LogId.GenerateNewId(), JobStepBatchState.Running));
 			Assert.IsNotNull(await TestSetup.JobService.UpdateStepAsync(Job, Job.Batches[1].Id, Job.Batches[1].Steps[0].Id, JobStepState.Running));
 
-			Assert.IsTrue(await TestSetup.JobService.UpdateBatchAsync(Job, Job.Batches[2].Id, LogId.GenerateNewId(), JobStepBatchState.Running));
+			Job = Deref(await TestSetup.JobService.UpdateBatchAsync(Job, Job.Batches[2].Id, LogId.GenerateNewId(), JobStepBatchState.Running));
 			Assert.IsNotNull(await TestSetup.JobService.UpdateStepAsync(Job, Job.Batches[2].Id, Job.Batches[2].Steps[0].Id, JobStepState.Running));
 			Assert.IsNotNull(await TestSetup.JobService.UpdateStepAsync(Job, Job.Batches[2].Id, Job.Batches[2].Steps[0].Id, JobStepState.Completed, JobStepOutcome.Success));
 
diff --git a/Engine/Source/Programs/Horde/Horde.Build.Tests/RpcServiceTest.cs b/Engine/Source/Programs/Horde/Horde.Build.Tests/RpcServiceTest.cs
--- a/Engine/Source/Programs/Horde/Horde.Build.Tests/RpcServiceTest.cs
+++ b/Engine/Source/Programs/Horde/Horde.Build.Tests/RpcServiceTest.cs
@@ -410,7 +410,7 @@ namespace HordeServerTests
 			};
 
 			// Set the session ID on the job batch to pass auth later
-			Assert.IsTrue(await TestSetup.JobCollection.TryAssignLeaseAsync(TestSetup.Fixture.Job1, 0, new PoolId("foo"),
+			Deref(await TestSetup.JobCollection.TryAssignLeaseAsync(TestSetup.Fixture.Job1, 0, new PoolId("foo"),
 				new AgentId("test"), SessionId,
 				LeaseId.GenerateNewId(), LogId.GenerateNewId()));
 /*
diff --git a/Engine/Source/Programs/Horde/Horde.Build.Tests/SchedulerTests.cs b/Engine/Source/Programs/Horde/Horde.Build.Tests/SchedulerTests.cs
--- a/Engine/Source/Programs/Horde/Horde.Build.Tests/SchedulerTests.cs
+++ b/Engine/Source/Programs/Horde/Horde.Build.Tests/SchedulerTests.cs
@@ -397,9 +397,9 @@ namespace HordeServerTests
 
 			// Create a job and fail it
 			IJob Job1 = await TestSetup.JobService.CreateJobAsync(null, Stream, NewTemplateRefId1, Template.Id, GraphA, "Hello", 1234, 1233, 999, null, null, null, null, null, null, true, true, null, null, new List<string> { "-Target=TriggerNext" });
-			Assert.IsTrue(await TestSetup.JobService.UpdateBatchAsync(Job1, Job1.Batches[0].Id, LogId.GenerateNewId(), JobStepBatchState.Running));
+			Job1 = Deref(await TestSetup.JobService.UpdateBatchAsync(Job1, Job1.Batches[0].Id, LogId.GenerateNewId(), JobStepBatchState.Running));
 			Assert.IsNotNull(await TestSetup.JobService.UpdateStepAsync(Job1, Job1.Batches[0].Id, Job1.Batches[0].Steps[0].Id, JobStepState.Completed, JobStepOutcome.Failure));
-			Assert.IsTrue(await TestSetup.JobService.UpdateBatchAsync(Job1, Job1.Batches[0].Id, LogId.GenerateNewId(), JobStepBatchState.Complete));
+			Job1 = Deref(await TestSetup.JobService.UpdateBatchAsync(Job1, Job1.Batches[0].Id, LogId.GenerateNewId(), JobStepBatchState.Complete));
 			await GetNewJobs();
 
 			// Tick the schedule and make sure it doesn't trigger
@@ -410,7 +410,7 @@ namespace HordeServerTests
 
 			// Create a job and make it succeed
 			IJob Job2 = await TestSetup.JobService.CreateJobAsync(null, Stream, NewTemplateRefId1, Template.Id, GraphA, "Hello", 1234, 1233, 999, null, null, null, null, null, null, true, true, null, null, new List<string> { "-Target=TriggerNext" });
-			Assert.IsTrue(await TestSetup.JobService.UpdateBatchAsync(Job2, Job2.Batches[0].Id, LogId.GenerateNewId(), JobStepBatchState.Running));
+			Job2 = Deref(await TestSetup.JobService.UpdateBatchAsync(Job2, Job2.Batches[0].Id, LogId.GenerateNewId(), JobStepBatchState.Running));
 			Assert.IsNotNull(await TestSetup.JobService.UpdateStepAsync(Job2, Job2.Batches[0].Id, Job2.Batches[0].Steps[0].Id, JobStepState.Completed, JobStepOutcome.Success));
 
 			// Tick the schedule and make sure it does trigger
diff --git a/Engine/Source/Programs/Horde/Horde.Build/Controllers/JobsController.cs b/Engine/Source/Programs/Horde/Horde.Build/Controllers/JobsController.cs
--- a/Engine/Source/Programs/Horde/Horde.Build/Controllers/JobsController.cs
+++ b/Engine/Source/Programs/Horde/Horde.Build/Controllers/JobsController.cs
@@ -250,7 +250,8 @@ namespace HordeServer.Controllers
 				AbortedByUserId = User.GetUserId();
 			}
 
-			if (!await JobService.UpdateJobAsync(Job, Name: Request.Name, Priority: Request.Priority, AutoSubmit: Request.AutoSubmit, AbortedByUserId: AbortedByUserId, Arguments: Request.Arguments))
+			IJob? NewJob = await JobService.UpdateJobAsync(Job, Name: Request.Name, Priority: Request.Priority, AutoSubmit: Request.AutoSubmit, AbortedByUserId: AbortedByUserId, Arguments: Request.Arguments);
+			if (NewJob == null)
 			{
 				return NotFound();
 			}
@@ -278,7 +279,9 @@ namespace HordeServer.Controllers
 			}
 
 			ObjectId TriggerId = Job.NotificationTriggerId ?? ObjectId.GenerateNewId();
-			if (!await JobService.UpdateJobAsync(Job, null, null, null, null, TriggerId, null, null))
+
+			Job = await JobService.UpdateJobAsync(Job, null, null, null, null, TriggerId, null, null);
+			if (Job == null)
 			{
 				return NotFound();
 			}
@@ -703,7 +706,8 @@ namespace HordeServer.Controllers
 				IGraph Graph = await JobService.GetGraphAsync(Job);
 				Graph = await Graphs.AppendAsync(Graph, Requests, null, null);
 
-				if (await JobService.TryUpdateGraphAsync(Job, Graph))
+				IJob? NewJob = await JobService.TryUpdateGraphAsync(Job, Graph);
+				if (NewJob != null)
 				{
 					return Ok();
 				}
@@ -887,7 +891,8 @@ namespace HordeServer.Controllers
 				return Forbid();
 			}
 
-			if (!await JobService.UpdateBatchAsync(Job, BatchIdValue, Request.LogId?.ToObjectId<ILogFile>(), Request.State))
+			IJob? NewJob = await JobService.UpdateBatchAsync(Job, BatchIdValue, Request.LogId?.ToObjectId<ILogFile>(), Request.State);
+			if (NewJob == null)
 			{
 				return NotFound();
 			}
@@ -1324,7 +1329,9 @@ namespace HordeServer.Controllers
 				}
 
 				NewTriggerId = ObjectId.GenerateNewId();
-				if (await JobService.UpdateJobAsync(Job, LabelIdxToTriggerId: new KeyValuePair<int, ObjectId>(LabelIndex, NewTriggerId)))
+
+				IJob? NewJob = await JobService.UpdateJobAsync(Job, LabelIdxToTriggerId: new KeyValuePair<int, ObjectId>(LabelIndex, NewTriggerId));
+				if (NewJob != null)
 				{
 					TriggerId = NewTriggerId;
 					break;
diff --git a/Engine/Source/Programs/Horde/Horde.Build/Jobs/IJobCollection.cs b/Engine/Source/Programs/Horde/Horde.Build/Jobs/IJobCollection.cs
--- a/Engine/Source/Programs/Horde/Horde.Build/Jobs/IJobCollection.cs
+++ b/Engine/Source/Programs/Horde/Horde.Build/Jobs/IJobCollection.cs
@@ -123,7 +123,7 @@ namespace HordeServer.Collections
 		/// <param name="Arguments">New arguments for the job</param>
 		/// <param name="LabelIdxToTriggerId">New trigger ID for a label in the job</param>
 		/// <param name="JobTrigger">New downstream job id</param>
-		Task<bool> TryUpdateJobAsync(IJob Job, IGraph Graph, string? Name = null, Priority? Priority = null, bool? AutoSubmit = null, int? AutoSubmitChange = null, string? AutoSubmitMessage = null, UserId? AbortedByUserId = null, ObjectId? NotificationTriggerId = null, List<Report>? Reports = null, List<string>? Arguments = null, KeyValuePair<int, ObjectId>? LabelIdxToTriggerId = null, KeyValuePair<TemplateRefId, JobId>? JobTrigger = null);
+		Task<IJob?> TryUpdateJobAsync(IJob Job, IGraph Graph, string? Name = null, Priority? Priority = null, bool? AutoSubmit = null, int? AutoSubmitChange = null, string? AutoSubmitMessage = null, UserId? AbortedByUserId = null, ObjectId? NotificationTriggerId = null, List<Report>? Reports = null, List<string>? Arguments = null, KeyValuePair<int, ObjectId>? LabelIdxToTriggerId = null, KeyValuePair<TemplateRefId, JobId>? JobTrigger = null);
 
 		/// <summary>
 		/// Updates the state of a batch
@@ -135,7 +135,7 @@ namespace HordeServer.Collections
 		/// <param name="NewState">New state of the jobstep</param>
 		/// <param name="NewError">Error code for the batch</param>
 		/// <returns>True if the job was updated, false if it was deleted</returns>
-		Task<bool> TryUpdateBatchAsync(IJob Job, IGraph Graph, SubResourceId BatchId, LogId? NewLogId, JobStepBatchState? NewState, JobStepBatchError? NewError);
+		Task<IJob?> TryUpdateBatchAsync(IJob Job, IGraph Graph, SubResourceId BatchId, LogId? NewLogId, JobStepBatchState? NewState, JobStepBatchError? NewError);
 
 		/// <summary>
 		/// Update a jobstep state
@@ -155,7 +155,7 @@ namespace HordeServer.Collections
 		/// <param name="NewReports">New report documents</param>
 		/// <param name="NewProperties">Property changes. Any properties with a null value will be removed.</param>
 		/// <returns>True if the job was updated, false if it was deleted in the meantime</returns>
-		Task<bool> TryUpdateStepAsync(IJob Job, IGraph Graph, SubResourceId BatchId, SubResourceId StepId, JobStepState NewState = default, JobStepOutcome NewOutcome = default, bool? NewAbortRequested = null, UserId? NewAbortByUserId = null, LogId? NewLogId = null, ObjectId? NewNotificationTriggerId = null, UserId? NewRetryByUserId = null, Priority? NewPriority = null, List<Report>? NewReports = null, Dictionary<string, string?>? NewProperties = null);
+		Task<IJob?> TryUpdateStepAsync(IJob Job, IGraph Graph, SubResourceId BatchId, SubResourceId StepId, JobStepState NewState = default, JobStepOutcome NewOutcome = default, bool? NewAbortRequested = null, UserId? NewAbortByUserId = null, LogId? NewLogId = null, ObjectId? NewNotificationTriggerId = null, UserId? NewRetryByUserId = null, Priority? NewPriority = null, List<Report>? NewReports = null, Dictionary<string, string?>? NewProperties = null);
 
 		/// <summary>
 		/// Attempts to update the node groups to be executed for a job. Fails if another write happens in the meantime.
@@ -163,14 +163,14 @@ namespace HordeServer.Collections
 		/// <param name="Job">The job to update</param>
 		/// <param name="NewGraph">New graph for this job</param>
 		/// <returns>True if the groups were updated to the given list. False if another write happened first.</returns>
-		Task<bool> TryUpdateGraphAsync(IJob Job, IGraph NewGraph);
+		Task<IJob?> TryUpdateGraphAsync(IJob Job, IGraph NewGraph);
 
 		/// <summary>
 		/// Removes a job from the dispatch queue. Ignores the state of any batches still remaining to execute. Should only be used to correct for inconsistent state.
 		/// </summary>
 		/// <param name="Job">Job</param>
 		/// <returns></returns>
-		Task<bool> TryRemoveFromDispatchQueueAsync(IJob Job);
+		Task<IJob?> TryRemoveFromDispatchQueueAsync(IJob Job);
 
 		/// <summary>
 		/// Adds an issue to a job
@@ -212,7 +212,7 @@ namespace HordeServer.Collections
 		/// <param name="BatchIdx">Index of the batch to cancel</param>
 		/// <param name="Graph">Graph for the job</param>
 		/// <returns>True if the job is updated</returns>
-		Task<bool> TryFailBatchAsync(IJob Job, int BatchIdx, IGraph Graph);
+		Task<IJob?> TryFailBatchAsync(IJob Job, int BatchIdx, IGraph Graph);
 
 		/// <summary>
 		/// Attempt to assign a lease to execute a batch
@@ -225,7 +225,7 @@ namespace HordeServer.Collections
 		/// <param name="LeaseId">The lease unique id</param>
 		/// <param name="LogId">Unique id of the log for the batch</param>
 		/// <returns>True if the batch is updated</returns>
-		Task<bool> TryAssignLeaseAsync(IJob Job, int BatchIdx, PoolId PoolId, AgentId AgentId, ObjectId SessionId, LeaseId LeaseId, LogId LogId);
+		Task<IJob?> TryAssignLeaseAsync(IJob Job, int BatchIdx, PoolId PoolId, AgentId AgentId, ObjectId SessionId, LeaseId LeaseId, LogId LogId);
 
 		/// <summary>
 		/// Cancel a lease reservation on a batch (before it has started)
@@ -233,7 +233,7 @@ namespace HordeServer.Collections
 		/// <param name="Job">The job containing the lease</param>
 		/// <param name="BatchIdx">Index of the batch to cancel</param>
 		/// <returns>True if the job is updated</returns>
-		Task<bool> TryCancelLeaseAsync(IJob Job, int BatchIdx);
+		Task<IJob?> TryCancelLeaseAsync(IJob Job, int BatchIdx);
 
 		/// <summary>
 		/// Upgrade all documents in the collection
diff --git a/Engine/Source/Programs/Horde/Horde.Build/Jobs/Impl/JobCollection.cs b/Engine/Source/Programs/Horde/Horde.Build/Jobs/Impl/JobCollection.cs
--- a/Engine/Source/Programs/Horde/Horde.Build/Jobs/Impl/JobCollection.cs
+++ b/Engine/Source/Programs/Horde/Horde.Build/Jobs/Impl/JobCollection.cs
@@ -18,6 +18,8 @@ using System.Linq;
 using System.Threading.Tasks;
 using OpenTracing.Util;
 using OpenTracing;
+using MongoDB.Bson.Serialization;
+using MongoDB.Bson.IO;
 
 namespace HordeServer.Collections.Impl
 {
@@ -429,6 +431,18 @@ namespace HordeServer.Collections.Impl
 			}
 		}
 
+		static JobDocument Clone(JobDocument Job)
+		{
+			using (MemoryStream Stream = new MemoryStream())
+			{
+				using (BsonBinaryWriter Writer = new BsonBinaryWriter(Stream))
+				{
+					BsonSerializer.Serialize(Writer, Job);
+				}
+				return BsonSerializer.Deserialize<JobDocument>(Stream.ToArray());
+			}
+		}
+
 		/// <inheritdoc/>
 		[SuppressMessage("Compiler", "CA1054:URI parameters should not be strings")]
 		public async Task<IJob> AddAsync(JobId JobId, StreamId StreamId, TemplateRefId TemplateRefId, ContentHash TemplateHash, IGraph Graph, string Name, int Change, int CodeChange, int? PreflightChange, int? ClonedPreflightChange, UserId? StartedByUserId, Priority? Priority, bool? AutoSubmit, bool? UpdateIssues, List<ChainedJobTemplate>? ChainedJobs, bool ShowUgsBadges, bool ShowUgsAlerts, string? NotificationChannel, string? NotificationChannelFilter, List<string>? Arguments)
@@ -562,7 +576,7 @@ namespace HordeServer.Collections.Impl
 		}
 
 		/// <inheritdoc/>
-		public async Task<bool> TryUpdateJobAsync(IJob Job, IGraph Graph, string? Name, Priority? Priority, bool? AutoSubmit, int? AutoSubmitChange, string? AutoSubmitMessage, UserId? AbortedByUserId, ObjectId? NotificationTriggerId, List<Report>? Reports, List<string>? Arguments, KeyValuePair<int, ObjectId>? LabelIdxToTriggerId, KeyValuePair<TemplateRefId, JobId>? JobTrigger)
+		public async Task<IJob?> TryUpdateJobAsync(IJob InJob, IGraph Graph, string? Name, Priority? Priority, bool? AutoSubmit, int? AutoSubmitChange, string? AutoSubmitMessage, UserId? AbortedByUserId, ObjectId? NotificationTriggerId, List<Report>? Reports, List<string>? Arguments, KeyValuePair<int, ObjectId>? LabelIdxToTriggerId, KeyValuePair<TemplateRefId, JobId>? JobTrigger)
 		{
 			// Create the update 
 			UpdateDefinitionBuilder<JobDocument> UpdateBuilder = Builders<JobDocument>.Update;
@@ -572,36 +586,36 @@ namespace HordeServer.Collections.Impl
 			bool bUpdateBatches = false;
 
 			// Build the update list
-			JobDocument JobDocument = (JobDocument)Job;
+			JobDocument JobDocument = Clone((JobDocument)InJob);
 			if (Name != null)
 			{
 				JobDocument.Name = Name;
-				Updates.Add(UpdateBuilder.Set(x => x.Name, Job.Name));
+				Updates.Add(UpdateBuilder.Set(x => x.Name, JobDocument.Name));
 			}
 			if (Priority != null)
 			{
 				JobDocument.Priority = Priority.Value;
-				Updates.Add(UpdateBuilder.Set(x => x.Priority, Job.Priority));
+				Updates.Add(UpdateBuilder.Set(x => x.Priority, JobDocument.Priority));
 			}
 			if (AutoSubmit != null)
 			{
 				JobDocument.AutoSubmit = AutoSubmit.Value;
-				Updates.Add(UpdateBuilder.Set(x => x.AutoSubmit, Job.AutoSubmit));
+				Updates.Add(UpdateBuilder.Set(x => x.AutoSubmit, JobDocument.AutoSubmit));
 			}
 			if (AutoSubmitChange != null)
 			{
 				JobDocument.AutoSubmitChange = AutoSubmitChange.Value;
-				Updates.Add(UpdateBuilder.Set(x => x.AutoSubmitChange, Job.AutoSubmitChange));
+				Updates.Add(UpdateBuilder.Set(x => x.AutoSubmitChange, JobDocument.AutoSubmitChange));
 			}
 			if (AutoSubmitMessage != null)
 			{
 				JobDocument.AutoSubmitMessage = (AutoSubmitMessage.Length == 0)? null : AutoSubmitMessage;
-				Updates.Add(UpdateBuilder.SetOrUnsetNullRef(x => x.AutoSubmitMessage, Job.AutoSubmitMessage));
+				Updates.Add(UpdateBuilder.SetOrUnsetNullRef(x => x.AutoSubmitMessage, JobDocument.AutoSubmitMessage));
 			}
-			if (AbortedByUserId != null && Job.AbortedByUserId == null)
+			if (AbortedByUserId != null && JobDocument.AbortedByUserId == null)
 			{
 				JobDocument.AbortedByUserId = AbortedByUserId;
-				Updates.Add(UpdateBuilder.Set(x => x.AbortedByUserId, Job.AbortedByUserId));
+				Updates.Add(UpdateBuilder.Set(x => x.AbortedByUserId, JobDocument.AbortedByUserId));
 				bUpdateBatches = true;
 			}
 			if (NotificationTriggerId != null)
@@ -640,7 +654,7 @@ namespace HordeServer.Collections.Impl
 			}
 			if (Arguments != null)
 			{
-				HashSet<string> ModifiedArguments = new HashSet<string>(Job.Arguments);
+				HashSet<string> ModifiedArguments = new HashSet<string>(JobDocument.Arguments);
 				ModifiedArguments.SymmetricExceptWith(Arguments);
 
 				foreach (string ModifiedArgument in ModifiedArguments)
@@ -662,19 +676,19 @@ namespace HordeServer.Collections.Impl
 			}
 
 			// Update the new list of job steps
-			return (Updates.Count == 0 || await TryUpdateAsync(JobDocument, UpdateBuilder.Combine(Updates)));
+			return await TryUpdateAsync(JobDocument, UpdateBuilder.Combine(Updates));
 		}
 
 		/// <inheritdoc/>
-		public async Task<bool> TryUpdateBatchAsync(IJob Job, IGraph Graph, SubResourceId BatchId, LogId? NewLogId, JobStepBatchState? NewState, JobStepBatchError? NewError)
+		public async Task<IJob?> TryUpdateBatchAsync(IJob Job, IGraph Graph, SubResourceId BatchId, LogId? NewLogId, JobStepBatchState? NewState, JobStepBatchError? NewError)
 		{
-			JobDocument JobDocument = (JobDocument)Job;
+			JobDocument JobDocument = Clone((JobDocument)Job);
 
 			// Find the index of the appropriate batch
 			int BatchIdx = JobDocument.Batches.FindIndex(x => x.Id == BatchId);
 			if (BatchIdx == -1)
 			{
-				return false;
+				return null;
 			}
 
 			// If we're marking the batch as complete and there are still steps to run (eg. because the agent crashed), we need to mark all the steps as complete first
@@ -748,13 +762,13 @@ namespace HordeServer.Collections.Impl
 			}
 
 			// Update the new list of job steps
-			return (Updates.Count == 0 || await TryUpdateAsync(JobDocument, UpdateBuilder.Combine(Updates)));
+			return await TryUpdateAsync(JobDocument, Updates);
 		}
 
 		/// <inheritdoc/>
-		public async Task<bool> TryUpdateStepAsync(IJob Job, IGraph Graph, SubResourceId BatchId, SubResourceId StepId, JobStepState NewState, JobStepOutcome NewOutcome, bool? NewAbortRequested, UserId? NewAbortByUserId, LogId? NewLogId, ObjectId? NewNotificationTriggerId, UserId? NewRetryByUserId, Priority? NewPriority, List<Report>? NewReports, Dictionary<string, string?>? NewProperties)
+		public Task<IJob?> TryUpdateStepAsync(IJob Job, IGraph Graph, SubResourceId BatchId, SubResourceId StepId, JobStepState NewState, JobStepOutcome NewOutcome, bool? NewAbortRequested, UserId? NewAbortByUserId, LogId? NewLogId, ObjectId? NewNotificationTriggerId, UserId? NewRetryByUserId, Priority? NewPriority, List<Report>? NewReports, Dictionary<string, string?>? NewProperties)
 		{
-			JobDocument JobDocument = (JobDocument)Job;
+			JobDocument JobDocument = Clone((JobDocument)Job);
 
 			// Create the update 
 			UpdateDefinitionBuilder<JobDocument> UpdateBuilder = Builders<JobDocument>.Update;
@@ -927,11 +941,22 @@ namespace HordeServer.Collections.Impl
 			}
 
 			// Update the new list of job steps
-			return (Updates.Count == 0 || await TryUpdateAsync(JobDocument, UpdateBuilder.Combine(Updates)));
+			return TryUpdateAsync(JobDocument, Updates);
 		}
 
-		/// <inheritdoc/>
-		async Task<bool> TryUpdateAsync(JobDocument Job, UpdateDefinition<JobDocument> Update)
+		Task<IJob?> TryUpdateAsync(JobDocument Job, List<UpdateDefinition<JobDocument>> Updates)
+		{
+			if (Updates.Count == 0)
+			{
+				return Task.FromResult<IJob?>(Job);
+			}
+			else
+			{
+				return TryUpdateAsync(Job, Builders<JobDocument>.Update.Combine(Updates));
+			}
+		}
+
+		async Task<IJob?> TryUpdateAsync(JobDocument Job, UpdateDefinition<JobDocument> Update)
 		{
 			int NewUpdateIndex = Job.UpdateIndex + 1;
 			Update = Update.Set(x => x.UpdateIndex, NewUpdateIndex);
@@ -944,15 +969,15 @@ namespace HordeServer.Collections.Impl
 			{
 				Job.UpdateIndex = NewUpdateIndex;
 				Job.UpdateTimeUtc = NewUpdateTimeUtc;
-				return true;
+				return Job;
 			}
-			return false;
+			return null;
 		}
 
 		/// <inheritdoc/>
-		public Task<bool> TryRemoveFromDispatchQueueAsync(IJob Job)
+		public Task<IJob?> TryRemoveFromDispatchQueueAsync(IJob Job)
 		{
-			JobDocument JobDocument = (JobDocument)Job;
+			JobDocument JobDocument = Clone((JobDocument)Job);
 
 			// Create the update 
 			UpdateDefinitionBuilder<JobDocument> UpdateBuilder = Builders<JobDocument>.Update;
@@ -962,11 +987,11 @@ namespace HordeServer.Collections.Impl
 			Updates.Add(UpdateBuilder.Set(x => x.SchedulePriority, JobDocument.SchedulePriority));
 
 			// Update the new list of job steps
-			return TryUpdateAsync(JobDocument, UpdateBuilder.Combine(Updates));
+			return TryUpdateAsync(JobDocument, Updates);
 		}
 
 		/// <inheritdoc/>
-		public Task<bool> TryUpdateGraphAsync(IJob Job, IGraph NewGraph)
+		public Task<IJob?> TryUpdateGraphAsync(IJob Job, IGraph NewGraph)
 		{
 			JobDocument JobDocument = (JobDocument)Job;
 
@@ -980,7 +1005,7 @@ namespace HordeServer.Collections.Impl
 			UpdateBatches(JobDocument, NewGraph, Updates, Logger);
 
 			// Update the new list of job steps
-			return TryUpdateAsync(JobDocument, UpdateBuilder.Combine(Updates));
+			return TryUpdateAsync(JobDocument, Updates);
 		}
 
 		/// <inheritdoc/>
@@ -1013,7 +1038,7 @@ namespace HordeServer.Collections.Impl
 		{
 			while (Job != null)
 			{
-				JobDocument JobDocument = (JobDocument)Job;
+				JobDocument JobDocument = Clone((JobDocument)Job);
 
 				for (int BatchIdx = 0; BatchIdx < JobDocument.Batches.Count; BatchIdx++)
 				{
@@ -1039,9 +1064,10 @@ namespace HordeServer.Collections.Impl
 				List<UpdateDefinition<JobDocument>> Updates = new List<UpdateDefinition<JobDocument>>();
 				UpdateBatches(JobDocument, Graph, Updates, Logger);
 
-				if (await TryUpdateAsync(JobDocument, Builders<JobDocument>.Update.Combine(Updates)))
+				IJob? NewJob = await TryUpdateAsync(JobDocument, Updates);
+				if (NewJob != null)
 				{
-					break;
+					return NewJob;
 				}
 
 				Job = await GetAsync(Job.Id);
@@ -1073,9 +1099,10 @@ namespace HordeServer.Collections.Impl
 				List<UpdateDefinition<JobDocument>> Updates = new List<UpdateDefinition<JobDocument>>();
 				UpdateBatches(JobDocument, Graph, Updates, Logger);
 
-				if (await TryUpdateAsync(JobDocument, Builders<JobDocument>.Update.Combine(Updates)))
+				IJob? NewJob = await TryUpdateAsync(JobDocument, Updates);
+				if(NewJob != null)
 				{
-					break;
+					return NewJob;
 				}
 
 				Job = await GetAsync(Job.Id);
@@ -1084,7 +1111,7 @@ namespace HordeServer.Collections.Impl
 		}
 
 		/// <inheritdoc/>
-		public async Task<bool> TryAssignLeaseAsync(IJob Job, int BatchIdx, PoolId PoolId, AgentId AgentId, ObjectId SessionId, LeaseId LeaseId, LogId LogId)
+		public async Task<IJob?> TryAssignLeaseAsync(IJob Job, int BatchIdx, PoolId PoolId, AgentId AgentId, ObjectId SessionId, LeaseId LeaseId, LogId LogId)
 		{
 			// Try to update the job with this agent id
 			UpdateDefinitionBuilder<JobDocument> UpdateBuilder = Builders<JobDocument>.Update;
@@ -1101,13 +1128,13 @@ namespace HordeServer.Collections.Impl
 			{
 				string CurrentSessionId = Job.Batches[BatchIdx].SessionId!.Value.ToString();
 				Logger.LogError("Attempt to replace current session ID {CurrSessionId} with {NewSessionId} for batch {JobId}:{BatchId}", CurrentSessionId, SessionId.ToString(), Job.Id.ToString(), Job.Batches[BatchIdx].Id);
-				return false;
+				return null;
 			}
 
-			JobDocument JobDocument = (JobDocument)Job;
-			if (!await TryUpdateAsync(JobDocument, UpdateBuilder.Combine(Updates)))
+			JobDocument JobDocument = Clone((JobDocument)Job);
+			if (await TryUpdateAsync(JobDocument, Updates) == null)
 			{
-				return false;
+				return null;
 			}
 
 			JobStepBatchDocument Batch = JobDocument.Batches[BatchIdx];
@@ -1116,30 +1143,32 @@ namespace HordeServer.Collections.Impl
 			Batch.LeaseId = LeaseId;
 			Batch.LogId = LogId;
 
-			return true;
+			return JobDocument;
 		}
 
 		/// <inheritdoc/>
-		public async Task<bool> TryCancelLeaseAsync(IJob Job, int BatchIdx)
+		public async Task<IJob?> TryCancelLeaseAsync(IJob Job, int BatchIdx)
 		{
-			JobDocument JobDocument = (JobDocument)Job;
+			JobDocument JobDocument = Clone((JobDocument)Job);
 
 			UpdateDefinition<JobDocument> Update = Builders<JobDocument>.Update.Unset(x => x.Batches[BatchIdx].AgentId).Unset(x => x.Batches[BatchIdx].SessionId).Unset(x => x.Batches[BatchIdx].LeaseId);
-			if (await TryUpdateAsync(JobDocument, Update))
+			if (await TryUpdateAsync(JobDocument, Update) == null)
 			{
-				JobStepBatchDocument Batch = JobDocument.Batches[BatchIdx];
-				Batch.AgentId = null;
-				Batch.SessionId = null;
-				Batch.LeaseId = null;
-				return true;
+				return null;
 			}
-			return false;
+
+			JobStepBatchDocument Batch = JobDocument.Batches[BatchIdx];
+			Batch.AgentId = null;
+			Batch.SessionId = null;
+			Batch.LeaseId = null;
+
+			return JobDocument;
 		}
 
 		/// <inheritdoc/>
-		public async Task<bool> TryFailBatchAsync(IJob Job, int BatchIdx, IGraph Graph)
+		public Task<IJob?> TryFailBatchAsync(IJob Job, int BatchIdx, IGraph Graph)
 		{
-			JobDocument JobDocument = (JobDocument)Job;
+			JobDocument JobDocument = Clone((JobDocument)Job);
 			JobStepBatchDocument Batch = JobDocument.Batches[BatchIdx];
 
 			UpdateDefinitionBuilder<JobDocument> UpdateBuilder = new UpdateDefinitionBuilder<JobDocument>();
@@ -1188,7 +1217,7 @@ namespace HordeServer.Collections.Impl
 			RefreshDependentJobSteps(JobDocument, Graph, Updates, Logger);
 			RefreshJobPriority(JobDocument, Updates);
 
-			return (Updates.Count == 0 || await TryUpdateAsync(JobDocument, UpdateBuilder.Combine(Updates)));
+			return TryUpdateAsync(JobDocument, UpdateBuilder.Combine(Updates));
 		}
 
 		/// <summary>
diff --git a/Engine/Source/Programs/Horde/Horde.Build/Jobs/Impl/JobService.cs b/Engine/Source/Programs/Horde/Horde.Build/Jobs/Impl/JobService.cs
--- a/Engine/Source/Programs/Horde/Horde.Build/Jobs/Impl/JobService.cs
+++ b/Engine/Source/Programs/Horde/Horde.Build/Jobs/Impl/JobService.cs
@@ -296,11 +296,17 @@ namespace HordeServer.Services
 				if (Job.StartedByUserId != NewJob.StartedByUserId) continue;
 				if (string.Join(",", Job.Arguments) != string.Join(",", NewJob.Arguments)) continue;
 
-				bool bUpdateSuccessful = await UpdateJobAsync(Job, null, null, null, KnownUsers.System, null, null, null);
-				if (!bUpdateSuccessful)
+				IJob? UpdatedJob = await UpdateJobAsync(Job, null, null, null, KnownUsers.System, null, null, null);
+				if (UpdatedJob == null)
 				{
 					Logger.LogError("Failed marking duplicate job as aborted! Job ID: {JobId}", Job.Id);
 				}
+
+				IJob? UpdatedJob2 = await GetJobAsync(Job.Id);
+				if (UpdatedJob2.AbortedByUserId != UpdatedJob.AbortedByUserId)
+				{
+					throw new NotImplementedException();
+				}
 			}
 		}
 
@@ -363,41 +369,35 @@ namespace HordeServer.Services
 		/// <param name="Reports">New reports to add</param>
 		/// <param name="Arguments">New arguments for the job</param>
 		/// <param name="LabelIdxToTriggerId">New trigger ID for a label in the job</param>
-		public async Task<bool> UpdateJobAsync(IJob Job, string? Name = null, Priority? Priority = null, bool? AutoSubmit = null, UserId? AbortedByUserId = null, ObjectId? OnCompleteTriggerId = null, List<Report>? Reports = null, List<string>? Arguments = null, KeyValuePair<int, ObjectId>? LabelIdxToTriggerId = null)
+		public async Task<IJob?> UpdateJobAsync(IJob Job, string? Name = null, Priority? Priority = null, bool? AutoSubmit = null, UserId? AbortedByUserId = null, ObjectId? OnCompleteTriggerId = null, List<Report>? Reports = null, List<string>? Arguments = null, KeyValuePair<int, ObjectId>? LabelIdxToTriggerId = null)
 		{
 			using IDisposable Scope = Logger.BeginScope("UpdateJobAsync({JobId})", Job.Id);
-			for (; ; )
+			for(IJob? NewJob = Job; NewJob != null; NewJob = await GetJobAsync(Job.Id))
 			{
-				IGraph Graph = await GetGraphAsync(Job);
+				IGraph Graph = await GetGraphAsync(NewJob);
 
 				// Capture the previous label states
-				IReadOnlyList<(LabelState, LabelOutcome)> OldLabelStates = Job.GetLabelStates(Graph);
+				IReadOnlyList<(LabelState, LabelOutcome)> OldLabelStates = NewJob.GetLabelStates(Graph);
 
 				// Update the new list of job steps
-				if (await Jobs.TryUpdateJobAsync(Job, Graph, Name, Priority, AutoSubmit, null, null, AbortedByUserId, OnCompleteTriggerId, Reports, Arguments, LabelIdxToTriggerId))
+				NewJob = await Jobs.TryUpdateJobAsync(NewJob, Graph, Name, Priority, AutoSubmit, null, null, AbortedByUserId, OnCompleteTriggerId, Reports, Arguments, LabelIdxToTriggerId);
+				if (NewJob != null)
 				{
 					// Update any badges that have been modified
-					await JobTaskSource.UpdateUgsBadges(Job, Graph, OldLabelStates);
+					await JobTaskSource.UpdateUgsBadges(NewJob, Graph, OldLabelStates);
 
 					// Cancel any leases which are no longer required
-					foreach (IJobStepBatch Batch in Job.Batches)
+					foreach (IJobStepBatch Batch in NewJob.Batches)
 					{
 						if (Batch.Error == JobStepBatchError.Cancelled && (Batch.State == JobStepBatchState.Starting || Batch.State == JobStepBatchState.Running) && Batch.AgentId != null && Batch.LeaseId != null)
 						{
 							await CancelLeaseAsync(Batch.AgentId.Value, Batch.LeaseId.Value);
 						}
 					}
-					return true;
-				}
-
-				// Get the new job state
-				IJob? NewJob = await GetJobAsync(Job.Id);
-				if (NewJob == null)
-				{
-					return false;
+					return NewJob;
 				}
-				Job = NewJob;
 			}
+			return null;
 		}
 
 		/// <summary>
@@ -552,19 +552,19 @@ namespace HordeServer.Services
 		/// <param name="Job">The job to update</param>
 		/// <param name="NewGraph">New graph for this job</param>
 		/// <returns>True if the groups were updated to the given list. False if another write happened first.</returns>
-		public async Task<bool> TryUpdateGraphAsync(IJob Job, IGraph NewGraph)
+		public async Task<IJob?> TryUpdateGraphAsync(IJob Job, IGraph NewGraph)
 		{
 			using IDisposable Scope = Logger.BeginScope("TryUpdateGraphAsync({JobId})", Job.Id);
 
 			IReadOnlyList<(LabelState, LabelOutcome)> OldLabelStates = Job.GetLabelStates(NewGraph);
-			if (await Jobs.TryUpdateGraphAsync(Job, NewGraph))
-			{
-				await JobTaskSource.UpdateUgsBadges(Job, NewGraph, OldLabelStates);
 
-				JobTaskSource.UpdateQueuedJob(Job, NewGraph);
-				return true;
+			IJob? NewJob = await Jobs.TryUpdateGraphAsync(Job, NewGraph);
+			if(NewJob != null)
+			{
+				await JobTaskSource.UpdateUgsBadges(NewJob, NewGraph, OldLabelStates);
+				JobTaskSource.UpdateQueuedJob(NewJob, NewGraph);
 			}
-			return false;
+			return NewJob;
 		}
 
 		/// <summary>
@@ -683,7 +683,7 @@ namespace HordeServer.Services
 		/// <param name="NewLogId">The new log file id</param>
 		/// <param name="NewState">New state of the jobstep</param>
 		/// <returns>True if the job was updated, false if it was deleted</returns>
-		public async Task<bool> UpdateBatchAsync(IJob Job, SubResourceId BatchId, LogId? NewLogId = null, JobStepBatchState? NewState = null)
+		public async Task<IJob?> UpdateBatchAsync(IJob Job, SubResourceId BatchId, LogId? NewLogId = null, JobStepBatchState? NewState = null)
 		{
 			using IDisposable Scope = Logger.BeginScope("UpdateBatchAsync({JobId})", Job.Id);
 
@@ -694,7 +694,7 @@ namespace HordeServer.Services
 				int BatchIdx = Job.Batches.FindIndex(x => x.Id == BatchId);
 				if (BatchIdx == -1)
 				{
-					break;
+					return null;
 				}
 
 				IJobStepBatch Batch = Job.Batches[BatchIdx];
@@ -702,7 +702,7 @@ namespace HordeServer.Services
 				// If the batch has already been marked as complete, error out
 				if (Batch.State == JobStepBatchState.Complete)
 				{
-					break;
+					return null;
 				}
 
 				// If we're marking the batch as complete before the agent has run everything, mark it to conform
@@ -736,20 +736,21 @@ namespace HordeServer.Services
 				}
 
 				// Update the batch state
-				if (await TryUpdateBatchAsync(Job, BatchId, NewLogId, NewState, NewError))
+				IJob? NewJob = await TryUpdateBatchAsync(Job, BatchId, NewLogId, NewState, NewError);
+				if (NewJob != null)
 				{
-					return true;
+					return NewJob;
 				}
 
 				// Update the job
-				IJob? NewJob = await GetJobAsync(Job.Id);
+				NewJob = await GetJobAsync(Job.Id);
 				if (NewJob == null)
 				{
-					break;
+					return null;
 				}
+
 				Job = NewJob;
 			}
-			return false;
 		}
 
 		/// <summary>
@@ -761,14 +762,10 @@ namespace HordeServer.Services
 		/// <param name="NewState">New state of the jobstep</param>
 		/// <param name="NewError">New error state</param>
 		/// <returns>The updated job, otherwise null</returns>
-		public async Task<bool> TryUpdateBatchAsync(IJob Job, SubResourceId BatchId, LogId? NewLogId = null, JobStepBatchState? NewState = null, JobStepBatchError? NewError = null)
+		public async Task<IJob?> TryUpdateBatchAsync(IJob Job, SubResourceId BatchId, LogId? NewLogId = null, JobStepBatchState? NewState = null, JobStepBatchError? NewError = null)
 		{
 			IGraph Graph = await GetGraphAsync(Job);
-			if (!await Jobs.TryUpdateBatchAsync(Job, Graph, BatchId, NewLogId, NewState, NewError))
-			{
-				return false;
-			}
-			return true;
+			return await Jobs.TryUpdateBatchAsync(Job, Graph, BatchId, NewLogId, NewState, NewError);
 		}
 
 		/// <summary>
@@ -793,12 +790,13 @@ namespace HordeServer.Services
 			using IDisposable Scope = Logger.BeginScope("UpdateStepAsync({JobId})", Job.Id);
 			for (; ;)
 			{
-				if (await TryUpdateStepAsync(Job, BatchId, StepId, NewState, NewOutcome, NewAbortRequested, NewAbortByUserId, NewLogId, NewNotificationTriggerId, NewRetryByUserId, NewPriority, NewReports, NewProperties))
+				IJob? NewJob = await TryUpdateStepAsync(Job, BatchId, StepId, NewState, NewOutcome, NewAbortRequested, NewAbortByUserId, NewLogId, NewNotificationTriggerId, NewRetryByUserId, NewPriority, NewReports, NewProperties);
+				if (NewJob != null)
 				{
-					return Job;
+					return NewJob;
 				}
 
-				IJob? NewJob = await GetJobAsync(Job.Id);
+				NewJob = await GetJobAsync(Job.Id);
 				if(NewJob == null)
 				{
 					return null;
@@ -825,7 +823,7 @@ namespace HordeServer.Services
 		/// <param name="NewReports">New reports</param>
 		/// <param name="NewProperties">Property changes. Any properties with a null value will be removed.</param>
 		/// <returns>True if the job was updated, false if it was deleted in the meantime</returns>
-		public async Task<bool> TryUpdateStepAsync(IJob Job, SubResourceId BatchId, SubResourceId StepId, JobStepState NewState = JobStepState.Unspecified, JobStepOutcome NewOutcome = JobStepOutcome.Unspecified, bool? NewAbortRequested = null, UserId? NewAbortByUserId = null, LogId? NewLogId = null, ObjectId? NewTriggerId = null, UserId? NewRetryByUserId = null, Priority? NewPriority = null, List<Report>? NewReports = null, Dictionary<string, string?>? NewProperties = null)
+		public async Task<IJob?> TryUpdateStepAsync(IJob Job, SubResourceId BatchId, SubResourceId StepId, JobStepState NewState = JobStepState.Unspecified, JobStepOutcome NewOutcome = JobStepOutcome.Unspecified, bool? NewAbortRequested = null, UserId? NewAbortByUserId = null, LogId? NewLogId = null, ObjectId? NewTriggerId = null, UserId? NewRetryByUserId = null, Priority? NewPriority = null, List<Report>? NewReports = null, Dictionary<string, string?>? NewProperties = null)
 		{
 			using IDisposable Scope = Logger.BeginScope("TryUpdateStepAsync({JobId})", Job.Id);
 
@@ -852,8 +850,11 @@ namespace HordeServer.Services
 			}
 
 			// Update the step
-			if (await Jobs.TryUpdateStepAsync(Job, Graph, BatchId, StepId, NewState, NewOutcome, NewAbortRequested, NewAbortByUserId, NewLogId, NewTriggerId, NewRetryByUserId, NewPriority, NewReports, NewProperties))
+			IJob? NewJob = await Jobs.TryUpdateStepAsync(Job, Graph, BatchId, StepId, NewState, NewOutcome, NewAbortRequested, NewAbortByUserId, NewLogId, NewTriggerId, NewRetryByUserId, NewPriority, NewReports, NewProperties);
+			if (NewJob != null)
 			{
+				Job = NewJob;
+
 				using IScope DdScope = GlobalTracer.Instance.BuildSpan("TryUpdateStepAsync").StartActive();
 				DdScope.Span.SetTag("JobId", Job.Id.ToString());
 				DdScope.Span.SetTag("BatchId", BatchId.ToString());
@@ -912,7 +913,7 @@ namespace HordeServer.Services
 								JobStepOutcome JobTriggerOutcome = Job.GetTargetOutcome(Graph, JobTrigger.Target);
 								if (JobTriggerOutcome == JobStepOutcome.Success || JobTriggerOutcome == JobStepOutcome.Warnings)
 								{
-									Job = await FireJobTriggerAsync(Job, Graph, JobTrigger);
+									Job = await FireJobTriggerAsync(Job, Graph, JobTrigger) ?? Job;
 								}
 							}
 						}
@@ -952,10 +953,10 @@ namespace HordeServer.Services
 					JobTaskSource.UpdateQueuedJob(Job, Graph);
 				}
 
-				return true;
+				return Job;
 			}
 
-			return false;
+			return null;
 		}
 
 		/// <summary>
@@ -1025,12 +1026,13 @@ namespace HordeServer.Services
 
 			for (; ; )
 			{
-				if (await Jobs.TryUpdateJobAsync(Job, Graph, AutoSubmitChange: Change, AutoSubmitMessage: Message))
+				IJob? NewJob = await Jobs.TryUpdateJobAsync(Job, Graph, AutoSubmitChange: Change, AutoSubmitMessage: Message);
+				if (NewJob != null)
 				{
-					return Job;
+					return NewJob;
 				}
 
-				IJob? NewJob = await GetJobAsync(Job.Id);
+				NewJob = await GetJobAsync(Job.Id);
 				if (NewJob == null)
 				{
 					return Job;
@@ -1091,25 +1093,27 @@ namespace HordeServer.Services
 		/// <param name="Graph">Graph for the job containing the trigger</param>
 		/// <param name="JobTrigger">The trigger object to fire</param>
 		/// <returns>New job instance</returns>
-		private async Task<IJob> FireJobTriggerAsync(IJob Job, IGraph Graph, IChainedJob JobTrigger)
+		private async Task<IJob?> FireJobTriggerAsync(IJob Job, IGraph Graph, IChainedJob JobTrigger)
 		{
 			for (; ; )
 			{
 				// Update the job
 				JobId ChainedJobId = JobId.GenerateNewId();
-				if (await Jobs.TryUpdateJobAsync(Job, Graph, JobTrigger: new KeyValuePair<TemplateRefId, JobId>(JobTrigger.TemplateRefId, ChainedJobId)))
+
+				IJob? NewJob = await Jobs.TryUpdateJobAsync(Job, Graph, JobTrigger: new KeyValuePair<TemplateRefId, JobId>(JobTrigger.TemplateRefId, ChainedJobId));
+				if(NewJob != null)
 				{
-					IStream? Stream = await StreamService.GetStreamAsync(Job.StreamId);
+					IStream? Stream = await StreamService.GetStreamAsync(NewJob.StreamId);
 					if(Stream == null)
 					{
-						Logger.LogWarning("Cannot find stream {StreamId} for downstream job", Job.StreamId);
+						Logger.LogWarning("Cannot find stream {StreamId} for downstream job", NewJob.StreamId);
 						break;
 					}
 
 					TemplateRef? TemplateRef;
 					if (!Stream.Templates.TryGetValue(JobTrigger.TemplateRefId, out TemplateRef))
 					{
-						Logger.LogWarning("Cannot find template {TemplateRefId} in stream {StreamId}", JobTrigger.TemplateRefId, Job.StreamId);
+						Logger.LogWarning("Cannot find template {TemplateRefId} in stream {StreamId}", JobTrigger.TemplateRefId, NewJob.StreamId);
 						break;
 					}
 
@@ -1121,21 +1125,22 @@ namespace HordeServer.Services
 					}
 
 					IGraph TriggerGraph = await Graphs.AddAsync(Template);
-					Logger.LogInformation("Creating downstream job {ChainedJobId} from job {JobId}", ChainedJobId, Job.Id);
+					Logger.LogInformation("Creating downstream job {ChainedJobId} from job {JobId}", ChainedJobId, NewJob.Id);
 
-					await CreateJobAsync(ChainedJobId, Stream, JobTrigger.TemplateRefId, TemplateRef.Hash, TriggerGraph, TemplateRef.Name, Job.Change, Job.CodeChange, Job.PreflightChange, Job.ClonedPreflightChange, Job.StartedByUserId, Template.Priority, null, Job.UpdateIssues, TemplateRef.ChainedJobs, false, false, TemplateRef.NotificationChannel, TemplateRef.NotificationChannelFilter, Template.Arguments);
-					break;
+					await CreateJobAsync(ChainedJobId, Stream, JobTrigger.TemplateRefId, TemplateRef.Hash, TriggerGraph, TemplateRef.Name, NewJob.Change, NewJob.CodeChange, NewJob.PreflightChange, NewJob.ClonedPreflightChange, NewJob.StartedByUserId, Template.Priority, null, NewJob.UpdateIssues, TemplateRef.ChainedJobs, false, false, TemplateRef.NotificationChannel, TemplateRef.NotificationChannelFilter, Template.Arguments);
+					return NewJob;
 				}
 
 				// Fetch the job again
-				IJob? NewJob = await Jobs.GetAsync(Job.Id);
+				NewJob = await Jobs.GetAsync(Job.Id);
 				if(NewJob == null)
 				{
 					break;
 				}
+
 				Job = NewJob;
 			}
-			return Job;
+			return null;
 		}
 
 		/// <summary>
diff --git a/Engine/Source/Programs/Horde/Horde.Build/Jobs/Impl/JobTaskSource.cs b/Engine/Source/Programs/Horde/Horde.Build/Jobs/Impl/JobTaskSource.cs
--- a/Engine/Source/Programs/Horde/Horde.Build/Jobs/Impl/JobTaskSource.cs
+++ b/Engine/Source/Programs/Horde/Horde.Build/Jobs/Impl/JobTaskSource.cs
@@ -754,8 +754,11 @@ namespace HordeServer.Tasks.Impl
 
 			// Try to update the job with this agent id
 			LogId LogId = (await LogFileService.CreateLogFileAsync(Job.Id, Agent.SessionId, LogType.Json)).Id;
-			if (await Jobs.TryAssignLeaseAsync(Item.Job, Item.BatchIdx, Item.PoolId, Agent.Id, Agent.SessionId!.Value, LeaseId, LogId))
+			IJob? NewJob = await Jobs.TryAssignLeaseAsync(Item.Job, Item.BatchIdx, Item.PoolId, Agent.Id, Agent.SessionId!.Value, LeaseId, LogId);
+			if (NewJob != null)
 			{
+				Job = NewJob;
+
 				// Get the lease name
 				StringBuilder LeaseName = new StringBuilder($"{Item.Stream.Name} - ");
 				if (Job.PreflightChange > 0)
@@ -796,7 +799,7 @@ namespace HordeServer.Tasks.Impl
 				Logger.LogTrace("Refreshing queue entries for job {JobId}", Job.Id);
 
 				// Get the new copy of the job
-				IJob? NewJob = await Jobs.GetAsync(Job.Id);
+				NewJob = await Jobs.GetAsync(Job.Id);
 				if (NewJob == null)
 				{
 					lock (LockObject)
@@ -1011,7 +1014,9 @@ namespace HordeServer.Tasks.Impl
 					int RunningStepIdx = Batch.Steps.FindIndex(x => x.State == JobStepState.Running);
 
 					IGraph Graph = await Graphs.GetAsync(Job.GraphHash);
-					if (await Jobs.TryFailBatchAsync(Job, BatchIdx, Graph))
+					Job = await Jobs.TryFailBatchAsync(Job, BatchIdx, Graph);
+
+					if (Job != null)
 					{
 						if (Batch.Error != JobStepBatchError.None)
 						{
@@ -1059,7 +1064,8 @@ namespace HordeServer.Tasks.Impl
 					break;
 				}
 
-				if (await Jobs.TryCancelLeaseAsync(Job, BatchIdx))
+				IJob? NewJob = await Jobs.TryCancelLeaseAsync(Job, BatchIdx);
+				if (NewJob != null)
 				{
 					break;
 				}
diff --git a/Engine/Source/Programs/Horde/Horde.Build/Services/RpcService.cs b/Engine/Source/Programs/Horde/Horde.Build/Services/RpcService.cs
--- a/Engine/Source/Programs/Horde/Horde.Build/Services/RpcService.cs
+++ b/Engine/Source/Programs/Horde/Horde.Build/Services/RpcService.cs
@@ -528,7 +528,9 @@ namespace HordeServer.Services
 			}
 
 			IJobStepBatch Batch = AuthorizeBatch(Job, Request.BatchId.ToSubResourceId(), Context);
-			if (!await JobService.UpdateBatchAsync(Job, BatchId, null, Api.JobStepBatchState.Starting))
+			Job = await JobService.UpdateBatchAsync(Job, BatchId, null, Api.JobStepBatchState.Starting);
+
+			if (Job == null)
 			{
 				throw new StructuredRpcException(StatusCode.NotFound, "Batch {JobId}:{BatchId} not found for updating", Request.JobId, Request.BatchId);
 			}
@@ -603,7 +605,7 @@ namespace HordeServer.Services
 				if (StepIdx == Batch.Steps.Count)
 				{
 					Logger.LogDebug("Job {JobId} batch {BatchId} is complete", Job.Id, Batch.Id);
-					if (!await JobService.TryUpdateBatchAsync(Job, Batch.Id, NewState: Api.JobStepBatchState.Stopping))
+					if (await JobService.TryUpdateBatchAsync(Job, Batch.Id, NewState: Api.JobStepBatchState.Stopping) == null)
 					{
 						return null;
 					}
@@ -646,7 +648,8 @@ namespace HordeServer.Services
 			//				}
 
 			// Update the step state
-			if (await JobService.TryUpdateStepAsync(Job, Batch.Id, Step.Id, JobStepState.Running, JobStepOutcome.Unspecified, null, null, Log.Value.Id, null, null, null, null))
+			IJob? NewJob = await JobService.TryUpdateStepAsync(Job, Batch.Id, Step.Id, JobStepState.Running, JobStepOutcome.Unspecified, null, null, Log.Value.Id, null, null, null, null);
+			if (NewJob != null)
 			{
 				BeginStepResponse Response = new BeginStepResponse();
 				Response.State = BeginStepResponse.Types.Result.Ready;
@@ -856,7 +859,8 @@ namespace HordeServer.Services
 				IGraph Graph = await JobService.GetGraphAsync(Job);
 				Graph = await Graphs.AppendAsync(Graph, NewGroups, NewAggregates, NewLabels);
 
-				if (await JobService.TryUpdateGraphAsync(Job, Graph))
+				IJob? NewJob = await JobService.TryUpdateGraphAsync(Job, Graph);
+				if (NewJob != null)
 				{
 					return new UpdateGraphResponse();
 				}
