Horde: Fix spammy build health notifications.

diff --git a/Engine/Source/Programs/Horde/Horde.Build/Notifications/Impl/SlackNotificationSink.cs b/Engine/Source/Programs/Horde/Horde.Build/Notifications/Impl/SlackNotificationSink.cs
--- a/Engine/Source/Programs/Horde/Horde.Build/Notifications/Impl/SlackNotificationSink.cs
+++ b/Engine/Source/Programs/Horde/Horde.Build/Notifications/Impl/SlackNotificationSink.cs
@@ -259,23 +259,15 @@ namespace HordeServer.Notifications.Impl
 
 		#region Message state 
 
-		async Task<MessageStateDocument> AddOrUpdateMessageStateAsync(string Recipient, string EventId, UserId? UserId, string Digest)
+		async Task<(MessageStateDocument, bool)> AddOrUpdateMessageStateAsync(string Recipient, string EventId, UserId? UserId, string Digest)
 		{
 			ObjectId NewId = ObjectId.GenerateNewId();
 
 			FilterDefinition<MessageStateDocument> Filter = Builders<MessageStateDocument>.Filter.Eq(x => x.Recipient, Recipient) & Builders<MessageStateDocument>.Filter.Eq(x => x.EventId, EventId);
 			UpdateDefinition<MessageStateDocument> Update = Builders<MessageStateDocument>.Update.SetOnInsert(x => x.Id, NewId).Set(x => x.UserId, UserId).Set(x => x.Digest, Digest);
 
-			MessageStateDocument? PrevState = await MessageStates.FindOneAndUpdateAsync(Filter, Update, new FindOneAndUpdateOptions<MessageStateDocument> { IsUpsert = true });
-			if (PrevState == null)
-			{
-				PrevState = new MessageStateDocument();
-				PrevState.Id = NewId;
-				PrevState.Recipient = Recipient;
-				PrevState.EventId = EventId;
-				PrevState.UserId = UserId;
-			}
-			return PrevState;
+			MessageStateDocument State = await MessageStates.FindOneAndUpdateAsync(Filter, Update, new FindOneAndUpdateOptions<MessageStateDocument> { IsUpsert = true, ReturnDocument = ReturnDocument.After });
+			return (State, State.Id == NewId);
 		}
 
 		async Task SetMessageTimestampAsync(ObjectId MessageId, string Channel, string Ts)
@@ -1197,10 +1189,10 @@ namespace HordeServer.Notifications.Impl
 
 			string RequestDigest = ContentHash.MD5(JsonSerializer.Serialize(Message, new JsonSerializerOptions { IgnoreNullValues = true })).ToString();
 
-			MessageStateDocument PrevState = await AddOrUpdateMessageStateAsync(Recipient, EventId, UserId, RequestDigest);
-			if (String.IsNullOrEmpty(PrevState.Ts))
+			(MessageStateDocument State, bool IsNew) = await AddOrUpdateMessageStateAsync(Recipient, EventId, UserId, RequestDigest);
+			if (IsNew)
 			{
-				Logger.LogInformation("Sending new slack message to {SlackUser}", Recipient);
+				Logger.LogInformation("Sending new slack message to {SlackUser} (msg: {MessageId})", Recipient, State.Id);
 				Message.Channel = Recipient;
 				Message.Ts = null;
 
@@ -1209,13 +1201,13 @@ namespace HordeServer.Notifications.Impl
 				{
 					Logger.LogWarning("Missing 'ts' or 'channel' field on slack response");
 				}
-				await SetMessageTimestampAsync(PrevState.Id, Response.Channel ?? String.Empty, Response.Ts ?? String.Empty);
+				await SetMessageTimestampAsync(State.Id, Response.Channel ?? String.Empty, Response.Ts ?? String.Empty);
 			}
-			else
+			else if (!String.IsNullOrEmpty(State.Ts))
 			{
-				Logger.LogInformation("Updating existing slack message for user {SlackUser} ({Channel}, {MessageTs})", Recipient, PrevState.Channel, PrevState.Ts);
-				Message.Channel = PrevState.Channel;
-				Message.Ts = PrevState.Ts;
+				Logger.LogInformation("Updating existing slack message {MessageId} for user {SlackUser} ({Channel}, {MessageTs})", State.Id, Recipient, State.Channel, State.Ts);
+				Message.Channel = State.Channel;
+				Message.Ts = State.Ts;
 				await SendRequestAsync<PostMessageResponse>(UpdateMessageUrl, Message);
 			}
 		}
