Horde: Fix Slack notifications not updating correctly.

diff --git a/Engine/Source/Programs/Horde/Horde.Build/Notifications/Impl/SlackNotificationSink.cs b/Engine/Source/Programs/Horde/Horde.Build/Notifications/Impl/SlackNotificationSink.cs
--- a/Engine/Source/Programs/Horde/Horde.Build/Notifications/Impl/SlackNotificationSink.cs
+++ b/Engine/Source/Programs/Horde/Horde.Build/Notifications/Impl/SlackNotificationSink.cs
@@ -272,7 +272,7 @@ namespace HordeServer.Notifications.Impl
 
 		async Task SetMessageTimestampAsync(ObjectId MessageId, string Channel, string Ts)
 		{
-			FilterDefinition<MessageStateDocument> Filter = Builders<MessageStateDocument>.Filter.Eq(x => x.Id, MessageId) & Builders<MessageStateDocument>.Filter.Eq(x => x.Ts, String.Empty);
+			FilterDefinition<MessageStateDocument> Filter = Builders<MessageStateDocument>.Filter.Eq(x => x.Id, MessageId);
 			UpdateDefinition<MessageStateDocument> Update = Builders<MessageStateDocument>.Update.Set(x => x.Channel, Channel).Set(x => x.Ts, Ts);
 			await MessageStates.FindOneAndUpdateAsync(Filter, Update);
 		}
