Horde: Add tracing spans to PerforceService

diff --git a/Engine/Source/Programs/Horde/Horde.Build/Perforce/Impl/PerforceService.cs b/Engine/Source/Programs/Horde/Horde.Build/Perforce/Impl/PerforceService.cs
--- a/Engine/Source/Programs/Horde/Horde.Build/Perforce/Impl/PerforceService.cs
+++ b/Engine/Source/Programs/Horde/Horde.Build/Perforce/Impl/PerforceService.cs
@@ -20,6 +20,8 @@ using System.Runtime.InteropServices;
 using System.Text;
 using System.Text.RegularExpressions;
 using System.Threading.Tasks;
+using OpenTracing;
+using OpenTracing.Util;
 
 namespace HordeServer.Services
 {
@@ -111,6 +113,10 @@ namespace HordeServer.Services
 
 		public async ValueTask<IUser> FindOrAddUserAsync(string ClusterName, string UserName)
 		{
+			using IScope Scope = GlobalTracer.Instance.BuildSpan("PerforceService.FindOrAddUserAsync").StartActive();
+			Scope.Span.SetTag("ClusterName", ClusterName);
+			Scope.Span.SetTag("UserName", UserName);
+			
 			IUser? User;
 			if (!UserCache.TryGetValue((ClusterName, UserName), out User))
 			{
@@ -146,6 +152,9 @@ namespace HordeServer.Services
 
 		async Task<IPerforceServer> SelectServer(PerforceCluster Cluster)
 		{
+			using IScope Scope = GlobalTracer.Instance.BuildSpan("PerforceService.SelectServer").StartActive();
+			Scope.Span.SetTag("ClusterName", Cluster.Name);
+			
 			IPerforceServer? Server = await LoadBalancer.SelectServerAsync(Cluster);
 			if (Server == null)
 			{
@@ -154,8 +163,14 @@ namespace HordeServer.Services
 			return Server;
 		}
 
+		[SuppressMessage("Microsoft.Reliability", "CA2000:DisposeObjectsBeforeLosingScope")]
 		static P4.Repository CreateConnection(IPerforceServer Server, string? UserName, string? Ticket)
 		{
+			using IScope Scope = GlobalTracer.Instance.BuildSpan("PerforceService.CreateConnection").StartActive();
+			Scope.Span.SetTag("Cluster", Server.Cluster);
+			Scope.Span.SetTag("ServerAndPort", Server.ServerAndPort);
+			Scope.Span.SetTag("UserName", UserName);
+			
 			P4.Repository Repository = new P4.Repository(new P4.Server(new P4.ServerAddress(Server.ServerAndPort)));
 			try
 			{
@@ -192,6 +207,9 @@ namespace HordeServer.Services
 		/// <returns></returns>
 		public async Task<NativePerforceConnection?> GetServiceUserConnection(string? ClusterName)
 		{
+			using IScope Scope = GlobalTracer.Instance.BuildSpan("PerforceService.GetServiceUserConnection").StartActive();
+			Scope.Span.SetTag("ClusterName", ClusterName);
+			
 			if(ClusterName != null)
 			{
 				Globals Globals = await CachedGlobals.GetCached();
@@ -227,12 +245,18 @@ namespace HordeServer.Services
 
 		async Task<P4.Repository> GetServiceUserConnection(PerforceCluster Cluster)
 		{
+			using IScope Scope = GlobalTracer.Instance.BuildSpan("PerforceService.GetServiceUserConnectionAsync").StartActive();
+			Scope.Span.SetTag("ClusterName", Cluster.Name);
+			
 			IPerforceServer Server = await SelectServer(Cluster);
 			return GetServiceUserConnection(Cluster, Server);
 		}
 
 		P4.Repository GetServiceUserConnection(PerforceCluster Cluster, IPerforceServer Server)
 		{
+			using IScope Scope = GlobalTracer.Instance.BuildSpan("PerforceService.GetServiceUserConnection").StartActive();
+			Scope.Span.SetTag("ClusterName", Cluster.Name);
+
 			if (Cluster.Name == PerforceCluster.DefaultName && Settings.P4BridgeServiceUsername != null && Settings.P4BridgeServicePassword != null)
 			{
 				return CreateConnection(Server, Settings.P4BridgeServiceUsername, Settings.P4BridgeServicePassword);
@@ -253,6 +277,10 @@ namespace HordeServer.Services
 
 		async Task<CachedTicketInfo> GetImpersonateCredential(PerforceCluster Cluster, string ImpersonateUser)
 		{
+			using IScope Scope = GlobalTracer.Instance.BuildSpan("PerforceService.GetImpersonateCredential").StartActive();
+			Scope.Span.SetTag("ClusterName", Cluster.Name);
+			Scope.Span.SetTag("ImpersonateUser", ImpersonateUser);
+			
 			if (!Cluster.CanImpersonate)
 			{
 				throw new Exception($"Service account required to impersonate user {ImpersonateUser}");
@@ -308,12 +336,24 @@ namespace HordeServer.Services
 
 		async Task<P4.Repository> GetImpersonatedConnection(PerforceCluster Cluster, string ImpersonateUser)
 		{
+			using IScope Scope = GlobalTracer.Instance.BuildSpan("PerforceService.GetImpersonatedConnection").StartActive();
+			Scope.Span.SetTag("ClusterName", Cluster.Name);
+			Scope.Span.SetTag("ImpersonateUser", ImpersonateUser);
+			
 			CachedTicketInfo TicketInfo = await GetImpersonateCredential(Cluster, ImpersonateUser);
 			return CreateConnection(TicketInfo.Server, TicketInfo.UserName, TicketInfo.Ticket.Ticket);
 		}
 
+		[SuppressMessage("Microsoft.Reliability", "CA2000:DisposeObjectsBeforeLosingScope")]
 		async Task<P4.Repository> GetConnection(PerforceCluster Cluster, string? Stream = null, string? Username = null, bool ReadOnly = true, bool CreateChange = false, int? ClientFromChange = null, bool UseClientFromChange = false, bool UsePortFromChange = false, string? ClientId = null, bool NoClient = false)
 		{
+			using IScope Scope = GlobalTracer.Instance.BuildSpan("PerforceService.GetConnection").StartActive();
+			Scope.Span.SetTag("ClusterName", Cluster.Name);
+			Scope.Span.SetTag("Stream", Stream);
+			Scope.Span.SetTag("Username", Username);
+			Scope.Span.SetTag("NoClient", NoClient);
+			Scope.Span.SetTag("ClientId", ClientId);
+			
 			P4.Repository Repository;
 			if (Username == null || !Cluster.CanImpersonate || Username.Equals(Cluster.ServiceAccount, StringComparison.OrdinalIgnoreCase))
 			{
@@ -425,6 +465,11 @@ namespace HordeServer.Services
 
 		static P4.Client GetOrCreateClient(string? ServiceUserName, P4.Repository Repository, string? Stream, string? Username = null, bool ReadOnly = true, bool CreateChange = false, int? ClientFromChange = null, bool UseClientFromChange = false, bool UsePortFromChange = false)
 		{
+			using IScope Scope = GlobalTracer.Instance.BuildSpan("PerforceService.GetOrCreateClient").StartActive();
+			Scope.Span.SetTag("ServiceUserName", ServiceUserName);
+			Scope.Span.SetTag("Stream", Stream);
+			Scope.Span.SetTag("Username", Username);
+			
 			P4.Client? Client = null;
 			P4.Changelist? Changelist = null;
 
@@ -586,6 +631,10 @@ namespace HordeServer.Services
 		/// <inheritdoc/>
 		public async Task<IStreamView> GetStreamViewAsync(string ClusterName, string StreamName)
 		{
+			using IScope Scope = GlobalTracer.Instance.BuildSpan("PerforceService.GetStreamViewAsync").StartActive();
+			Scope.Span.SetTag("ClusterName", ClusterName);
+			Scope.Span.SetTag("StreamName", StreamName);
+			
 			PerforceCluster Cluster = await GetClusterAsync(ClusterName);
 			using (P4.Repository Repository = await GetConnection(Cluster, NoClient: true))
 			{
@@ -653,6 +702,9 @@ namespace HordeServer.Services
 		/// <returns></returns>
 		static ChangeFile CreateChangeFile(string RelativePath, P4.FileMetaData MetaData)
 		{
+			using IScope Scope = GlobalTracer.Instance.BuildSpan("PerforceService.CreateChangeFile (FileMetaData)").StartActive();
+			Scope.Span.SetTag("RelativePath", RelativePath);
+			
 			int Revision = GetSyncRevision(MetaData.DepotPath.Path, MetaData.HeadAction, MetaData.HeadRev);
 			Md5Hash? Digest = String.IsNullOrEmpty(MetaData.Digest) ? (Md5Hash?)null : Md5Hash.Parse(MetaData.Digest);
 			return new ChangeFile(RelativePath, MetaData.DepotPath.Path, Revision, MetaData.FileSize, Digest, MetaData.HeadType.ToString());
@@ -666,6 +718,9 @@ namespace HordeServer.Services
 		/// <returns></returns>
 		static ChangeFile CreateChangeFile(string RelativePath, P4.ShelvedFile MetaData)
 		{
+			using IScope Scope = GlobalTracer.Instance.BuildSpan("PerforceService.CreateChangeFile (ShelvedFile)").StartActive();
+			Scope.Span.SetTag("RelativePath", RelativePath);
+			
 			int Revision = GetSyncRevision(MetaData.Path.Path, MetaData.Action, MetaData.Revision);
 			Md5Hash? Digest = String.IsNullOrEmpty(MetaData.Digest) ? (Md5Hash?)null : Md5Hash.Parse(MetaData.Digest);
 			return new ChangeFile(RelativePath, MetaData.Path.Path, Revision, MetaData.Size, Digest, MetaData.Type.ToString());
@@ -674,6 +729,11 @@ namespace HordeServer.Services
 		/// <inheritdoc/>
 		public async Task<List<ChangeFile>> GetStreamSnapshotAsync(string ClusterName, string StreamName, int Change)
 		{
+			using IScope Scope = GlobalTracer.Instance.BuildSpan("PerforceService.GetStreamSnapshotAsync").StartActive();
+			Scope.Span.SetTag("ClusterName", ClusterName);
+			Scope.Span.SetTag("StreamName", StreamName);
+			Scope.Span.SetTag("Change", Change);
+			
 			PerforceCluster Cluster = await GetClusterAsync(ClusterName);
 			using (P4.Repository Repository = await GetConnection(Cluster, Stream: StreamName))
 			{
@@ -706,6 +766,11 @@ namespace HordeServer.Services
 		/// <inheritdoc/>
 		public async Task<List<ChangeSummary>> GetChangesAsync(string ClusterName, int? MinChange, int MaxResults)
 		{
+			using IScope Scope = GlobalTracer.Instance.BuildSpan("PerforceService.GetChangesAsync").StartActive();
+			Scope.Span.SetTag("ClusterName", ClusterName);
+			Scope.Span.SetTag("MinChange", MinChange ?? -1);
+			Scope.Span.SetTag("MaxResults", MaxResults);
+			
 			PerforceCluster Cluster = await GetClusterAsync(ClusterName);
 			using (P4.Repository Repository = await GetConnection(Cluster, NoClient: true))
 			{
@@ -737,6 +802,13 @@ namespace HordeServer.Services
 		/// <inheritdoc/>
 		public async Task<List<ChangeSummary>> GetChangesAsync(string ClusterName, string StreamName, int? MinChange, int? MaxChange, int Results, string? ImpersonateUser)
 		{
+			using IScope Scope = GlobalTracer.Instance.BuildSpan("PerforceService.GetChangesAsync").StartActive();
+			Scope.Span.SetTag("ClusterName", ClusterName);
+			Scope.Span.SetTag("MinChange", MinChange ?? -1);
+			Scope.Span.SetTag("MaxChange", MaxChange ?? -1);
+			Scope.Span.SetTag("Results", Results);
+			Scope.Span.SetTag("ImpersonateUser", ImpersonateUser);
+			
 			PerforceCluster Cluster = await GetClusterAsync(ClusterName);
 			using (P4.Repository Repository = await GetConnection(Cluster, StreamName, ImpersonateUser))
 			{
@@ -767,6 +839,11 @@ namespace HordeServer.Services
 		/// <inheritdoc/>
 		public async Task<ChangeDetails> GetChangeDetailsAsync(string ClusterName, string StreamName, int ChangeNumber)
 		{
+			using IScope Scope = GlobalTracer.Instance.BuildSpan("PerforceService.GetChangeDetailsAsync").StartActive();
+			Scope.Span.SetTag("ClusterName", ClusterName);
+			Scope.Span.SetTag("StreamName", StreamName);
+			Scope.Span.SetTag("ChangeNumber", ChangeNumber);
+			
 			PerforceCluster Cluster = await GetClusterAsync(ClusterName);
 			using (P4.Repository Repository = await GetConnection(Cluster, NoClient: true))
 			{
@@ -821,6 +898,12 @@ namespace HordeServer.Services
 		/// <inheritdoc/>
 		public async Task<List<ChangeDetails>> GetChangeDetailsAsync(string ClusterName, string StreamName, IReadOnlyList<int> ChangeNumbers, string? ImpersonateUser)
 		{
+			using IScope Scope = GlobalTracer.Instance.BuildSpan("PerforceService.GetChangeDetailsAsync").StartActive();
+			Scope.Span.SetTag("ClusterName", ClusterName);
+			Scope.Span.SetTag("StreamName", StreamName);
+			Scope.Span.SetTag("ChangeNumbers.Count", ChangeNumbers.Count);
+			Scope.Span.SetTag("ImpersonateUser", ImpersonateUser);
+			
 			PerforceCluster Cluster = await GetClusterAsync(ClusterName);
 			using (P4.Repository Repository = await GetConnection(Cluster, Stream: StreamName, Username: ImpersonateUser))
 			{
@@ -897,6 +980,10 @@ namespace HordeServer.Services
 		/// <inheritdoc />
 		public async Task<string> CreateTicket(string ClusterName, string ImpersonateUser)
 		{
+			using IScope Scope = GlobalTracer.Instance.BuildSpan("PerforceService.CreateTicket").StartActive();
+			Scope.Span.SetTag("ClusterName", ClusterName);
+			Scope.Span.SetTag("ImpersonateUser", ImpersonateUser);
+			
 			PerforceCluster Cluster = await GetClusterAsync(ClusterName);
 
 			CachedTicketInfo Credential = await GetImpersonateCredential(Cluster, ImpersonateUser);
@@ -912,6 +999,9 @@ namespace HordeServer.Services
 		/// <inheritdoc/>
 		public async Task<List<FileSummary>> FindFilesAsync(string ClusterName, IEnumerable<string> Paths)
 		{
+			using IScope Scope = GlobalTracer.Instance.BuildSpan("PerforceService.FindFilesAsync").StartActive();
+			Scope.Span.SetTag("ClusterName", ClusterName);
+			
 			List<FileSummary> Results = new List<FileSummary>();
 
 			PerforceCluster Cluster = await GetClusterAsync(ClusterName);
@@ -948,6 +1038,10 @@ namespace HordeServer.Services
 		/// <inheritdoc/>
 		public async Task<byte[]> PrintAsync(string ClusterName, string DepotPath)
 		{
+			using IScope Scope = GlobalTracer.Instance.BuildSpan("PerforceService.PrintAsync").StartActive();
+			Scope.Span.SetTag("ClusterName", ClusterName);
+			Scope.Span.SetTag("DepotPath", DepotPath);
+			
 			if (DepotPath.EndsWith("...", StringComparison.OrdinalIgnoreCase) || DepotPath.EndsWith("*", StringComparison.OrdinalIgnoreCase))
 			{
 				throw new Exception("PrintAsync requires exactly one file to be specified");
@@ -975,6 +1069,10 @@ namespace HordeServer.Services
 		/// <inheritdoc/>
 		public async Task<int> DuplicateShelvedChangeAsync(string ClusterName, int ShelvedChange)
 		{
+			using IScope Scope = GlobalTracer.Instance.BuildSpan("PerforceService.DuplicateShelvedChangeAsync").StartActive();
+			Scope.Span.SetTag("ClusterName", ClusterName);
+			Scope.Span.SetTag("ShelvedChange", ShelvedChange);
+
 			string? ChangeOwner = null;
 
 			// Get the owner of the shelf
@@ -1021,6 +1119,10 @@ namespace HordeServer.Services
 		/// <inheritdoc/>
 		public async Task DeleteShelvedChangeAsync(string ClusterName, int ShelvedChange)
 		{
+			using IScope Scope = GlobalTracer.Instance.BuildSpan("PerforceService.DeleteShelvedChangeAsync").StartActive();
+			Scope.Span.SetTag("ClusterName", ClusterName);
+			Scope.Span.SetTag("ShelvedChange", ShelvedChange);
+			
 			PerforceCluster Cluster = await GetClusterAsync(ClusterName);
 			using (P4.Repository Repository = await GetConnection(Cluster, NoClient: true))
 			{
@@ -1055,6 +1157,10 @@ namespace HordeServer.Services
 		/// <inheritdoc/>
 		public async Task UpdateChangelistDescription(string ClusterName, int Change, string Description)
 		{
+			using IScope Scope = GlobalTracer.Instance.BuildSpan("PerforceService.UpdateChangelistDescription").StartActive();
+			Scope.Span.SetTag("ClusterName", ClusterName);
+			Scope.Span.SetTag("Change", Change);
+
 			try
 			{
 				PerforceCluster Cluster = await GetClusterAsync(ClusterName);
@@ -1082,6 +1188,11 @@ namespace HordeServer.Services
 		/// <inheritdoc/>
 		public async Task<int> CreateNewChangeAsync(string ClusterName, string StreamName, string FilePath)
 		{
+			using IScope Scope = GlobalTracer.Instance.BuildSpan("PerforceService.CreateNewChangeAsync").StartActive();
+			Scope.Span.SetTag("ClusterName", ClusterName);
+			Scope.Span.SetTag("StreamName", StreamName);
+			Scope.Span.SetTag("FilePath", FilePath);
+			
 			PerforceCluster Cluster = await GetClusterAsync(ClusterName);
 			P4.SubmitResults? SubmitResults = null;
 
@@ -1241,6 +1352,11 @@ namespace HordeServer.Services
 		/// <inheritdoc/>
 		public async Task<(int? Change, string Message)> SubmitShelvedChangeAsync(string ClusterName, int Change, int OriginalChange)
 		{
+			using IScope Scope = GlobalTracer.Instance.BuildSpan("PerforceService.SubmitShelvedChangeAsync").StartActive();
+			Scope.Span.SetTag("ClusterName", ClusterName);
+			Scope.Span.SetTag("Change", Change);
+			Scope.Span.SetTag("OriginalChange", OriginalChange);
+			
 			PerforceCluster Cluster = await GetClusterAsync(ClusterName);
 			using (P4.Repository Repository = await GetConnection(Cluster, NoClient: true))
 			{
@@ -1350,6 +1466,10 @@ namespace HordeServer.Services
 		/// <inheritdoc/>		
 		public async Task<PerforceUserInfo?> GetUserInfoAsync(string ClusterName, string UserName)
 		{
+			using IScope Scope = GlobalTracer.Instance.BuildSpan("PerforceService.GetUserInfoAsync").StartActive();
+			Scope.Span.SetTag("ClusterName", ClusterName);
+			Scope.Span.SetTag("UserName", UserName);
+			
 			PerforceCluster Cluster = await GetClusterAsync(ClusterName);
 			using (P4.Repository Repository = await GetConnection(Cluster, NoClient: true))
 			{
@@ -1366,6 +1486,8 @@ namespace HordeServer.Services
 
 		static int ReshelveChange(P4.Repository Repository, int Change)
 		{
+			using IScope Scope = GlobalTracer.Instance.BuildSpan("PerforceService.ReshelveChange").StartActive();
+			Scope.Span.SetTag("Change", Change);
 
 			bool EdgeServer = false;
 			string? Value;
@@ -1432,6 +1554,11 @@ namespace HordeServer.Services
 		/// <inheritdoc/>
 		public virtual async Task<int> GetCodeChangeAsync(string ClusterName, string StreamName, int Change)
 		{
+			using IScope Scope = GlobalTracer.Instance.BuildSpan("PerforceService.GetCodeChangeAsync").StartActive();
+			Scope.Span.SetTag("ClusterName", ClusterName);
+			Scope.Span.SetTag("StreamName", StreamName);
+			Scope.Span.SetTag("Change", Change);
+			
 			int MaxChange = Change;
 			for (; ; )
 			{
