Horde: Add tracing spans to JobService

diff --git a/Engine/Source/Programs/Horde/Horde.Build/Jobs/Impl/JobService.cs b/Engine/Source/Programs/Horde/Horde.Build/Jobs/Impl/JobService.cs
--- a/Engine/Source/Programs/Horde/Horde.Build/Jobs/Impl/JobService.cs
+++ b/Engine/Source/Programs/Horde/Horde.Build/Jobs/Impl/JobService.cs
@@ -24,6 +24,7 @@ using System.Text.RegularExpressions;
 using System.Threading;
 using System.Threading.Tasks;
 using System.Globalization;
+using Horde.Build.Utilities;
 using HordeServer.Tasks.Impl;
 using OpenTracing.Util;
 using OpenTracing;
@@ -233,6 +234,28 @@ namespace HordeServer.Services
 		/// <returns>Unique id representing the job</returns>
 		public async Task<IJob> CreateJobAsync(JobId? JobId, IStream Stream, TemplateRefId TemplateRefId, ContentHash TemplateHash, IGraph Graph, string Name, int Change, int CodeChange, int? PreflightChange, int? ClonedPreflightChange, UserId? StartedByUserId, Priority? Priority, bool? AutoSubmit, bool? UpdateIssues, List<ChainedJobTemplate>? JobTriggers, bool ShowUgsBadges, bool ShowUgsAlerts, string? NotificationChannel, string? NotificationChannelFilter, IReadOnlyList<string> Arguments)
 		{
+			using IScope TraceScope = GlobalTracer.Instance.BuildSpan("JobService.CreateJobAsync").StartActive();
+			TraceScope.Span.SetTag("JobId", JobId);
+			TraceScope.Span.SetTag("Stream", Stream.Name);
+			TraceScope.Span.SetTag("TemplateRefId", TemplateRefId);
+			TraceScope.Span.SetTag("TemplateHash", TemplateHash);
+			TraceScope.Span.SetTag("GraphId", Graph.Id);
+			TraceScope.Span.SetTag("Name", Name);
+			TraceScope.Span.SetTag("Change", Change);
+			TraceScope.Span.SetTag("CodeChange", CodeChange);
+			TraceScope.Span.SetTag("PreflightChange", PreflightChange);
+			TraceScope.Span.SetTag("ClonedPreflightChange", ClonedPreflightChange);
+			TraceScope.Span.SetTag("StartedByUserId", StartedByUserId.ToString());
+			TraceScope.Span.SetTag("Priority", Priority.ToString());
+			TraceScope.Span.SetTag("ShowUgsBadges", ShowUgsBadges);
+			TraceScope.Span.SetTag("NotificationChannel", NotificationChannel ?? "null");
+			TraceScope.Span.SetTag("NotificationChannelFilter", NotificationChannelFilter ?? "null");
+			TraceScope.Span.SetTag("Arguments", string.Join(',', Arguments));
+			
+			if (AutoSubmit != null) TraceScope.Span.SetTag("AutoSubmit", AutoSubmit.Value);
+			if (UpdateIssues != null) TraceScope.Span.SetTag("UpdateIssues", UpdateIssues.Value);
+			if (JobTriggers != null) TraceScope.Span.SetTag("JobTriggers.Count", JobTriggers.Count);
+
 			JobId JobIdValue = JobId ?? HordeServer.Utilities.ObjectId<IJob>.GenerateNewId();
 			using IDisposable Scope = Logger.BeginScope("CreateJobAsync({JobId})", JobIdValue);
 
@@ -281,6 +304,10 @@ namespace HordeServer.Services
 
 		private async Task AbortAnyDuplicateJobs(IJob NewJob)
 		{
+			using IScope Scope = GlobalTracer.Instance.BuildSpan("JobService.AbortAnyDuplicateJobs").StartActive();
+			Scope.Span.SetTag("JobId", NewJob.Id.ToString());
+			Scope.Span.SetTag("JobName", NewJob.Name);
+			
 			List<IJob> JobsToAbort = new List<IJob>();
 			if (NewJob.PreflightChange > 0)
 			{
@@ -316,6 +343,10 @@ namespace HordeServer.Services
 		/// <param name="Job">The job to delete</param>
 		public async Task<bool> DeleteJobAsync(IJob Job)
 		{
+			using IScope TraceScope = GlobalTracer.Instance.BuildSpan("JobService.DeleteJobAsync").StartActive();
+			TraceScope.Span.SetTag("JobId", Job.Id.ToString());
+			TraceScope.Span.SetTag("JobName", Job.Name);
+
 			using IDisposable Scope = Logger.BeginScope("DeleteJobAsync({JobId})", Job.Id);
 
 			// Delete the job
@@ -371,6 +402,10 @@ namespace HordeServer.Services
 		/// <param name="LabelIdxToTriggerId">New trigger ID for a label in the job</param>
 		public async Task<IJob?> UpdateJobAsync(IJob Job, string? Name = null, Priority? Priority = null, bool? AutoSubmit = null, UserId? AbortedByUserId = null, ObjectId? OnCompleteTriggerId = null, List<Report>? Reports = null, List<string>? Arguments = null, KeyValuePair<int, ObjectId>? LabelIdxToTriggerId = null)
 		{
+			using IScope TraceScope = GlobalTracer.Instance.BuildSpan("JobService.UpdateJobAsync").StartActive();
+			TraceScope.Span.SetTag("JobId", Job.Id.ToString());
+			TraceScope.Span.SetTag("Name", Name);
+			
 			using IDisposable Scope = Logger.BeginScope("UpdateJobAsync({JobId})", Job.Id);
 			for(IJob? NewJob = Job; NewJob != null; NewJob = await GetJobAsync(Job.Id))
 			{
@@ -408,6 +443,10 @@ namespace HordeServer.Services
 		/// <returns></returns>
 		async Task CancelLeaseAsync(AgentId AgentId, LeaseId LeaseId)
 		{
+			using IScope Scope = GlobalTracer.Instance.BuildSpan("JobService.CancelLeaseAsync").StartActive();
+			Scope.Span.SetTag("AgentId", AgentId.ToString());
+			Scope.Span.SetTag("LeaseId", LeaseId.ToString());
+			
 			for (; ; )
 			{
 				IAgent? Agent = await Agents.GetAsync(AgentId);
@@ -495,6 +534,26 @@ namespace HordeServer.Services
 		/// <returns>List of jobs matching the given criteria</returns>
 		public async Task<List<IJob>> FindJobsAsync(JobId[]? JobIds = null, StreamId? StreamId = null, string? Name = null, TemplateRefId[]? Templates = null, int? MinChange = null, int? MaxChange = null, int? PreflightChange = null, UserId? PreflightStartedByUser = null, UserId? StartedByUser = null, DateTimeOffset ? MinCreateTime = null, DateTimeOffset? MaxCreateTime = null, string? Target = null, JobStepState[]? State = null, JobStepOutcome[]? Outcome = null, DateTimeOffset? ModifiedBefore = null, DateTimeOffset? ModifiedAfter = null, int? Index = null, int? Count = null)
 		{
+			using IScope Scope = GlobalTracer.Instance.BuildSpan("JobService.FindJobsAsync").StartActive();
+			Scope.Span.SetTag("JobIds", JobIds);
+			Scope.Span.SetTag("StreamId", StreamId);
+			Scope.Span.SetTag("Name", Name);
+			Scope.Span.SetTag("Templates", Templates);
+			Scope.Span.SetTag("MinChange", MinChange);
+			Scope.Span.SetTag("MaxChange", MaxChange);
+			Scope.Span.SetTag("PreflightChange", PreflightChange);
+			Scope.Span.SetTag("PreflightStartedByUser", PreflightStartedByUser);
+			Scope.Span.SetTag("StartedByUser", StartedByUser);
+			Scope.Span.SetTag("MinCreateTime", MinCreateTime);
+			Scope.Span.SetTag("MaxCreateTime", MaxCreateTime);
+			Scope.Span.SetTag("Target", Target);
+			Scope.Span.SetTag("State", State?.ToString());
+			Scope.Span.SetTag("Outcome", Outcome?.ToString());
+			Scope.Span.SetTag("ModifiedBefore", ModifiedBefore);
+			Scope.Span.SetTag("ModifiedAfter", ModifiedAfter);
+			Scope.Span.SetTag("Index", Index);
+			Scope.Span.SetTag("Count", Count);
+			
 			if (Target == null && (State == null || State.Length == 0) && (Outcome == null || Outcome.Length == 0))
 			{
 				return await Jobs.FindAsync(JobIds, StreamId, Name, Templates, MinChange, MaxChange, PreflightChange, PreflightStartedByUser, StartedByUser, MinCreateTime, MaxCreateTime, ModifiedBefore, ModifiedAfter, Index, Count);
@@ -554,6 +613,10 @@ namespace HordeServer.Services
 		/// <returns>True if the groups were updated to the given list. False if another write happened first.</returns>
 		public async Task<IJob?> TryUpdateGraphAsync(IJob Job, IGraph NewGraph)
 		{
+			using IScope TraceScope = GlobalTracer.Instance.BuildSpan("JobService.TryUpdateGraphAsync").StartActive();
+			TraceScope.Span.SetTag("Job", Job.Id);
+			TraceScope.Span.SetTag("NewGraph", NewGraph.Id);
+
 			using IDisposable Scope = Logger.BeginScope("TryUpdateGraphAsync({JobId})", Job.Id);
 
 			IReadOnlyList<(LabelState, LabelOutcome)> OldLabelStates = Job.GetLabelStates(NewGraph);
@@ -574,6 +637,9 @@ namespace HordeServer.Services
 		/// <returns>Timing info for the given job</returns>
 		public async Task<IJobTiming> GetJobTimingAsync(IJob Job)
 		{
+			using IScope TraceScope = GlobalTracer.Instance.BuildSpan("JobService.GetJobTimingAsync").StartActive();
+			TraceScope.Span.SetTag("Job", Job.Id);
+
 			using IDisposable Scope = Logger.BeginScope("GetJobTimingAsync({JobId})", Job.Id);
 
 			IGraph Graph = await Graphs.GetAsync(Job.GraphHash);
@@ -644,6 +710,12 @@ namespace HordeServer.Services
 		/// <returns>Expected duration for the given step</returns>
 		async Task<JobStepTimingData> GetStepTimingInfo(StreamId StreamId, TemplateRefId TemplateId, string NodeName, int? Change)
 		{
+			using IScope TraceScope = GlobalTracer.Instance.BuildSpan("JobService.GetStepTimingInfo").StartActive();
+			TraceScope.Span.SetTag("StreamId", StreamId);
+			TraceScope.Span.SetTag("TemplateId", TemplateId);
+			TraceScope.Span.SetTag("NodeName", NodeName);
+			TraceScope.Span.SetTag("Change", Change);
+			
 			// Find all the steps matching the given criteria
 			List<IJobStepRef> Steps = await JobStepRefs.GetStepsForNodeAsync(StreamId, TemplateId, NodeName, Change, false, 10);
 
@@ -685,6 +757,12 @@ namespace HordeServer.Services
 		/// <returns>True if the job was updated, false if it was deleted</returns>
 		public async Task<IJob?> UpdateBatchAsync(IJob Job, SubResourceId BatchId, LogId? NewLogId = null, JobStepBatchState? NewState = null)
 		{
+			using IScope TraceScope = GlobalTracer.Instance.BuildSpan("JobService.UpdateBatchAsync").StartActive();
+			TraceScope.Span.SetTag("Job", Job.Id);
+			TraceScope.Span.SetTag("BatchId", BatchId);
+			TraceScope.Span.SetTag("NewLogId", NewLogId);
+			TraceScope.Span.SetTag("NewState", NewState.ToString());
+
 			using IDisposable Scope = Logger.BeginScope("UpdateBatchAsync({JobId})", Job.Id);
 
 			bool bCheckForBadAgent = true;
@@ -787,6 +865,11 @@ namespace HordeServer.Services
 		/// <returns>True if the job was updated, false if it was deleted in the meantime</returns>
 		public async Task<IJob?> UpdateStepAsync(IJob Job, SubResourceId BatchId, SubResourceId StepId, JobStepState NewState = JobStepState.Unspecified, JobStepOutcome NewOutcome = JobStepOutcome.Unspecified, bool? NewAbortRequested = null, UserId? NewAbortByUserId = null, LogId? NewLogId = null, ObjectId? NewNotificationTriggerId = null, UserId? NewRetryByUserId = null, Priority? NewPriority = null, List<Report>? NewReports = null, Dictionary<string, string?>? NewProperties = null)
 		{
+			using IScope TraceScope = GlobalTracer.Instance.BuildSpan("JobService.UpdateStepAsync").StartActive();
+			TraceScope.Span.SetTag("Job", Job.Id);
+			TraceScope.Span.SetTag("BatchId", BatchId);
+			TraceScope.Span.SetTag("StepId", StepId);
+			
 			using IDisposable Scope = Logger.BeginScope("UpdateStepAsync({JobId})", Job.Id);
 			for (; ;)
 			{
@@ -825,6 +908,11 @@ namespace HordeServer.Services
 		/// <returns>True if the job was updated, false if it was deleted in the meantime</returns>
 		public async Task<IJob?> TryUpdateStepAsync(IJob Job, SubResourceId BatchId, SubResourceId StepId, JobStepState NewState = JobStepState.Unspecified, JobStepOutcome NewOutcome = JobStepOutcome.Unspecified, bool? NewAbortRequested = null, UserId? NewAbortByUserId = null, LogId? NewLogId = null, ObjectId? NewTriggerId = null, UserId? NewRetryByUserId = null, Priority? NewPriority = null, List<Report>? NewReports = null, Dictionary<string, string?>? NewProperties = null)
 		{
+			using IScope TraceScope = GlobalTracer.Instance.BuildSpan("JobService.TryUpdateStepAsync").StartActive();
+			TraceScope.Span.SetTag("Job", Job.Id);
+			TraceScope.Span.SetTag("BatchId", BatchId);
+			TraceScope.Span.SetTag("StepId", StepId);
+
 			using IDisposable Scope = Logger.BeginScope("TryUpdateStepAsync({JobId})", Job.Id);
 
 			// Get the graph for this job
@@ -968,6 +1056,10 @@ namespace HordeServer.Services
 		[SuppressMessage("Design", "CA1031:Do not catch general exception types", Justification = "<Pending>")]
 		private async Task<IJob> AutoSubmitChangeAsync(IJob Job, IGraph Graph)
 		{
+			using IScope TraceScope = GlobalTracer.Instance.BuildSpan("JobService.AutoSubmitChangeAsync").StartActive();
+			TraceScope.Span.SetTag("Job", Job.Id);
+			TraceScope.Span.SetTag("Graph", Graph.Id);
+			
 			int? Change;
 			string Message;
 			try
@@ -1095,6 +1187,13 @@ namespace HordeServer.Services
 		/// <returns>New job instance</returns>
 		private async Task<IJob?> FireJobTriggerAsync(IJob Job, IGraph Graph, IChainedJob JobTrigger)
 		{
+			using IScope TraceScope = GlobalTracer.Instance.BuildSpan("JobService.FireJobTriggerAsync").StartActive();
+			TraceScope.Span.SetTag("Job", Job.Id);
+			TraceScope.Span.SetTag("Graph", Graph.Id);
+			TraceScope.Span.SetTag("JobTrigger.Target", JobTrigger.Target);
+			TraceScope.Span.SetTag("JobTrigger.JobId", JobTrigger.JobId);
+			TraceScope.Span.SetTag("JobTrigger.TemplateRefId", JobTrigger.TemplateRefId);
+			
 			for (; ; )
 			{
 				// Update the job
diff --git a/Engine/Source/Programs/Horde/Horde.Build/Utilities/TracingExtensions.cs b/Engine/Source/Programs/Horde/Horde.Build/Utilities/TracingExtensions.cs
new file mode 100644
--- /dev/null
+++ b/Engine/Source/Programs/Horde/Horde.Build/Utilities/TracingExtensions.cs
@@ -0,0 +1,95 @@
+using System;
+using System.Collections.Generic;
+using EpicGames.Core;
+using HordeServer.Models;
+using HordeServer.Utilities;
+using OpenTracing;
+
+namespace Horde.Build.Utilities
+{
+	/// <summary>
+	/// Extensions to handle Horde specific data types in the OpenTracing library
+	/// </summary>
+	public static class OpenTracingSpanExtensions
+	{
+		/// <summary>Set a key:value tag on the span</summary>
+		/// <returns>This span instance, for chaining</returns>
+		public static ISpan SetTag<T>(this ISpan Span, string Key, StringId<T>? Value)
+		{
+			if (Value != null) Span.SetTag(Key, Value.ToString());
+			return Span;
+		}
+		
+		/// <summary>Set a key:value tag on the span</summary>
+		/// <returns>This span instance, for chaining</returns>
+		public static ISpan SetTag<T>(this ISpan Span, string Key, StringId<T> Value)
+		{
+			Span.SetTag(Key, Value.ToString());
+			return Span;
+		}
+
+		/// <summary>Set a key:value tag on the span</summary>
+		/// <returns>This span instance, for chaining</returns>
+		public static ISpan SetTag<T>(this ISpan Span, string Key, IEnumerable<StringId<T>>? Value)
+		{
+			if (Value != null) Span.SetTag(Key, string.Join(',', Value));
+			return Span;
+		}
+
+		/// <summary>Set a key:value tag on the span</summary>
+		/// <returns>This span instance, for chaining</returns>
+		public static ISpan SetTag<T>(this ISpan Span, string Key, ObjectId<T>? Value)
+		{
+			if (Value != null) Span.SetTag(Key, Value.ToString());
+			return Span;
+		}
+		
+		/// <summary>Set a key:value tag on the span</summary>
+		/// <returns>This span instance, for chaining</returns>
+		public static ISpan SetTag<T>(this ISpan Span, string Key, ObjectId<T> Value)
+		{
+			Span.SetTag(Key, Value.ToString());
+			return Span;
+		}
+		
+		/// <summary>Set a key:value tag on the span</summary>
+		/// <returns>This span instance, for chaining</returns>
+		public static ISpan SetTag(this ISpan Span, string Key, ContentHash Value)
+		{
+			Span.SetTag(Key, Value.ToString());
+			return Span;
+		}
+		
+		/// <summary>Set a key:value tag on the span</summary>
+		/// <returns>This span instance, for chaining</returns>
+		public static ISpan SetTag(this ISpan Span, string Key, SubResourceId Value)
+		{
+			Span.SetTag(Key, Value.ToString());
+			return Span;
+		}
+		
+		/// <summary>Set a key:value tag on the span</summary>
+		/// <returns>This span instance, for chaining</returns>
+		public static ISpan SetTag<T>(this ISpan Span, string Key, IEnumerable<ObjectId<T>>? Value)
+		{
+			if (Value != null) Span.SetTag(Key, string.Join(',', Value));
+			return Span;
+		}
+		
+		/// <summary>Set a key:value tag on the span</summary>
+		/// <returns>This span instance, for chaining</returns>
+		public static ISpan SetTag(this ISpan Span, string Key, int? Value)
+		{
+			if (Value != null) Span.SetTag(Key, Value);
+			return Span;
+		}
+		
+		/// <summary>Set a key:value tag on the span</summary>
+		/// <returns>This span instance, for chaining</returns>
+		public static ISpan SetTag(this ISpan Span, string Key, DateTimeOffset? Value)
+		{
+			if (Value != null) Span.SetTag(Key, Value.ToString());
+			return Span;
+		}
+	}
+}
\ No newline at end of file
