Horde: Add more trace spans for scheduled background jobs

diff --git a/Engine/Source/Programs/Horde/Horde.Build/Fleet/AutoscaleService.cs b/Engine/Source/Programs/Horde/Horde.Build/Fleet/AutoscaleService.cs
--- a/Engine/Source/Programs/Horde/Horde.Build/Fleet/AutoscaleService.cs
+++ b/Engine/Source/Programs/Horde/Horde.Build/Fleet/AutoscaleService.cs
@@ -22,6 +22,8 @@ using System.Threading;
 using System.Threading.Tasks;
 using HordeCommon;
 using Microsoft.Extensions.Options;
+using OpenTracing;
+using OpenTracing.Util;
 using StatsdClient;
 
 namespace HordeServer.Services
@@ -120,58 +122,65 @@ namespace HordeServer.Services
 			Logger.LogInformation("Autoscaling pools...");
 			Stopwatch Stopwatch = Stopwatch.StartNew();
 
-			// Find all the current agents
-			List<IAgent> Agents = await AgentCollection.FindAsync(Status: AgentStatus.Ok);
+			Dictionary<PoolId, PoolData> PoolToData;
+			using (IScope Scope = GlobalTracer.Instance.BuildSpan("Compute utilization").StartActive())
+			{
+				// Find all the current agents
+				List<IAgent> Agents = await AgentCollection.FindAsync(Status: AgentStatus.Ok);
 
-			// Query leases in last interval
-			DateTime MaxTime = DateTime.UtcNow;
-			DateTime MinTime = MaxTime - (SampleTime * NumSamples);
-			List<ILease> Leases = await LeaseCollection.FindLeasesAsync(null, null, MinTime, MaxTime, null, null);
+				// Query leases in last interval
+				DateTime MaxTime = DateTime.UtcNow;
+				DateTime MinTime = MaxTime - (SampleTime * NumSamples);
+				List<ILease> Leases = await LeaseCollection.FindLeasesAsync(null, null, MinTime, MaxTime, null, null);
 
-			// Add all the leases to a data object for each agent
-			Dictionary<AgentId, AgentData> AgentIdToData = Agents.ToDictionary(x => x.Id, x => new AgentData(x));
-			foreach (ILease Lease in Leases)
-			{
-				AgentData? AgentData;
-				if (AgentIdToData.TryGetValue(Lease.AgentId, out AgentData) && AgentData.Agent.SessionId == Lease.SessionId)
+				// Add all the leases to a data object for each agent
+				Dictionary<AgentId, AgentData> AgentIdToData = Agents.ToDictionary(x => x.Id, x => new AgentData(x));
+				foreach (ILease Lease in Leases)
 				{
-					AgentData.Leases.Add(Lease);
+					AgentData? AgentData;
+					if (AgentIdToData.TryGetValue(Lease.AgentId, out AgentData) &&
+					    AgentData.Agent.SessionId == Lease.SessionId)
+					{
+						AgentData.Leases.Add(Lease);
+					}
 				}
-			}
 
-			// Compute utilization for each agent
-			foreach (AgentData AgentData in AgentIdToData.Values)
-			{
-				foreach (ILease Lease in AgentData.Leases.OrderBy(x => x.StartTime))
+				// Compute utilization for each agent
+				foreach (AgentData AgentData in AgentIdToData.Values)
 				{
-					double MinT = (Lease.StartTime - MinTime).TotalSeconds / SampleTime.TotalSeconds;
-					double MaxT = (Lease.FinishTime == null) ? NumSamples : ((Lease.FinishTime.Value - MinTime).TotalSeconds / SampleTime.TotalSeconds);
-
-					Any Payload = Any.Parser.ParseFrom(Lease.Payload.ToArray());
-					if (Payload.Is(ExecuteJobTask.Descriptor))
-					{
-						AgentData.Add(MinT, MaxT, 1.0, 0.0);
-					}
-					else
+					foreach (ILease Lease in AgentData.Leases.OrderBy(x => x.StartTime))
 					{
-						AgentData.Add(MinT, MaxT, 0.0, 1.0);
+						double MinT = (Lease.StartTime - MinTime).TotalSeconds / SampleTime.TotalSeconds;
+						double MaxT = (Lease.FinishTime == null)
+							? NumSamples
+							: ((Lease.FinishTime.Value - MinTime).TotalSeconds / SampleTime.TotalSeconds);
+
+						Any Payload = Any.Parser.ParseFrom(Lease.Payload.ToArray());
+						if (Payload.Is(ExecuteJobTask.Descriptor))
+						{
+							AgentData.Add(MinT, MaxT, 1.0, 0.0);
+						}
+						else
+						{
+							AgentData.Add(MinT, MaxT, 0.0, 1.0);
+						}
 					}
 				}
-			}
 
-			// Get all the pools
-			List<IPool> Pools = await PoolCollection.GetAsync();
-			Dictionary<PoolId, PoolData> PoolToData = Pools.ToDictionary(x => x.Id, x => new PoolData(x));
+				// Get all the pools
+				List<IPool> Pools = await PoolCollection.GetAsync();
+				PoolToData = Pools.ToDictionary(x => x.Id, x => new PoolData(x));
 
-			// Find pool utilization over the query period
-			foreach (AgentData AgentData in AgentIdToData.Values)
-			{
-				foreach (PoolId PoolId in AgentData.Agent.GetPools())
+				// Find pool utilization over the query period
+				foreach (AgentData AgentData in AgentIdToData.Values)
 				{
-					PoolData? PoolData;
-					if (PoolToData.TryGetValue(PoolId, out PoolData))
+					foreach (PoolId PoolId in AgentData.Agent.GetPools())
 					{
-						PoolData.Add(AgentData);
+						PoolData? PoolData;
+						if (PoolToData.TryGetValue(PoolId, out PoolData))
+						{
+							PoolData.Add(AgentData);
+						}
 					}
 				}
 			}
@@ -203,22 +212,29 @@ namespace HordeServer.Services
 
 				if (Pool.EnableAutoscaling)
 				{
-					if (Delta > 0)
-					{
-						await FleetManager.ExpandPool(Pool, PoolData.Agents, Delta);
-						await PoolCollection.TryUpdateAsync(Pool, LastScaleUpTime: DateTime.UtcNow);
-					}
-					if (Delta < 0)
+					using (IScope Scope = GlobalTracer.Instance.BuildSpan("Scaling pool").StartActive())
 					{
-						bool bShrinkIsOnCoolDown = Pool.LastScaleDownTime != null && Pool.LastScaleDownTime + ShrinkPoolCoolDown > DateTime.UtcNow;
-						if (!bShrinkIsOnCoolDown)
+						Scope.Span.SetTag("poolName", Pool.Name);
+						Scope.Span.SetTag("delta", Delta);
+
+						if (Delta > 0)
 						{
-							await FleetManager.ShrinkPool(Pool, PoolData.Agents, -Delta);
-							await PoolCollection.TryUpdateAsync(Pool, LastScaleDownTime: DateTime.UtcNow);
+							await FleetManager.ExpandPool(Pool, PoolData.Agents, Delta);
+							await PoolCollection.TryUpdateAsync(Pool, LastScaleUpTime: DateTime.UtcNow);
 						}
-						else
+
+						if (Delta < 0)
 						{
-							Logger.LogDebug("Cannot shrink {PoolName} right now, it's on cool-down until {CoolDownTimeEnds}", Pool.Name, Pool.LastScaleDownTime + ShrinkPoolCoolDown);
+							bool bShrinkIsOnCoolDown = Pool.LastScaleDownTime != null && Pool.LastScaleDownTime + ShrinkPoolCoolDown > DateTime.UtcNow;
+							if (!bShrinkIsOnCoolDown)
+							{
+								await FleetManager.ShrinkPool(Pool, PoolData.Agents, -Delta);
+								await PoolCollection.TryUpdateAsync(Pool, LastScaleDownTime: DateTime.UtcNow);
+							}
+							else
+							{
+								Logger.LogDebug("Cannot shrink {PoolName} right now, it's on cool-down until {CoolDownTimeEnds}", Pool.Name, Pool.LastScaleDownTime + ShrinkPoolCoolDown);
+							}
 						}
 					}
 				}
diff --git a/Engine/Source/Programs/Horde/Horde.Build/Projects/ScheduleService.cs b/Engine/Source/Programs/Horde/Horde.Build/Projects/ScheduleService.cs
--- a/Engine/Source/Programs/Horde/Horde.Build/Projects/ScheduleService.cs
+++ b/Engine/Source/Programs/Horde/Horde.Build/Projects/ScheduleService.cs
@@ -8,6 +8,7 @@ using HordeServer.Collections;
 using HordeServer.Controllers;
 using HordeServer.Models;
 using HordeServer.Utilities;
+using Horde.Build.Utilities;
 using Microsoft.AspNetCore.Server.Kestrel.Core.Features;
 using Microsoft.Extensions.Hosting;
 using Microsoft.Extensions.Logging;
@@ -187,6 +188,10 @@ namespace HordeServer.Services
 		[System.Diagnostics.CodeAnalysis.SuppressMessage("Design", "CA1031:Do not catch general exception types", Justification = "<Pending>")]
 		private async Task<DateTimeOffset?> UpdateScheduleAsync(IStream Stream, TemplateRefId TemplateRefId, TemplateRef TemplateRef, Schedule Schedule, DateTimeOffset Now)
 		{
+			using IScope Scope = GlobalTracer.Instance.BuildSpan("ScheduleService.UpdateScheduleAsync").StartActive();
+			Scope.Span.SetTag("Stream.Id", Stream.Id);
+			Scope.Span.SetTag("Stream.Name", Stream.Name);
+			
 			// Check if we can run the trigger
 			DateTimeOffset? NextTriggerTime = Schedule.GetNextTriggerTime(TimeZone);
 			if (!NextTriggerTime.HasValue)
@@ -243,6 +248,11 @@ namespace HordeServer.Services
 		/// <returns>Async task</returns>
 		private async Task TriggerAsync(IStream Stream, TemplateRefId TemplateRefId, TemplateRef TemplateRef, Schedule Schedule, int NumActiveJobs, DateTimeOffset Now)
 		{
+			using IScope Scope = GlobalTracer.Instance.BuildSpan("ScheduleService.TriggerAsync").StartActive();
+			Scope.Span.SetTag("Stream.Id", Stream.Id);
+			Scope.Span.SetTag("Stream.Name", Stream.Name);
+			Scope.Span.SetTag("TemplateRefId", TemplateRefId);
+			
 			// Check we're not already at the maximum number of allowed jobs
 			if (Schedule.MaxActive != 0 && NumActiveJobs >= Schedule.MaxActive)
 			{
