Horde: Add null conditional check to silence warnings

diff --git a/Engine/Source/Programs/Horde/Horde.Build/Jobs/Impl/JobService.cs b/Engine/Source/Programs/Horde/Horde.Build/Jobs/Impl/JobService.cs
--- a/Engine/Source/Programs/Horde/Horde.Build/Jobs/Impl/JobService.cs
+++ b/Engine/Source/Programs/Horde/Horde.Build/Jobs/Impl/JobService.cs
@@ -330,7 +330,7 @@ namespace HordeServer.Services
 				}
 
 				IJob? UpdatedJob2 = await GetJobAsync(Job.Id);
-				if (UpdatedJob2.AbortedByUserId != UpdatedJob.AbortedByUserId)
+				if (UpdatedJob2?.AbortedByUserId != UpdatedJob?.AbortedByUserId)
 				{
 					throw new NotImplementedException();
 				}
