Horde: Add missing copyright header

diff --git a/Engine/Source/Programs/Horde/Horde.Build/Utilities/TracingExtensions.cs b/Engine/Source/Programs/Horde/Horde.Build/Utilities/TracingExtensions.cs
--- a/Engine/Source/Programs/Horde/Horde.Build/Utilities/TracingExtensions.cs
+++ b/Engine/Source/Programs/Horde/Horde.Build/Utilities/TracingExtensions.cs
@@ -1,3 +1,5 @@
+// Copyright Epic Games, Inc. All Rights Reserved.
+
 using System;
 using System.Collections.Generic;
 using EpicGames.Core;
