Horde: Fix bug recursively calling wrong ISpan.SetTag method

diff --git a/Engine/Source/Programs/Horde/Horde.Build/Utilities/TracingExtensions.cs b/Engine/Source/Programs/Horde/Horde.Build/Utilities/TracingExtensions.cs
--- a/Engine/Source/Programs/Horde/Horde.Build/Utilities/TracingExtensions.cs
+++ b/Engine/Source/Programs/Horde/Horde.Build/Utilities/TracingExtensions.cs
@@ -82,7 +82,7 @@ namespace Horde.Build.Utilities
 		/// <returns>This span instance, for chaining</returns>
 		public static ISpan SetTag(this ISpan Span, string Key, int? Value)
 		{
-			if (Value != null) Span.SetTag(Key, Value);
+			if (Value != null) Span.SetTag(Key, Value.Value);
 			return Span;
 		}
 		
