Horde: Add collection name to MongoDB tracing spans

diff --git a/Engine/Source/Programs/Horde/Horde.Build/Utilities/MongoTracingCollection.cs b/Engine/Source/Programs/Horde/Horde.Build/Utilities/MongoTracingCollection.cs
--- a/Engine/Source/Programs/Horde/Horde.Build/Utilities/MongoTracingCollection.cs
+++ b/Engine/Source/Programs/Horde/Horde.Build/Utilities/MongoTracingCollection.cs
@@ -51,7 +51,8 @@ namespace HordeServer.Utiltiies
 		public async Task<IAsyncCursor<TResult>> AggregateAsync<TResult>(PipelineDefinition<T, TResult> pipeline, AggregateOptions Options = null!,
 			CancellationToken  CancellationToken = new CancellationToken())
 		{
-			using IScope _ = GlobalTracer.Instance.BuildSpan(Prefix + "AggregateAsync").StartActive();
+			using IScope Scope = GlobalTracer.Instance.BuildSpan(Prefix + "AggregateAsync").StartActive();
+			Scope.Span.SetTag("CollectionName", Collection.CollectionNamespace.CollectionName);
 			return await Collection.AggregateAsync(pipeline, Options,  CancellationToken);
 		}
 
@@ -59,35 +60,40 @@ namespace HordeServer.Utiltiies
 		public async Task<IAsyncCursor<TResult>> AggregateAsync<TResult>(IClientSessionHandle Session, PipelineDefinition<T, TResult> Pipeline, AggregateOptions Options = null!,
 			CancellationToken  CancellationToken = new CancellationToken())
 		{
-			using IScope _ = GlobalTracer.Instance.BuildSpan(Prefix + "AggregateAsync").StartActive();
+			using IScope Scope = GlobalTracer.Instance.BuildSpan(Prefix + "AggregateAsync").StartActive();
+			Scope.Span.SetTag("CollectionName", Collection.CollectionNamespace.CollectionName);
 			return await Collection.AggregateAsync(Session, Pipeline, Options, CancellationToken);
 		}
 
 		/// <inheritdoc />
 		public void AggregateToCollection<TResult>(PipelineDefinition<T, TResult> Pipeline, AggregateOptions? Options = null, CancellationToken CancellationToken = default(CancellationToken))
 		{
-			using IScope _ = GlobalTracer.Instance.BuildSpan(Prefix + "AggregateToCollection").StartActive();
+			using IScope Scope = GlobalTracer.Instance.BuildSpan(Prefix + "AggregateToCollection").StartActive();
+			Scope.Span.SetTag("CollectionName", Collection.CollectionNamespace.CollectionName);
 			Collection.AggregateToCollection(Pipeline, Options, CancellationToken);
 		}
 
 		/// <inheritdoc />
 		public void AggregateToCollection<TResult>(IClientSessionHandle Session, PipelineDefinition<T, TResult> pipeline, AggregateOptions? options = null, CancellationToken cancellationToken = default(CancellationToken))
 		{
-			using IScope _ = GlobalTracer.Instance.BuildSpan(Prefix + "AggregateToCollection").StartActive();
+			using IScope Scope = GlobalTracer.Instance.BuildSpan(Prefix + "AggregateToCollection").StartActive();
+			Scope.Span.SetTag("CollectionName", Collection.CollectionNamespace.CollectionName);
 			Collection.AggregateToCollection(Session, pipeline, options, cancellationToken);
 		}
 
 		/// <inheritdoc />
 		public async Task AggregateToCollectionAsync<TResult>(PipelineDefinition<T, TResult> Pipeline, AggregateOptions? Options = null, CancellationToken CancellationToken = default(CancellationToken))
 		{
-			using IScope _ = GlobalTracer.Instance.BuildSpan(Prefix + "AggregateToCollectionAsync").StartActive();
+			using IScope Scope = GlobalTracer.Instance.BuildSpan(Prefix + "AggregateToCollectionAsync").StartActive();
+			Scope.Span.SetTag("CollectionName", Collection.CollectionNamespace.CollectionName);
 			await Collection.AggregateToCollectionAsync(Pipeline, Options, CancellationToken);
 		}
 
 		/// <inheritdoc />
 		public async Task AggregateToCollectionAsync<TResult>(IClientSessionHandle Session, PipelineDefinition<T, TResult> Pipeline, AggregateOptions? Options = null, CancellationToken CancellationToken = default(CancellationToken))
 		{
-			using IScope _ = GlobalTracer.Instance.BuildSpan(Prefix + "AggregateToCollectionAsync").StartActive();
+			using IScope Scope = GlobalTracer.Instance.BuildSpan(Prefix + "AggregateToCollectionAsync").StartActive();
+			Scope.Span.SetTag("CollectionName", Collection.CollectionNamespace.CollectionName);
 			await Collection.AggregateToCollectionAsync(Session, Pipeline, Options, CancellationToken);
 		}
 
@@ -109,7 +115,8 @@ namespace HordeServer.Utiltiies
 		public async Task<BulkWriteResult<T>> BulkWriteAsync(IEnumerable<WriteModel<T>> requests, BulkWriteOptions Options = null!,
 			CancellationToken  CancellationToken = new CancellationToken())
 		{
-			using IScope _ = GlobalTracer.Instance.BuildSpan(Prefix + "BulkWriteAsync").StartActive();
+			using IScope Scope = GlobalTracer.Instance.BuildSpan(Prefix + "BulkWriteAsync").StartActive();
+			Scope.Span.SetTag("CollectionName", Collection.CollectionNamespace.CollectionName);
 			return await Collection.BulkWriteAsync(requests, Options,  CancellationToken);
 		}
 
@@ -117,7 +124,8 @@ namespace HordeServer.Utiltiies
 		public async Task<BulkWriteResult<T>> BulkWriteAsync(IClientSessionHandle Session, IEnumerable<WriteModel<T>> requests, BulkWriteOptions Options = null!,
 			CancellationToken  CancellationToken = new CancellationToken())
 		{
-			using IScope _ = GlobalTracer.Instance.BuildSpan(Prefix + "BulkWriteAsync").StartActive();
+			using IScope Scope = GlobalTracer.Instance.BuildSpan(Prefix + "BulkWriteAsync").StartActive();
+			Scope.Span.SetTag("CollectionName", Collection.CollectionNamespace.CollectionName);
 			return await Collection.BulkWriteAsync(Session, requests, Options,  CancellationToken);
 		}
 
@@ -139,7 +147,8 @@ namespace HordeServer.Utiltiies
 		public async Task<long> CountAsync(FilterDefinition<T> filter, CountOptions Options = null!,
 			CancellationToken  CancellationToken = new CancellationToken())
 		{
-			using IScope _ = GlobalTracer.Instance.BuildSpan(Prefix + "CountAsync").StartActive();
+			using IScope Scope = GlobalTracer.Instance.BuildSpan(Prefix + "CountAsync").StartActive();
+			Scope.Span.SetTag("CollectionName", Collection.CollectionNamespace.CollectionName);
 			return await Collection.CountAsync(filter, Options,  CancellationToken);
 		}
 
@@ -147,7 +156,8 @@ namespace HordeServer.Utiltiies
 		public async Task<long> CountAsync(IClientSessionHandle Session, FilterDefinition<T> filter, CountOptions Options = null!,
 			CancellationToken  CancellationToken = new CancellationToken())
 		{
-			using IScope _ = GlobalTracer.Instance.BuildSpan(Prefix + "CountAsync").StartActive();
+			using IScope Scope = GlobalTracer.Instance.BuildSpan(Prefix + "CountAsync").StartActive();
+			Scope.Span.SetTag("CollectionName", Collection.CollectionNamespace.CollectionName);
 			return await Collection.CountAsync(Session, filter, Options,  CancellationToken);
 		}
 
@@ -162,7 +172,8 @@ namespace HordeServer.Utiltiies
 		public long CountDocuments(IClientSessionHandle Session, FilterDefinition<T> filter, CountOptions Options = null!,
 			CancellationToken  CancellationToken = new CancellationToken())
 		{
-			using IScope _ = GlobalTracer.Instance.BuildSpan(Prefix + "CountDocumentsAsync").StartActive();
+			using IScope Scope = GlobalTracer.Instance.BuildSpan(Prefix + "CountDocumentsAsync").StartActive();
+			Scope.Span.SetTag("CollectionName", Collection.CollectionNamespace.CollectionName);
 			return Collection.CountDocuments(Session, filter, Options,  CancellationToken);
 		}
 
@@ -170,7 +181,8 @@ namespace HordeServer.Utiltiies
 		public async Task<long> CountDocumentsAsync(FilterDefinition<T> filter, CountOptions Options = null!,
 			CancellationToken  CancellationToken = new CancellationToken())
 		{
-			using IScope _ = GlobalTracer.Instance.BuildSpan(Prefix + "CountDocumentsAsync").StartActive();
+			using IScope Scope = GlobalTracer.Instance.BuildSpan(Prefix + "CountDocumentsAsync").StartActive();
+			Scope.Span.SetTag("CollectionName", Collection.CollectionNamespace.CollectionName);
 			return await Collection.CountDocumentsAsync(filter, Options,  CancellationToken);
 		}
 
@@ -178,7 +190,8 @@ namespace HordeServer.Utiltiies
 		public async Task<long> CountDocumentsAsync(IClientSessionHandle Session, FilterDefinition<T> filter, CountOptions Options = null!,
 			CancellationToken  CancellationToken = new CancellationToken())
 		{
-			using IScope _ = GlobalTracer.Instance.BuildSpan(Prefix + "CountDocumentsAsync").StartActive();
+			using IScope Scope = GlobalTracer.Instance.BuildSpan(Prefix + "CountDocumentsAsync").StartActive();
+			Scope.Span.SetTag("CollectionName", Collection.CollectionNamespace.CollectionName);
 			return await Collection.CountDocumentsAsync(Session, filter, Options,  CancellationToken);
 		}
 
@@ -205,7 +218,8 @@ namespace HordeServer.Utiltiies
 		/// <inheritdoc />
 		public async Task<DeleteResult> DeleteManyAsync(FilterDefinition<T> filter, CancellationToken  CancellationToken = new CancellationToken())
 		{
-			using IScope _ = GlobalTracer.Instance.BuildSpan(Prefix + "DeleteManyAsync").StartActive();
+			using IScope Scope = GlobalTracer.Instance.BuildSpan(Prefix + "DeleteManyAsync").StartActive();
+			Scope.Span.SetTag("CollectionName", Collection.CollectionNamespace.CollectionName);
 			return await Collection.DeleteManyAsync(filter,  CancellationToken);
 		}
 
@@ -213,7 +227,8 @@ namespace HordeServer.Utiltiies
 		public async Task<DeleteResult> DeleteManyAsync(FilterDefinition<T> filter, DeleteOptions Options,
 			CancellationToken  CancellationToken = new CancellationToken())
 		{
-			using IScope _ = GlobalTracer.Instance.BuildSpan(Prefix + "DeleteManyAsync").StartActive();
+			using IScope Scope = GlobalTracer.Instance.BuildSpan(Prefix + "DeleteManyAsync").StartActive();
+			Scope.Span.SetTag("CollectionName", Collection.CollectionNamespace.CollectionName);
 			return await Collection.DeleteManyAsync(filter, Options,  CancellationToken);
 		}
 
@@ -221,7 +236,8 @@ namespace HordeServer.Utiltiies
 		public async Task<DeleteResult> DeleteManyAsync(IClientSessionHandle Session, FilterDefinition<T> filter, DeleteOptions Options = null!,
 			CancellationToken  CancellationToken = new CancellationToken())
 		{
-			using IScope _ = GlobalTracer.Instance.BuildSpan(Prefix + "DeleteManyAsync").StartActive();
+			using IScope Scope = GlobalTracer.Instance.BuildSpan(Prefix + "DeleteManyAsync").StartActive();
+			Scope.Span.SetTag("CollectionName", Collection.CollectionNamespace.CollectionName);
 			return await Collection.DeleteManyAsync(Session, filter, Options,  CancellationToken);
 		}
 
@@ -248,7 +264,8 @@ namespace HordeServer.Utiltiies
 		/// <inheritdoc />
 		public async Task<DeleteResult> DeleteOneAsync(FilterDefinition<T> filter, CancellationToken  CancellationToken = new CancellationToken())
 		{
-			using IScope _ = GlobalTracer.Instance.BuildSpan(Prefix + "DeleteOneAsync").StartActive();
+			using IScope Scope = GlobalTracer.Instance.BuildSpan(Prefix + "DeleteOneAsync").StartActive();
+			Scope.Span.SetTag("CollectionName", Collection.CollectionNamespace.CollectionName);
 			return await Collection.DeleteOneAsync(filter,  CancellationToken);
 		}
 
@@ -256,7 +273,8 @@ namespace HordeServer.Utiltiies
 		public async Task<DeleteResult> DeleteOneAsync(FilterDefinition<T> filter, DeleteOptions Options,
 			CancellationToken  CancellationToken = new CancellationToken())
 		{
-			using IScope _ = GlobalTracer.Instance.BuildSpan(Prefix + "DeleteOneAsync").StartActive();
+			using IScope Scope = GlobalTracer.Instance.BuildSpan(Prefix + "DeleteOneAsync").StartActive();
+			Scope.Span.SetTag("CollectionName", Collection.CollectionNamespace.CollectionName);
 			return await Collection.DeleteOneAsync(filter, Options,  CancellationToken);
 		}
 
@@ -264,7 +282,8 @@ namespace HordeServer.Utiltiies
 		public Task<DeleteResult> DeleteOneAsync(IClientSessionHandle Session, FilterDefinition<T> filter, DeleteOptions Options = null!,
 			CancellationToken  CancellationToken = new CancellationToken())
 		{
-			using IScope _ = GlobalTracer.Instance.BuildSpan(Prefix + "DeleteOneAsync").StartActive();
+			using IScope Scope = GlobalTracer.Instance.BuildSpan(Prefix + "DeleteOneAsync").StartActive();
+			Scope.Span.SetTag("CollectionName", Collection.CollectionNamespace.CollectionName);
 			return Collection.DeleteOneAsync(Session, filter, Options,  CancellationToken);
 		}
 
@@ -286,7 +305,8 @@ namespace HordeServer.Utiltiies
 		public async Task<IAsyncCursor<TField>> DistinctAsync<TField>(FieldDefinition<T, TField> field, FilterDefinition<T> filter, DistinctOptions Options = null!,
 			CancellationToken  CancellationToken = new CancellationToken())
 		{
-			using IScope _ = GlobalTracer.Instance.BuildSpan(Prefix + "DistinctAsync").StartActive();
+			using IScope Scope = GlobalTracer.Instance.BuildSpan(Prefix + "DistinctAsync").StartActive();
+			Scope.Span.SetTag("CollectionName", Collection.CollectionNamespace.CollectionName);
 			return await Collection.DistinctAsync(field, filter, Options,  CancellationToken);
 		}
 
@@ -294,7 +314,8 @@ namespace HordeServer.Utiltiies
 		public async Task<IAsyncCursor<TField>> DistinctAsync<TField>(IClientSessionHandle Session, FieldDefinition<T, TField> field, FilterDefinition<T> filter,
 			DistinctOptions Options = null!, CancellationToken  CancellationToken = new CancellationToken())
 		{
-			using IScope _ = GlobalTracer.Instance.BuildSpan(Prefix + "DistinctAsync").StartActive();
+			using IScope Scope = GlobalTracer.Instance.BuildSpan(Prefix + "DistinctAsync").StartActive();
+			Scope.Span.SetTag("CollectionName", Collection.CollectionNamespace.CollectionName);
 			return await Collection.DistinctAsync(Session, field, filter, Options,  CancellationToken);
 		}
 
@@ -309,7 +330,8 @@ namespace HordeServer.Utiltiies
 		public async Task<long> EstimatedDocumentCountAsync(EstimatedDocumentCountOptions Options = null!,
 			CancellationToken  CancellationToken = new CancellationToken())
 		{
-			using IScope _ = GlobalTracer.Instance.BuildSpan(Prefix + "EstimatedDocumentCountAsync").StartActive();
+			using IScope Scope = GlobalTracer.Instance.BuildSpan(Prefix + "EstimatedDocumentCountAsync").StartActive();
+			Scope.Span.SetTag("CollectionName", Collection.CollectionNamespace.CollectionName);
 			return await Collection.EstimatedDocumentCountAsync(Options,  CancellationToken);
 		}
 
@@ -331,7 +353,8 @@ namespace HordeServer.Utiltiies
 		public async Task<IAsyncCursor<TProjection>> FindAsync<TProjection>(FilterDefinition<T> filter, FindOptions<T, TProjection> Options = null!,
 			CancellationToken  CancellationToken = new CancellationToken())
 		{
-			using IScope _ = GlobalTracer.Instance.BuildSpan(Prefix + "FindAsync").StartActive();
+			using IScope Scope = GlobalTracer.Instance.BuildSpan(Prefix + "FindAsync").StartActive();
+			Scope.Span.SetTag("CollectionName", Collection.CollectionNamespace.CollectionName);
 			return await Collection.FindAsync(filter, Options,  CancellationToken);
 		}
 
@@ -339,7 +362,8 @@ namespace HordeServer.Utiltiies
 		public async Task<IAsyncCursor<TProjection>> FindAsync<TProjection>(IClientSessionHandle Session, FilterDefinition<T> filter, FindOptions<T, TProjection> Options = null!,
 			CancellationToken  CancellationToken = new CancellationToken())
 		{
-			using IScope _ = GlobalTracer.Instance.BuildSpan(Prefix + "FindAsync").StartActive();
+			using IScope Scope = GlobalTracer.Instance.BuildSpan(Prefix + "FindAsync").StartActive();
+			Scope.Span.SetTag("CollectionName", Collection.CollectionNamespace.CollectionName);
 			return await Collection.FindAsync(Session, filter, Options,  CancellationToken);
 		}
 
@@ -361,7 +385,8 @@ namespace HordeServer.Utiltiies
 		public async Task<TProjection> FindOneAndDeleteAsync<TProjection>(FilterDefinition<T> filter, FindOneAndDeleteOptions<T, TProjection> Options = null!,
 			CancellationToken  CancellationToken = new CancellationToken())
 		{
-			using IScope _ = GlobalTracer.Instance.BuildSpan(Prefix + "FindOneAndDeleteAsync").StartActive();
+			using IScope Scope = GlobalTracer.Instance.BuildSpan(Prefix + "FindOneAndDeleteAsync").StartActive();
+			Scope.Span.SetTag("CollectionName", Collection.CollectionNamespace.CollectionName);
 			return await Collection.FindOneAndDeleteAsync(filter, Options,  CancellationToken);
 		}
 
@@ -369,7 +394,8 @@ namespace HordeServer.Utiltiies
 		public async Task<TProjection> FindOneAndDeleteAsync<TProjection>(IClientSessionHandle Session, FilterDefinition<T> filter,
 			FindOneAndDeleteOptions<T, TProjection> Options = null!, CancellationToken  CancellationToken = new CancellationToken())
 		{
-			using IScope _ = GlobalTracer.Instance.BuildSpan(Prefix + "FindOneAndDeleteAsync").StartActive();
+			using IScope Scope = GlobalTracer.Instance.BuildSpan(Prefix + "FindOneAndDeleteAsync").StartActive();
+			Scope.Span.SetTag("CollectionName", Collection.CollectionNamespace.CollectionName);
 			return await Collection.FindOneAndDeleteAsync(Session, filter, Options,  CancellationToken);
 		}
 
@@ -391,7 +417,8 @@ namespace HordeServer.Utiltiies
 		public async Task<TProjection> FindOneAndReplaceAsync<TProjection>(FilterDefinition<T> filter, T replacement,
 			FindOneAndReplaceOptions<T, TProjection> Options = null!, CancellationToken  CancellationToken = new CancellationToken())
 		{
-			using IScope _ = GlobalTracer.Instance.BuildSpan(Prefix + "FindOneAndReplaceAsync").StartActive();
+			using IScope Scope = GlobalTracer.Instance.BuildSpan(Prefix + "FindOneAndReplaceAsync").StartActive();
+			Scope.Span.SetTag("CollectionName", Collection.CollectionNamespace.CollectionName);
 			return await Collection.FindOneAndReplaceAsync(filter, replacement, Options,  CancellationToken);
 		}
 
@@ -399,7 +426,8 @@ namespace HordeServer.Utiltiies
 		public async Task<TProjection> FindOneAndReplaceAsync<TProjection>(IClientSessionHandle Session, FilterDefinition<T> filter, T replacement,
 			FindOneAndReplaceOptions<T, TProjection> Options = null!, CancellationToken  CancellationToken = new CancellationToken())
 		{
-			using IScope _ = GlobalTracer.Instance.BuildSpan(Prefix + "FindOneAndReplaceAsync").StartActive();
+			using IScope Scope = GlobalTracer.Instance.BuildSpan(Prefix + "FindOneAndReplaceAsync").StartActive();
+			Scope.Span.SetTag("CollectionName", Collection.CollectionNamespace.CollectionName);
 			return await Collection.FindOneAndReplaceAsync(Session, filter, replacement, Options,  CancellationToken);
 		}
 
@@ -422,7 +450,8 @@ namespace HordeServer.Utiltiies
 		public async Task<TProjection> FindOneAndUpdateAsync<TProjection>(FilterDefinition<T> filter, UpdateDefinition<T> update,
 			FindOneAndUpdateOptions<T, TProjection> Options = null!, CancellationToken  CancellationToken = new CancellationToken())
 		{
-			using IScope _ = GlobalTracer.Instance.BuildSpan(Prefix + "FindOneAndUpdateAsync").StartActive();
+			using IScope Scope = GlobalTracer.Instance.BuildSpan(Prefix + "FindOneAndUpdateAsync").StartActive();
+			Scope.Span.SetTag("CollectionName", Collection.CollectionNamespace.CollectionName);
 			return await Collection.FindOneAndUpdateAsync(filter, update, Options,  CancellationToken);
 		}
 
@@ -430,7 +459,8 @@ namespace HordeServer.Utiltiies
 		public async Task<TProjection> FindOneAndUpdateAsync<TProjection>(IClientSessionHandle Session, FilterDefinition<T> filter, UpdateDefinition<T> update,
 			FindOneAndUpdateOptions<T, TProjection> Options = null!, CancellationToken  CancellationToken = new CancellationToken())
 		{
-			using IScope _ = GlobalTracer.Instance.BuildSpan(Prefix + "FindOneAndUpdateAsync").StartActive();
+			using IScope Scope = GlobalTracer.Instance.BuildSpan(Prefix + "FindOneAndUpdateAsync").StartActive();
+			Scope.Span.SetTag("CollectionName", Collection.CollectionNamespace.CollectionName);
 			return await Collection.FindOneAndUpdateAsync(Session, filter, update, Options,  CancellationToken);
 		}
 
@@ -451,7 +481,8 @@ namespace HordeServer.Utiltiies
 		/// <inheritdoc />
 		public async Task InsertOneAsync( T Document, CancellationToken CancellationToken)
 		{
-			using IScope _ = GlobalTracer.Instance.BuildSpan(Prefix + "InsertOneAsync").StartActive();
+			using IScope Scope = GlobalTracer.Instance.BuildSpan(Prefix + "InsertOneAsync").StartActive();
+			Scope.Span.SetTag("CollectionName", Collection.CollectionNamespace.CollectionName);
 			await Collection.InsertOneAsync(Document, CancellationToken);
 		}
 
@@ -459,7 +490,8 @@ namespace HordeServer.Utiltiies
 		public async Task InsertOneAsync( T Document, InsertOneOptions Options = null!,
 			CancellationToken  CancellationToken = new CancellationToken())
 		{
-			using IScope _ = GlobalTracer.Instance.BuildSpan(Prefix + "InsertOneAsync").StartActive();
+			using IScope Scope = GlobalTracer.Instance.BuildSpan(Prefix + "InsertOneAsync").StartActive();
+			Scope.Span.SetTag("CollectionName", Collection.CollectionNamespace.CollectionName);
 			await Collection.InsertOneAsync(Document, Options,  CancellationToken);
 		}
 
@@ -467,7 +499,8 @@ namespace HordeServer.Utiltiies
 		public async Task InsertOneAsync(IClientSessionHandle Session, T Document, InsertOneOptions Options = null!,
 			CancellationToken CancellationToken = new CancellationToken())
 		{
-			using IScope _ = GlobalTracer.Instance.BuildSpan(Prefix + "InsertOneAsync").StartActive();
+			using IScope Scope = GlobalTracer.Instance.BuildSpan(Prefix + "InsertOneAsync").StartActive();
+			Scope.Span.SetTag("CollectionName", Collection.CollectionNamespace.CollectionName);
 			await Collection.InsertOneAsync(Session, Document, Options,  CancellationToken);
 		}
 
@@ -489,7 +522,8 @@ namespace HordeServer.Utiltiies
 		public async Task InsertManyAsync(IEnumerable<T> Documents, InsertManyOptions Options = null!,
 			CancellationToken  CancellationToken = new CancellationToken())
 		{
-			using IScope _ = GlobalTracer.Instance.BuildSpan(Prefix + "InsertManyAsync").StartActive();
+			using IScope Scope = GlobalTracer.Instance.BuildSpan(Prefix + "InsertManyAsync").StartActive();
+			Scope.Span.SetTag("CollectionName", Collection.CollectionNamespace.CollectionName);
 			await Collection.InsertManyAsync(Documents, Options,  CancellationToken);
 		}
 
@@ -497,7 +531,8 @@ namespace HordeServer.Utiltiies
 		public async Task InsertManyAsync(IClientSessionHandle Session, IEnumerable<T> documents, InsertManyOptions Options = null!,
 			CancellationToken  CancellationToken = new CancellationToken())
 		{
-			using IScope _ = GlobalTracer.Instance.BuildSpan(Prefix + "InsertManyAsync").StartActive();
+			using IScope Scope = GlobalTracer.Instance.BuildSpan(Prefix + "InsertManyAsync").StartActive();
+			Scope.Span.SetTag("CollectionName", Collection.CollectionNamespace.CollectionName);
 			await Collection.InsertManyAsync(Session, documents, Options,  CancellationToken);
 		}
 
@@ -519,7 +554,8 @@ namespace HordeServer.Utiltiies
 		public async Task<IAsyncCursor<TResult>> MapReduceAsync<TResult>(BsonJavaScript map, BsonJavaScript reduce, MapReduceOptions<T, TResult> Options = null!,
 			CancellationToken  CancellationToken = new CancellationToken())
 		{
-			using IScope _ = GlobalTracer.Instance.BuildSpan(Prefix + "MapReduceAsync").StartActive();
+			using IScope Scope = GlobalTracer.Instance.BuildSpan(Prefix + "MapReduceAsync").StartActive();
+			Scope.Span.SetTag("CollectionName", Collection.CollectionNamespace.CollectionName);
 			return await Collection.MapReduceAsync(map, reduce, Options,  CancellationToken);
 		}
 
@@ -527,7 +563,8 @@ namespace HordeServer.Utiltiies
 		public async Task<IAsyncCursor<TResult>> MapReduceAsync<TResult>(IClientSessionHandle Session, BsonJavaScript map, BsonJavaScript reduce,
 			MapReduceOptions<T, TResult> Options = null!, CancellationToken  CancellationToken = new CancellationToken())
 		{
-			using IScope _ = GlobalTracer.Instance.BuildSpan(Prefix + "MapReduceAsync").StartActive();
+			using IScope Scope = GlobalTracer.Instance.BuildSpan(Prefix + "MapReduceAsync").StartActive();
+			Scope.Span.SetTag("CollectionName", Collection.CollectionNamespace.CollectionName);
 			return await Collection.MapReduceAsync(Session, map, reduce, Options,  CancellationToken);
 		}
 
@@ -569,7 +606,8 @@ namespace HordeServer.Utiltiies
 		public async Task<ReplaceOneResult> ReplaceOneAsync(FilterDefinition<T> filter, T replacement, ReplaceOptions Options = null!,
 			CancellationToken  CancellationToken = new CancellationToken())
 		{
-			using IScope _ = GlobalTracer.Instance.BuildSpan(Prefix + "ReplaceOneAsync").StartActive();
+			using IScope Scope = GlobalTracer.Instance.BuildSpan(Prefix + "ReplaceOneAsync").StartActive();
+			Scope.Span.SetTag("CollectionName", Collection.CollectionNamespace.CollectionName);
 			return await Collection.ReplaceOneAsync(filter, replacement, Options,  CancellationToken);
 		}
 
@@ -577,7 +615,8 @@ namespace HordeServer.Utiltiies
 		public async Task<ReplaceOneResult> ReplaceOneAsync(FilterDefinition<T> filter, T replacement, UpdateOptions Options,
 			CancellationToken  CancellationToken = new CancellationToken())
 		{
-			using IScope _ = GlobalTracer.Instance.BuildSpan(Prefix + "ReplaceOneAsync").StartActive();
+			using IScope Scope = GlobalTracer.Instance.BuildSpan(Prefix + "ReplaceOneAsync").StartActive();
+			Scope.Span.SetTag("CollectionName", Collection.CollectionNamespace.CollectionName);
 			return await Collection.ReplaceOneAsync(filter, replacement, Options,  CancellationToken);
 		}
 
@@ -585,7 +624,8 @@ namespace HordeServer.Utiltiies
 		public async Task<ReplaceOneResult> ReplaceOneAsync(IClientSessionHandle Session, FilterDefinition<T> filter, T replacement,
 			ReplaceOptions Options = null!, CancellationToken  CancellationToken = new CancellationToken())
 		{
-			using IScope _ = GlobalTracer.Instance.BuildSpan(Prefix + "ReplaceOneAsync").StartActive();
+			using IScope Scope = GlobalTracer.Instance.BuildSpan(Prefix + "ReplaceOneAsync").StartActive();
+			Scope.Span.SetTag("CollectionName", Collection.CollectionNamespace.CollectionName);
 			return await Collection.ReplaceOneAsync(Session, filter, replacement, Options,  CancellationToken);
 		}
 
@@ -593,7 +633,8 @@ namespace HordeServer.Utiltiies
 		public async Task<ReplaceOneResult> ReplaceOneAsync(IClientSessionHandle Session, FilterDefinition<T> filter, T replacement, UpdateOptions Options,
 			CancellationToken  CancellationToken = new CancellationToken())
 		{
-			using IScope _ = GlobalTracer.Instance.BuildSpan(Prefix + "ReplaceOneAsync").StartActive();
+			using IScope Scope = GlobalTracer.Instance.BuildSpan(Prefix + "ReplaceOneAsync").StartActive();
+			Scope.Span.SetTag("CollectionName", Collection.CollectionNamespace.CollectionName);
 			return await Collection.ReplaceOneAsync(Session, filter, replacement, Options,  CancellationToken);
 		}
 
@@ -615,7 +656,8 @@ namespace HordeServer.Utiltiies
 		public async Task<UpdateResult> UpdateManyAsync(FilterDefinition<T> filter, UpdateDefinition<T> update, UpdateOptions Options = null!,
 			CancellationToken  CancellationToken = new CancellationToken())
 		{
-			using IScope _ = GlobalTracer.Instance.BuildSpan(Prefix + "UpdateManyAsync").StartActive();
+			using IScope Scope = GlobalTracer.Instance.BuildSpan(Prefix + "UpdateManyAsync").StartActive();
+			Scope.Span.SetTag("CollectionName", Collection.CollectionNamespace.CollectionName);
 			return await Collection.UpdateManyAsync(filter, update, Options,  CancellationToken);
 		}
 
@@ -623,7 +665,8 @@ namespace HordeServer.Utiltiies
 		public async Task<UpdateResult> UpdateManyAsync(IClientSessionHandle Session, FilterDefinition<T> filter, UpdateDefinition<T> update,
 			UpdateOptions Options = null!, CancellationToken  CancellationToken = new CancellationToken())
 		{
-			using IScope _ = GlobalTracer.Instance.BuildSpan(Prefix + "UpdateManyAsync").StartActive();
+			using IScope Scope = GlobalTracer.Instance.BuildSpan(Prefix + "UpdateManyAsync").StartActive();
+			Scope.Span.SetTag("CollectionName", Collection.CollectionNamespace.CollectionName);
 			return await Collection.UpdateManyAsync(Session, filter, update, Options,  CancellationToken);
 		}
 
@@ -645,7 +688,8 @@ namespace HordeServer.Utiltiies
 		public async Task<UpdateResult> UpdateOneAsync(FilterDefinition<T> filter, UpdateDefinition<T> update, UpdateOptions Options = null!,
 			CancellationToken  CancellationToken = new CancellationToken())
 		{
-			using IScope _ = GlobalTracer.Instance.BuildSpan(Prefix + "UpdateOneAsync").StartActive();
+			using IScope Scope = GlobalTracer.Instance.BuildSpan(Prefix + "UpdateOneAsync").StartActive();
+			Scope.Span.SetTag("CollectionName", Collection.CollectionNamespace.CollectionName);
 			return await Collection.UpdateOneAsync(filter, update, Options,  CancellationToken);
 		}
 
@@ -653,7 +697,8 @@ namespace HordeServer.Utiltiies
 		public async Task<UpdateResult> UpdateOneAsync(IClientSessionHandle Session, FilterDefinition<T> filter, UpdateDefinition<T> update,
 			UpdateOptions Options = null!, CancellationToken  CancellationToken = new CancellationToken())
 		{
-			using IScope _ = GlobalTracer.Instance.BuildSpan(Prefix + "UpdateOneAsync").StartActive();
+			using IScope Scope = GlobalTracer.Instance.BuildSpan(Prefix + "UpdateOneAsync").StartActive();
+			Scope.Span.SetTag("CollectionName", Collection.CollectionNamespace.CollectionName);
 			return await Collection.UpdateOneAsync(Session, filter, update, Options,  CancellationToken);
 		}
 
@@ -675,7 +720,8 @@ namespace HordeServer.Utiltiies
 		public async Task<IChangeStreamCursor<TResult>> WatchAsync<TResult>(PipelineDefinition<ChangeStreamDocument<T>, TResult> pipeline, ChangeStreamOptions Options = null!,
 			CancellationToken  CancellationToken = new CancellationToken())
 		{
-			using IScope _ = GlobalTracer.Instance.BuildSpan(Prefix + "WatchAsync").StartActive();
+			using IScope Scope = GlobalTracer.Instance.BuildSpan(Prefix + "WatchAsync").StartActive();
+			Scope.Span.SetTag("CollectionName", Collection.CollectionNamespace.CollectionName);
 			return await Collection.WatchAsync(pipeline, Options,  CancellationToken);
 		}
 
@@ -683,7 +729,8 @@ namespace HordeServer.Utiltiies
 		public async Task<IChangeStreamCursor<TResult>> WatchAsync<TResult>(IClientSessionHandle Session, PipelineDefinition<ChangeStreamDocument<T>, TResult> pipeline, ChangeStreamOptions Options = null!,
 			CancellationToken  CancellationToken = new CancellationToken())
 		{
-			using IScope _ = GlobalTracer.Instance.BuildSpan(Prefix + "WatchAsync").StartActive();
+			using IScope Scope = GlobalTracer.Instance.BuildSpan(Prefix + "WatchAsync").StartActive();
+			Scope.Span.SetTag("CollectionName", Collection.CollectionNamespace.CollectionName);
 			return await Collection.WatchAsync(Session, pipeline, Options,  CancellationToken);
 		}
 
