Horde: Add device pool authorization
#rnx

diff --git a/Engine/Source/Programs/Horde/Horde.Build/Api/Devices.cs b/Engine/Source/Programs/Horde/Horde.Build/Api/Devices.cs
--- a/Engine/Source/Programs/Horde/Horde.Build/Api/Devices.cs
+++ b/Engine/Source/Programs/Horde/Horde.Build/Api/Devices.cs
@@ -103,6 +103,28 @@ namespace HordeServer.Api
 		[Required]
 		public DevicePoolType PoolType { get; set; }
 
+		/// <summary>
+		/// Projects associated with this device pool
+		/// </summary>
+		public List<string>? ProjectIds { get; set; }
+
+	}
+
+	/// <summary>
+	/// Device pool update request object
+	/// </summary>
+	public class UpdateDevicePoolRequest
+	{
+		/// <summary>
+		/// Id of the device pool to update
+		/// </summary>
+		public string Id { get; set; } = null!;
+
+		/// <summary>
+		/// Projects associated with this device pool
+		/// </summary>
+		public List<string>? ProjectIds { get; set; }
+
 	}
 
 	/// <summary>
@@ -144,14 +166,20 @@ namespace HordeServer.Api
 		/// </summary>
 		public DevicePoolType PoolType { get; set; }
 
+		/// <summary>
+		/// Whether there is write access to the pool
+		/// </summary>
+		public bool WriteAccess { get; set; }
+
 		/// <summary>
 		/// Constructor
 		/// </summary>
-		public GetDevicePoolResponse(string Id, string Name, DevicePoolType PoolType)
+		public GetDevicePoolResponse(string Id, string Name, DevicePoolType PoolType, bool WriteAccess)
 		{
 			this.Id = Id;
 			this.Name = Name;
             this.PoolType = PoolType;
+			this.WriteAccess = WriteAccess;
         }
 	}
 
diff --git a/Engine/Source/Programs/Horde/Horde.Build/Controllers/DevicesController.cs b/Engine/Source/Programs/Horde/Horde.Build/Controllers/DevicesController.cs
--- a/Engine/Source/Programs/Horde/Horde.Build/Controllers/DevicesController.cs
+++ b/Engine/Source/Programs/Horde/Horde.Build/Controllers/DevicesController.cs
@@ -21,6 +21,7 @@ namespace HordeServer.Controllers
 	using DeviceId = StringId<IDevice>;
 	using DevicePlatformId = StringId<IDevicePlatform>;
 	using DevicePoolId = StringId<IDevicePool>;
+	using ProjectId = StringId<IProject>;
 
 	/// <summary>
 	/// Controller for device service
@@ -71,7 +72,9 @@ namespace HordeServer.Controllers
 		public async Task<ActionResult<CreateDeviceResponse>> CreateDeviceAsync([FromBody] CreateDeviceRequest DeviceRequest)
 		{
 
-			if (!await DeviceService.AuthorizeAsync(AclAction.DeviceWrite, User))
+			DevicePoolAuthorization? PoolAuth = await DeviceService.GetUserPoolAuthorizationAsync(new DevicePoolId(DeviceRequest.PoolId!), User);
+
+			if (PoolAuth == null || !PoolAuth.Write)
 			{
 				return Forbid();
 			}
@@ -132,10 +135,7 @@ namespace HordeServer.Controllers
 		public async Task<ActionResult<List<object>>> GetDevicesAsync()
 		{
 
-			if (!await DeviceService.AuthorizeAsync(AclAction.DeviceRead, User))
-			{
-				return Forbid();
-			}
+			List<DevicePoolAuthorization> PoolAuth = await DeviceService.GetUserPoolAuthorizationsAsync(User);
 
 			List<IDevice> Devices = await DeviceService.GetDevicesAsync();
 
@@ -143,6 +143,13 @@ namespace HordeServer.Controllers
 
 			foreach (IDevice Device in Devices)
 			{
+				DevicePoolAuthorization? Auth = PoolAuth.Where(x => x.Pool.Id == Device.PoolId).FirstOrDefault();
+
+				if (Auth == null || !Auth.Read)
+				{
+					continue;
+				}
+
 				Responses.Add(new GetDeviceResponse(Device.Id.ToString(), Device.PlatformId.ToString(), Device.PoolId.ToString(), Device.Name, Device.Enabled, Device.Address, Device.ModelId?.ToString(), Device.ModifiedByUser, Device.Notes, Device.ProblemTimeUtc, Device.MaintenanceTimeUtc, Device.Utilization, Device.CheckedOutByUser, Device.CheckOutTime));
 			}
 
@@ -159,11 +166,6 @@ namespace HordeServer.Controllers
 		public async Task<ActionResult<object>> GetDeviceAsync(string DeviceId)
 		{
 
-			if (!await DeviceService.AuthorizeAsync(AclAction.DeviceRead, User))
-			{
-				return Forbid();
-			}
-
 			DeviceId DeviceIdValue = new DeviceId(DeviceId);
 
 			IDevice? Device = await DeviceService.GetDeviceAsync(DeviceIdValue);
@@ -173,6 +175,14 @@ namespace HordeServer.Controllers
 				return BadRequest($"Unable to find device with id {DeviceId}");
 			}
 
+			DevicePoolAuthorization? PoolAuth = await DeviceService.GetUserPoolAuthorizationAsync(Device.PoolId, User);
+
+			if (PoolAuth == null || !PoolAuth.Write)
+			{
+				return Forbid();
+			}
+
+
 			return new GetDeviceResponse(Device.Id.ToString(), Device.PlatformId.ToString(), Device.PoolId.ToString(), Device.Name, Device.Enabled, Device.Address, Device.ModelId?.ToString(), Device.ModifiedByUser?.ToString(), Device.Notes, Device.ProblemTimeUtc, Device.MaintenanceTimeUtc, Device.Utilization, Device.CheckedOutByUser, Device.CheckOutTime);
 		}
 
@@ -185,11 +195,6 @@ namespace HordeServer.Controllers
 		[ProducesResponseType(typeof(List<GetDeviceResponse>), 200)]
 		public async Task<ActionResult> UpdateDeviceAsync(string DeviceId, [FromBody] UpdateDeviceRequest Update)
 		{
-			if (!await DeviceService.AuthorizeAsync(AclAction.DeviceWrite, User))
-			{
-				return Forbid();
-			}
-
 			IUser? InternalUser = await UserCollection.GetUserAsync(User);
 			if (InternalUser == null)
 			{
@@ -205,6 +210,13 @@ namespace HordeServer.Controllers
 				return BadRequest($"Device with id ${DeviceId} does not exist");
 			}
 
+			DevicePoolAuthorization? PoolAuth = await DeviceService.GetUserPoolAuthorizationAsync(Device.PoolId, User);
+
+			if (PoolAuth == null || !PoolAuth.Write)
+			{
+				return Forbid();
+			}
+
 			IDevicePlatform? Platform = await DeviceService.GetPlatformAsync(Device.PlatformId);
 
 			if (Platform == null)
@@ -249,10 +261,6 @@ namespace HordeServer.Controllers
 		[ProducesResponseType(typeof(List<GetDeviceResponse>), 200)]
 		public async Task<ActionResult> CheckoutDeviceAsync(string DeviceId, [FromBody] CheckoutDeviceRequest Request)
 		{
-			if (!await DeviceService.AuthorizeAsync(AclAction.DeviceWrite, User))
-			{
-				return Forbid();
-			}
 
 			IUser? InternalUser = await UserCollection.GetUserAsync(User);
 			if (InternalUser == null)
@@ -269,6 +277,13 @@ namespace HordeServer.Controllers
 				return BadRequest($"Device with id ${DeviceId} does not exist");
 			}
 
+			DevicePoolAuthorization? PoolAuth = await DeviceService.GetUserPoolAuthorizationAsync(Device.PoolId, User);
+
+			if (PoolAuth == null || !PoolAuth.Write)
+			{
+				return Forbid();
+			}
+
 			if (Request.Checkout)
 			{
 				if (!string.IsNullOrEmpty(Device.CheckedOutByUser))
@@ -309,6 +324,14 @@ namespace HordeServer.Controllers
 				return NotFound();
 			}
 
+			DevicePoolAuthorization? PoolAuth = await DeviceService.GetUserPoolAuthorizationAsync(Device.PoolId, User);
+
+			if (PoolAuth == null || !PoolAuth.Write)
+			{
+				return Forbid();
+			}
+
+
 			await DeviceService.DeleteDeviceAsync(DeviceIdValue);
 			return Ok();
 		}
@@ -408,7 +431,7 @@ namespace HordeServer.Controllers
 
 			string Name = Request.Name.Trim();
 
-			IDevicePool? Pool = await DeviceService.TryCreatePoolAsync(DevicePoolId.Sanitize(Name), Name, Request.PoolType);
+			IDevicePool? Pool = await DeviceService.TryCreatePoolAsync(DevicePoolId.Sanitize(Name), Name, Request.PoolType, Request.ProjectIds?.Select(x => new ProjectId(x)).ToList());
 
 			if (Pool == null)
 			{
@@ -417,9 +440,37 @@ namespace HordeServer.Controllers
 
 			return new CreateDevicePoolResponse(Pool.Id.ToString());
 
+		}
+
+		/// <summary>
+		/// Update a device pool
+		/// </summary>
+		[HttpPut]
+		[Authorize]
+		[Route("/api/v2/devices/pools")]
+		public async Task<ActionResult> UpdatePoolAsync([FromBody] UpdateDevicePoolRequest Request)
+		{
+
+			DevicePoolAuthorization? PoolAuth = await DeviceService.GetUserPoolAuthorizationAsync(new DevicePoolId(Request.Id), User);
 
+			if (PoolAuth == null || !PoolAuth.Write)
+			{
+				return Forbid();
+			}
+
+			List<ProjectId>? ProjectIds = null;
+
+			if (Request.ProjectIds != null)
+			{
+				ProjectIds = Request.ProjectIds.Select(x => new ProjectId(x)).ToList();
+			}
+
+			await DeviceService.UpdatePoolAsync(new DevicePoolId(Request.Id), ProjectIds);
+			
+			return Ok();
 		}
 
+
 		/// <summary>
 		/// Get a list of existing device pools
 		/// </summary>
@@ -430,10 +481,7 @@ namespace HordeServer.Controllers
 		public async Task<ActionResult<List<object>>> GetDevicePoolsAsync()
 		{
 
-			if (!await DeviceService.AuthorizeAsync(AclAction.DeviceRead, User))
-			{
-				return Forbid();
-			}
+			List<DevicePoolAuthorization> PoolAuth = await DeviceService.GetUserPoolAuthorizationsAsync(User);
 
 			List<IDevicePool> Pools = await DeviceService.GetPoolsAsync();
 
@@ -441,8 +489,14 @@ namespace HordeServer.Controllers
 
 			foreach (IDevicePool Pool in Pools)
 			{
-				// @todo: ACL per platform
-				Responses.Add(new GetDevicePoolResponse(Pool.Id.ToString(), Pool.Name, Pool.PoolType));
+				DevicePoolAuthorization? Auth = PoolAuth.Where(x => x.Pool.Id == Pool.Id).FirstOrDefault();
+
+				if (Auth == null || !Auth.Read)
+				{
+					continue;
+				}
+				
+				Responses.Add(new GetDevicePoolResponse(Pool.Id.ToString(), Pool.Name, Pool.PoolType, Auth.Write));
 			}
 
 			return Responses;
diff --git a/Engine/Source/Programs/Horde/Horde.Build/Devices/IDevice.cs b/Engine/Source/Programs/Horde/Horde.Build/Devices/IDevice.cs
--- a/Engine/Source/Programs/Horde/Horde.Build/Devices/IDevice.cs
+++ b/Engine/Source/Programs/Horde/Horde.Build/Devices/IDevice.cs
@@ -13,6 +13,7 @@ namespace HordeServer.Models
 	using DevicePlatformId = StringId<IDevicePlatform>;
 	using DevicePoolId = StringId<IDevicePool>;
 	using UserId = ObjectId<IUser>;
+	using ProjectId = StringId<IProject>;
 
 	/// <summary>
 	/// A reservation containing one or more devices
@@ -130,10 +131,15 @@ namespace HordeServer.Models
 		/// </summary>
 		public DevicePoolType PoolType { get; }
 
-        /// <summary>
-        /// Friendly name of the pool
-        /// </summary>
-        string Name { get; }
+		/// <summary>
+		/// Projects associated with this device pool
+		/// </summary>
+		public List<ProjectId>? ProjectIds { get; }
+
+		/// <summary>
+		/// Friendly name of the pool
+		/// </summary>
+		string Name { get; }
 
 		/// <summary>
 		/// Acl access to the pool
diff --git a/Engine/Source/Programs/Horde/Horde.Build/Devices/IDeviceCollection.cs b/Engine/Source/Programs/Horde/Horde.Build/Devices/IDeviceCollection.cs
--- a/Engine/Source/Programs/Horde/Horde.Build/Devices/IDeviceCollection.cs
+++ b/Engine/Source/Programs/Horde/Horde.Build/Devices/IDeviceCollection.cs
@@ -12,6 +12,7 @@ namespace HordeServer.Collections
 	using DevicePlatformId = StringId<IDevicePlatform>;
 	using DevicePoolId = StringId<IDevicePool>;
 	using UserId = ObjectId<IUser>;
+	using ProjectId = StringId<IProject>;
 
 	/// <summary>
 	/// Device reservation request data
@@ -83,8 +84,16 @@ namespace HordeServer.Collections
 		/// </summary>
 		/// <param name="Id">The id of the new pool</param>
 		/// <param name="Name">The friendly name of the new pool</param>
-        /// <param name="PoolType">The pool type</param>
-		Task<IDevicePool?> TryAddPoolAsync(DevicePoolId Id, string Name, DevicePoolType PoolType);
+		/// <param name="PoolType">The pool type</param>
+		/// <param name="ProjectIds">Projects associated with this pool</param>
+		Task<IDevicePool?> TryAddPoolAsync(DevicePoolId Id, string Name, DevicePoolType PoolType, List<ProjectId>? ProjectIds );
+
+		/// <summary>
+		/// Update a device pool
+		/// </summary>
+		/// <param name="Id">The id of the device pool to update</param>
+		/// <param name="ProjectIds">Associated project ids</param>
+		Task UpdatePoolAsync(DevicePoolId Id, List<ProjectId>? ProjectIds);
 
 		/// <summary>
 		/// Get a pool by id
diff --git a/Engine/Source/Programs/Horde/Horde.Build/Devices/Impl/DeviceService.cs b/Engine/Source/Programs/Horde/Horde.Build/Devices/Impl/DeviceService.cs
--- a/Engine/Source/Programs/Horde/Horde.Build/Devices/Impl/DeviceService.cs
+++ b/Engine/Source/Programs/Horde/Horde.Build/Devices/Impl/DeviceService.cs
@@ -25,6 +25,41 @@ namespace HordeServer.Services
 	using DevicePoolId = StringId<IDevicePool>;
 	using JobId = ObjectId<IJob>;
 	using UserId = ObjectId<IUser>;
+	using ProjectId = StringId<IProject>;
+
+	/// <summary>
+	///  Device Pool Authorization (convenience class for pool ACL's)
+	/// </summary>
+	public class DevicePoolAuthorization
+	{
+		/// <summary>
+		/// The device pool
+		/// </summary>
+		public IDevicePool Pool { get; private set; }
+
+		/// <summary>
+		///  Read access to pool
+		/// </summary>
+		public bool Read { get; private set; }
+
+		/// <summary>
+		///  Write access to pool
+		/// </summary>
+		public bool Write { get; private set; }
+
+		/// <summary>
+		/// Constructor
+		/// </summary>
+		/// <param name="Pool"></param>
+		/// <param name="Read"></param>
+		/// <param name="Write"></param>
+		public DevicePoolAuthorization(IDevicePool Pool, bool Read, bool Write)
+		{
+			this.Pool = Pool;
+			this.Read = Read;
+			this.Write = Write;				
+		}
+	}
 
 	/// <summary>
 	/// Platform map required by V1 API
@@ -76,6 +111,11 @@ namespace HordeServer.Services
 		/// </summary>
 		StreamService StreamService;
 
+		/// <summary>
+		/// Singleton instance of the project service
+		/// </summary>
+		ProjectService ProjectService;
+
 		/// <summary>
 		/// Log output writer
 		/// </summary>
@@ -94,11 +134,12 @@ namespace HordeServer.Services
 		/// <summary>
 		/// Device service constructor
 		/// </summary>
-		public DeviceService(IDeviceCollection Devices, ISingletonDocument<DevicePlatformMapV1> PlatformMapSingleton, JobService JobService, StreamService StreamService, AclService AclService, INotificationService NotificationService, ILogger<DeviceService> Logger)
+		public DeviceService(IDeviceCollection Devices, ISingletonDocument<DevicePlatformMapV1> PlatformMapSingleton, JobService JobService, ProjectService ProjectService, StreamService StreamService, AclService AclService, INotificationService NotificationService, ILogger<DeviceService> Logger)
 			: base(TimeSpan.FromMinutes(1.0), Logger)
 		{
 			this.Devices = Devices;
             this.JobService = JobService;
+			this.ProjectService = ProjectService;
             this.StreamService = StreamService;
             this.AclService = AclService;
             this.NotificationService = NotificationService;
@@ -164,11 +205,20 @@ namespace HordeServer.Services
 		/// <summary>
 		/// Create a new device pool
 		/// </summary>
-		public Task<IDevicePool?> TryCreatePoolAsync(DevicePoolId Id, string Name, DevicePoolType PoolType)
+		public Task<IDevicePool?> TryCreatePoolAsync(DevicePoolId Id, string Name, DevicePoolType PoolType, List<ProjectId>? ProjectIds)
 		{
-			return Devices.TryAddPoolAsync(Id, Name, PoolType);
+			return Devices.TryAddPoolAsync(Id, Name, PoolType, ProjectIds);
 		}
 
+		/// <summary>
+		/// Update a device pool
+		/// </summary>
+		public Task UpdatePoolAsync(DevicePoolId Id, List<ProjectId>? ProjectIds)
+		{
+			return Devices.UpdatePoolAsync(Id, ProjectIds);
+		}
+
+
 		/// <summary>
 		/// Get a list of existing device pools
 		/// </summary>
@@ -368,6 +418,86 @@ namespace HordeServer.Services
 			return Task.FromResult(User.IsInRole("Internal-Employees"));			
 		}
 
+		/// <summary>
+		/// Get list of pool authorizations for a user
+		/// </summary>
+		/// <param name="User"></param>
+		/// <returns></returns>
+		async public Task<List<DevicePoolAuthorization>> GetUserPoolAuthorizationsAsync(ClaimsPrincipal User)
+		{
+			
+			List<DevicePoolAuthorization> AuthPools = new List<DevicePoolAuthorization>();
+
+			List<IDevicePool> AllPools = await GetPoolsAsync();
+
+			// Set of projects associated with device pools
+			HashSet<ProjectId> ProjectIds = new HashSet<ProjectId>(AllPools.Where(x => x.ProjectIds != null).SelectMany(x => x.ProjectIds!));
+
+			Dictionary<ProjectId, bool> DeviceRead = new Dictionary<ProjectId, bool>();
+			Dictionary<ProjectId, bool> DeviceWrite = new Dictionary<ProjectId, bool>();
+			ProjectPermissionsCache PermissionsCache = new ProjectPermissionsCache();
+
+			foreach (ProjectId ProjectId in ProjectIds)
+			{
+				DeviceRead.Add(ProjectId,  await ProjectService.AuthorizeAsync(ProjectId, AclAction.DeviceRead, User, PermissionsCache));
+				DeviceWrite.Add(ProjectId, await ProjectService.AuthorizeAsync(ProjectId, AclAction.DeviceWrite, User, PermissionsCache));
+			}
+
+			bool InternalEmployee = User.IsInRole("Internal-Employees");
+
+			foreach (IDevicePool Pool in AllPools)
+			{
+				// for global pools which aren't associated with a project
+				if (Pool.ProjectIds == null || Pool.ProjectIds.Count == 0)
+				{
+					AuthPools.Add(new DevicePoolAuthorization(Pool, true, InternalEmployee));
+					continue;
+				}
+
+				bool Read = false;
+				bool Write = false;
+
+				foreach (ProjectId ProjectId in Pool.ProjectIds)
+				{
+					if (DeviceRead.ContainsKey(ProjectId))
+					{
+						Read = true;
+					}
+
+					if (DeviceWrite.ContainsKey(ProjectId))
+					{
+						Write = true;
+					}
+				}
+
+				AuthPools.Add(new DevicePoolAuthorization(Pool, Read, Write));
+
+			}
+
+			return AuthPools;
+		}
+
+		/// <summary>
+		/// Get list of pool authorizations for a user
+		/// </summary>
+		/// <param name="Id"></param>
+		/// <param name="User"></param>
+		/// <returns></returns>
+		async public Task<DevicePoolAuthorization?> GetUserPoolAuthorizationAsync(DevicePoolId Id, ClaimsPrincipal User)
+		{
+			List<DevicePoolAuthorization> Auth = await GetUserPoolAuthorizationsAsync(User);
+			return Auth.Where(x => x.Pool.Id == Id).FirstOrDefault();
+		}
+
+		/// <summary>
+		/// Update a device
+		/// </summary>
+		async public Task UpdateDevicePoolAsync(DevicePoolId PoolId, List<ProjectId>? ProjectIds)
+		{
+			await Devices.UpdatePoolAsync(PoolId, ProjectIds);
+		}
+
+
 		/// <summary>
 		/// Get Platform mappings for V1 API
 		/// </summary>		
