Horde: Add DeviceCollection changes which were somehow missed in previous CL
#rnx

diff --git a/Engine/Source/Programs/Horde/Horde.Build/Devices/Impl/DeviceCollection.cs b/Engine/Source/Programs/Horde/Horde.Build/Devices/Impl/DeviceCollection.cs
--- a/Engine/Source/Programs/Horde/Horde.Build/Devices/Impl/DeviceCollection.cs
+++ b/Engine/Source/Programs/Horde/Horde.Build/Devices/Impl/DeviceCollection.cs
@@ -19,6 +19,7 @@ namespace HordeServer.Collections.Impl
 	using DevicePlatformId = StringId<IDevicePlatform>;
 	using DevicePoolId = StringId<IDevicePool>;
 	using UserId = ObjectId<IUser>;
+	using ProjectId = StringId<IProject>;
 
 	/// <summary>
 	/// Collection of device documents
@@ -65,7 +66,10 @@ namespace HordeServer.Collections.Impl
 			[BsonRequired]
             public DevicePoolType PoolType { get; set; }
 
-            [BsonRequired]
+			[BsonIgnoreIfNull]
+			public List<ProjectId>? ProjectIds { get; set; }
+
+			[BsonRequired]
 			public string Name { get; set; } = null!;
 
 			[BsonIgnoreIfNull]
@@ -76,11 +80,12 @@ namespace HordeServer.Collections.Impl
 			{
             }
 
-			public DevicePoolDocument(DevicePoolId Id, string Name, DevicePoolType PoolType)
+			public DevicePoolDocument(DevicePoolId Id, string Name, DevicePoolType PoolType, List<ProjectId>? ProjectIds)
 			{
 				this.Id = Id;
 				this.Name = Name;
                 this.PoolType = PoolType;
+				this.ProjectIds = ProjectIds;
             }
 
 
@@ -320,9 +325,9 @@ namespace HordeServer.Collections.Impl
 		}
 
 		/// <inheritdoc/>
-		public async Task<IDevicePool?> TryAddPoolAsync(DevicePoolId Id, string Name, DevicePoolType PoolType)
+		public async Task<IDevicePool?> TryAddPoolAsync(DevicePoolId Id, string Name, DevicePoolType PoolType, List<ProjectId>? ProjectIds)
 		{
-			DevicePoolDocument NewPool = new DevicePoolDocument(Id, Name, PoolType);
+			DevicePoolDocument NewPool = new DevicePoolDocument(Id, Name, PoolType, ProjectIds);
 
 			try
 			{
@@ -342,6 +347,25 @@ namespace HordeServer.Collections.Impl
 			}
 		}
 
+		/// <inheritdoc/>
+		public async Task UpdatePoolAsync(DevicePoolId Id, List<ProjectId>? ProjectIds)
+		{
+			UpdateDefinitionBuilder<DevicePoolDocument> UpdateBuilder = Builders<DevicePoolDocument>.Update;
+
+			List<UpdateDefinition<DevicePoolDocument>> Updates = new List<UpdateDefinition<DevicePoolDocument>>();
+
+			if (ProjectIds != null)
+			{
+				Updates.Add(UpdateBuilder.Set(x => x.ProjectIds, ProjectIds));
+			}
+
+			if (Updates.Count > 0)
+			{
+				await Pools.FindOneAndUpdateAsync<DevicePoolDocument>(x => x.Id == Id, UpdateBuilder.Combine(Updates));
+			}
+
+		}
+
 		/// <inheritdoc/>
 		public async Task<List<IDevice>> FindAllDevicesAsync(List<DeviceId>? DeviceIds = null)
 		{
