Horde: Add an experimental option to parse BuildGraph scripts on the server.

diff --git a/Engine/Source/Programs/Horde/Horde.Build/Controllers/JobsController.cs b/Engine/Source/Programs/Horde/Horde.Build/Controllers/JobsController.cs
--- a/Engine/Source/Programs/Horde/Horde.Build/Controllers/JobsController.cs
+++ b/Engine/Source/Programs/Horde/Horde.Build/Controllers/JobsController.cs
@@ -1,6 +1,7 @@
 // Copyright Epic Games, Inc. All Rights Reserved.
 
 using EpicGames.Core;
+using Horde.Build.Utilities;
 using HordeServer.Api;
 using HordeServer.Collections;
 using HordeCommon;
@@ -24,6 +25,7 @@ using HordeServer.Notifications;
 using OpenTracing.Util;
 using OpenTracing;
 using HordeServer.Jobs;
+using EpicGames.Perforce;
 
 namespace HordeServer.Controllers
 {
@@ -157,6 +159,16 @@ namespace HordeServer.Controllers
 			// New properties for the job
 			List<string> Arguments = Create.Arguments ?? Template.GetDefaultArguments();
 
+			// Special handling for internal graph evaluation
+			if (Arguments.Any(x => x.Equals("-InternalParser", StringComparison.OrdinalIgnoreCase)))
+			{
+				IPerforceConnection? PerforceConnection = await Perforce.GetServiceUserConnection(Stream.ClusterName);
+				if (PerforceConnection != null)
+				{
+					Graph = await Graphs.AddAsync(PerforceConnection, Stream, Change, CodeChange, Arguments);
+				}
+			}
+
 			// Create the job
 			IJob Job = await JobService.CreateJobAsync(null, Stream, TemplateRefId, Template.Id, Graph, Name, Change, CodeChange, Create.PreflightChange, null, User.GetUserId(), Priority, Create.AutoSubmit, Create.UpdateIssues, TemplateRef.ChainedJobs, TemplateRef.ShowUgsBadges, TemplateRef.ShowUgsAlerts, TemplateRef.NotificationChannel, TemplateRef.NotificationChannelFilter, Arguments);
 			await UpdateNotificationsAsync(Job.Id, new UpdateNotificationsRequest { Slack = true });
diff --git a/Engine/Source/Programs/Horde/Horde.Build/Horde.Build.csproj b/Engine/Source/Programs/Horde/Horde.Build/Horde.Build.csproj
--- a/Engine/Source/Programs/Horde/Horde.Build/Horde.Build.csproj
+++ b/Engine/Source/Programs/Horde/Horde.Build/Horde.Build.csproj
@@ -90,6 +90,7 @@
   </ItemGroup>
 
   <ItemGroup>
+    <ProjectReference Include="..\..\Shared\EpicGames.BuildGraph\EpicGames.BuildGraph.csproj" />
     <ProjectReference Include="..\..\Shared\EpicGames.Core\EpicGames.Core.csproj" />
     <ProjectReference Include="..\..\Shared\EpicGames.Horde\EpicGames.Horde.csproj" />
     <ProjectReference Include="..\..\Shared\EpicGames.IoHash\EpicGames.IoHash.csproj" />
diff --git a/Engine/Source/Programs/Horde/Horde.Build/Utilities/BuildGraphExtensions.cs b/Engine/Source/Programs/Horde/Horde.Build/Utilities/BuildGraphExtensions.cs
new file mode 100644
--- /dev/null
+++ b/Engine/Source/Programs/Horde/Horde.Build/Utilities/BuildGraphExtensions.cs
@@ -0,0 +1,198 @@
+// Copyright Epic Games, Inc. All Rights Reserved.
+
+using EpicGames.BuildGraph;
+using EpicGames.Perforce;
+using HordeServer.Collections;
+using HordeServer.Models;
+using System;
+using System.Collections.Generic;
+using System.Globalization;
+using System.IO;
+using System.Linq;
+using System.Threading.Tasks;
+
+namespace Horde.Build.Utilities
+{
+	/// <summary>
+	/// Extension methods to construct graphs from BuildGraph scripts in Perforce
+	/// </summary>
+	static class BuildGraphExtensions
+	{
+		class PerforceScriptContext : IBgScriptReaderContext
+		{
+			IPerforceConnection PerforceConnection;
+			string Stream;
+			int ChangeNumber;
+
+			public PerforceScriptContext(IPerforceConnection Perforce, string Stream, int ChangeNumber)
+			{
+				this.PerforceConnection = Perforce;
+				this.Stream = Stream;
+				this.ChangeNumber = ChangeNumber;
+			}
+
+			public bool Exists(string Path)
+			{
+				List<FStatRecord> Records = PerforceConnection.FStatAsync($"{Stream}/{Path}@{ChangeNumber}").Result;
+				return Records.Count > 0;
+			}
+
+			public object GetNativePath(string Path)
+			{
+				return $"{Stream}/{Path}";
+			}
+
+			public bool TryRead(string Path, out byte[]? Data)
+			{
+				string TempFileName = System.IO.Path.GetTempFileName();
+				try
+				{
+					PerforceConnection.PrintAsync(TempFileName, $"{Stream}/{Path}").Wait();
+					Data = File.ReadAllBytes(TempFileName);
+					return true;
+				}
+				catch
+				{
+					Data = null;
+					return false;
+				}
+				finally
+				{
+					try
+					{
+						File.SetAttributes(TempFileName, FileAttributes.Normal);
+						File.Delete(TempFileName);
+					}
+					catch { }
+				}
+			}
+		}
+		
+		/// <summary>
+		/// Create a graph from a script in Perforce
+		/// </summary>
+		/// <param name="GraphCollection"></param>
+		/// <param name="PerforceConnection"></param>
+		/// <param name="Stream"></param>
+		/// <param name="ChangeNumber"></param>
+		/// <param name="CodeChangeNumber"></param>
+		/// <param name="Arguments"></param>
+		/// <returns></returns>
+		/// <exception cref="NotImplementedException"></exception>
+		public static async Task<IGraph> AddAsync(this IGraphCollection GraphCollection, IPerforceConnection PerforceConnection, IStream Stream, int ChangeNumber, int CodeChangeNumber, List<string> Arguments)
+		{
+			PerforceScriptContext Context = new PerforceScriptContext(PerforceConnection, Stream.Name, ChangeNumber);
+
+			BgScriptSchema? Schema = null!;
+
+			Dictionary<string, string> Properties = new Dictionary<string, string>();
+			Properties["HostPlatform"] = "";
+			Properties["RootDir"] = "";
+			Properties["IsBuildMachine"] = "true";
+			Properties["Branch"] = Stream.Name;
+			Properties["Change"] = ChangeNumber.ToString(CultureInfo.InvariantCulture);
+			Properties["EscapedBranch"] = Stream.Name.Replace('/', '+');
+			Properties["CodeChange"] = CodeChangeNumber.ToString(CultureInfo.InvariantCulture);
+
+			string? ScriptFile = null;
+			foreach (string Argument in Arguments)
+			{
+				const string ScriptPrefix = "-Script=";
+				if (Argument.StartsWith(ScriptPrefix, StringComparison.OrdinalIgnoreCase))
+				{
+					ScriptFile = Argument.Substring(ScriptPrefix.Length);
+				}
+			}
+
+			List<string> Targets = new List<string>();
+			foreach (string Argument in Arguments)
+			{
+				const string TargetPrefix = "-Target=";
+				if (Argument.StartsWith(TargetPrefix, StringComparison.OrdinalIgnoreCase))
+				{
+					Targets.AddRange(Argument.Substring(TargetPrefix.Length).Split(';'));
+				}
+			}
+
+			Dictionary<string, string> Options = new Dictionary<string, string>();
+			foreach (string Argument in Arguments)
+			{
+				const string SetPrefix = "-Set:";
+				if (Argument.StartsWith(SetPrefix, StringComparison.OrdinalIgnoreCase))
+				{
+					int EqualsIdx = Argument.IndexOf('=', StringComparison.Ordinal);
+					if (EqualsIdx != -1)
+					{
+						Options[Argument.Substring(SetPrefix.Length, EqualsIdx - SetPrefix.Length)] = Argument.Substring(EqualsIdx + 1);
+					}
+				}
+			}
+
+			BgScript Script;
+			if (!BgScriptReader.TryRead(Context, ScriptFile, Options, Properties, false, Schema, EpicGames.Core.Log.Logger, out Script))
+			{
+				throw new NotImplementedException();
+			}
+
+			HashSet<BgNode> TargetNodes = new HashSet<BgNode>();
+			foreach (string Target in Targets)
+			{
+				if (Script.NameToNode.TryGetValue(Target, out BgNode? Node))
+				{
+					TargetNodes.Add(Node);
+				}
+				else if (Script.NameToAgent.TryGetValue(Target, out BgAgent? Agent))
+				{
+					TargetNodes.UnionWith(Agent.Nodes);
+				}
+				else if (Script.NameToAggregate.TryGetValue(Target, out BgAggregate? Aggregate))
+				{
+					TargetNodes.UnionWith(Aggregate.RequiredNodes);
+				}
+				else
+				{
+					throw new NotImplementedException();
+				}
+			}
+			Script.Select(TargetNodes);
+
+			List<NewGroup> NewGroups = new List<NewGroup>();
+			foreach (BgAgent Agent in Script.Agents)
+			{
+				List<NewNode> NewNodes = new List<NewNode>();
+				foreach (BgNode Node in Agent.Nodes)
+				{
+					const bool bAllowRetry = true;
+					NewNode NewNode = new NewNode(Node.Name, Node.InputDependencies.Select(x => x.Name).ToList(), Node.OrderDependencies.Select(x => x.Name).ToList(), HordeCommon.Priority.Normal, bAllowRetry, Node.bRunEarly, Node.bNotifyOnWarnings, new Dictionary<string, string>(), new Dictionary<string, string>());
+					NewNodes.Add(NewNode);
+				}
+
+				string AgentType = Agent.PossibleTypes.FirstOrDefault(x => Stream.AgentTypes.ContainsKey(x)) ?? Agent.PossibleTypes.FirstOrDefault() ?? "Unknown";
+				NewGroups.Add(new NewGroup(AgentType, NewNodes));
+			}
+
+			List<NewAggregate> NewAggregates = new List<NewAggregate>();
+			foreach (BgAggregate Aggregate in Script.NameToAggregate.Values)
+			{
+				NewAggregate NewAggregate = new NewAggregate(Aggregate.Name, Aggregate.RequiredNodes.Select(x => x.Name).ToList());
+				NewAggregates.Add(NewAggregate);
+			}
+
+			List<NewLabel> NewLabels = new List<NewLabel>();
+			foreach (BgLabel Label in Script.Labels)
+			{
+				NewLabel NewLabel = new NewLabel();
+				NewLabel.DashboardName = String.IsNullOrEmpty(Label.DashboardName) ? null : Label.DashboardName;
+				NewLabel.DashboardCategory = String.IsNullOrEmpty(Label.DashboardCategory) ? null : Label.DashboardCategory;
+				NewLabel.UgsName = String.IsNullOrEmpty(Label.UgsBadge) ? null : Label.UgsBadge;
+				NewLabel.UgsProject = String.IsNullOrEmpty(Label.UgsProject) ? null : Label.UgsProject;
+				NewLabel.Change = (Label.Change == BgLabelChange.Code) ? HordeCommon.LabelChange.Code : HordeCommon.LabelChange.Current;
+				NewLabel.RequiredNodes = Label.RequiredNodes.Select(x => x.Name).ToList();
+				NewLabel.IncludedNodes = Label.IncludedNodes.Select(x => x.Name).ToList();
+				NewLabels.Add(NewLabel);
+			}
+
+			return await GraphCollection.AppendAsync(null, NewGroups, NewAggregates, NewLabels);
+		}
+	}
+}
