Horde: Frontend for device pool auth
#rnx

diff --git a/Engine/Source/Programs/Horde/Horde.Build/Devices/Impl/DeviceService.cs b/Engine/Source/Programs/Horde/Horde.Build/Devices/Impl/DeviceService.cs
--- a/Engine/Source/Programs/Horde/Horde.Build/Devices/Impl/DeviceService.cs
+++ b/Engine/Source/Programs/Horde/Horde.Build/Devices/Impl/DeviceService.cs
@@ -429,6 +429,7 @@ namespace HordeServer.Services
 			List<DevicePoolAuthorization> AuthPools = new List<DevicePoolAuthorization>();
 
 			List<IDevicePool> AllPools = await GetPoolsAsync();
+			List<IProject> Projects = await ProjectService.GetProjectsAsync();
 
 			// Set of projects associated with device pools
 			HashSet<ProjectId> ProjectIds = new HashSet<ProjectId>(AllPools.Where(x => x.ProjectIds != null).SelectMany(x => x.ProjectIds!));
@@ -439,6 +440,13 @@ namespace HordeServer.Services
 
 			foreach (ProjectId ProjectId in ProjectIds)
 			{
+
+				if (Projects.Where(x => x.Id == ProjectId).FirstOrDefault() == null)
+				{
+					Logger.LogWarning($"Device pool authorization references missing project id {ProjectId}");
+					continue;
+				}
+
 				DeviceRead.Add(ProjectId,  await ProjectService.AuthorizeAsync(ProjectId, AclAction.DeviceRead, User, PermissionsCache));
 				DeviceWrite.Add(ProjectId, await ProjectService.AuthorizeAsync(ProjectId, AclAction.DeviceWrite, User, PermissionsCache));
 			}
diff --git a/Engine/Source/Programs/Horde/HordeDashboard/src/backend/Api.ts b/Engine/Source/Programs/Horde/HordeDashboard/src/backend/Api.ts
--- a/Engine/Source/Programs/Horde/HordeDashboard/src/backend/Api.ts
+++ b/Engine/Source/Programs/Horde/HordeDashboard/src/backend/Api.ts
@@ -2850,6 +2850,9 @@ export type GetDevicePoolResponse = {
 	/**  Name of the device pool */
 	name: string;
 
+	/** Whether there is write access to the pool */
+	writeAccess: boolean;
+
 }
 
 
