Horde: Allow retrieving a server cluster with the default name.

diff --git a/Engine/Source/Programs/Horde/Horde.Build/Perforce/Impl/PerforceService.cs b/Engine/Source/Programs/Horde/Horde.Build/Perforce/Impl/PerforceService.cs
--- a/Engine/Source/Programs/Horde/Horde.Build/Perforce/Impl/PerforceService.cs
+++ b/Engine/Source/Programs/Horde/Horde.Build/Perforce/Impl/PerforceService.cs
@@ -210,35 +210,32 @@ namespace HordeServer.Services
 			using IScope Scope = GlobalTracer.Instance.BuildSpan("PerforceService.GetServiceUserConnection").StartActive();
 			Scope.Span.SetTag("ClusterName", ClusterName);
 			
-			if(ClusterName != null)
+			Globals Globals = await CachedGlobals.GetCached();
+			PerforceCluster? Cluster = Globals.FindPerforceCluster(ClusterName);
+			if(Cluster != null)
 			{
-				Globals Globals = await CachedGlobals.GetCached();
-				PerforceCluster? Cluster = Globals.FindPerforceCluster(ClusterName);
-				if(Cluster != null)
+				string? UserName = null;
+				string? Password = null;
+				if (Cluster.ServiceAccount != null)
 				{
-					string? UserName = null;
-					string? Password = null;
-					if (Cluster.ServiceAccount != null)
+					PerforceCredentials? Credentials = Cluster.Credentials.FirstOrDefault(x => x.UserName.Equals(Cluster.ServiceAccount, StringComparison.OrdinalIgnoreCase));
+					if (Credentials == null)
 					{
-						PerforceCredentials? Credentials = Cluster.Credentials.FirstOrDefault(x => x.UserName.Equals(Cluster.ServiceAccount, StringComparison.OrdinalIgnoreCase));
-						if (Credentials == null)
-						{
-							throw new Exception($"No credentials defined for {Cluster.ServiceAccount} on {Cluster.Name}");
-						}
+						throw new Exception($"No credentials defined for {Cluster.ServiceAccount} on {Cluster.Name}");
 					}
+				}
 
-					PerforceSettings Settings = new PerforceSettings();
-					Settings.User = UserName;
-					Settings.Client = "__DOES_NOT_EXIST__";
+				PerforceSettings Settings = new PerforceSettings();
+				Settings.User = UserName;
+				Settings.Client = "__DOES_NOT_EXIST__";
 
-					NativePerforceConnection NativeConnection = new NativePerforceConnection(Logger);
-					await NativeConnection.ConnectAsync(Settings);
-					if (Password != null)
-					{
-						await NativeConnection.LoginAsync(Password);
-					}
-					return NativeConnection;
+				NativePerforceConnection NativeConnection = new NativePerforceConnection(Logger);
+				await NativeConnection.ConnectAsync(Settings);
+				if (Password != null)
+				{
+					await NativeConnection.LoginAsync(Password);
 				}
+				return NativeConnection;
 			}
 			return null;
 		}
