Horde: Fix missing server setting on native P4 connection.

diff --git a/Engine/Source/Programs/Horde/Horde.Build/Perforce/Impl/PerforceService.cs b/Engine/Source/Programs/Horde/Horde.Build/Perforce/Impl/PerforceService.cs
--- a/Engine/Source/Programs/Horde/Horde.Build/Perforce/Impl/PerforceService.cs
+++ b/Engine/Source/Programs/Horde/Horde.Build/Perforce/Impl/PerforceService.cs
@@ -137,7 +137,7 @@ namespace HordeServer.Services
 			return User!;
 		}
 
-		async Task<PerforceCluster> GetClusterAsync(string ClusterName)
+		async Task<PerforceCluster> GetClusterAsync(string? ClusterName)
 		{
 			Globals Globals = await CachedGlobals.GetCached();
 
@@ -209,35 +209,38 @@ namespace HordeServer.Services
 		{
 			using IScope Scope = GlobalTracer.Instance.BuildSpan("PerforceService.GetServiceUserConnection").StartActive();
 			Scope.Span.SetTag("ClusterName", ClusterName);
-			
-			Globals Globals = await CachedGlobals.GetCached();
-			PerforceCluster? Cluster = Globals.FindPerforceCluster(ClusterName);
-			if(Cluster != null)
+
+			PerforceCluster Cluster = await GetClusterAsync(ClusterName);
+
+			string? UserName = null;
+			string? Password = null;
+			if (Cluster.ServiceAccount != null)
 			{
-				string? UserName = null;
-				string? Password = null;
-				if (Cluster.ServiceAccount != null)
+				PerforceCredentials? Credentials = Cluster.Credentials.FirstOrDefault(x => x.UserName.Equals(Cluster.ServiceAccount, StringComparison.OrdinalIgnoreCase));
+				if (Credentials == null)
 				{
-					PerforceCredentials? Credentials = Cluster.Credentials.FirstOrDefault(x => x.UserName.Equals(Cluster.ServiceAccount, StringComparison.OrdinalIgnoreCase));
-					if (Credentials == null)
-					{
-						throw new Exception($"No credentials defined for {Cluster.ServiceAccount} on {Cluster.Name}");
-					}
+					throw new Exception($"No credentials defined for {Cluster.ServiceAccount} on {Cluster.Name}");
 				}
 
-				PerforceSettings Settings = new PerforceSettings();
-				Settings.User = UserName;
-				Settings.Client = "__DOES_NOT_EXIST__";
+				UserName = Credentials.UserName;
+				Password = Credentials.Password;
+			}
 
-				NativePerforceConnection NativeConnection = new NativePerforceConnection(Logger);
-				await NativeConnection.ConnectAsync(Settings);
-				if (Password != null)
-				{
-					await NativeConnection.LoginAsync(Password);
-				}
-				return NativeConnection;
+			IPerforceServer Server = await SelectServer(Cluster);
+
+			PerforceSettings Settings = new PerforceSettings();
+			Settings.Port = Server.ServerAndPort;
+			Settings.User = UserName;
+			Settings.AppName = "Horde.Build";
+			Settings.Client = "__DOES_NOT_EXIST__";
+
+			NativePerforceConnection NativeConnection = new NativePerforceConnection(Logger);
+			await NativeConnection.ConnectAsync(Settings);
+			if (Password != null)
+			{
+				await NativeConnection.LoginAsync(Password);
 			}
-			return null;
+			return NativeConnection;
 		}
 
 		async Task<P4.Repository> GetServiceUserConnection(PerforceCluster Cluster)
