Horde: Fix for DeviceService access
#rnx

diff --git a/Engine/Source/Programs/Horde/Horde.Build/Devices/Impl/DeviceService.cs b/Engine/Source/Programs/Horde/Horde.Build/Devices/Impl/DeviceService.cs
--- a/Engine/Source/Programs/Horde/Horde.Build/Devices/Impl/DeviceService.cs
+++ b/Engine/Source/Programs/Horde/Horde.Build/Devices/Impl/DeviceService.cs
@@ -467,14 +467,16 @@ namespace HordeServer.Services
 
 				foreach (ProjectId ProjectId in Pool.ProjectIds)
 				{
-					if (DeviceRead.ContainsKey(ProjectId))
+					bool Value;
+
+					if (!Read && DeviceRead.TryGetValue(ProjectId, out Value)
 					{
-						Read = true;
+						Read = Value;
 					}
 
-					if (DeviceWrite.ContainsKey(ProjectId))
+					if (!Write && DeviceWrite.TryGetValue(ProjectId, out Value))
 					{
-						Write = true;
+						Write = Value;
 					}
 				}
 
