Horde: Fix Horde.Build compile
#rnx

diff --git a/Engine/Source/Programs/Horde/Horde.Build/Devices/Impl/DeviceService.cs b/Engine/Source/Programs/Horde/Horde.Build/Devices/Impl/DeviceService.cs
--- a/Engine/Source/Programs/Horde/Horde.Build/Devices/Impl/DeviceService.cs
+++ b/Engine/Source/Programs/Horde/Horde.Build/Devices/Impl/DeviceService.cs
@@ -469,7 +469,7 @@ namespace HordeServer.Services
 				{
 					bool Value;
 
-					if (!Read && DeviceRead.TryGetValue(ProjectId, out Value)
+					if (!Read && DeviceRead.TryGetValue(ProjectId, out Value))
 					{
 						Read = Value;
 					}
