Horde: Fix password not being correctly forwarded to P4 API.

diff --git a/Engine/Source/Programs/Horde/Horde.Build/Perforce/Impl/PerforceService.cs b/Engine/Source/Programs/Horde/Horde.Build/Perforce/Impl/PerforceService.cs
--- a/Engine/Source/Programs/Horde/Horde.Build/Perforce/Impl/PerforceService.cs
+++ b/Engine/Source/Programs/Horde/Horde.Build/Perforce/Impl/PerforceService.cs
@@ -212,34 +212,30 @@ namespace HordeServer.Services
 
 			PerforceCluster Cluster = await GetClusterAsync(ClusterName);
 
-			string? UserName = null;
+			string? UserName = Cluster.ServiceAccount ?? Environment.UserName;
+
 			string? Password = null;
 			if (Cluster.ServiceAccount != null)
 			{
-				PerforceCredentials? Credentials = Cluster.Credentials.FirstOrDefault(x => x.UserName.Equals(Cluster.ServiceAccount, StringComparison.OrdinalIgnoreCase));
+				PerforceCredentials? Credentials = Cluster.Credentials.FirstOrDefault(x => x.UserName.Equals(UserName, StringComparison.OrdinalIgnoreCase));
 				if (Credentials == null)
 				{
 					throw new Exception($"No credentials defined for {Cluster.ServiceAccount} on {Cluster.Name}");
 				}
-
-				UserName = Credentials.UserName;
 				Password = Credentials.Password;
 			}
 
 			IPerforceServer Server = await SelectServer(Cluster);
 
 			PerforceSettings Settings = new PerforceSettings();
-			Settings.Port = Server.ServerAndPort;
+			Settings.ServerAndPort = Server.ServerAndPort;
 			Settings.User = UserName;
+			Settings.Password = Password;
 			Settings.AppName = "Horde.Build";
 			Settings.Client = "__DOES_NOT_EXIST__";
 
 			NativePerforceConnection NativeConnection = new NativePerforceConnection(Logger);
 			await NativeConnection.ConnectAsync(Settings);
-			if (Password != null)
-			{
-				await NativeConnection.LoginAsync(Password);
-			}
 			return NativeConnection;
 		}
 
@@ -271,6 +267,8 @@ namespace HordeServer.Services
 				{
 					throw new Exception($"No credentials defined for {Cluster.ServiceAccount} on {Cluster.Name}");
 				}
+				UserName = Credentials.UserName;
+				Password = Credentials.Password;
 			}
 			return CreateConnection(Server, UserName, Password);
 		}
