Horde: Upgrade server to NET 6.0.

diff --git a/Engine/Source/Programs/Horde/Horde.Build.Tests/AgentServiceTest.cs b/Engine/Source/Programs/Horde/Horde.Build.Tests/AgentServiceTest.cs
--- a/Engine/Source/Programs/Horde/Horde.Build.Tests/AgentServiceTest.cs
+++ b/Engine/Source/Programs/Horde/Horde.Build.Tests/AgentServiceTest.cs
@@ -28,12 +28,12 @@ namespace HordeServerTests
 	        TestSetup TestSetup = await GetTestSetup();
 
 	        ActionResult<List<object>> Res = await TestSetup.JobsController.FindJobsAsync();
-	        Assert.AreEqual(2, Res.Value.Count);
+	        Assert.AreEqual(2, Res.Value!.Count);
 	        Assert.AreEqual("hello2", (Res.Value[0] as GetJobResponse)!.Name);
 	        Assert.AreEqual("hello1", (Res.Value[1] as GetJobResponse)!.Name);
 	        
 	        Res = await TestSetup.JobsController.FindJobsAsync(IncludePreflight: false);
-	        Assert.AreEqual(1, Res.Value.Count);
+	        Assert.AreEqual(1, Res.Value!.Count);
 	        Assert.AreEqual("hello2", (Res.Value[0] as GetJobResponse)!.Name);
         }
         
diff --git a/Engine/Source/Programs/Horde/Horde.Build.Tests/Horde.Build.Tests.csproj b/Engine/Source/Programs/Horde/Horde.Build.Tests/Horde.Build.Tests.csproj
--- a/Engine/Source/Programs/Horde/Horde.Build.Tests/Horde.Build.Tests.csproj
+++ b/Engine/Source/Programs/Horde/Horde.Build.Tests/Horde.Build.Tests.csproj
@@ -1,7 +1,7 @@
 <Project Sdk="Microsoft.NET.Sdk">
 
   <PropertyGroup>
-    <TargetFramework>net5.0</TargetFramework>
+    <TargetFramework>net6.0</TargetFramework>
     <Nullable>enable</Nullable>
     <IsPackable>false</IsPackable>
   </PropertyGroup>
diff --git a/Engine/Source/Programs/Horde/Horde.Build.Tests/JobsControllerDbTest.cs b/Engine/Source/Programs/Horde/Horde.Build.Tests/JobsControllerDbTest.cs
--- a/Engine/Source/Programs/Horde/Horde.Build.Tests/JobsControllerDbTest.cs
+++ b/Engine/Source/Programs/Horde/Horde.Build.Tests/JobsControllerDbTest.cs
@@ -27,12 +27,12 @@ namespace HordeServerTests
 	        TestSetup TestSetup = await GetTestSetup();
 
 	        ActionResult<List<object>> Res = await TestSetup.JobsController.FindJobsAsync();
-	        Assert.AreEqual(2, Res.Value.Count);
+	        Assert.AreEqual(2, Res.Value!.Count);
 	        Assert.AreEqual("hello2", (Res.Value[0] as GetJobResponse)!.Name);
 	        Assert.AreEqual("hello1", (Res.Value[1] as GetJobResponse)!.Name);
 	        
 	        Res = await TestSetup.JobsController.FindJobsAsync(IncludePreflight: false);
-	        Assert.AreEqual(1, Res.Value.Count);
+	        Assert.AreEqual(1, Res.Value!.Count);
 	        Assert.AreEqual("hello2", (Res.Value[0] as GetJobResponse)!.Name);
         }
         
@@ -46,16 +46,16 @@ namespace HordeServerTests
 	        SubResourceId BatchId = Job.Batches[0].Id;
 	        SubResourceId StepId = Job.Batches[0].Steps[0].Id;
 
-	        object Obj = (await JobsController.GetStepAsync(Job.Id, BatchId.ToString(), StepId.ToString())).Value;
+	        object Obj = (await JobsController.GetStepAsync(Job.Id, BatchId.ToString(), StepId.ToString())).Value!;
 	        GetStepResponse StepRes = (Obj as GetStepResponse)!;
 	        Assert.IsFalse(StepRes.AbortRequested);
 	        
 	        UpdateStepRequest UpdateReq = new UpdateStepRequest();
 	        UpdateReq.AbortRequested = true;
-	        Obj = (await JobsController.UpdateStepAsync(Job.Id, BatchId.ToString(), StepId.ToString(), UpdateReq)).Value;
+	        Obj = (await JobsController.UpdateStepAsync(Job.Id, BatchId.ToString(), StepId.ToString(), UpdateReq)).Value!;
 	        UpdateStepResponse UpdateRes = (Obj as UpdateStepResponse)!;
 	        
-	        Obj = (await JobsController.GetStepAsync(Job.Id, BatchId.ToString(), StepId.ToString())).Value;
+	        Obj = (await JobsController.GetStepAsync(Job.Id, BatchId.ToString(), StepId.ToString())).Value!;
 	        StepRes = (Obj as GetStepResponse)!;
 	        Assert.IsTrue(StepRes.AbortRequested);
 //	        Assert.AreEqual("Anonymous", StepRes.AbortByUser);
diff --git a/Engine/Source/Programs/Horde/Horde.Build/Agents/ITaskSource.cs b/Engine/Source/Programs/Horde/Horde.Build/Agents/ITaskSource.cs
--- a/Engine/Source/Programs/Horde/Horde.Build/Agents/ITaskSource.cs
+++ b/Engine/Source/Programs/Horde/Horde.Build/Agents/ITaskSource.cs
@@ -11,6 +11,7 @@ using Microsoft.Extensions.Logging;
 using MongoDB.Bson;
 using System;
 using System.Collections.Generic;
+using System.Globalization;
 using System.Linq;
 using System.Linq.Expressions;
 using System.Reflection;
@@ -106,7 +107,7 @@ namespace HordeServer.Tasks
 			/// <param name="Accessor">Accessor for the property</param>
 			public PropertyList Add(string Name, Func<TMessage, object> Accessor)
 			{
-				FormatString.Append($", {Name}={{{Name}}}");
+				FormatString.Append(CultureInfo.InvariantCulture, $", {Name}={{{Name}}}");
 				Accessors.Add(Accessor);
 				return this;
 			}
diff --git a/Engine/Source/Programs/Horde/Horde.Build/Auditing/Impl/AuditLog.cs b/Engine/Source/Programs/Horde/Horde.Build/Auditing/Impl/AuditLog.cs
--- a/Engine/Source/Programs/Horde/Horde.Build/Auditing/Impl/AuditLog.cs
+++ b/Engine/Source/Programs/Horde/Horde.Build/Auditing/Impl/AuditLog.cs
@@ -58,6 +58,11 @@ namespace HordeServer.Collections.Impl
 
 		class AuditLogChannel : IAuditLogChannel<TSubject>
 		{
+			sealed class Scope : IDisposable
+			{
+				public void Dispose() { }
+			}
+
 			public readonly AuditLog<TSubject> Outer;
 			public TSubject Subject { get; }
 
@@ -67,7 +72,7 @@ namespace HordeServer.Collections.Impl
 				this.Subject = Subject;
 			}
 
-			public IDisposable? BeginScope<TState>(TState state) => null;
+			public IDisposable BeginScope<TState>(TState state) => new Scope();
 
 			public bool IsEnabled(LogLevel logLevel) => true;
 
diff --git a/Engine/Source/Programs/Horde/Horde.Build/Controllers/AccountController.cs b/Engine/Source/Programs/Horde/Horde.Build/Controllers/AccountController.cs
--- a/Engine/Source/Programs/Horde/Horde.Build/Controllers/AccountController.cs
+++ b/Engine/Source/Programs/Horde/Horde.Build/Controllers/AccountController.cs
@@ -10,6 +10,7 @@ using Microsoft.AspNetCore.Authorization;
 using Microsoft.AspNetCore.Mvc;
 using System;
 using System.Collections.Generic;
+using System.Globalization;
 using System.Linq;
 using System.Net;
 using System.Text;
@@ -62,7 +63,7 @@ namespace HordeServer.Controllers
 			Content.Append($"<html><style>{StyleSheet}</style><h1>Horde Server</h1>");
 			if (User.Identity?.IsAuthenticated ?? false)
 			{
-				Content.Append($"<p>User <b>{User.Identity?.Name}</b> is logged in. <a href=\"/account/logout\">Log out</a></p>");
+				Content.Append(CultureInfo.InvariantCulture, $"<p>User <b>{User.Identity?.Name}</b> is logged in. <a href=\"/account/logout\">Log out</a></p>");
 				if (await AclService.AuthorizeAsync(AclAction.AdminWrite, User))
 				{
 					Content.Append("<p>");
@@ -74,16 +75,16 @@ namespace HordeServer.Controllers
 					Content.Append("<a href=\"/api/v1/admin/chainedjobtoken\">Get chained job token</a><br/>");
 					Content.Append("</p>");
 				}
-				Content.Append($"<p>Claims for {User.Identity?.Name}:");
+				Content.Append(CultureInfo.InvariantCulture, $"<p>Claims for {User.Identity?.Name}:");
 				Content.Append("<table>");
 				foreach (System.Security.Claims.Claim Claim in User.Claims)
 				{
-					Content.Append($"<tr><td>{Claim.Type}</td><td>{Claim.Value}</td></tr>");
+					Content.Append(CultureInfo.InvariantCulture, $"<tr><td>{Claim.Type}</td><td>{Claim.Value}</td></tr>");
 				}
 				Content.Append("</table>");
 				Content.Append("</p>");
 
-				Content.Append($"<p>Built from Perforce</p>");
+				Content.Append(CultureInfo.InvariantCulture, $"<p>Built from Perforce</p>");
 			}
 			else
 			{
diff --git a/Engine/Source/Programs/Horde/Horde.Build/Horde.Build.csproj b/Engine/Source/Programs/Horde/Horde.Build/Horde.Build.csproj
--- a/Engine/Source/Programs/Horde/Horde.Build/Horde.Build.csproj
+++ b/Engine/Source/Programs/Horde/Horde.Build/Horde.Build.csproj
@@ -1,7 +1,7 @@
 <Project Sdk="Microsoft.NET.Sdk.Web">
 
   <PropertyGroup>
-    <TargetFramework>net5.0</TargetFramework>
+    <TargetFramework>net6.0</TargetFramework>
     <Nullable>enable</Nullable>
     <Deterministic>true</Deterministic>
     <AssemblyName>Horde.Build</AssemblyName>
diff --git a/Engine/Source/Programs/Horde/Horde.Build/Jobs/Impl/JobTaskSource.cs b/Engine/Source/Programs/Horde/Horde.Build/Jobs/Impl/JobTaskSource.cs
--- a/Engine/Source/Programs/Horde/Horde.Build/Jobs/Impl/JobTaskSource.cs
+++ b/Engine/Source/Programs/Horde/Horde.Build/Jobs/Impl/JobTaskSource.cs
@@ -29,6 +29,7 @@ using System.Threading;
 using System.Threading.Tasks;
 using HordeServer.Services;
 using EpicGames.Horde.Compute;
+using System.Globalization;
 
 namespace HordeServer.Tasks.Impl
 {
@@ -769,7 +770,7 @@ namespace HordeServer.Tasks.Impl
 				{
 					LeaseName.Append((Job.Change > 0) ? $"CL {Job.Change}" : "Latest CL");
 				}
-				LeaseName.Append($" - {Job.Name}");
+				LeaseName.Append(CultureInfo.InvariantCulture, $" - {Job.Name}");
 
 				// Get the global settings
 				Globals Globals = await DatabaseService.GetGlobalsAsync();
@@ -836,7 +837,7 @@ namespace HordeServer.Tasks.Impl
 			{
 				LeaseName.Append((Job.Change > 0) ? $"CL {Job.Change}" : "Latest CL");
 			}
-			LeaseName.Append($" - {Job.Name}");
+			LeaseName.Append(CultureInfo.InvariantCulture, $" - {Job.Name}");
 
 			// Get the global settings
 			Globals Globals = await DatabaseService.GetGlobalsAsync();
diff --git a/Engine/Source/Programs/Horde/Horde.Build/Models/Globals.cs b/Engine/Source/Programs/Horde/Horde.Build/Models/Globals.cs
--- a/Engine/Source/Programs/Horde/Horde.Build/Models/Globals.cs
+++ b/Engine/Source/Programs/Horde/Horde.Build/Models/Globals.cs
@@ -363,11 +363,7 @@ namespace HordeServer.Models
 		/// </summary>
 		public void RotateSigningKey()
 		{
-			using (RNGCryptoServiceProvider Generator = new RNGCryptoServiceProvider())
-			{
-				JwtSigningKey = new byte[128];
-				Generator.GetBytes(JwtSigningKey);
-			}
+			JwtSigningKey = RandomNumberGenerator.GetBytes(128);
 		}
 
 		/// <summary>
diff --git a/Engine/Source/Programs/Horde/Horde.Build/Notifications/Impl/SlackNotificationSink.cs b/Engine/Source/Programs/Horde/Horde.Build/Notifications/Impl/SlackNotificationSink.cs
--- a/Engine/Source/Programs/Horde/Horde.Build/Notifications/Impl/SlackNotificationSink.cs
+++ b/Engine/Source/Programs/Horde/Horde.Build/Notifications/Impl/SlackNotificationSink.cs
@@ -649,7 +649,7 @@ namespace HordeServer.Notifications.Impl
 							{
 								SampleText.Append("\n\n");
 							}
-							SampleText.Append($"```{Data.Message}```");
+							SampleText.Append(CultureInfo.InvariantCulture, $"```{Data.Message}```");
 						}
 					}
 				}
@@ -1187,7 +1187,7 @@ namespace HordeServer.Notifications.Impl
 				Message.Attachments.AddRange(Attachments);
 			}
 
-			string RequestDigest = ContentHash.MD5(JsonSerializer.Serialize(Message, new JsonSerializerOptions { IgnoreNullValues = true })).ToString();
+			string RequestDigest = ContentHash.MD5(JsonSerializer.Serialize(Message, new JsonSerializerOptions { DefaultIgnoreCondition = JsonIgnoreCondition.WhenWritingNull })).ToString();
 
 			(MessageStateDocument State, bool IsNew) = await AddOrUpdateMessageStateAsync(Recipient, EventId, UserId, RequestDigest);
 			if (IsNew)
@@ -1218,7 +1218,7 @@ namespace HordeServer.Notifications.Impl
 			{
 				using (HttpRequestMessage SendMessageRequest = new HttpRequestMessage(HttpMethod.Post, RequestUrl))
 				{
-					string RequestJson = JsonSerializer.Serialize(Request, new JsonSerializerOptions { IgnoreNullValues = true });
+					string RequestJson = JsonSerializer.Serialize(Request, new JsonSerializerOptions { DefaultIgnoreCondition = JsonIgnoreCondition.WhenWritingNull });
 					using (StringContent MessageContent = new StringContent(RequestJson, Encoding.UTF8, "application/json"))
 					{
 						SendMessageRequest.Content = MessageContent;
diff --git a/Engine/Source/Programs/Horde/Horde.Build/Notifications/Impl/SubscriptionCollection.cs b/Engine/Source/Programs/Horde/Horde.Build/Notifications/Impl/SubscriptionCollection.cs
--- a/Engine/Source/Programs/Horde/Horde.Build/Notifications/Impl/SubscriptionCollection.cs
+++ b/Engine/Source/Programs/Horde/Horde.Build/Notifications/Impl/SubscriptionCollection.cs
@@ -12,6 +12,7 @@ using MongoDB.Driver;
 using Polly;
 using System;
 using System.Collections.Generic;
+using System.Globalization;
 using System.Linq;
 using System.Reflection.Metadata.Ecma335;
 using System.Text;
@@ -115,12 +116,12 @@ namespace HordeServer.Collections.Impl
 			public override string ToString()
 			{
 				StringBuilder Result = new StringBuilder();
-				Result.Append($"stream={StreamId}, template={TemplateId}");
+				Result.Append(CultureInfo.InvariantCulture, $"stream={StreamId}, template={TemplateId}");
 				if(CategoryName != null)
 				{
-					Result.Append($", category={CategoryName}");
+					Result.Append(CultureInfo.InvariantCulture, $", category={CategoryName}");
 				}
-				Result.Append($", label={LabelName}, outcome={Outcome}");
+				Result.Append(CultureInfo.InvariantCulture, $", label={LabelName}, outcome={Outcome}");
 				return Result.ToString();
 			}
 		}
diff --git a/Engine/Source/Programs/Horde/Horde.Build/Perforce/Impl/PerforceService.cs b/Engine/Source/Programs/Horde/Horde.Build/Perforce/Impl/PerforceService.cs
--- a/Engine/Source/Programs/Horde/Horde.Build/Perforce/Impl/PerforceService.cs
+++ b/Engine/Source/Programs/Horde/Horde.Build/Perforce/Impl/PerforceService.cs
@@ -1597,15 +1597,15 @@ namespace HordeServer.Services
 			StringBuilder Filter = new StringBuilder(BasePath);
 			if (MinChange != null && MaxChange != null)
 			{
-				Filter.Append($"@{MinChange},{MaxChange}");
+				Filter.Append(CultureInfo.InvariantCulture, $"@{MinChange},{MaxChange}");
 			}
 			else if (MinChange != null)
 			{
-				Filter.Append($"@>={MinChange}");
+				Filter.Append(CultureInfo.InvariantCulture, $"@>={MinChange}");
 			}
 			else if (MaxChange != null)
 			{
-				Filter.Append($"@<={MaxChange}");
+				Filter.Append(CultureInfo.InvariantCulture, $"@<={MaxChange}");
 			}
 			return Filter.ToString();
 		}
diff --git a/Engine/Source/Programs/Horde/Horde.Build/Server/ConfigService.cs b/Engine/Source/Programs/Horde/Horde.Build/Server/ConfigService.cs
--- a/Engine/Source/Programs/Horde/Horde.Build/Server/ConfigService.cs
+++ b/Engine/Source/Programs/Horde/Horde.Build/Server/ConfigService.cs
@@ -27,6 +27,7 @@ using System.Reflection;
 using PoolId = HordeServer.Utilities.StringId<HordeServer.Models.IPool>;
 using System.Globalization;
 using HordeServer.Storage;
+using System.Text.Json.Serialization;
 
 namespace HordeServer.Services
 {
@@ -378,7 +379,7 @@ namespace HordeServer.Services
 
 		static string GetMimeTypeFromPath(Uri Path)
 		{
-			string ContentType;
+			string? ContentType;
 			if (!ContentTypeProvider.TryGetContentType(Path.AbsolutePath, out ContentType))
 			{
 				ContentType = "application/octet-stream";
@@ -813,7 +814,7 @@ namespace HordeServer.Services
 				// This will trigger a setting update as the user config json is set to reload on change
 				try
 				{
-					await FileReference.WriteAllBytesAsync(Program.UserConfigFile, JsonSerializer.SerializeToUtf8Bytes(NewLocalSettings, new JsonSerializerOptions { WriteIndented = true, IgnoreNullValues  = true }));
+					await FileReference.WriteAllBytesAsync(Program.UserConfigFile, JsonSerializer.SerializeToUtf8Bytes(NewLocalSettings, new JsonSerializerOptions { WriteIndented = true, DefaultIgnoreCondition = JsonIgnoreCondition.WhenWritingNull }));
 				}
 				catch (Exception Ex)
 				{
diff --git a/Engine/Source/Programs/Horde/Horde.Build/Server/DatabaseService.cs b/Engine/Source/Programs/Horde/Horde.Build/Server/DatabaseService.cs
--- a/Engine/Source/Programs/Horde/Horde.Build/Server/DatabaseService.cs
+++ b/Engine/Source/Programs/Horde/Horde.Build/Server/DatabaseService.cs
@@ -33,6 +33,7 @@ using System.Net.NetworkInformation;
 using OpenTracing.Util;
 using OpenTracing;
 using HordeServer.Utiltiies;
+using System.Globalization;
 
 namespace HordeServer.Services
 {
@@ -503,11 +504,11 @@ namespace HordeServer.Services
 			if (Sender != null)
 			{
 				string SenderInfo = StringUtils.Indent(Sender.ToString() ?? String.Empty, "    ");
-				Builder.Append($"\nSender:\n{SenderInfo}");
+				Builder.Append(CultureInfo.InvariantCulture, $"\nSender:\n{SenderInfo}");
 			}
 			if (Certificate != null)
 			{
-				Builder.Append($"\nCertificate: {Certificate.Subject}");
+				Builder.Append(CultureInfo.InvariantCulture, $"\nCertificate: {Certificate.Subject}");
 			}
 			if (Chain != null)
 			{
@@ -516,7 +517,7 @@ namespace HordeServer.Services
 					Builder.Append("\nChain status:");
 					foreach (X509ChainStatus Status in Chain.ChainStatus)
 					{
-						Builder.Append($"\n  {Status.StatusInformation}");
+						Builder.Append(CultureInfo.InvariantCulture, $"\n  {Status.StatusInformation}");
 					}
 				}
 				if (Chain.ChainElements != null)
@@ -525,17 +526,17 @@ namespace HordeServer.Services
 					for (int Idx = 0; Idx < Chain.ChainElements.Count; Idx++)
 					{
 						X509ChainElement Element = Chain.ChainElements[Idx];
-						Builder.Append($"\n  {Idx,4} - Certificate: {Element.Certificate.Subject}");
+						Builder.Append(CultureInfo.InvariantCulture, $"\n  {Idx,4} - Certificate: {Element.Certificate.Subject}");
 						if (Element.ChainElementStatus != null && Element.ChainElementStatus.Length > 0)
 						{
 							foreach (X509ChainStatus Status in Element.ChainElementStatus)
 							{
-								Builder.Append($"\n         Status: {Status.StatusInformation} ({Status.Status})");
+								Builder.Append(CultureInfo.InvariantCulture, $"\n         Status: {Status.StatusInformation} ({Status.Status})");
 							}
 						}
 						if (!String.IsNullOrEmpty(Element.Information))
 						{
-							Builder.Append($"\n         Info: {Element.Information}");
+							Builder.Append(CultureInfo.InvariantCulture, $"\n         Info: {Element.Information}");
 						}
 					}
 				}
diff --git a/Engine/Source/Programs/Horde/Horde.Build/Server/RequestTrackerService.cs b/Engine/Source/Programs/Horde/Horde.Build/Server/RequestTrackerService.cs
--- a/Engine/Source/Programs/Horde/Horde.Build/Server/RequestTrackerService.cs
+++ b/Engine/Source/Programs/Horde/Horde.Build/Server/RequestTrackerService.cs
@@ -3,6 +3,7 @@
 using System;
 using System.Collections.Concurrent;
 using System.Collections.Generic;
+using System.Globalization;
 using System.Linq;
 using System.Text;
 using System.Threading.Tasks;
@@ -70,7 +71,7 @@ namespace HordeServer.Utilities
 			{
 				int AgeInMs = Pair.Value.GetTimeSinceStartInMs();
 				string Path = Pair.Value.Request.Path;
-				Content.AppendLine($"{AgeInMs,9}  {Path}");
+				Content.AppendLine(CultureInfo.InvariantCulture, $"{AgeInMs,9}  {Path}");
 			}
 
 			return Content.ToString();
diff --git a/Engine/Source/Programs/Horde/Horde.Build/Startup.cs b/Engine/Source/Programs/Horde/Horde.Build/Startup.cs
--- a/Engine/Source/Programs/Horde/Horde.Build/Startup.cs
+++ b/Engine/Source/Programs/Horde/Horde.Build/Startup.cs
@@ -619,7 +619,7 @@ namespace HordeServer
 
 		public static void ConfigureJsonSerializer(JsonSerializerOptions Options)
 		{
-			Options.IgnoreNullValues = true;
+			Options.DefaultIgnoreCondition = JsonIgnoreCondition.WhenWritingNull;
 			Options.PropertyNamingPolicy = JsonNamingPolicy.CamelCase;
 			Options.PropertyNameCaseInsensitive = true;
 			Options.Converters.Add(new JsonObjectIdConverter());
@@ -849,8 +849,8 @@ namespace HordeServer
 
 		static void LaunchBrowser(IApplicationBuilder App)
 		{
-			IServerAddressesFeature Feature = App.ServerFeatures.Get<IServerAddressesFeature>();
-			if (Feature.Addresses.Count > 0)
+			IServerAddressesFeature? Feature = App.ServerFeatures.Get<IServerAddressesFeature>();
+			if (Feature != null && Feature.Addresses.Count > 0)
 			{
 				// with a development cert, host will be set by default to localhost, otherwise there will be no host in address
 				string Address = Feature.Addresses.First().Replace("[::]", System.Net.Dns.GetHostName(), StringComparison.OrdinalIgnoreCase);
diff --git a/Engine/Source/Programs/Horde/Horde.Build/Utilities/CompactBinaryFormatter.cs b/Engine/Source/Programs/Horde/Horde.Build/Utilities/CompactBinaryFormatter.cs
--- a/Engine/Source/Programs/Horde/Horde.Build/Utilities/CompactBinaryFormatter.cs
+++ b/Engine/Source/Programs/Horde/Horde.Build/Utilities/CompactBinaryFormatter.cs
@@ -108,7 +108,7 @@ namespace HordeServer.Utilities
 		}
 
 		/// <inheritdoc/>
-		protected override bool CanWriteType(Type Type)
+		protected override bool CanWriteType(Type? Type)
 		{
 			return true;
 		}
@@ -136,7 +136,7 @@ namespace HordeServer.Utilities
 	public class CbPreferredOutputFormatter : CbOutputFormatter
 	{
 		/// <inheritdoc/>
-		protected override bool CanWriteType(Type Type)
+		protected override bool CanWriteType(Type? Type)
 		{
 			return Type == typeof(CbObject);
 		}
diff --git a/Engine/Source/Programs/Horde/Horde.Build/Utilities/ObjectId.cs b/Engine/Source/Programs/Horde/Horde.Build/Utilities/ObjectId.cs
--- a/Engine/Source/Programs/Horde/Horde.Build/Utilities/ObjectId.cs
+++ b/Engine/Source/Programs/Horde/Horde.Build/Utilities/ObjectId.cs
@@ -195,20 +195,24 @@ namespace HordeServer.Utilities
 		}
 
 		/// <inheritdoc/>
-		public override bool CanConvertFrom(ITypeDescriptorContext Context, Type SourceType)
+		public override bool CanConvertFrom(ITypeDescriptorContext? Context, Type SourceType)
 		{
 			return SourceType == typeof(string) || base.CanConvertFrom(Context, SourceType);
 		}
 
 		/// <inheritdoc/>
-		public override object ConvertFrom(ITypeDescriptorContext Context, CultureInfo Culture, object Value)
+		public override object ConvertFrom(ITypeDescriptorContext? Context, CultureInfo? Culture, object Value)
 		{
 			return Activator.CreateInstance(Type, Value)!;
 		}
 
 		/// <inheritdoc/>
-		public override bool CanConvertTo(ITypeDescriptorContext Context, Type DestinationType)
+		public override bool CanConvertTo(ITypeDescriptorContext? Context, Type? DestinationType)
 		{
+			if (DestinationType == null)
+			{
+				return false;
+			}
 			if (DestinationType == typeof(string))
 			{
 				return true;
@@ -229,11 +233,11 @@ namespace HordeServer.Utilities
 		}
 
 		/// <inheritdoc/>
-		public override object? ConvertTo(ITypeDescriptorContext Context, CultureInfo Culture, object Value, Type DestinationType)
+		public override object? ConvertTo(ITypeDescriptorContext? Context, CultureInfo? Culture, object? Value, Type DestinationType)
 		{
 			if (DestinationType == typeof(string))
 			{
-				return Value.ToString();
+				return Value?.ToString();
 			}
 			else
 			{
diff --git a/Engine/Source/Programs/Horde/Horde.Build/Utilities/PropertyFilter.cs b/Engine/Source/Programs/Horde/Horde.Build/Utilities/PropertyFilter.cs
--- a/Engine/Source/Programs/Horde/Horde.Build/Utilities/PropertyFilter.cs
+++ b/Engine/Source/Programs/Horde/Horde.Build/Utilities/PropertyFilter.cs
@@ -209,15 +209,15 @@ namespace HordeServer.Utilities
 	public class PropertyFilterTypeConverter : TypeConverter
 	{
 		/// <inheritdoc/>
-		public override bool CanConvertFrom(ITypeDescriptorContext context, Type sourceType)
+		public override bool CanConvertFrom(ITypeDescriptorContext? Context, Type SourceType)
 		{
-			return sourceType == typeof(string);
+			return SourceType == typeof(string);
 		}
 
 		/// <inheritdoc/>
-		public override object ConvertFrom(ITypeDescriptorContext context, CultureInfo culture, object value)
+		public override object ConvertFrom(ITypeDescriptorContext? Context, CultureInfo? Culture, object Value)
 		{
-			return PropertyFilter.Parse((string)value);
+			return PropertyFilter.Parse((string)Value);
 		}
 	}
 
diff --git a/Engine/Source/Programs/Horde/Horde.Build/Utilities/StringId.cs b/Engine/Source/Programs/Horde/Horde.Build/Utilities/StringId.cs
--- a/Engine/Source/Programs/Horde/Horde.Build/Utilities/StringId.cs
+++ b/Engine/Source/Programs/Horde/Horde.Build/Utilities/StringId.cs
@@ -320,20 +320,24 @@ namespace HordeServer.Utilities
 		}
 
 		/// <inheritdoc/>
-		public override bool CanConvertFrom(ITypeDescriptorContext Context, Type SourceType)
+		public override bool CanConvertFrom(ITypeDescriptorContext? Context, Type SourceType)
 		{
 			return SourceType == typeof(string) || base.CanConvertFrom(Context, SourceType);
 		}
 
 		/// <inheritdoc/>
-		public override object ConvertFrom(ITypeDescriptorContext Context, CultureInfo Culture, object Value)
+		public override object ConvertFrom(ITypeDescriptorContext? Context, CultureInfo? Culture, object Value)
 		{
 			return Activator.CreateInstance(Type, Value)!;
 		}
 
 		/// <inheritdoc/>
-		public override bool CanConvertTo(ITypeDescriptorContext Context, Type DestinationType)
+		public override bool CanConvertTo(ITypeDescriptorContext? Context, Type? DestinationType)
 		{
+			if (DestinationType == null)
+			{
+				return false;
+			}
 			if (DestinationType == typeof(string))
 			{
 				return true;
@@ -354,11 +358,11 @@ namespace HordeServer.Utilities
 		}
 
 		/// <inheritdoc/>
-		public override object? ConvertTo(ITypeDescriptorContext Context, CultureInfo Culture, object Value, Type DestinationType)
+		public override object? ConvertTo(ITypeDescriptorContext? Context, CultureInfo? Culture, object? Value, Type DestinationType)
 		{
 			if (DestinationType == typeof(string))
 			{
-				return Value.ToString();
+				return Value?.ToString();
 			}
 			else
 			{
diff --git a/Engine/Source/Programs/Horde/Horde.Build/Utilities/ViewMap.cs b/Engine/Source/Programs/Horde/Horde.Build/Utilities/ViewMap.cs
--- a/Engine/Source/Programs/Horde/Horde.Build/Utilities/ViewMap.cs
+++ b/Engine/Source/Programs/Horde/Horde.Build/Utilities/ViewMap.cs
@@ -5,6 +5,7 @@ using System;
 using System.Collections.Generic;
 using System.Diagnostics;
 using System.Diagnostics.CodeAnalysis;
+using System.Globalization;
 using System.Text;
 using System.Text.RegularExpressions;
 
@@ -314,7 +315,7 @@ namespace HordeServer.Utilities
 			{
 				Builder.Append("-");
 			}
-			Builder.Append($"{SourcePrefix}{Wildcard}{SourceSuffix} {TargetPrefix}{Wildcard}{TargetSuffix}");
+			Builder.Append(CultureInfo.InvariantCulture, $"{SourcePrefix}{Wildcard}{SourceSuffix} {TargetPrefix}{Wildcard}{TargetSuffix}");
 			return Builder.ToString();
 		}
 	}
