Horde: Fix cast exception with audit log messages.

diff --git a/Engine/Source/Programs/Horde/Horde.Build/Auditing/IAuditLog.cs b/Engine/Source/Programs/Horde/Horde.Build/Auditing/IAuditLog.cs
--- a/Engine/Source/Programs/Horde/Horde.Build/Auditing/IAuditLog.cs
+++ b/Engine/Source/Programs/Horde/Horde.Build/Auditing/IAuditLog.cs
@@ -141,7 +141,7 @@ namespace HordeServer.Collections
 			await BodyWriter.WriteAsync(Encoding.UTF8.GetBytes(Prefix));
 
 			string Separator = "";
-			await foreach (IAuditLogMessage<AgentId> Message in Channel.FindAsync(MinTime, MaxTime, Index, Count))
+			await foreach (IAuditLogMessage<T> Message in Channel.FindAsync(MinTime, MaxTime, Index, Count))
 			{
 				string Line = $"{Separator}\n\t\t{Message.Data}";
 				await BodyWriter.WriteAsync(Encoding.UTF8.GetBytes(Line));
