Horde: Add issue audit logging and issue view updates
#rnx

diff --git a/Engine/Source/Programs/Horde/HordeDashboard/src/backend/Backend.ts b/Engine/Source/Programs/Horde/HordeDashboard/src/backend/Backend.ts
--- a/Engine/Source/Programs/Horde/HordeDashboard/src/backend/Backend.ts
+++ b/Engine/Source/Programs/Horde/HordeDashboard/src/backend/Backend.ts
@@ -857,6 +857,15 @@ export class Backend {
 
     }
 
+    getIssueHistory(id: string, query: AuditLogQuery): Promise<AuditLogEntry[]> {
+        return new Promise<AuditLogEntry[]>((resolve, reject) => {
+            this.backend.get(`/api/v1/issues/${encodeURIComponent(id)}/history`, { params: query }).then((response) => {
+                const history = (response.data?.entries ?? []) as AuditLogEntry[];
+                resolve(history);
+            }).catch(reason => { reject(reason); });
+        });
+    }
+
 
     getIssues(queryIn?: IssueQuery): Promise<IssueData[]> {
 
