Horde: Fix race condition where remote execution queues would not correctly be flagged as active.

diff --git a/Engine/Source/Programs/Horde/Horde.Build/Compute/Impl/ComputeService.cs b/Engine/Source/Programs/Horde/Horde.Build/Compute/Impl/ComputeService.cs
--- a/Engine/Source/Programs/Horde/Horde.Build/Compute/Impl/ComputeService.cs
+++ b/Engine/Source/Programs/Horde/Horde.Build/Compute/Impl/ComputeService.cs
@@ -175,6 +175,7 @@ namespace HordeServer.Compute.Impl
 			List<IoHash> RequirementsHashes = await TaskScheduler.GetInactiveQueuesAsync();
 			foreach (IoHash RequirementsHash in RequirementsHashes)
 			{
+				Logger.LogInformation("Inactive queue: {RequirementsHash}", RequirementsHash);
 				for (; ; )
 				{
 					ComputeTaskInfo? ComputeTask = await TaskScheduler.DequeueAsync(RequirementsHash);
@@ -185,7 +186,7 @@ namespace HordeServer.Compute.Impl
 
 					ComputeTaskStatus Status = new ComputeTaskStatus(ComputeTask.TaskHash, ComputeTaskState.Complete, null, null);
 					Status.Outcome = ComputeTaskOutcome.Expired;
-					Logger.LogInformation("Compute task expired (task: {TaskHash}, channel: {ChannelId})", ComputeTask.TaskHash, ComputeTask.ChannelId);
+					Logger.LogInformation("Compute task expired (queue: {RequirementsHash}, task: {TaskHash}, channel: {ChannelId})", RequirementsHash, ComputeTask.TaskHash, ComputeTask.ChannelId);
 					await MessageQueue.PostAsync(ComputeTask.ChannelId.ToString(), Status);
 				}
 			}
@@ -206,7 +207,7 @@ namespace HordeServer.Compute.Impl
 			foreach (CbObjectAttachment TaskHash in TaskHashes)
 			{
 				ComputeTaskInfo TaskInfo = new ComputeTaskInfo(TaskHash, ChannelId);
-				Logger.LogDebug("Adding task {TaskHash} to queue {RequirementsHash}", TaskInfo.TaskHash.Hash, RequirementsHash);
+				Logger.LogDebug("Adding task {TaskHash} from channel {ChannelId} to queue {RequirementsHash}", TaskInfo.TaskHash.Hash, ChannelId, RequirementsHash);
 				Tasks.Add(TaskScheduler.EnqueueAsync(RequirementsHash, TaskInfo, false));
 			}
 			await Task.WhenAll(Tasks);
diff --git a/Engine/Source/Programs/Horde/Horde.Build/Compute/Impl/RedisTaskScheduler.cs b/Engine/Source/Programs/Horde/Horde.Build/Compute/Impl/RedisTaskScheduler.cs
--- a/Engine/Source/Programs/Horde/Horde.Build/Compute/Impl/RedisTaskScheduler.cs
+++ b/Engine/Source/Programs/Horde/Horde.Build/Compute/Impl/RedisTaskScheduler.cs
@@ -325,22 +325,22 @@ namespace HordeServer.Compute.Impl
 			// Check if the set of active keys already contains the key we're adding. In order to optimize the 
 			// common case under heavy load where the key is in the set, updating it creates a full copy of it. Any
 			// readers can thus access it without the need for any locking.
-			ReadOnlyHashSet<TQueueId> LocalActiveKeysCopy = LocalActiveQueues;
-			if (!LocalActiveKeysCopy.Contains(QueueId))
+			for(; ;)
 			{
-				for (; ; )
+				ReadOnlyHashSet<TQueueId> LocalActiveQueuesCopy = LocalActiveQueues;
+				if (LocalActiveQueuesCopy.Contains(QueueId))
 				{
-					HashSet<TQueueId> NewLocalActiveKeys = new HashSet<TQueueId>(LocalActiveKeysCopy);
-					if (!NewLocalActiveKeys.Add(QueueId))
-					{
-						break;
-					}
-					if (Interlocked.CompareExchange(ref LocalActiveQueues, NewLocalActiveKeys, LocalActiveKeysCopy) == LocalActiveKeysCopy)
-					{
-						Logger.LogInformation("Refreshing active queue {QueueId}", QueueId);
-						await ActiveQueues.SetAsync(QueueId, DateTime.UtcNow);
-						break;
-					}
+					break;
+				}
+
+				HashSet<TQueueId> NewLocalActiveQueues = new HashSet<TQueueId>(LocalActiveQueuesCopy);
+				NewLocalActiveQueues.Add(QueueId);
+
+				if (Interlocked.CompareExchange(ref LocalActiveQueues, NewLocalActiveQueues, LocalActiveQueuesCopy) == LocalActiveQueuesCopy)
+				{
+					Logger.LogInformation("Refreshing active queue {QueueId}", QueueId);
+					await ActiveQueues.SetAsync(QueueId, DateTime.UtcNow);
+					break;
 				}
 			}
 		}
