Horde: Additional logging for Slack config update notifications.

diff --git a/Engine/Source/Programs/Horde/Horde.Build/Notifications/Impl/SlackNotificationSink.cs b/Engine/Source/Programs/Horde/Horde.Build/Notifications/Impl/SlackNotificationSink.cs
--- a/Engine/Source/Programs/Horde/Horde.Build/Notifications/Impl/SlackNotificationSink.cs
+++ b/Engine/Source/Programs/Horde/Horde.Build/Notifications/Impl/SlackNotificationSink.cs
@@ -889,12 +889,20 @@ namespace HordeServer.Notifications.Impl
 		/// <inheritdoc/>
 		public async Task NotifyConfigUpdateFailureAsync(string ErrorMessage, string FileName, int? Change = null, IUser? Author = null, string? Description = null)
 		{
-			Logger.LogInformation("Sending config update failure notification for {FileName}", FileName);
+			Logger.LogInformation("Sending config update failure notification for {FileName} (change: {Change}, author: {UserId})", FileName, Change ?? -1, Author?.Id ?? UserId.Empty);
 
 			string? SlackUserId = null;
 			if (Author != null)
 			{
 				SlackUserId = await GetSlackUserId(Author);
+				if (SlackUserId == null)
+				{
+					Logger.LogWarning("Unable to identify Slack user id for {UserId}", Author.Id);
+				}
+				else
+				{
+					Logger.LogInformation("Mappsed user {UserId} to Slack user {SlackUserId}", Author.Id, SlackUserId);
+				}
 			}
 
 			if (SlackUserId != null)
