Horde: Fix accumulated dashboard warnings
#rnx

diff --git a/Engine/Source/Programs/Horde/HordeDashboard/src/backend/Api.ts b/Engine/Source/Programs/Horde/HordeDashboard/src/backend/Api.ts
--- a/Engine/Source/Programs/Horde/HordeDashboard/src/backend/Api.ts
+++ b/Engine/Source/Programs/Horde/HordeDashboard/src/backend/Api.ts
@@ -896,12 +896,18 @@ export type GetJobResponse = {
 	/** Hash of the graph for this job */
 	graphHash?: string;
 
-	/** The user that started this job */
+	/** @deprecated The user that started this job */
 	startedByUser?: string;
 
-	/** The user that aborted this job */
+	/** The user that started this job */
+	startedByUserInfo?: GetThinUserInfoResponse;
+
+	/** @deprecated The user that aborted this job */
 	abortedByUser?: string;
 
+	/** The user that started this job */
+	abortedByUserInfo?: GetThinUserInfoResponse;
+
 	/** The roles to impersonate when executing this job */
 	roles: string[];
 
@@ -989,12 +995,18 @@ export type GetStepResponse = {
 	/** If the step has been requested to abort	*/
 	abortRequested?: boolean;
 
-	/** Name of the user that requested the abort of this step */
+	/** @deprecated Name of the user that requested the abort of this step */
 	abortByUser?: string;
 
-	/** Name of the user that requested this step be run again */
+	/* The user that requested the abort of this step */
+	abortedByUserInfo?: GetThinUserInfoResponse;
+
+	/** @deprecated Name of the user that requested this step be run again */
 	retryByUser?: string;
 
+	/* The user that retried this step */
+	retriedByUserInfo?: GetThinUserInfoResponse;
+
 	/**The log id for this step */
 	logId?: string;
 
@@ -1963,12 +1975,16 @@ export type GetChangeSummaryResponse = {
 	/**  The source changelist number */
 	number: number;
 
-	/**  Name of the user that authored this change */
+	/**  @deprecated Name of the user that authored this change */
 	author: string;
 
 	/**  The description text */
 	description: string;
 
+	/**  Information about the change author */
+	authorInfo: GetThinUserInfoResponse;
+
+
 }
 
 
@@ -2595,19 +2611,25 @@ export type GetIssueResponse = {
 
 	/** Severity of this issue	*/
 	severity: IssueSeverity;
-	
+
 	/**Whether the issue is promoted */
 	promoted: boolean;
 
-	/**Owner of the issue */
+	/** @deprecated Owner of the issue */
 	owner?: string;
 
-	/**Owner id of the issue */
+	/** @deprecated Owner id of the issue */
 	ownerId?: string;
 
-	/**User that nominated the current owner */
+	/** Owner of the issue */
+	ownerInfo: GetThinUserInfoResponse;
+
+	/** @deprecated User that nominated the current owner */
 	nominatedBy?: string;
 
+	/** Use that nominated the current owner */
+	nominatedByInfo: GetThinUserInfoResponse;
+
 	/**Time that the issue was acknowledged */
 	acknowledgedAt?: Date | string;
 
@@ -2617,12 +2639,15 @@ export type GetIssueResponse = {
 	/**Time at which the issue was resolved */
 	resolvedAt?: Date | string;
 
-	/** Name of the user that resolved the issue */
+	/** @deprecated Name of the user that resolved the issue */
 	resolvedBy?: string;
 
-	/** User id of the person that resolved the issue */
+	/** @deprecated User id of the person that resolved the issue */
 	resolvedById?: string;
 
+	/** Use info for the person that resolved the issue */
+	resolvedByInfo: GetThinUserInfoResponse;
+
 	/**  List of stream paths affected by this issue */
 	streams: string[];
 
@@ -2634,9 +2659,13 @@ export type GetIssueResponse = {
 
 	affectedStreams: GetIssueAffectedStreamResponse[];
 
-	/** User id's of the Most likely suspects for causing this issue */
+	/** @deprecated User id's of the Most likely suspects for causing this issue */
 	primarySuspectIds: string[];
 
+	/** Use info for the person that resolved the issue */
+	primarySuspectsInfo: GetThinUserInfoResponse[];
+
+
 	/** Whether to show alerts for this issue */
 	showDesktopAlerts: boolean;
 
@@ -2784,6 +2813,21 @@ export type GetUserResponse = {
 
 }
 
+/** Basic information about a user. May be embedded in other responses.*/
+export type GetThinUserInfoResponse = {
+
+	/**  Id of the user */
+	id: string;
+
+	/** Name of the user */
+	name: string;
+
+	/**  The user's email address */
+	email: string;
+
+}
+
+
 /// <summary>
 /// Request to update settings for a user
 /// </summary>
diff --git a/Engine/Source/Programs/Horde/HordeDashboard/src/backend/Backend.ts b/Engine/Source/Programs/Horde/HordeDashboard/src/backend/Backend.ts
--- a/Engine/Source/Programs/Horde/HordeDashboard/src/backend/Backend.ts
+++ b/Engine/Source/Programs/Horde/HordeDashboard/src/backend/Backend.ts
@@ -3,7 +3,7 @@
 import { configure } from "mobx";
 import { isNumber } from 'util';
 import templateCache from '../backend/TemplateCache';
-import { AgentData, AuditLogEntry, AuditLogQuery, ArtifactData, BatchUpdatePoolRequest, ChangeSummaryData, CreateDeviceRequest, CreateDeviceResponse, CreateJobRequest, CreateJobResponse, CreatePoolRequest, CreateSoftwareResponse, CreateSubscriptionRequest, CreateSubscriptionResponse, DashboardPreference, EventData, GetArtifactZipRequest, GetDevicePlatformResponse, GetDevicePoolResponse, GetDeviceReservationResponse, GetDeviceResponse, GetGraphResponse, GetIssueStreamResponse, GetJobsTabResponse, GetJobStepRefResponse, GetJobStepTraceResponse, GetJobTimingResponse, GetLogEventResponse, GetNotificationResponse, GetPerforceServerStatusResponse, GetServerInfoResponse, GetServerSettingsResponse, GetSoftwareResponse, GetSubscriptionResponse, GetUserResponse, GetUtilizationTelemetryResponse, GlobalConfig, IssueData, IssueQuery, JobData, JobQuery, JobsTabColumnType, LeaseData, LogData, LogLineData, PerforceCluster, PoolData, ProjectData, ScheduleData, ScheduleQuery, SearchLogFileResponse, ServerUpdateResponse, SessionData, StreamData, TabType, TemplateData, TestData, UpdateAgentRequest, UpdateDeviceRequest, UpdateGlobalConfigRequest, UpdateIssueRequest, UpdateJobRequest, UpdateLeaseRequest, UpdateNotificationsRequest, UpdatePoolRequest, UpdateServerSettingsRequest, UpdateStepRequest, UpdateStepResponse, UpdateUserRequest, UsersQuery } from './Api';
+import { AgentData, ArtifactData, AuditLogEntry, AuditLogQuery, BatchUpdatePoolRequest, ChangeSummaryData, CreateDeviceRequest, CreateDeviceResponse, CreateJobRequest, CreateJobResponse, CreatePoolRequest, CreateSoftwareResponse, CreateSubscriptionRequest, CreateSubscriptionResponse, DashboardPreference, EventData, GetArtifactZipRequest, GetDevicePlatformResponse, GetDevicePoolResponse, GetDeviceReservationResponse, GetDeviceResponse, GetGraphResponse, GetIssueStreamResponse, GetJobsTabResponse, GetJobStepRefResponse, GetJobStepTraceResponse, GetJobTimingResponse, GetLogEventResponse, GetNotificationResponse, GetPerforceServerStatusResponse, GetServerInfoResponse, GetServerSettingsResponse, GetSoftwareResponse, GetSubscriptionResponse, GetUserResponse, GetUtilizationTelemetryResponse, GlobalConfig, IssueData, IssueQuery, JobData, JobQuery, JobsTabColumnType, LeaseData, LogData, LogLineData, PoolData, ProjectData, ScheduleData, ScheduleQuery, SearchLogFileResponse, ServerUpdateResponse, SessionData, StreamData, TabType, TemplateData, TestData, UpdateAgentRequest, UpdateDeviceRequest, UpdateGlobalConfigRequest, UpdateIssueRequest, UpdateJobRequest, UpdateLeaseRequest, UpdateNotificationsRequest, UpdatePoolRequest, UpdateServerSettingsRequest, UpdateStepRequest, UpdateStepResponse, UpdateUserRequest, UsersQuery } from './Api';
 import dashboard from './Dashboard';
 import { ChallengeStatus, Fetch } from './Fetch';
 import graphCache, { GraphQuery } from './GraphCache';
diff --git a/Engine/Source/Programs/Horde/HordeDashboard/src/backend/CommitCache.ts b/Engine/Source/Programs/Horde/HordeDashboard/src/backend/CommitCache.ts
--- a/Engine/Source/Programs/Horde/HordeDashboard/src/backend/CommitCache.ts
+++ b/Engine/Source/Programs/Horde/HordeDashboard/src/backend/CommitCache.ts
@@ -76,7 +76,7 @@ export class CommitCache {
             }
 
             map!.set(c, {
-                number: 0, author: "", description: ""
+                number: 0, author: "", description: "", authorInfo : { id : "0", email:"unknown", name: "unknown"}
             })
 
         })
diff --git a/Engine/Source/Programs/Horde/HordeDashboard/src/base/utilities/issueUtils.ts b/Engine/Source/Programs/Horde/HordeDashboard/src/base/utilities/issueUtils.ts
--- a/Engine/Source/Programs/Horde/HordeDashboard/src/base/utilities/issueUtils.ts
+++ b/Engine/Source/Programs/Horde/HordeDashboard/src/base/utilities/issueUtils.ts
@@ -9,6 +9,8 @@ export type IssueJira = {
 	link: string;
 }
 
+// This regex could be optimized, though it works
+// eslint-disable-next-line
 const urlRegex = /([\w+]+\:\/\/)?([\w\d-]+\.)*[\w-]+[\.\:]\w+([\/\?\=\&\#\.]?[\w-]+)*\/?/gm
 
 export function getIssueJiras(issue: GetIssueResponse): IssueJira[] {
