Horde: Use a different database for each test to avoid non-deterministic execution.

diff --git a/Engine/Source/Programs/Horde/Horde.Build.Tests/AgentControllerTest.cs b/Engine/Source/Programs/Horde/Horde.Build.Tests/AgentControllerTest.cs
--- a/Engine/Source/Programs/Horde/Horde.Build.Tests/AgentControllerTest.cs
+++ b/Engine/Source/Programs/Horde/Horde.Build.Tests/AgentControllerTest.cs
@@ -10,18 +10,17 @@ using Microsoft.AspNetCore.Mvc;
 namespace HordeServerTests
 {
 	[TestClass]
-    public class AgentControllerDbTest : DatabaseIntegrationTest
+    public class AgentControllerDbTest : TestSetup
     {
         [TestMethod]
         public async Task UpdateAgent()
         {
-	        TestSetup TestSetup = await GetTestSetup();
-	        AgentsController AgentsController = TestSetup.AgentsController;
-	        IAgent FixtureAgent = TestSetup.Fixture!.Agent1;
+			Fixture Fixture = await CreateFixtureAsync();
+	        IAgent FixtureAgent = Fixture.Agent1;
 
 	        ActionResult<object> Obj = await AgentsController.GetAgentAsync(FixtureAgent.Id.ToString());
 	        GetAgentResponse GetRes = (Obj.Value as GetAgentResponse)!;
-	        Assert.AreEqual(TestSetup.Fixture!.Agent1Name.ToUpper(), GetRes.Name);
+	        Assert.AreEqual(Fixture!.Agent1Name.ToUpper(), GetRes.Name);
 	        Assert.IsNull(GetRes.Comment);
 	        
 	        UpdateAgentRequest UpdateReq = new UpdateAgentRequest();
diff --git a/Engine/Source/Programs/Horde/Horde.Build.Tests/AgentServiceTest.cs b/Engine/Source/Programs/Horde/Horde.Build.Tests/AgentServiceTest.cs
--- a/Engine/Source/Programs/Horde/Horde.Build.Tests/AgentServiceTest.cs
+++ b/Engine/Source/Programs/Horde/Horde.Build.Tests/AgentServiceTest.cs
@@ -20,19 +20,19 @@ namespace HordeServerTests
     /// Testing the agent service
     /// </summary>
 	[TestClass]
-    public class AgentServiceTest : DatabaseIntegrationTest
+    public class AgentServiceTest : TestSetup
     {
         [TestMethod]
         public async Task GetJobs()
         {
-	        TestSetup TestSetup = await GetTestSetup();
+			Fixture Fixture = await CreateFixtureAsync();
 
-	        ActionResult<List<object>> Res = await TestSetup.JobsController.FindJobsAsync();
+			ActionResult<List<object>> Res = await JobsController.FindJobsAsync();
 	        Assert.AreEqual(2, Res.Value!.Count);
 	        Assert.AreEqual("hello2", (Res.Value[0] as GetJobResponse)!.Name);
 	        Assert.AreEqual("hello1", (Res.Value[1] as GetJobResponse)!.Name);
 	        
-	        Res = await TestSetup.JobsController.FindJobsAsync(IncludePreflight: false);
+	        Res = await JobsController.FindJobsAsync(IncludePreflight: false);
 	        Assert.AreEqual(1, Res.Value!.Count);
 	        Assert.AreEqual("hello2", (Res.Value[0] as GetJobResponse)!.Name);
         }
@@ -40,14 +40,13 @@ namespace HordeServerTests
         [TestMethod]
         public async Task CreateSessionTest()
         {
-	        TestSetup TestSetup = await GetTestSetup();
-	        AgentService AgentService = TestSetup.AgentService;
+			Fixture Fixture = await CreateFixtureAsync();
 
-	        IAgent Agent = await AgentService.CreateSessionAsync(TestSetup.Fixture!.Agent1, AgentStatus.Ok, new List<string>(), new Dictionary<string, int>(),
+	        IAgent Agent = await AgentService.CreateSessionAsync(Fixture.Agent1, AgentStatus.Ok, new List<string>(), new Dictionary<string, int>(),
 		        "test");
 	        
 	        Assert.IsTrue(AgentService.AuthorizeSession(Agent, GetUser(Agent)));
-	        TestSetup.Clock.Advance(TimeSpan.FromMinutes(20));
+	        Clock.Advance(TimeSpan.FromMinutes(20));
 	        Assert.IsFalse(AgentService.AuthorizeSession(Agent, GetUser(Agent)));
         }
 
diff --git a/Engine/Source/Programs/Horde/Horde.Build.Tests/ControllerIntegrationTest.cs b/Engine/Source/Programs/Horde/Horde.Build.Tests/ControllerIntegrationTest.cs
--- a/Engine/Source/Programs/Horde/Horde.Build.Tests/ControllerIntegrationTest.cs
+++ b/Engine/Source/Programs/Horde/Horde.Build.Tests/ControllerIntegrationTest.cs
@@ -23,12 +23,19 @@ namespace HordeServerTests
 {
 	public class TestWebApplicationFactory<TStartup> : WebApplicationFactory<TStartup> where TStartup : class
     {
+		MongoDbInstance MongoDbInstance;
+
+		public TestWebApplicationFactory(MongoDbInstance MongoDbInstance)
+		{
+			this.MongoDbInstance = MongoDbInstance;
+		}
+
 		protected override void ConfigureWebHost(IWebHostBuilder builder)
         {
             var Dict = new Dictionary<string, string?>
             {
-                {"Horde:DatabaseConnectionString", DatabaseIntegrationTest.GetMongoDbRunner().GetConnectionString()},
-                {"Horde:DatabaseName", DatabaseIntegrationTest.MongoDbDatabaseName},
+                {"Horde:DatabaseConnectionString", MongoDbInstance.ConnectionString},
+                {"Horde:DatabaseName", MongoDbInstance.DatabaseName},
                 {"Horde:LogServiceWriteCacheType", "inmemory"},
                 {"Horde:DisableAuth", "true"},
                 {"Horde:OidcAuthority", null},
@@ -39,53 +46,45 @@ namespace HordeServerTests
         }
     }
 
-    public class ControllerIntegrationTest
+    public class ControllerIntegrationTest : IDisposable
     {
-        private HttpClient? _client;
-        private TestWebApplicationFactory<Startup>? _factory;
-        protected readonly HttpClient client;
-        private Fixture? _fixture;
-
-        public ControllerIntegrationTest()
-        {
-            client = GetClientForTestServer();
-        }
+		protected MongoDbInstance MongoDbInstance { get; }
+		TestWebApplicationFactory<Startup> _factory { get; }
+		protected HttpClient client { get; }
+		private Lazy<Task<Fixture>> _fixture;
 
-        public async Task<Fixture> GetFixture()
-        {
-            if (_fixture != null) return _fixture;
+		public ControllerIntegrationTest()
+		{
+			MongoDbInstance = new MongoDbInstance();
+			_factory = new TestWebApplicationFactory<Startup>(MongoDbInstance);
+			client = _factory.CreateClient();
 
-            IServiceProvider Services = GetFactory().Services;
-            DatabaseService DatabaseService = Services.GetRequiredService<DatabaseService>();
-            ITemplateCollection TemplateService = Services.GetRequiredService<ITemplateCollection>();
-            JobService JobService = Services.GetRequiredService<JobService>();
-            IArtifactCollection ArtifactCollection = Services.GetRequiredService<IArtifactCollection>();
-            StreamService StreamService = Services.GetRequiredService<StreamService>();
-            AgentService AgentService = Services.GetRequiredService<AgentService>();
-            IPerforceService PerforceService = Services.GetRequiredService<IPerforceService>();
-            GraphCollection GraphCollection = new GraphCollection(DatabaseService);
+			_fixture = new Lazy<Task<Fixture>>(Task.Run(() => CreateFixture()));
+		}
 
-            _fixture = new Fixture();
-            _fixture = await Fixture.Create(false, GraphCollection, TemplateService, JobService, ArtifactCollection, StreamService, AgentService, PerforceService);
-            return _fixture;
-        }
+		public void Dispose()
+		{
+			MongoDbInstance.Dispose();
+		}
 
-        [MethodImpl(MethodImplOptions.Synchronized)]
-        public HttpClient GetClientForTestServer()
-        {
-            if (_client != null) return _client;
+		public Task<Fixture> GetFixture()
+		{
+			return _fixture.Value;
+		}
 
-            _factory = GetFactory();
-            _client = _factory.CreateClient();
-            return _client;
-        }
+		async Task<Fixture> CreateFixture()
+		{
+			IServiceProvider Services = _factory.Services;
+			DatabaseService DatabaseService = Services.GetRequiredService<DatabaseService>();
+			ITemplateCollection TemplateService = Services.GetRequiredService<ITemplateCollection>();
+			JobService JobService = Services.GetRequiredService<JobService>();
+			IArtifactCollection ArtifactCollection = Services.GetRequiredService<IArtifactCollection>();
+			StreamService StreamService = Services.GetRequiredService<StreamService>();
+			AgentService AgentService = Services.GetRequiredService<AgentService>();
+			IPerforceService PerforceService = Services.GetRequiredService<IPerforceService>();
+			GraphCollection GraphCollection = new GraphCollection(DatabaseService);
 
-        [MethodImpl(MethodImplOptions.Synchronized)]
-        TestWebApplicationFactory<Startup> GetFactory()
-        {
-            if (_factory != null) return _factory;
-            _factory = new TestWebApplicationFactory<Startup>();
-            return _factory;
+			return await Fixture.Create(GraphCollection, TemplateService, JobService, ArtifactCollection, StreamService, AgentService, PerforceService);
         }
     }
 }
\ No newline at end of file
diff --git a/Engine/Source/Programs/Horde/Horde.Build.Tests/DatabaseIntegrationTest.cs b/Engine/Source/Programs/Horde/Horde.Build.Tests/DatabaseIntegrationTest.cs
--- a/Engine/Source/Programs/Horde/Horde.Build.Tests/DatabaseIntegrationTest.cs
+++ b/Engine/Source/Programs/Horde/Horde.Build.Tests/DatabaseIntegrationTest.cs
@@ -1,6 +1,7 @@
 // Copyright Epic Games, Inc. All Rights Reserved.
 
 using System;
+using System.Collections.Generic;
 using System.Net;
 using System.Threading.Tasks;
 using HordeServer;
@@ -9,6 +10,7 @@ using HordeServer.Utilities;
 using Microsoft.Extensions.Logging;
 using Microsoft.Extensions.Options;
 using Microsoft.VisualStudio.TestTools.UnitTesting;
+using MongoDB.Driver;
 using StackExchange.Redis;
 
 namespace HordeServerTests
@@ -39,73 +41,116 @@ namespace HordeServerTests
 	        // Dummy stub to satisfy return value of OnChange 
         }
     }
-    
-    public class DatabaseIntegrationTest
+
+	public class MongoDbInstance : IDisposable
+	{
+		public string DatabaseName { get; }
+		public string ConnectionString { get; }
+		MongoClient Client { get; }
+
+		private static object LockObject = new object();
+		private static MongoDbRunnerLocal? _mongoDbRunner;
+		private static int NextDatabaseIndex = 1;
+		public const string MongoDbDatabaseNamePrefix = "HordeServerTest_";
+
+		public MongoDbInstance()
+		{
+			int DatabaseIndex;
+			lock (LockObject)
+			{
+				if (_mongoDbRunner == null)
+				{
+					// One-time setup per test run to avoid overhead of starting the external MongoDB process
+					Startup.ConfigureMongoDbClient();
+					_mongoDbRunner = new MongoDbRunnerLocal();
+					_mongoDbRunner.Start();
+
+					// Drop all the previous databases
+					MongoClientSettings MongoSettings = MongoClientSettings.FromConnectionString(_mongoDbRunner.GetConnectionString());
+					MongoClient Client = new MongoClient(MongoSettings);
+
+					List<string> DropDatabaseNames = Client.ListDatabaseNames().ToList();
+					foreach (string DropDatabaseName in DropDatabaseNames)
+					{
+						if (DropDatabaseName.StartsWith(MongoDbDatabaseNamePrefix, StringComparison.Ordinal))
+						{
+							Client.DropDatabase(DropDatabaseName);
+						}
+					}
+				}
+				DatabaseIndex = NextDatabaseIndex++;
+			}
+
+			DatabaseName = $"{MongoDbDatabaseNamePrefix}{DatabaseIndex}";
+			ConnectionString = $"{_mongoDbRunner.GetConnectionString()}/{DatabaseName}";
+			Client = new MongoClient(MongoClientSettings.FromConnectionString(ConnectionString));
+		}
+
+		public void Dispose()
+		{
+			Client.DropDatabase(DatabaseName);
+		}
+	}
+
+	public class DatabaseIntegrationTest : IDisposable
     {
-        private static MongoDbRunnerLocal? _mongoDbRunner;
+		private static object LockObject = new object();
+		private MongoDbInstance? MongoDbInstance;
+		private DatabaseService? DatabaseService;
+		
         private static RedisRunner? _redisRunner;
-        private static DatabaseService? _databaseService;
         private static IDatabase? _redisDb;
         private static RedisConnectionPool? _redisConnectionPool;
         private static ConnectionMultiplexer? _conMux;
-        public const string MongoDbDatabaseName = "HordeServerTest";
         public const int RedisDbNum = 15;
-        
-        protected async Task<TestSetup> GetTestSetup()
-        {
-	        TestSetup TestSetup = new TestSetup(GetDatabaseService(true));
-	        await TestSetup.CreateFixture(true);
-	        return TestSetup;
-        }
 
-        public static DatabaseService GetDatabaseService(bool ForceDropDatabase = true)
-        {
-            if (_mongoDbRunner == null)
-            {
-                // One-time setup per test run to avoid overhead of starting the external MongoDB process
-                
-                Startup.ConfigureMongoDbClient();
-                _mongoDbRunner = new MongoDbRunnerLocal(MongoDbDatabaseName);
-                _mongoDbRunner.Start();
-                
-                ServerSettings Ss = new ServerSettings();
-                Ss.DatabaseName = MongoDbDatabaseName;
-                Ss.DatabaseConnectionString = _mongoDbRunner.GetConnectionString();
-            
-                ILoggerFactory LoggerFactory = new LoggerFactory();
-
-                _databaseService = new DatabaseService(Options.Create(Ss), LoggerFactory);
-                _databaseService.Database.Client.DropDatabase(Ss.DatabaseName);
-            }
-
-            if (ForceDropDatabase)
-            {
-	            _databaseService!.Database.Client.DropDatabase(MongoDbDatabaseName);
-            }
-
-            return _databaseService!;
-        }
+		public DatabaseIntegrationTest()
+		{
+		}
 
-        public static MongoDbRunnerLocal GetMongoDbRunner()
+		public void Dispose()
+		{
+			MongoDbInstance?.Dispose();
+			DatabaseService?.Dispose();
+		}
+
+		public DatabaseService GetDatabaseService()
         {
-            GetDatabaseService();
-            return _mongoDbRunner!;
+			lock(LockObject)
+			{
+				if (DatabaseService == null)
+				{
+					MongoDbInstance = new MongoDbInstance();
+
+					ServerSettings Ss = new ServerSettings();
+					Ss.DatabaseName = MongoDbInstance.DatabaseName;
+					Ss.DatabaseConnectionString = MongoDbInstance.ConnectionString;
+
+					ILoggerFactory LoggerFactory = new LoggerFactory();
+
+					DatabaseService = new DatabaseService(Options.Create(Ss), LoggerFactory);
+				}
+			}
+			return DatabaseService;
         }
 
         public static IDatabase GetRedisDatabase()
         {
-	        if (_redisDb == null)
-	        {
-		        // One-time setup per test run to avoid overhead of starting the external Redis process
-		        _redisRunner = new RedisRunner();
-		        _redisRunner.Start();
-		        
-		        (string Host, int Port) = _redisRunner.GetListenAddress();
-		        _conMux = ConnectionMultiplexer.Connect($"{Host}:{Port},allowAdmin=true");
-		        ClearRedisDatabase();
-
-		        _redisDb = _conMux.GetDatabase(RedisDbNum);
-	        }
+			lock (LockObject)
+			{
+				if (_redisDb == null)
+				{
+					// One-time setup per test run to avoid overhead of starting the external Redis process
+					_redisRunner = new RedisRunner();
+					_redisRunner.Start();
+
+					(string Host, int Port) = _redisRunner.GetListenAddress();
+					_conMux = ConnectionMultiplexer.Connect($"{Host}:{Port},allowAdmin=true");
+					ClearRedisDatabase();
+
+					_redisDb = _conMux.GetDatabase(RedisDbNum);
+				}
+			}
 
 	        return _redisDb;
         }
diff --git a/Engine/Source/Programs/Horde/Horde.Build.Tests/DatabaseRunner.cs b/Engine/Source/Programs/Horde/Horde.Build.Tests/DatabaseRunner.cs
--- a/Engine/Source/Programs/Horde/Horde.Build.Tests/DatabaseRunner.cs
+++ b/Engine/Source/Programs/Horde/Horde.Build.Tests/DatabaseRunner.cs
@@ -193,11 +193,8 @@ namespace HordeServerTests
 
 	public class MongoDbRunnerLocal : DatabaseRunner
 	{
-		private readonly string DatabaseName;
-		
-		public MongoDbRunnerLocal(string DatabaseName) : base("mongodb", "ThirdParty/Mongo/mongod.exe", 27017, true)
+		public MongoDbRunnerLocal() : base("mongodb", "ThirdParty/Mongo/mongod.exe", 27017, true)
 		{
-			this.DatabaseName = DatabaseName;
 		}
 
 		protected override string GetArguments()
@@ -208,7 +205,7 @@ namespace HordeServerTests
 		public string GetConnectionString()
 		{
 			(var Host, int ListenPort) = GetListenAddress();
-			return $"mongodb://{Host}:{ListenPort}/{DatabaseName}";
+			return $"mongodb://{Host}:{ListenPort}";
 		}
 	}
 	
diff --git a/Engine/Source/Programs/Horde/Horde.Build.Tests/DatabaseRunnerTest.cs b/Engine/Source/Programs/Horde/Horde.Build.Tests/DatabaseRunnerTest.cs
--- a/Engine/Source/Programs/Horde/Horde.Build.Tests/DatabaseRunnerTest.cs
+++ b/Engine/Source/Programs/Horde/Horde.Build.Tests/DatabaseRunnerTest.cs
@@ -11,7 +11,7 @@ namespace HordeServerTests
         [TestMethod]
         public void RunMongoDbTest()
         {
-	        MongoDbRunnerLocal Runner = new MongoDbRunnerLocal(DatabaseIntegrationTest.MongoDbDatabaseName);
+	        MongoDbRunnerLocal Runner = new MongoDbRunnerLocal();
 	        Runner.Start();
             Thread.Sleep(100);
             Runner.Stop();
diff --git a/Engine/Source/Programs/Horde/Horde.Build.Tests/DebugTest.cs b/Engine/Source/Programs/Horde/Horde.Build.Tests/DebugTest.cs
--- a/Engine/Source/Programs/Horde/Horde.Build.Tests/DebugTest.cs
+++ b/Engine/Source/Programs/Horde/Horde.Build.Tests/DebugTest.cs
@@ -15,6 +15,7 @@ using Microsoft.Extensions.Logging;
 using Microsoft.Extensions.Logging.Console;
 using MongoDB.Bson;
 
+#if false
 using PoolId = HordeServer.Utilities.StringId<HordeServer.Models.IPool>;
 using Microsoft.Extensions.Options;
 
@@ -31,7 +32,7 @@ namespace HordeServerTests
     {
         private async Task<TestSetup> GetTestSetupExternal()
         {
-	        TestSetup TestSetup = new TestSetup(GetDatabaseServiceExternal());
+	        TestSetup TestSetup = ebugnew TestSetup(GetDatabaseServiceExternal());
 	        await TestSetup.CreateFixture(true);
 	        return TestSetup;
         }
@@ -94,4 +95,5 @@ namespace HordeServerTests
 	        await AutoscaleService.TickSharedOnlyForTestingAsync();
         }
     }
-}
\ No newline at end of file
+}
+#endif
diff --git a/Engine/Source/Programs/Horde/Horde.Build.Tests/Fixture.cs b/Engine/Source/Programs/Horde/Horde.Build.Tests/Fixture.cs
--- a/Engine/Source/Programs/Horde/Horde.Build.Tests/Fixture.cs
+++ b/Engine/Source/Programs/Horde/Horde.Build.Tests/Fixture.cs
@@ -36,7 +36,7 @@ namespace HordeServerTests
 		public IAgent Agent1 { get; private set; } = null!;
 		public string Agent1Name { get; private set; } = null!;
 
-		public static async Task<Fixture> Create(bool ForceNewFixture, IGraphCollection GraphCollection, ITemplateCollection TemplateCollection, JobService JobService, IArtifactCollection ArtifactCollection, StreamService StreamService, AgentService AgentService, IPerforceService PerforceService)
+		public static async Task<Fixture> Create(IGraphCollection GraphCollection, ITemplateCollection TemplateCollection, JobService JobService, IArtifactCollection ArtifactCollection, StreamService StreamService, AgentService AgentService, IPerforceService PerforceService)
 		{
 			Fixture _fixture = new Fixture();
 			await _fixture.Populate(GraphCollection, TemplateCollection, JobService, ArtifactCollection, StreamService, AgentService, PerforceService);
diff --git a/Engine/Source/Programs/Horde/Horde.Build.Tests/IssueServiceTests.cs b/Engine/Source/Programs/Horde/Horde.Build.Tests/IssueServiceTests.cs
--- a/Engine/Source/Programs/Horde/Horde.Build.Tests/IssueServiceTests.cs
+++ b/Engine/Source/Programs/Horde/Horde.Build.Tests/IssueServiceTests.cs
@@ -45,7 +45,7 @@ namespace HordeServerTests
 	using UserId = ObjectId<IUser>;
 
 	[TestClass]
-	public class IssueServiceTests : DatabaseIntegrationTest
+	public class IssueServiceTests : TestSetup
 	{
 		class TestJsonLogger : JsonLogger, IAsyncDisposable
 		{
@@ -108,10 +108,7 @@ namespace HordeServerTests
 		StreamId DevStreamId = StreamId.Sanitize(DevStreamName);
 
 		IGraph Graph;
-		TestSetup TestSetup;
 		PerforceServiceStub Perforce;
-		ILogFileService LogFileService => TestSetup.LogFileService;
-		IIssueService IssueService => TestSetup.IssueService;
 
 		UserId TimId;
 		UserId JerryId;
@@ -121,8 +118,6 @@ namespace HordeServerTests
 
 		public IssueServiceTests()
 		{
-			this.TestSetup = GetTestSetup().Result;
-
 			if (RuntimeInformation.IsOSPlatform(OSPlatform.Windows))
 			{
 				WorkspaceDir = new DirectoryReference("C:\\Horde");
@@ -132,23 +127,23 @@ namespace HordeServerTests
 				WorkspaceDir = new DirectoryReference("/Horde");
 			}
 
-			IProject Project = TestSetup.ProjectCollection.AddOrUpdateAsync(new ProjectId("ue4"), "", "", 0, new ProjectConfig { Name = "UE4" }).Result!;
+			IProject Project = ProjectCollection.AddOrUpdateAsync(new ProjectId("ue4"), "", "", 0, new ProjectConfig { Name = "UE4" }).Result!;
 
-			IStream MainStream = TestSetup.StreamCollection.TryCreateOrReplaceAsync(MainStreamId, null, "", "", Project.Id, new StreamConfig { Name = MainStreamName }).Result!;
-			IStream ReleaseStream = TestSetup.StreamCollection.TryCreateOrReplaceAsync(ReleaseStreamId, null, "", "", Project.Id, new StreamConfig { Name = ReleaseStreamName }).Result!;
-			IStream DevStream = TestSetup.StreamCollection.TryCreateOrReplaceAsync(DevStreamId, null, "", "", Project.Id, new StreamConfig { Name = DevStreamName }).Result!;
+			IStream MainStream = StreamCollection.TryCreateOrReplaceAsync(MainStreamId, null, "", "", Project.Id, new StreamConfig { Name = MainStreamName }).Result!;
+			IStream ReleaseStream = StreamCollection.TryCreateOrReplaceAsync(ReleaseStreamId, null, "", "", Project.Id, new StreamConfig { Name = ReleaseStreamName }).Result!;
+			IStream DevStream = StreamCollection.TryCreateOrReplaceAsync(DevStreamId, null, "", "", Project.Id, new StreamConfig { Name = DevStreamName }).Result!;
 
-			IUser Bill = TestSetup.UserCollection.FindOrAddUserByLoginAsync("Bill").Result;
-			IUser Anne = TestSetup.UserCollection.FindOrAddUserByLoginAsync("Anne").Result;
-			IUser Bob = TestSetup.UserCollection.FindOrAddUserByLoginAsync("Bob").Result;
-			IUser Jerry = TestSetup.UserCollection.FindOrAddUserByLoginAsync("Jerry").Result;
-			IUser Chris = TestSetup.UserCollection.FindOrAddUserByLoginAsync("Chris").Result;
+			IUser Bill = UserCollection.FindOrAddUserByLoginAsync("Bill").Result;
+			IUser Anne = UserCollection.FindOrAddUserByLoginAsync("Anne").Result;
+			IUser Bob = UserCollection.FindOrAddUserByLoginAsync("Bob").Result;
+			IUser Jerry = UserCollection.FindOrAddUserByLoginAsync("Jerry").Result;
+			IUser Chris = UserCollection.FindOrAddUserByLoginAsync("Chris").Result;
 
-			TimId = TestSetup.UserCollection.FindOrAddUserByLoginAsync("Tim").Result.Id;
-			JerryId = TestSetup.UserCollection.FindOrAddUserByLoginAsync("Jerry").Result.Id;
-			BobId = TestSetup.UserCollection.FindOrAddUserByLoginAsync("Bob").Result.Id;
+			TimId = UserCollection.FindOrAddUserByLoginAsync("Tim").Result.Id;
+			JerryId = UserCollection.FindOrAddUserByLoginAsync("Jerry").Result.Id;
+			BobId = UserCollection.FindOrAddUserByLoginAsync("Bob").Result.Id;
 
-			Perforce = TestSetup.PerforceService;
+			Perforce = PerforceService;
 			Perforce.AddChange(MainStreamName, 100, Bill, "Description", new string[] { "a/b.cpp" });
 			Perforce.AddChange(MainStreamName, 105, Anne, "Description", new string[] { "a/c.cpp" });
 			Perforce.AddChange(MainStreamName, 110, Bob, "Description", new string[] { "a/d.cpp" });
@@ -232,7 +227,7 @@ namespace HordeServerTests
 
 			JobStepRefId JobStepRefId = new JobStepRefId(Job.Id, Batch.Id, Step.Id);
 			string NodeName = Graph.Groups[Batch.GroupIdx].Nodes[Step.NodeIdx].Name;
-			await TestSetup.JobStepRefCollection.InsertOrReplaceAsync(JobStepRefId, "TestJob", NodeName, Job.StreamId, Job.TemplateId, Job.Change, Step.LogId, null, null, Outcome, 0.0f, 0.0f, Step.StartTimeUtc!.Value, Step.StartTimeUtc);
+			await JobStepRefCollection.InsertOrReplaceAsync(JobStepRefId, "TestJob", NodeName, Job.StreamId, Job.TemplateId, Job.Change, Step.LogId, null, null, Outcome, 0.0f, 0.0f, Step.StartTimeUtc!.Value, Step.StartTimeUtc);
 		}
 
 		async Task AddEvent(IJob Job, int BatchIdx, int StepIdx, object Data, EventSeverity Severity = EventSeverity.Error)
diff --git a/Engine/Source/Programs/Horde/Horde.Build.Tests/JobCollectionTests.cs b/Engine/Source/Programs/Horde/Horde.Build.Tests/JobCollectionTests.cs
--- a/Engine/Source/Programs/Horde/Horde.Build.Tests/JobCollectionTests.cs
+++ b/Engine/Source/Programs/Horde/Horde.Build.Tests/JobCollectionTests.cs
@@ -31,20 +31,8 @@ namespace HordeServerTests
 	using LogId = ObjectId<ILogFile>;
 
 	[TestClass]
-	public class JobCollectionTests : DatabaseIntegrationTest
+	public class JobCollectionTests : TestSetup
 	{
-		IGraphCollection GraphCollection;
-		IUserCollection UserCollection;
-		IJobCollection JobCollection;
-
-		public JobCollectionTests()
-		{
-			DatabaseService DatabaseService = GetDatabaseService();
-			this.GraphCollection = new GraphCollection(DatabaseService);
-			this.UserCollection = new UserCollectionV2(DatabaseService, NullLogger<UserCollectionV2>.Instance);
-			this.JobCollection = new JobCollection(DatabaseService, UserCollection, NullLogger<JobCollection>.Instance);
-		}
-
 		NewGroup AddGroup(List<NewGroup> Groups)
 		{
 			NewGroup Group = new NewGroup("win64", new List<NewNode>());
@@ -132,12 +120,13 @@ namespace HordeServerTests
 		[TestMethod]
 		public async Task TryAssignLeaseTest()
 		{
-			TestSetup TestSetup = await GetTestSetup();
-			await TestSetup.JobCollection.TryAssignLeaseAsync(TestSetup.Fixture!.Job1, 0, new PoolId("foo"), TestSetup.Fixture.Agent1.Id,
+			Fixture Fixture = await CreateFixtureAsync();
+
+			await JobCollection.TryAssignLeaseAsync(Fixture.Job1, 0, new PoolId("foo"), Fixture.Agent1.Id,
 				ObjectId.GenerateNewId(), LeaseId.GenerateNewId(), LogId.GenerateNewId());
 			
-			IJob Job = (await TestSetup.JobCollection.GetAsync(TestSetup.Fixture!.Job1.Id))!;
-			await TestSetup.JobCollection.TryAssignLeaseAsync(Job, 0, new PoolId("foo"), TestSetup.Fixture.Agent1.Id,
+			IJob Job = (await JobCollection.GetAsync(Fixture.Job1.Id))!;
+			await JobCollection.TryAssignLeaseAsync(Job, 0, new PoolId("foo"), Fixture.Agent1.Id,
 				ObjectId.GenerateNewId(), LeaseId.GenerateNewId(), LogId.GenerateNewId());
 			
 			// Manually verify the log output
diff --git a/Engine/Source/Programs/Horde/Horde.Build.Tests/JobServiceTests.cs b/Engine/Source/Programs/Horde/Horde.Build.Tests/JobServiceTests.cs
--- a/Engine/Source/Programs/Horde/Horde.Build.Tests/JobServiceTests.cs
+++ b/Engine/Source/Programs/Horde/Horde.Build.Tests/JobServiceTests.cs
@@ -21,19 +21,17 @@ namespace HordeServerTests
 	using LogId = ObjectId<ILogFile>;
 
 	[TestClass]
-	public class JobServiceTests : DatabaseIntegrationTest
+	public class JobServiceTests : TestSetup
 	{
 		[TestMethod]
 		public async Task TestChainedJobs()
 		{
-			TestSetup TestSetup = await GetTestSetup();
-
 			ProjectId ProjectId = new ProjectId("ue5");
-			IProject? Project = await TestSetup.ProjectService.Collection.AddOrUpdateAsync(ProjectId, "", "", 0, new ProjectConfig { Name = "UE5" });
+			IProject? Project = await ProjectService.Collection.AddOrUpdateAsync(ProjectId, "", "", 0, new ProjectConfig { Name = "UE5" });
 			Assert.IsNotNull(Project);
 
-			ITemplate Template = await TestSetup.TemplateCollection.AddAsync("Test template", null, false, null, null, new List<string>(), new List<Parameter>());
-			IGraph Graph = await TestSetup.GraphCollection.AddAsync(Template);
+			ITemplate Template = await TemplateCollection.AddAsync("Test template", null, false, null, null, new List<string>(), new List<Parameter>());
+			IGraph Graph = await GraphCollection.AddAsync(Template);
 
 			TemplateRefId TemplateRefId1 = new TemplateRefId("template1");
 			TemplateRefId TemplateRefId2 = new TemplateRefId("template2");
@@ -44,19 +42,19 @@ namespace HordeServerTests
 			StreamConfig.Tabs.Add(new CreateJobsTabRequest { Title = "foo", Templates = new List<string> { TemplateRefId1.ToString(), TemplateRefId2.ToString() } });
 
 			StreamId StreamId = new StreamId("ue5-main");
-			IStream? Stream = await TestSetup.StreamService.GetStreamAsync(StreamId);
-			Stream = await TestSetup.StreamService.StreamCollection.TryCreateOrReplaceAsync(new StreamId("ue5-main"), Stream, String.Empty, String.Empty, ProjectId, StreamConfig);
+			IStream? Stream = await StreamService.GetStreamAsync(StreamId);
+			Stream = await StreamService.StreamCollection.TryCreateOrReplaceAsync(new StreamId("ue5-main"), Stream, String.Empty, String.Empty, ProjectId, StreamConfig);
 
-			IJob Job = await TestSetup.JobService.CreateJobAsync(null, Stream!, TemplateRefId1, Template.Id, Graph, "Hello", 1234, 1233, 999, null, null, null, null, null, Stream!.Templates[TemplateRefId1].ChainedJobs, true, true, null, null, new List<string>());
+			IJob Job = await JobService.CreateJobAsync(null, Stream!, TemplateRefId1, Template.Id, Graph, "Hello", 1234, 1233, 999, null, null, null, null, null, Stream!.Templates[TemplateRefId1].ChainedJobs, true, true, null, null, new List<string>());
 			Assert.AreEqual(1, Job.ChainedJobs.Count);
 
-			Job = Deref(await TestSetup.JobService.UpdateBatchAsync(Job, Job.Batches[0].Id, LogId.GenerateNewId(), JobStepBatchState.Running));
-			Job = Deref(await TestSetup.JobService.UpdateStepAsync(Job, Job.Batches[0].Id, Job.Batches[0].Steps[0].Id, JobStepState.Running));
-			Job = Deref(await TestSetup.JobService.UpdateStepAsync(Job, Job.Batches[0].Id, Job.Batches[0].Steps[0].Id, JobStepState.Completed, JobStepOutcome.Success));
+			Job = Deref(await JobService.UpdateBatchAsync(Job, Job.Batches[0].Id, LogId.GenerateNewId(), JobStepBatchState.Running));
+			Job = Deref(await JobService.UpdateStepAsync(Job, Job.Batches[0].Id, Job.Batches[0].Steps[0].Id, JobStepState.Running));
+			Job = Deref(await JobService.UpdateStepAsync(Job, Job.Batches[0].Id, Job.Batches[0].Steps[0].Id, JobStepState.Completed, JobStepOutcome.Success));
 
 			Assert.IsNotNull(Job.ChainedJobs[0].JobId);
 
-			IJob? ChainedJob = await TestSetup.JobCollection.GetAsync(Job.ChainedJobs[0].JobId!.Value);
+			IJob? ChainedJob = await JobCollection.GetAsync(Job.ChainedJobs[0].JobId!.Value);
 			Assert.IsNotNull(ChainedJob);
 			Assert.AreEqual(ChainedJob!.Id, Job!.ChainedJobs[0].JobId);
 
@@ -69,22 +67,22 @@ namespace HordeServerTests
 		[TestMethod]
 		public async Task StopAnyDuplicateJobsByPreflight()
 		{
-			TestSetup TestSetup = await GetTestSetup();
+			Fixture Fixture = await CreateFixtureAsync();
 
 			string[] Args = {"-Target=bogus"};
-			IJob OrgJob = await CreatePreflightJob(TestSetup, "tpl-ref-1", "tpl-hash-1", "elvis", 1000, Args);
-			IJob NewJob = await CreatePreflightJob(TestSetup, "tpl-ref-1", "tpl-hash-1", "elvis", 1000, Args);
-			IJob DifferentTplRef = await CreatePreflightJob(TestSetup, "tpl-ref-other", "tpl-hash-1", "elvis", 1000, Args);
-			IJob DifferentTplHash = await CreatePreflightJob(TestSetup, "tpl-ref-1", "tpl-hash-other", "elvis", 1000, Args);
-			IJob DifferentUserName = await CreatePreflightJob(TestSetup, "tpl-ref-1", "tpl-hash-1", "julia", 1000, Args);
-			IJob DifferentArgs = await CreatePreflightJob(TestSetup, "tpl-ref-1", "tpl-hash-1", "elvis", 1000, new []{"-Target=other"});
+			IJob OrgJob = await CreatePreflightJob(Fixture, "tpl-ref-1", "tpl-hash-1", "elvis", 1000, Args);
+			IJob NewJob = await CreatePreflightJob(Fixture, "tpl-ref-1", "tpl-hash-1", "elvis", 1000, Args);
+			IJob DifferentTplRef = await CreatePreflightJob(Fixture, "tpl-ref-other", "tpl-hash-1", "elvis", 1000, Args);
+			IJob DifferentTplHash = await CreatePreflightJob(Fixture, "tpl-ref-1", "tpl-hash-other", "elvis", 1000, Args);
+			IJob DifferentUserName = await CreatePreflightJob(Fixture, "tpl-ref-1", "tpl-hash-1", "julia", 1000, Args);
+			IJob DifferentArgs = await CreatePreflightJob(Fixture, "tpl-ref-1", "tpl-hash-1", "elvis", 1000, new []{"-Target=other"});
 			
-			OrgJob = (await TestSetup.JobService.GetJobAsync(OrgJob.Id))!;
-			NewJob = (await TestSetup.JobService.GetJobAsync(NewJob.Id))!;
-			DifferentTplRef = (await TestSetup.JobService.GetJobAsync(DifferentTplRef.Id))!;
-			DifferentTplHash = (await TestSetup.JobService.GetJobAsync(DifferentTplHash.Id))!;
-			DifferentUserName = (await TestSetup.JobService.GetJobAsync(DifferentUserName.Id))!;
-			DifferentArgs = (await TestSetup.JobService.GetJobAsync(DifferentArgs.Id))!;
+			OrgJob = (await JobService.GetJobAsync(OrgJob.Id))!;
+			NewJob = (await JobService.GetJobAsync(NewJob.Id))!;
+			DifferentTplRef = (await JobService.GetJobAsync(DifferentTplRef.Id))!;
+			DifferentTplHash = (await JobService.GetJobAsync(DifferentTplHash.Id))!;
+			DifferentUserName = (await JobService.GetJobAsync(DifferentUserName.Id))!;
+			DifferentArgs = (await JobService.GetJobAsync(DifferentArgs.Id))!;
 			
 			Assert.AreEqual(KnownUsers.System, OrgJob.AbortedByUserId);
 			Assert.IsNull(NewJob.AbortedByUserId);
@@ -94,15 +92,15 @@ namespace HordeServerTests
 			Assert.IsNull(DifferentArgs.AbortedByUserId);
 		}
 
-		private async Task<IJob> CreatePreflightJob(TestSetup TestSetup, string TemplateRefId, string TemplateHash, string StartedByUserName, int PreflightChange, string[] Arguments)
+		private async Task<IJob> CreatePreflightJob(Fixture Fixture, string TemplateRefId, string TemplateHash, string StartedByUserName, int PreflightChange, string[] Arguments)
 		{
-			IUser User = await TestSetup.UserCollection.FindOrAddUserByLoginAsync(StartedByUserName);
-			return await TestSetup.JobService.CreateJobAsync(
+			IUser User = await UserCollection.FindOrAddUserByLoginAsync(StartedByUserName);
+			return await JobService.CreateJobAsync(
 				JobId: JobId.GenerateNewId(),
-				Stream: TestSetup.Fixture!.Stream!,
+				Stream: Fixture!.Stream!,
 				TemplateRefId: new TemplateRefId(TemplateRefId),
 				TemplateHash: new ContentHash(Encoding.ASCII.GetBytes(TemplateHash)),
-				Graph: TestSetup.Fixture!.Graph,
+				Graph: Fixture!.Graph,
 				Name: "hello1",
 				Change: 1000001,
 				CodeChange: 1000002,
@@ -126,21 +124,21 @@ namespace HordeServerTests
 		// public async Task StopAnyDuplicateJobsByChange()
 		// {
 		// 	TestSetup TestSetup = await GetTestSetup();
-		// 	Fixture Fix = TestSetup.Fixture!;
+		// 	Fixture Fix = Fixture!;
 		// 	IJob Job1 = Fix.Job1;
 		// 	IJob Job2 = Fix.Job2;
 		// 	
 		// 	Assert.AreEqual(JobState.Waiting, Job1.GetState());
-		// 	Assert.IsTrue(await TestSetup.JobService.UpdateBatchAsync(Job1, Job1.Batches[0].Id, ObjectId.GenerateNewId(), JobStepBatchState.Running));
+		// 	Assert.IsTrue(await JobService.UpdateBatchAsync(Job1, Job1.Batches[0].Id, ObjectId.GenerateNewId(), JobStepBatchState.Running));
 		//
-		// 	Job1 = (await TestSetup.JobService.GetJobAsync(Job1.Id))!;
-		// 	Job2 = (await TestSetup.JobService.GetJobAsync(Job2.Id))!;
+		// 	Job1 = (await JobService.GetJobAsync(Job1.Id))!;
+		// 	Job2 = (await JobService.GetJobAsync(Job2.Id))!;
 		// 	Assert.AreEqual(JobState.Running, Job1.GetState());
 		// 	Assert.AreEqual(JobState.Complete, Job2.GetState());
 		// 	
-		// 	IJob Job3 = await TestSetup.JobService.CreateJobAsync(null, Fix.Stream!.Id, Fix.TemplateRefId1, Fix.Template.Id, Fix.Graph, "Hello", Job1.Change, Job1.CodeChange, null, "joe", null, null, true, true, null, null, new List<string>());
-		// 	await TestSetup.DispatchService.TickOnlyForTestingAsync();
-		// 	Job1 = (await TestSetup.JobService.GetJobAsync(Job1.Id))!;
+		// 	IJob Job3 = await JobService.CreateJobAsync(null, Fix.Stream!.Id, Fix.TemplateRefId1, Fix.Template.Id, Fix.Graph, "Hello", Job1.Change, Job1.CodeChange, null, "joe", null, null, true, true, null, null, new List<string>());
+		// 	await DispatchService.TickOnlyForTestingAsync();
+		// 	Job1 = (await JobService.GetJobAsync(Job1.Id))!;
 		// 	Assert.AreEqual(JobState.Complete, Job1.GetState());
 		// 	Assert.AreEqual(JobState.Waiting, Job3.GetState());
 		// }
@@ -148,17 +146,15 @@ namespace HordeServerTests
 		[TestMethod]
 		public async Task TestRunEarly()
 		{
-			TestSetup TestSetup = await GetTestSetup();
-
-			IProject? Project = await TestSetup.ProjectService.Collection.AddOrUpdateAsync(new ProjectId("ue5"), "", "", 0, new ProjectConfig { Name = "UE5" });
+			IProject? Project = await ProjectService.Collection.AddOrUpdateAsync(new ProjectId("ue5"), "", "", 0, new ProjectConfig { Name = "UE5" });
 			Assert.IsNotNull(Project);
 
 			StreamId StreamId = new StreamId("ue5-main");
-			IStream? Stream = await TestSetup.StreamCollection.GetAsync(StreamId);
-			Stream = await TestSetup.StreamCollection.TryCreateOrReplaceAsync(StreamId, Stream, "", "", Project!.Id, new StreamConfig { Name = "//UE5/Main" });
+			IStream? Stream = await StreamCollection.GetAsync(StreamId);
+			Stream = await StreamCollection.TryCreateOrReplaceAsync(StreamId, Stream, "", "", Project!.Id, new StreamConfig { Name = "//UE5/Main" });
 
-			ITemplate Template = await TestSetup.TemplateCollection.AddAsync("Test template", null, false, null, null, new List<string>(), new List<Parameter>());
-			IGraph Graph = await TestSetup.GraphCollection.AddAsync(Template);
+			ITemplate Template = await TemplateCollection.AddAsync("Test template", null, false, null, null, new List<string>(), new List<Parameter>());
+			IGraph Graph = await GraphCollection.AddAsync(Template);
 
 			NewGroup GroupA = new NewGroup("win", new List<NewNode>());
 			GroupA.Nodes.Add(new NewNode("Compile"));
@@ -168,22 +164,22 @@ namespace HordeServerTests
 			GroupB.Nodes.Add(new NewNode("Middle"));
 			GroupB.Nodes.Add(new NewNode("Pak", InputDependencies: new List<string> { "Compile", "Cook" }));
 
-			Graph = await TestSetup.GraphCollection.AppendAsync(Graph, new List<NewGroup> { GroupA, GroupB });
+			Graph = await GraphCollection.AppendAsync(Graph, new List<NewGroup> { GroupA, GroupB });
 
-			IJob Job = await TestSetup.JobService.CreateJobAsync(null, Stream!, new TemplateRefId("temp"), Template.Id, Graph, "Hello", 1234, 1233, 999, null, null, null, null, null, null, true, true, null, null, new List<string> { "-Target=Pak" });
+			IJob Job = await JobService.CreateJobAsync(null, Stream!, new TemplateRefId("temp"), Template.Id, Graph, "Hello", 1234, 1233, 999, null, null, null, null, null, null, true, true, null, null, new List<string> { "-Target=Pak" });
 
-			Job = Deref(await TestSetup.JobService.UpdateBatchAsync(Job, Job.Batches[0].Id, LogId.GenerateNewId(), JobStepBatchState.Running));
-			Assert.IsNotNull(await TestSetup.JobService.UpdateStepAsync(Job, Job.Batches[0].Id, Job.Batches[0].Steps[0].Id, JobStepState.Running));
-			Assert.IsNotNull(await TestSetup.JobService.UpdateStepAsync(Job, Job.Batches[0].Id, Job.Batches[0].Steps[0].Id, JobStepState.Completed, JobStepOutcome.Success));
+			Job = Deref(await JobService.UpdateBatchAsync(Job, Job.Batches[0].Id, LogId.GenerateNewId(), JobStepBatchState.Running));
+			Assert.IsNotNull(await JobService.UpdateStepAsync(Job, Job.Batches[0].Id, Job.Batches[0].Steps[0].Id, JobStepState.Running));
+			Assert.IsNotNull(await JobService.UpdateStepAsync(Job, Job.Batches[0].Id, Job.Batches[0].Steps[0].Id, JobStepState.Completed, JobStepOutcome.Success));
 
-			Job = Deref(await TestSetup.JobService.UpdateBatchAsync(Job, Job.Batches[1].Id, LogId.GenerateNewId(), JobStepBatchState.Running));
-			Assert.IsNotNull(await TestSetup.JobService.UpdateStepAsync(Job, Job.Batches[1].Id, Job.Batches[1].Steps[0].Id, JobStepState.Running));
+			Job = Deref(await JobService.UpdateBatchAsync(Job, Job.Batches[1].Id, LogId.GenerateNewId(), JobStepBatchState.Running));
+			Assert.IsNotNull(await JobService.UpdateStepAsync(Job, Job.Batches[1].Id, Job.Batches[1].Steps[0].Id, JobStepState.Running));
 
-			Job = Deref(await TestSetup.JobService.UpdateBatchAsync(Job, Job.Batches[2].Id, LogId.GenerateNewId(), JobStepBatchState.Running));
-			Assert.IsNotNull(await TestSetup.JobService.UpdateStepAsync(Job, Job.Batches[2].Id, Job.Batches[2].Steps[0].Id, JobStepState.Running));
-			Assert.IsNotNull(await TestSetup.JobService.UpdateStepAsync(Job, Job.Batches[2].Id, Job.Batches[2].Steps[0].Id, JobStepState.Completed, JobStepOutcome.Success));
+			Job = Deref(await JobService.UpdateBatchAsync(Job, Job.Batches[2].Id, LogId.GenerateNewId(), JobStepBatchState.Running));
+			Assert.IsNotNull(await JobService.UpdateStepAsync(Job, Job.Batches[2].Id, Job.Batches[2].Steps[0].Id, JobStepState.Running));
+			Assert.IsNotNull(await JobService.UpdateStepAsync(Job, Job.Batches[2].Id, Job.Batches[2].Steps[0].Id, JobStepState.Completed, JobStepOutcome.Success));
 
-			Job = (await TestSetup.JobService.GetJobAsync(Job.Id))!;
+			Job = (await JobService.GetJobAsync(Job.Id))!;
 
 			Assert.AreEqual(JobStepState.Waiting, Job.Batches[2].Steps[1].State);
 		}
diff --git a/Engine/Source/Programs/Horde/Horde.Build.Tests/JobsControllerDbTest.cs b/Engine/Source/Programs/Horde/Horde.Build.Tests/JobsControllerDbTest.cs
--- a/Engine/Source/Programs/Horde/Horde.Build.Tests/JobsControllerDbTest.cs
+++ b/Engine/Source/Programs/Horde/Horde.Build.Tests/JobsControllerDbTest.cs
@@ -18,20 +18,20 @@ namespace HordeServerTests
     /// the entire ASP.NET chain.
     /// </summary>
 	[TestClass]
-    public class JobsControllerDbTest : DatabaseIntegrationTest
+    public class JobsControllerDbTest : TestSetup
     {
         
         [TestMethod]
         public async Task GetJobs()
         {
-	        TestSetup TestSetup = await GetTestSetup();
+			Fixture Fixture = await CreateFixtureAsync();
 
-	        ActionResult<List<object>> Res = await TestSetup.JobsController.FindJobsAsync();
+			ActionResult<List<object>> Res = await JobsController.FindJobsAsync();
 	        Assert.AreEqual(2, Res.Value!.Count);
 	        Assert.AreEqual("hello2", (Res.Value[0] as GetJobResponse)!.Name);
 	        Assert.AreEqual("hello1", (Res.Value[1] as GetJobResponse)!.Name);
 	        
-	        Res = await TestSetup.JobsController.FindJobsAsync(IncludePreflight: false);
+	        Res = await JobsController.FindJobsAsync(IncludePreflight: false);
 	        Assert.AreEqual(1, Res.Value!.Count);
 	        Assert.AreEqual("hello2", (Res.Value[0] as GetJobResponse)!.Name);
         }
@@ -39,10 +39,9 @@ namespace HordeServerTests
         [TestMethod]
         public async Task AbortStepTest()
         {
-	        TestSetup TestSetup = await GetTestSetup();
-	        JobsController JobsController = TestSetup.JobsController;
+			Fixture Fixture = await CreateFixtureAsync();
 
-	        IJob Job = TestSetup.Fixture!.Job1;
+	        IJob Job = Fixture.Job1;
 	        SubResourceId BatchId = Job.Batches[0].Id;
 	        SubResourceId StepId = Job.Batches[0].Steps[0].Id;
 
diff --git a/Engine/Source/Programs/Horde/Horde.Build.Tests/PoolServiceTests.cs b/Engine/Source/Programs/Horde/Horde.Build.Tests/PoolServiceTests.cs
--- a/Engine/Source/Programs/Horde/Horde.Build.Tests/PoolServiceTests.cs
+++ b/Engine/Source/Programs/Horde/Horde.Build.Tests/PoolServiceTests.cs
@@ -13,9 +13,8 @@ using PoolId = HordeServer.Utilities.StringId<HordeServer.Models.IPool>;
 namespace HordeServerTests
 {
     [TestClass]
-    public class PoolServiceTests : DatabaseIntegrationTest
+    public class PoolServiceTests : TestSetup
     {
-        private readonly PoolService PoolService;
         private readonly Dictionary<string, string> FixtureProps = new Dictionary<string, string>
         {
             {"foo", "bar"},
@@ -24,8 +23,6 @@ namespace HordeServerTests
 
         public PoolServiceTests()
         {
-            var PoolCollection = new PoolCollection(GetDatabaseService());
-            PoolService = new PoolService(GetDatabaseService(), PoolCollection, new FakeClock());
         }
 
         private async Task<IPool> CreatePoolFixture(String Name)
@@ -96,10 +93,8 @@ namespace HordeServerTests
         [TestMethod]
         public async Task UpdatePoolCollectionTest()
         {
-	        TestSetup TestSetup = await GetTestSetup();
-	        
 	        IPool Pool = await CreatePoolFixture("update-pool-2");
-	        await TestSetup.PoolCollection.TryUpdateAsync(Pool, LastScaleUpTime: DateTime.UtcNow);
+	        await PoolCollection.TryUpdateAsync(Pool, LastScaleUpTime: DateTime.UtcNow);
         }
     }
 }
\ No newline at end of file
diff --git a/Engine/Source/Programs/Horde/Horde.Build.Tests/RpcServiceTest.cs b/Engine/Source/Programs/Horde/Horde.Build.Tests/RpcServiceTest.cs
--- a/Engine/Source/Programs/Horde/Horde.Build.Tests/RpcServiceTest.cs
+++ b/Engine/Source/Programs/Horde/Horde.Build.Tests/RpcServiceTest.cs
@@ -56,7 +56,7 @@ namespace HordeServerTests
 	}
 
 	[TestClass]
-	public class RpcServiceTest : DatabaseIntegrationTest
+	public class RpcServiceTest : TestSetup
 	{
 		private readonly ServerCallContext AdminContext = new ServerCallContextStub("app-horde-admins");
 
@@ -298,15 +298,14 @@ namespace HordeServerTests
 		[TestMethod]
 		public async Task CreateSessionTest()
 		{
-			TestSetup TestSetup = await GetTestSetup();
 			CreateSessionRequest Req = new CreateSessionRequest();
-			await Assert.ThrowsExceptionAsync<StructuredRpcException>(() => TestSetup.RpcService.CreateSession(Req, AdminContext));
+			await Assert.ThrowsExceptionAsync<StructuredRpcException>(() => RpcService.CreateSession(Req, AdminContext));
 
 			Req.Name = "MyName";
-			await Assert.ThrowsExceptionAsync<StructuredRpcException>(() => TestSetup.RpcService.CreateSession(Req, AdminContext));
+			await Assert.ThrowsExceptionAsync<StructuredRpcException>(() => RpcService.CreateSession(Req, AdminContext));
 
 			Req.Capabilities = new AgentCapabilities();
-			CreateSessionResponse Res = await TestSetup.RpcService.CreateSession(Req, AdminContext);
+			CreateSessionResponse Res = await RpcService.CreateSession(Req, AdminContext);
 
 			Assert.AreEqual("MYNAME", Res.AgentId);
 			// TODO: Check Token, ExpiryTime, SessionId 
@@ -315,12 +314,11 @@ namespace HordeServerTests
 		[TestMethod]
 		public async Task UpdateSessionTest()
 		{
-			TestSetup TestSetup = await GetTestSetup();
 			CreateSessionRequest CreateReq = new CreateSessionRequest
 			{
 				Name = "UpdateSessionTest1", Capabilities = new AgentCapabilities()
 			};
-			CreateSessionResponse CreateRes = await TestSetup.RpcService.CreateSession(CreateReq, AdminContext);
+			CreateSessionResponse CreateRes = await RpcService.CreateSession(CreateReq, AdminContext);
 			string AgentId = CreateRes.AgentId;
 			string SessionId = CreateRes.SessionId;
 
@@ -328,7 +326,7 @@ namespace HordeServerTests
 				new TestAsyncStreamReader<UpdateSessionRequest>(AdminContext);
 			TestServerStreamWriter<UpdateSessionResponse> ResponseStream =
 				new TestServerStreamWriter<UpdateSessionResponse>(AdminContext);
-			Task Call = TestSetup.RpcService.UpdateSession(RequestStream, ResponseStream, AdminContext);
+			Task Call = RpcService.UpdateSession(RequestStream, ResponseStream, AdminContext);
 
 			RequestStream.AddMessage(new UpdateSessionRequest {AgentId = "does-not-exist", SessionId = SessionId});
 			StructuredRpcException Re = await Assert.ThrowsExceptionAsync<StructuredRpcException>(() => Call);
@@ -339,14 +337,13 @@ namespace HordeServerTests
 		[TestMethod]
 		public async Task QueryServerSessionTest()
 		{
-			TestSetup TestSetup = await GetTestSetup();
-			TestSetup.RpcService.LongPollTimeout = TimeSpan.FromMilliseconds(200);
+			RpcService.LongPollTimeout = TimeSpan.FromMilliseconds(200);
 
 			TestAsyncStreamReader<QueryServerStateRequest> RequestStream =
 				new TestAsyncStreamReader<QueryServerStateRequest>(AdminContext);
 			TestServerStreamWriter<QueryServerStateResponse> ResponseStream =
 				new TestServerStreamWriter<QueryServerStateResponse>(AdminContext);
-			Task Call = TestSetup.RpcService.QueryServerState(RequestStream, ResponseStream, AdminContext);
+			Task Call = RpcService.QueryServerState(RequestStream, ResponseStream, AdminContext);
 
 			RequestStream.AddMessage(new QueryServerStateRequest {Name = "bogusAgentName"});
 			QueryServerStateResponse? Res = await ResponseStream.ReadNextAsync();
@@ -362,13 +359,11 @@ namespace HordeServerTests
 		[TestMethod]
 		public async Task FinishBatchTest()
 		{
-			TestSetup TestSetup = await GetTestSetup();
-			
 			CreateSessionRequest CreateReq = new CreateSessionRequest
 			{
 				Name = "UpdateSessionTest1", Capabilities = new AgentCapabilities()
 			};
-			CreateSessionResponse CreateRes = await TestSetup.RpcService.CreateSession(CreateReq, AdminContext);
+			CreateSessionResponse CreateRes = await RpcService.CreateSession(CreateReq, AdminContext);
 			string AgentId = CreateRes.AgentId;
 			string SessionId = CreateRes.SessionId;
 
@@ -376,7 +371,7 @@ namespace HordeServerTests
 				new TestAsyncStreamReader<UpdateSessionRequest>(AdminContext);
 			TestServerStreamWriter<UpdateSessionResponse> ResponseStream =
 				new TestServerStreamWriter<UpdateSessionResponse>(AdminContext);
-			Task Call = TestSetup.RpcService.UpdateSession(RequestStream, ResponseStream, AdminContext);
+			Task Call = RpcService.UpdateSession(RequestStream, ResponseStream, AdminContext);
 
 			RequestStream.AddMessage(new UpdateSessionRequest {AgentId = "does-not-exist", SessionId = SessionId});
 			StructuredRpcException Re = await Assert.ThrowsExceptionAsync<StructuredRpcException>(() => Call);
@@ -387,7 +382,8 @@ namespace HordeServerTests
 		[TestMethod]
 		public async Task UploadArtifactTest()
 		{
-			TestSetup TestSetup = await GetTestSetup();
+			Fixture Fixture = await CreateFixtureAsync();
+
 			ObjectId SessionId = ObjectId.GenerateNewId();
 			ServerCallContext Context = new ServerCallContextStub(new ClaimsPrincipal(new ClaimsIdentity(new List<Claim>
 			{
@@ -401,16 +397,16 @@ namespace HordeServerTests
 			
 			UploadArtifactMetadata Metadata = new UploadArtifactMetadata
 			{
-				JobId = TestSetup.Fixture!.Job1.Id.ToString(),
-				BatchId = TestSetup.Fixture!.Job1.Batches[0].Id.ToString(),
-				StepId = TestSetup.Fixture!.Job1.Batches[0].Steps[0].Id.ToString(),
+				JobId = Fixture.Job1.Id.ToString(),
+				BatchId = Fixture.Job1.Batches[0].Id.ToString(),
+				StepId = Fixture.Job1.Batches[0].Steps[0].Id.ToString(),
 				Name = "testfile.txt",
 				MimeType = "text/plain",
 				Length = DataStr.Length
 			};
 
 			// Set the session ID on the job batch to pass auth later
-			Deref(await TestSetup.JobCollection.TryAssignLeaseAsync(TestSetup.Fixture.Job1, 0, new PoolId("foo"),
+			Deref(await JobCollection.TryAssignLeaseAsync(Fixture.Job1, 0, new PoolId("foo"),
 				new AgentId("test"), SessionId,
 				LeaseId.GenerateNewId(), LogId.GenerateNewId()));
 /*
@@ -426,7 +422,7 @@ namespace HordeServerTests
 			await Task.Delay(500);
 */			
 			TestAsyncStreamReader<UploadArtifactRequest> RequestStream = new TestAsyncStreamReader<UploadArtifactRequest>(Context);
-			Task<UploadArtifactResponse> Call = TestSetup.RpcService.UploadArtifact(RequestStream,  Context);
+			Task<UploadArtifactResponse> Call = RpcService.UploadArtifact(RequestStream,  Context);
 			RequestStream.AddMessage(new UploadArtifactRequest { Metadata = Metadata });
 			RequestStream.AddMessage(new UploadArtifactRequest { Data = ByteString.CopyFromUtf8(Data[0]) });
 			RequestStream.AddMessage(new UploadArtifactRequest { Data = ByteString.CopyFromUtf8(Data[1]) });
@@ -435,9 +431,9 @@ namespace HordeServerTests
 			UploadArtifactResponse Res = await Call;
 
 
-			IArtifact? Artifact = await TestSetup.ArtifactCollection.GetArtifactAsync(ObjectId.Parse(Res.Id));
+			IArtifact? Artifact = await ArtifactCollection.GetArtifactAsync(ObjectId.Parse(Res.Id));
 			Assert.IsNotNull(Artifact);
-			Stream Stream = await TestSetup.ArtifactCollection.OpenArtifactReadStreamAsync(Artifact!);
+			Stream Stream = await ArtifactCollection.OpenArtifactReadStreamAsync(Artifact!);
 			StreamReader Reader = new StreamReader(Stream);
 			string text = Reader.ReadToEnd();
 			Assert.AreEqual(DataStr, text);
diff --git a/Engine/Source/Programs/Horde/Horde.Build.Tests/SchedulerTests.cs b/Engine/Source/Programs/Horde/Horde.Build.Tests/SchedulerTests.cs
--- a/Engine/Source/Programs/Horde/Horde.Build.Tests/SchedulerTests.cs
+++ b/Engine/Source/Programs/Horde/Horde.Build.Tests/SchedulerTests.cs
@@ -26,36 +26,26 @@ namespace HordeServerTests
 	using LogId = ObjectId<ILogFile>;
 
 	[TestClass]
-    public class SchedulerTests : DatabaseIntegrationTest
+    public class SchedulerTests : TestSetup
 	{
 		ProjectId ProjectId { get; } = new ProjectId("ue5");
 		StreamId StreamId { get; } = new StreamId("ue5-main");
 		TemplateRefId TemplateRefId { get; } = new TemplateRefId("template1");
 
-		IServiceProvider ServiceProvider;
-		TestSetup TestSetup;
 		ITemplate Template;
 		HashSet<JobId> InitialJobIds;
-		ScheduleService ScheduleService;
-		PerforceServiceStub PerforceService;
 
 		public SchedulerTests()
 		{
-			TestSetup = GetTestSetup().Result;
+			IUser Bob = UserCollection.FindOrAddUserByLoginAsync("Bob").Result;
 
-			IUser Bob = TestSetup.UserCollection.FindOrAddUserByLoginAsync("Bob").Result;
-
-			IProject ? Project = TestSetup.ProjectService.Collection.AddOrUpdateAsync(ProjectId, "", "", 0, new ProjectConfig { Name = "UE5" }).Result;
+			IProject ? Project = ProjectService.Collection.AddOrUpdateAsync(ProjectId, "", "", 0, new ProjectConfig { Name = "UE5" }).Result;
 			Assert.IsNotNull(Project);
 
-			Template = TestSetup.TemplateCollection.AddAsync("Test template", null, false, null, null, new List<string>(), new List<Parameter>()).Result;
-
-			InitialJobIds = new HashSet<JobId>(TestSetup.JobCollection.FindAsync().Result.Select(x => x.Id));
+			Template = TemplateCollection.AddAsync("Test template", null, false, null, null, new List<string>(), new List<Parameter>()).Result;
 
-			ServiceProvider = TestSetup.ServiceProvider;
-			ScheduleService = ServiceProvider.GetRequiredService<ScheduleService>();
+			InitialJobIds = new HashSet<JobId>(JobCollection.FindAsync().Result.Select(x => x.Id));
 
-			PerforceService = (PerforceServiceStub)TestSetup.PerforceService;
 			PerforceService.Changes.Clear();
 			PerforceService.AddChange("//UE5/Main", 100, Bob, "", new[] { "code.cpp" });
 			PerforceService.AddChange("//UE5/Main", 101, Bob, "", new[] { "content.uasset" });
@@ -64,20 +54,19 @@ namespace HordeServerTests
 
 		async Task<IStream> SetScheduleAsync(CreateScheduleRequest Schedule)
 		{
-			IStream? Stream = await TestSetup.StreamService.GetStreamAsync(StreamId);
+			IStream? Stream = await StreamService.GetStreamAsync(StreamId);
 
 			StreamConfig Config = new StreamConfig();
 			Config.Name = "//UE5/Main";
 			Config.Tabs.Add(new CreateJobsTabRequest { Title = "foo", Templates = new List<string> { TemplateRefId.ToString() } });
 			Config.Templates.Add(new CreateTemplateRefRequest { Id = TemplateRefId.ToString(), Name = "Test", Schedule = Schedule });
 
-			StreamService StreamService = ServiceProvider.GetRequiredService<StreamService>();
 			return (await StreamService.StreamCollection.TryCreateOrReplaceAsync(StreamId, Stream, "", "", ProjectId, Config))!; 
 		}
 
 		public async Task<List<IJob>> FileTestHelperAsync(params string[] Files)
 		{
-			IUser Bob = TestSetup.UserCollection.FindOrAddUserByLoginAsync("Bob", "Bob").Result;
+			IUser Bob = UserCollection.FindOrAddUserByLoginAsync("Bob", "Bob").Result;
 
 			PerforceService.Changes.Clear();
 			PerforceService.AddChange("//UE5/Main", 100, Bob, "", new[] { "code.cpp" });
@@ -89,7 +78,7 @@ namespace HordeServerTests
 			PerforceService.AddChange("//UE5/Main", 106, Bob, "", new[] { "bar/foo/content.uasset" });
 
 			DateTime StartTime = new DateTime(2021, 1, 1, 12, 0, 0, DateTimeKind.Local); // Friday Jan 1, 2021 
-			TestSetup.Clock.UtcNow = StartTime;
+			Clock.UtcNow = StartTime;
 
 			CreateScheduleRequest Schedule = new CreateScheduleRequest();
 			Schedule.Enabled = true;
@@ -98,7 +87,7 @@ namespace HordeServerTests
 			Schedule.Files = Files.ToList();
 			await SetScheduleAsync(Schedule);
 
-			TestSetup.Clock.Advance(TimeSpan.FromHours(1.25));
+			Clock.Advance(TimeSpan.FromHours(1.25));
 			await ScheduleService.TickSharedOnlyForTestingAsync();
 
 			return await GetNewJobs();
@@ -137,7 +126,7 @@ namespace HordeServerTests
 		public void DayScheduleTest()
 		{
 			DateTime StartTime = new DateTime(2021, 1, 1, 12, 0, 0, DateTimeKind.Utc); // Friday Jan 1, 2021 
-			TestSetup.Clock.UtcNow = StartTime;
+			Clock.UtcNow = StartTime;
 
 			Schedule Schedule = new Schedule(RequireSubmittedChange: false);
 			Schedule.Patterns.Add(new SchedulePattern(new List<DayOfWeek> { DayOfWeek.Friday, DayOfWeek.Sunday }, 13 * 60, null, null));
@@ -159,7 +148,7 @@ namespace HordeServerTests
 		public void MulticheduleTest()
 		{
 			DateTime StartTime = new DateTime(2021, 1, 1, 12, 0, 0, DateTimeKind.Utc); // Friday Jan 1, 2021 
-			TestSetup.Clock.UtcNow = StartTime;
+			Clock.UtcNow = StartTime;
 
 			Schedule Schedule = new Schedule(RequireSubmittedChange: false);
 			Schedule.Patterns.Add(new SchedulePattern(null, 13 * 60, 14 * 60, 15));
@@ -187,7 +176,7 @@ namespace HordeServerTests
 		public async Task NoSubmittedChangeTestAsync()
 		{
 			DateTime StartTime = new DateTime(2021, 1, 1, 12, 0, 0, DateTimeKind.Local); // Friday Jan 1, 2021 
-			TestSetup.Clock.UtcNow = StartTime;
+			Clock.UtcNow = StartTime;
 
 			CreateScheduleRequest Schedule = new CreateScheduleRequest();
 			Schedule.Enabled = true;
@@ -202,7 +191,7 @@ namespace HordeServerTests
 			Assert.AreEqual(0, Jobs1.Count);
 
 			// Trigger a job
-			TestSetup.Clock.Advance(TimeSpan.FromHours(1.25));
+			Clock.Advance(TimeSpan.FromHours(1.25));
 			await ScheduleService.TickSharedOnlyForTestingAsync();
 
 			List<IJob> Jobs2 = await GetNewJobs();
@@ -215,7 +204,7 @@ namespace HordeServerTests
 		public async Task RequireSubmittedChangeTestAsync()
 		{
 			DateTime StartTime = new DateTime(2021, 1, 1, 12, 0, 0, DateTimeKind.Local); // Friday Jan 1, 2021 
-			TestSetup.Clock.UtcNow = StartTime;
+			Clock.UtcNow = StartTime;
 
 			CreateScheduleRequest Schedule = new CreateScheduleRequest();
 			Schedule.Enabled = true;
@@ -229,7 +218,7 @@ namespace HordeServerTests
 			Assert.AreEqual(0, Jobs1.Count);
 
 			// Trigger a job
-			TestSetup.Clock.Advance(TimeSpan.FromHours(1.25));
+			Clock.Advance(TimeSpan.FromHours(1.25));
 			await ScheduleService.TickSharedOnlyForTestingAsync();
 
 			List<IJob> Jobs2 = await GetNewJobs();
@@ -237,13 +226,13 @@ namespace HordeServerTests
 			Assert.AreEqual(102, Jobs2[0].Change);
 			Assert.AreEqual(100, Jobs2[0].CodeChange);
 
-			IStream Stream2 = (await TestSetup.StreamCollection.GetAsync(Stream.Id))!;
+			IStream Stream2 = (await StreamCollection.GetAsync(Stream.Id))!;
 			Schedule Schedule2 = Stream2.Templates.First().Value.Schedule!;
 			Assert.AreEqual(102, Schedule2.LastTriggerChange);
-			Assert.AreEqual(TestSetup.Clock.UtcNow, Schedule2.LastTriggerTime);
+			Assert.AreEqual(Clock.UtcNow, Schedule2.LastTriggerTime);
 
 			// Trigger another job
-			TestSetup.Clock.Advance(TimeSpan.FromHours(0.5));
+			Clock.Advance(TimeSpan.FromHours(0.5));
 			await ScheduleService.TickSharedOnlyForTestingAsync();
 
 			List<IJob> Jobs3 = await GetNewJobs();
@@ -254,7 +243,7 @@ namespace HordeServerTests
 		public async Task MultipleJobsTestAsync()
 		{
 			DateTime StartTime = new DateTime(2021, 1, 1, 12, 0, 0, DateTimeKind.Local); // Friday Jan 1, 2021 
-			TestSetup.Clock.UtcNow = StartTime;
+			Clock.UtcNow = StartTime;
 
 			CreateScheduleRequest Schedule = new CreateScheduleRequest();
 			Schedule.Enabled = true;
@@ -270,13 +259,13 @@ namespace HordeServerTests
 			Assert.AreEqual(0, Jobs1.Count);
 
 			// Trigger some jobs
-			IUser Bob = TestSetup.UserCollection.FindOrAddUserByLoginAsync("Bob").Result;
+			IUser Bob = UserCollection.FindOrAddUserByLoginAsync("Bob").Result;
 			PerforceService.AddChange("//UE5/Main", 103, Bob, "", new string[] { "foo.cpp" });
 			PerforceService.AddChange("//UE5/Main", 104, Bob, "", new string[] { "foo.cpp" });
 			PerforceService.AddChange("//UE5/Main", 105, Bob, "", new string[] { "foo.uasset" });
 			PerforceService.AddChange("//UE5/Main", 106, Bob, "", new string[] { "foo.cpp" });
 
-			TestSetup.Clock.Advance(TimeSpan.FromHours(1.25));
+			Clock.Advance(TimeSpan.FromHours(1.25));
 			await ScheduleService.TickSharedOnlyForTestingAsync();
 
 			List<IJob> Jobs2 = await GetNewJobs();
@@ -291,7 +280,7 @@ namespace HordeServerTests
 		public async Task MaxActiveTestAsync()
 		{
 			DateTime StartTime = new DateTime(2021, 1, 1, 12, 0, 0, DateTimeKind.Local); // Friday Jan 1, 2021 
-			TestSetup.Clock.UtcNow = StartTime;
+			Clock.UtcNow = StartTime;
 
 			CreateScheduleRequest Schedule = new CreateScheduleRequest();
 			Schedule.Enabled = true;
@@ -301,7 +290,7 @@ namespace HordeServerTests
 			await SetScheduleAsync(Schedule);
 
 			// Trigger a job
-			TestSetup.Clock.Advance(TimeSpan.FromHours(1.25));
+			Clock.Advance(TimeSpan.FromHours(1.25));
 			await ScheduleService.TickSharedOnlyForTestingAsync();
 
 			List<IJob> Jobs2 = await GetNewJobs();
@@ -310,17 +299,17 @@ namespace HordeServerTests
 			Assert.AreEqual(100, Jobs2[0].CodeChange);
 
 			// Test that another job does not trigger
-			TestSetup.Clock.Advance(TimeSpan.FromHours(0.5));
+			Clock.Advance(TimeSpan.FromHours(0.5));
 			await ScheduleService.TickSharedOnlyForTestingAsync();
 
 			List<IJob> Jobs3 = await GetNewJobs();
 			Assert.AreEqual(0, Jobs3.Count);
 
 			// Mark the original job as complete
-			await TestSetup.JobService.UpdateJobAsync(Jobs2[0], AbortedByUserId: KnownUsers.System);
+			await JobService.UpdateJobAsync(Jobs2[0], AbortedByUserId: KnownUsers.System);
 
 			// Test that another job does not trigger
-			TestSetup.Clock.Advance(TimeSpan.FromHours(0.5));
+			Clock.Advance(TimeSpan.FromHours(0.5));
 			await ScheduleService.TickSharedOnlyForTestingAsync();
 
 			List<IJob> Jobs4 = await GetNewJobs();
@@ -333,7 +322,7 @@ namespace HordeServerTests
 		public async Task CreateNewChangeTestAsync()
 		{
 			DateTime StartTime = new DateTime(2021, 1, 1, 12, 0, 0, DateTimeKind.Local); // Friday Jan 1, 2021 
-			TestSetup.Clock.UtcNow = StartTime;
+			Clock.UtcNow = StartTime;
 
 			CreateScheduleRequest Schedule = new CreateScheduleRequest();
 			Schedule.Enabled = true;
@@ -341,7 +330,7 @@ namespace HordeServerTests
 			IStream Stream = await SetScheduleAsync(Schedule);
 
 			// Trigger a job
-			TestSetup.Clock.Advance(TimeSpan.FromHours(1.25));
+			Clock.Advance(TimeSpan.FromHours(1.25));
 			await ScheduleService.TickSharedOnlyForTestingAsync();
 
 			List<IJob> Jobs2 = await GetNewJobs();
@@ -350,7 +339,7 @@ namespace HordeServerTests
 			Assert.AreEqual(100, Jobs2[0].CodeChange);
 
 			// Check another job does not trigger due to the change above
-			TestSetup.Clock.Advance(TimeSpan.FromHours(1.25));
+			Clock.Advance(TimeSpan.FromHours(1.25));
 			await ScheduleService.TickSharedOnlyForTestingAsync();
 
 			List<IJob> Jobs3 = await GetNewJobs();
@@ -361,14 +350,14 @@ namespace HordeServerTests
 		public async Task GateTestAsync()
 		{
 			DateTime StartTime = new DateTime(2021, 1, 1, 12, 0, 0, DateTimeKind.Local); // Friday Jan 1, 2021 
-			TestSetup.Clock.UtcNow = StartTime;
+			Clock.UtcNow = StartTime;
 
 			// Create two templates, the second dependent on the first
-			ITemplate? NewTemplate1 = await TestSetup.TemplateCollection.AddAsync("Test template 1");
+			ITemplate? NewTemplate1 = await TemplateCollection.AddAsync("Test template 1");
 			TemplateRef NewTemplateRef1 = new TemplateRef(NewTemplate1);
 			TemplateRefId NewTemplateRefId1 = new TemplateRefId("new-template-1");
 
-			ITemplate? NewTemplate2 = await TestSetup.TemplateCollection.AddAsync("Test template 2");
+			ITemplate? NewTemplate2 = await TemplateCollection.AddAsync("Test template 2");
 			TemplateRef NewTemplateRef2 = new TemplateRef(NewTemplate2);
 			NewTemplateRef2.Schedule = new Schedule();
 			NewTemplateRef2.Schedule.Gate = new ScheduleGate(NewTemplateRefId1, "TriggerNext");
@@ -376,19 +365,18 @@ namespace HordeServerTests
 			NewTemplateRef2.Schedule.LastTriggerTime = StartTime;
 			TemplateRefId NewTemplateRefId2 = new TemplateRefId("new-template-2");
 
-			IStream? Stream = await TestSetup.StreamService.GetStreamAsync(StreamId);
+			IStream? Stream = await StreamService.GetStreamAsync(StreamId);
 
 			StreamConfig Config = new StreamConfig();
 			Config.Name = "//UE5/Main";
 			Config.Tabs.Add(new CreateJobsTabRequest { Title = "foo", Templates = new List<string> { NewTemplateRefId1.ToString(), NewTemplateRefId2.ToString() } });
 
-			StreamService StreamService = ServiceProvider.GetRequiredService<StreamService>();
 			Stream = (await StreamService.StreamCollection.TryCreateOrReplaceAsync(StreamId, Stream, "", "", ProjectId, Config))!;
 
 			// Create the TriggerNext step and mark it as complete
-			IGraph GraphA = await TestSetup.GraphCollection.AddAsync(NewTemplate1);
+			IGraph GraphA = await GraphCollection.AddAsync(NewTemplate1);
 			NewGroup GroupA = new NewGroup("win", new List<NewNode> { new NewNode("TriggerNext") });
-			GraphA = await TestSetup.GraphCollection.AppendAsync(GraphA, new List<NewGroup> { GroupA });
+			GraphA = await GraphCollection.AppendAsync(GraphA, new List<NewGroup> { GroupA });
 
 			// Tick the schedule and make sure it doesn't trigger
 			await ScheduleService.TickSharedOnlyForTestingAsync();
@@ -396,22 +384,22 @@ namespace HordeServerTests
 			Assert.AreEqual(0, Jobs2.Count);
 
 			// Create a job and fail it
-			IJob Job1 = await TestSetup.JobService.CreateJobAsync(null, Stream, NewTemplateRefId1, Template.Id, GraphA, "Hello", 1234, 1233, 999, null, null, null, null, null, null, true, true, null, null, new List<string> { "-Target=TriggerNext" });
-			Job1 = Deref(await TestSetup.JobService.UpdateBatchAsync(Job1, Job1.Batches[0].Id, LogId.GenerateNewId(), JobStepBatchState.Running));
-			Assert.IsNotNull(await TestSetup.JobService.UpdateStepAsync(Job1, Job1.Batches[0].Id, Job1.Batches[0].Steps[0].Id, JobStepState.Completed, JobStepOutcome.Failure));
-			Job1 = Deref(await TestSetup.JobService.UpdateBatchAsync(Job1, Job1.Batches[0].Id, LogId.GenerateNewId(), JobStepBatchState.Complete));
+			IJob Job1 = await JobService.CreateJobAsync(null, Stream, NewTemplateRefId1, Template.Id, GraphA, "Hello", 1234, 1233, 999, null, null, null, null, null, null, true, true, null, null, new List<string> { "-Target=TriggerNext" });
+			Job1 = Deref(await JobService.UpdateBatchAsync(Job1, Job1.Batches[0].Id, LogId.GenerateNewId(), JobStepBatchState.Running));
+			Assert.IsNotNull(await JobService.UpdateStepAsync(Job1, Job1.Batches[0].Id, Job1.Batches[0].Steps[0].Id, JobStepState.Completed, JobStepOutcome.Failure));
+			Job1 = Deref(await JobService.UpdateBatchAsync(Job1, Job1.Batches[0].Id, LogId.GenerateNewId(), JobStepBatchState.Complete));
 			await GetNewJobs();
 
 			// Tick the schedule and make sure it doesn't trigger
-			TestSetup.Clock.Advance(TimeSpan.FromMinutes(30.0));
+			Clock.Advance(TimeSpan.FromMinutes(30.0));
 			await ScheduleService.TickSharedOnlyForTestingAsync();
 			List<IJob> Jobs3 = await GetNewJobs();
 			Assert.AreEqual(0, Jobs3.Count);
 
 			// Create a job and make it succeed
-			IJob Job2 = await TestSetup.JobService.CreateJobAsync(null, Stream, NewTemplateRefId1, Template.Id, GraphA, "Hello", 1234, 1233, 999, null, null, null, null, null, null, true, true, null, null, new List<string> { "-Target=TriggerNext" });
-			Job2 = Deref(await TestSetup.JobService.UpdateBatchAsync(Job2, Job2.Batches[0].Id, LogId.GenerateNewId(), JobStepBatchState.Running));
-			Assert.IsNotNull(await TestSetup.JobService.UpdateStepAsync(Job2, Job2.Batches[0].Id, Job2.Batches[0].Steps[0].Id, JobStepState.Completed, JobStepOutcome.Success));
+			IJob Job2 = await JobService.CreateJobAsync(null, Stream, NewTemplateRefId1, Template.Id, GraphA, "Hello", 1234, 1233, 999, null, null, null, null, null, null, true, true, null, null, new List<string> { "-Target=TriggerNext" });
+			Job2 = Deref(await JobService.UpdateBatchAsync(Job2, Job2.Batches[0].Id, LogId.GenerateNewId(), JobStepBatchState.Running));
+			Assert.IsNotNull(await JobService.UpdateStepAsync(Job2, Job2.Batches[0].Id, Job2.Batches[0].Steps[0].Id, JobStepState.Completed, JobStepOutcome.Success));
 
 			// Tick the schedule and make sure it does trigger
 			await ScheduleService.TickSharedOnlyForTestingAsync();
@@ -425,7 +413,7 @@ namespace HordeServerTests
 		public async Task UpdateConfigAsync()
 		{
 			DateTime StartTime = new DateTime(2021, 1, 1, 12, 0, 0, DateTimeKind.Local); // Friday Jan 1, 2021 
-			TestSetup.Clock.UtcNow = StartTime;
+			Clock.UtcNow = StartTime;
 
 			CreateScheduleRequest Schedule = new CreateScheduleRequest();
 			Schedule.Enabled = true;
@@ -434,7 +422,7 @@ namespace HordeServerTests
 			Schedule.MaxActive = 2;
 			await SetScheduleAsync(Schedule);
 
-			TestSetup.Clock.Advance(TimeSpan.FromHours(1.25));
+			Clock.Advance(TimeSpan.FromHours(1.25));
 			await ScheduleService.TickSharedOnlyForTestingAsync();
 
 			List<IJob> Jobs1 = await GetNewJobs();
@@ -443,7 +431,7 @@ namespace HordeServerTests
 			Assert.AreEqual(100, Jobs1[0].CodeChange);
 
 			// Make sure the job is registered
-			IStream? Stream1 = await TestSetup.StreamService.GetStreamAsync(StreamId);
+			IStream? Stream1 = await StreamService.GetStreamAsync(StreamId);
 			TemplateRef TemplateRef1 = Stream1!.Templates.First().Value;
 			Assert.AreEqual(1, TemplateRef1.Schedule!.ActiveJobs.Count);
 			Assert.AreEqual(Jobs1[0].Id, TemplateRef1.Schedule!.ActiveJobs[0]);
@@ -452,7 +440,7 @@ namespace HordeServerTests
 			await SetScheduleAsync(Schedule);
 
 			// Make sure the job is still registered
-			IStream? Stream2 = await TestSetup.StreamService.GetStreamAsync(StreamId);
+			IStream? Stream2 = await StreamService.GetStreamAsync(StreamId);
 			TemplateRef TemplateRef2 = Stream2!.Templates.First().Value;
 			Assert.AreEqual(1, TemplateRef2.Schedule!.ActiveJobs.Count);
 			Assert.AreEqual(Jobs1[0].Id, TemplateRef2.Schedule!.ActiveJobs[0]);
@@ -462,7 +450,7 @@ namespace HordeServerTests
 		public async Task StreamPausing()
 		{
 			DateTime StartTime = new DateTime(2021, 1, 1, 12, 0, 0, DateTimeKind.Utc); // Friday Jan 1, 2021 
-			TestSetup.Clock.UtcNow = StartTime;
+			Clock.UtcNow = StartTime;
 
 			CreateScheduleRequest Schedule = new CreateScheduleRequest();
 			Schedule.Enabled = true;
@@ -473,13 +461,13 @@ namespace HordeServerTests
 			await StreamCollection.TryUpdatePauseStateAsync(Stream, NewPausedUntil: StartTime.AddHours(5), NewPauseComment: "testing");
 			
 			// Try trigger a job. No job should be scheduled as the stream is paused
-			TestSetup.Clock.Advance(TimeSpan.FromHours(1.25));
+			Clock.Advance(TimeSpan.FromHours(1.25));
 			await ScheduleService.TickSharedOnlyForTestingAsync();
 			List<IJob> Jobs2 = await GetNewJobs();
 			Assert.AreEqual(0, Jobs2.Count);
 
 			// Advance time beyond the pause period. A build should now trigger
-			TestSetup.Clock.Advance(TimeSpan.FromHours(5.25));
+			Clock.Advance(TimeSpan.FromHours(5.25));
 			await ScheduleService.TickSharedOnlyForTestingAsync();
 
 			List<IJob> Jobs3 = await GetNewJobs();
@@ -490,7 +478,7 @@ namespace HordeServerTests
 
 		async Task<List<IJob>> GetNewJobs()
 		{
-			List<IJob> Jobs = await TestSetup.JobCollection.FindAsync();
+			List<IJob> Jobs = await JobCollection.FindAsync();
 			Jobs.RemoveAll(x => InitialJobIds.Contains(x.Id));
 			InitialJobIds.UnionWith(Jobs.Select(x => x.Id));
 			return Jobs.OrderBy(x => x.Change).ToList();
diff --git a/Engine/Source/Programs/Horde/Horde.Build.Tests/StorageTests.cs b/Engine/Source/Programs/Horde/Horde.Build.Tests/StorageTests.cs
--- a/Engine/Source/Programs/Horde/Horde.Build.Tests/StorageTests.cs
+++ b/Engine/Source/Programs/Horde/Horde.Build.Tests/StorageTests.cs
@@ -2,6 +2,7 @@
 
 using EpicGames.Core;
 using EpicGames.Serialization;
+using HordeServer.Services;
 using HordeServer.Storage;
 using HordeServer.Utilities;
 using Microsoft.Extensions.DependencyInjection;
@@ -20,11 +21,6 @@ namespace HordeServerTests
 	[TestClass]
 	public class StorageTests : TestSetup
 	{
-		public StorageTests()
-			: base(DatabaseIntegrationTest.GetDatabaseService())
-		{
-		}
-
 		[TestMethod]
 		public async Task BlobCollectionTest()
 		{
diff --git a/Engine/Source/Programs/Horde/Horde.Build.Tests/StreamServiceTest.cs b/Engine/Source/Programs/Horde/Horde.Build.Tests/StreamServiceTest.cs
--- a/Engine/Source/Programs/Horde/Horde.Build.Tests/StreamServiceTest.cs
+++ b/Engine/Source/Programs/Horde/Horde.Build.Tests/StreamServiceTest.cs
@@ -10,28 +10,28 @@ using PoolId = HordeServer.Utilities.StringId<HordeServer.Models.IPool>;
 namespace HordeServerTests
 {
     [TestClass]
-    public class StreamServiceTests : DatabaseIntegrationTest
+    public class StreamServiceTests : TestSetup
     {
         [TestMethod]
         public async Task Pausing()
         {
-	        TestSetup TestSetup = await GetTestSetup();
+			Fixture Fixture = await CreateFixtureAsync();
 
-	        IStream Stream = (await TestSetup.StreamService.GetStreamAsync(TestSetup.Fixture!.Stream!.Id))!;
+	        IStream Stream = (await StreamService.GetStreamAsync(Fixture!.Stream!.Id))!;
 	        Assert.IsFalse(Stream.IsPaused(DateTime.UtcNow));
 	        Assert.IsNull(Stream.PausedUntil);
 	        Assert.IsNull(Stream.PauseComment);
 
 	        DateTime PausedUntil = DateTime.UtcNow.AddHours(1);
-	        await TestSetup.StreamService.UpdatePauseStateAsync(Stream, NewPausedUntil: PausedUntil, NewPauseComment: "mycomment");
-	        Stream = (await TestSetup.StreamService.GetStreamAsync(TestSetup.Fixture!.Stream!.Id))!;
+	        await StreamService.UpdatePauseStateAsync(Stream, NewPausedUntil: PausedUntil, NewPauseComment: "mycomment");
+	        Stream = (await StreamService.GetStreamAsync(Fixture!.Stream!.Id))!;
 	        // Comparing by string to avoid comparing exact milliseconds as those are not persisted in MongoDB fields
 	        Assert.IsTrue(Stream.IsPaused(DateTime.UtcNow));
 	        Assert.AreEqual(PausedUntil.ToString(CultureInfo.InvariantCulture), Stream.PausedUntil!.Value.ToString(CultureInfo.InvariantCulture));
 	        Assert.AreEqual("mycomment", Stream.PauseComment);
 	        
-	        await TestSetup.StreamService.UpdatePauseStateAsync(Stream, NewPausedUntil: null, NewPauseComment: null);
-	        Stream = (await TestSetup.StreamService.GetStreamAsync(TestSetup.Fixture!.Stream!.Id))!;
+	        await StreamService.UpdatePauseStateAsync(Stream, NewPausedUntil: null, NewPauseComment: null);
+	        Stream = (await StreamService.GetStreamAsync(Fixture!.Stream!.Id))!;
 	        Assert.IsFalse(Stream.IsPaused(DateTime.UtcNow));
 	        Assert.IsNull(Stream.PausedUntil);
 	        Assert.IsNull(Stream.PauseComment);
diff --git a/Engine/Source/Programs/Horde/Horde.Build.Tests/TestSetup.cs b/Engine/Source/Programs/Horde/Horde.Build.Tests/TestSetup.cs
--- a/Engine/Source/Programs/Horde/Horde.Build.Tests/TestSetup.cs
+++ b/Engine/Source/Programs/Horde/Horde.Build.Tests/TestSetup.cs
@@ -49,7 +49,7 @@ namespace HordeServerTests
 	///
 	/// Easier to pass all these things around in a single object.
 	/// </summary>
-	public class TestSetup
+	public class TestSetup : DatabaseIntegrationTest
 	{
 		public IServiceProvider ServiceProvider { get; }
 		public FakeClock Clock { get; set; } = null!;
@@ -91,24 +91,24 @@ namespace HordeServerTests
 		public CredentialService CredentialService => ServiceProvider.GetRequiredService<CredentialService>();
 		public PoolService PoolService => ServiceProvider.GetRequiredService<PoolService>();
 		public LifetimeService LifetimeService => ServiceProvider.GetRequiredService<LifetimeService>();
+		public ScheduleService ScheduleService => ServiceProvider.GetRequiredService<ScheduleService>();
 
 		public ServerSettings ServerSettings => ServiceProvider.GetRequiredService<IOptions<ServerSettings>>().Value;
 		public IOptionsMonitor<ServerSettings> ServerSettingsMon => ServiceProvider.GetRequiredService<IOptionsMonitor<ServerSettings>>();
-		public Fixture? Fixture;
 
 		public JobsController JobsController => GetJobsController();
 		public AgentsController AgentsController => GetAgentsController();
 
 		private static bool IsDatadogWriterPatched;
 
-		public TestSetup(DatabaseService DbService)
+		public TestSetup()
 		{
 			PatchDatadogWriter();
 
 			IConfiguration Config = new ConfigurationBuilder().Build();
 
 			IServiceCollection Services = new ServiceCollection();
-			Services.AddSingleton(DbService);
+			Services.AddSingleton(GetDatabaseService());
 			Services.Configure<ServerSettings>(ConfigureSettings);
 			Services.AddSingleton<IConfiguration>(Config);
 
@@ -217,12 +217,9 @@ namespace HordeServerTests
 			Services.AddSingleton<ISingletonDocument<AgentSoftwareChannels>>(new SingletonDocumentStub<AgentSoftwareChannels>());
 		}
 
-		public async Task CreateFixture(bool ForceNewFixture = false)
+		public Task<Fixture> CreateFixtureAsync()
 		{
-			if (!DatabaseService.ReadOnlyMode)
-			{
-				Fixture = await Fixture.Create(ForceNewFixture, GraphCollection, TemplateCollection, JobService, ArtifactCollection, StreamService, AgentService, PerforceService);
-			}
+			return Fixture.Create(GraphCollection, TemplateCollection, JobService, ArtifactCollection, StreamService, AgentService, PerforceService);
 		}
 
 		private JobsController GetJobsController()
