Horde: Fix notifications for config updates not correctly identifying the user that last modified the file.

diff --git a/Engine/Source/Programs/Horde/Horde.Build/Perforce/IPerforceService.cs b/Engine/Source/Programs/Horde/Horde.Build/Perforce/IPerforceService.cs
--- a/Engine/Source/Programs/Horde/Horde.Build/Perforce/IPerforceService.cs
+++ b/Engine/Source/Programs/Horde/Horde.Build/Perforce/IPerforceService.cs
@@ -126,9 +126,10 @@ namespace HordeServer.Services
 		/// </summary>
 		/// <param name="ClusterName">Name of the Perforce cluster</param>
 		/// <param name="MinChange">The minimum changelist number</param>
+		/// <param name="MaxChange">The maximum changelist number</param>
 		/// <param name="MaxResults"></param>
 		/// <returns>Changelist information</returns>
-		public Task<List<ChangeSummary>> GetChangesAsync(string ClusterName, int? MinChange, int MaxResults);
+		public Task<List<ChangeSummary>> GetChangesAsync(string ClusterName, int? MinChange, int? MaxChange, int MaxResults);
 
 		/// <summary>
 		/// Gets the latest change for a particular stream
diff --git a/Engine/Source/Programs/Horde/Horde.Build/Perforce/Impl/PerforceService.cs b/Engine/Source/Programs/Horde/Horde.Build/Perforce/Impl/PerforceService.cs
--- a/Engine/Source/Programs/Horde/Horde.Build/Perforce/Impl/PerforceService.cs
+++ b/Engine/Source/Programs/Horde/Horde.Build/Perforce/Impl/PerforceService.cs
@@ -762,7 +762,7 @@ namespace HordeServer.Services
 		}
 
 		/// <inheritdoc/>
-		public async Task<List<ChangeSummary>> GetChangesAsync(string ClusterName, int? MinChange, int MaxResults)
+		public async Task<List<ChangeSummary>> GetChangesAsync(string ClusterName, int? MinChange, int? MaxChange, int MaxResults)
 		{
 			using IScope Scope = GlobalTracer.Instance.BuildSpan("PerforceService.GetChangesAsync").StartActive();
 			Scope.Span.SetTag("ClusterName", ClusterName);
@@ -774,15 +774,7 @@ namespace HordeServer.Services
 			{
 				P4.ChangesCmdOptions Options = new P4.ChangesCmdOptions(P4.ChangesCmdFlags.IncludeTime | P4.ChangesCmdFlags.FullDescription, null, MaxResults, P4.ChangeListStatus.Submitted, null);
 
-				IList<P4.Changelist> Changelists;
-				if (MinChange == null)
-				{
-					Changelists = Repository.GetChangelists(Options);
-				}
-				else
-				{
-					Changelists = Repository.GetChangelists(Options, new P4.FileSpec(new P4.DepotPath($"//...@>={MinChange.Value}"), null, null, null));
-				}
+				IList<P4.Changelist> Changelists = Repository.GetChangelists(Options, new P4.FileSpec(new P4.DepotPath(GetFilter("//...", MinChange, MaxChange)), null, null, null));
 
 				List<ChangeSummary> Changes = new List<ChangeSummary>();
 				if (Changelists != null)
diff --git a/Engine/Source/Programs/Horde/Horde.Build/Server/ConfigService.cs b/Engine/Source/Programs/Horde/Horde.Build/Server/ConfigService.cs
--- a/Engine/Source/Programs/Horde/Horde.Build/Server/ConfigService.cs
+++ b/Engine/Source/Programs/Horde/Horde.Build/Server/ConfigService.cs
@@ -485,7 +485,7 @@ namespace HordeServer.Services
 					List<FileSummary> Files = await PerforceService.FindFilesAsync(ConfigPath.Host, new[] { FileName });
 					Change = Files[0].Change;
 
-					List<ChangeSummary> Changes = await PerforceService.GetChangesAsync(ConfigPath.Host, Change, 1);
+					List<ChangeSummary> Changes = await PerforceService.GetChangesAsync(ConfigPath.Host, Change, Change, 1);
 					if (Changes.Count > 0 && Changes[0].Number == Change)
 					{
 						(Author, Description) = (Changes[0].Author, Changes[0].Description);
