Fix test compilation error.

diff --git a/Engine/Source/Programs/Horde/Horde.Build.Tests/Stubs/Services/PerforceServiceStub.cs b/Engine/Source/Programs/Horde/Horde.Build.Tests/Stubs/Services/PerforceServiceStub.cs
--- a/Engine/Source/Programs/Horde/Horde.Build.Tests/Stubs/Services/PerforceServiceStub.cs
+++ b/Engine/Source/Programs/Horde/Horde.Build.Tests/Stubs/Services/PerforceServiceStub.cs
@@ -197,7 +197,7 @@ namespace HordeServerTests.Stubs.Services
 			throw new NotImplementedException();
 		}
 
-		public Task<List<ChangeSummary>> GetChangesAsync(string ClusterName, int? MinChange, int MaxResults)
+		public Task<List<ChangeSummary>> GetChangesAsync(string ClusterName, int? MinChange, int? MaxChange, int MaxResults)
 		{
 			throw new NotImplementedException();
 		}
