Horde: Fix compute expiry timers not being restarted correctly.

diff --git a/Engine/Source/Programs/Horde/Horde.Build/Compute/Impl/RedisTaskScheduler.cs b/Engine/Source/Programs/Horde/Horde.Build/Compute/Impl/RedisTaskScheduler.cs
--- a/Engine/Source/Programs/Horde/Horde.Build/Compute/Impl/RedisTaskScheduler.cs
+++ b/Engine/Source/Programs/Horde/Horde.Build/Compute/Impl/RedisTaskScheduler.cs
@@ -317,7 +317,7 @@ namespace HordeServer.Compute.Impl
 					if (ResetActiveQueuesTimer.Elapsed > ResetTime)
 					{
 						LocalActiveQueues = new HashSet<TQueueId>();
-						ResetActiveQueuesTimer.Reset();
+						ResetActiveQueuesTimer.Restart();
 					}
 				}
 			}
