Horde: Fix build error.

diff --git a/Engine/Source/Programs/Horde/Horde.Build/Utilities/BuildGraphExtensions.cs b/Engine/Source/Programs/Horde/Horde.Build/Utilities/BuildGraphExtensions.cs
--- a/Engine/Source/Programs/Horde/Horde.Build/Utilities/BuildGraphExtensions.cs
+++ b/Engine/Source/Programs/Horde/Horde.Build/Utilities/BuildGraphExtensions.cs
@@ -64,8 +64,13 @@ namespace Horde.Build.Utilities
 					catch { }
 				}
 			}
+
+			public Task<string[]> FindAsync(string Path)
+			{
+				throw new NotImplementedException();
+			}
 		}
-		
+
 		/// <summary>
 		/// Create a graph from a script in Perforce
 		/// </summary>
