Horde: Another build fix.

diff --git a/Engine/Source/Programs/Horde/Horde.Build/Utilities/BuildGraphExtensions.cs b/Engine/Source/Programs/Horde/Horde.Build/Utilities/BuildGraphExtensions.cs
--- a/Engine/Source/Programs/Horde/Horde.Build/Utilities/BuildGraphExtensions.cs
+++ b/Engine/Source/Programs/Horde/Horde.Build/Utilities/BuildGraphExtensions.cs
@@ -136,7 +136,7 @@ namespace Horde.Build.Utilities
 				throw new NotImplementedException();
 			}
 
-			BgScript? Script = await BgScriptReader.ReadAsync(Context, ScriptFile, Options, Properties, false, Schema, EpicGames.Core.Log.Logger);
+			BgGraph? Script = await BgScriptReader.ReadAsync(Context, ScriptFile, Options, Properties, false, Schema, EpicGames.Core.Log.Logger);
 			if (Script == null)
 			{
 				throw new NotImplementedException();
