Horde: Add settings flag for disabling background services

Useful when doing canary deployments where background services are doing writes to the database.
As a comparison, requests/responses are easy to control for writes with.

diff --git a/Engine/Source/Programs/Horde/Horde.Build/ServerSettings.cs b/Engine/Source/Programs/Horde/Horde.Build/ServerSettings.cs
--- a/Engine/Source/Programs/Horde/Horde.Build/ServerSettings.cs
+++ b/Engine/Source/Programs/Horde/Horde.Build/ServerSettings.cs
@@ -353,6 +353,12 @@ namespace HordeServer
 		/// </summary>
 		public bool DisableSchedules { get; set; }
 
+		/// <summary>
+		/// Whether to run background services
+		///
+		/// </summary>
+		public bool EnableBackgroundServices { get; set; } = true;
+
 		/// <summary>
 		/// Timezone for evaluating schedules
 		/// </summary>
diff --git a/Engine/Source/Programs/Horde/Horde.Build/Startup.cs b/Engine/Source/Programs/Horde/Horde.Build/Startup.cs
--- a/Engine/Source/Programs/Horde/Horde.Build/Startup.cs
+++ b/Engine/Source/Programs/Horde/Horde.Build/Startup.cs
@@ -531,26 +531,29 @@ namespace HordeServer
 						.Build();
 				});
 
-			Services.AddHostedService(Provider => Provider.GetRequiredService<AgentService>());
-			Services.AddHostedService(Provider => Provider.GetRequiredService<AutoscaleService>());
-			Services.AddHostedService(Provider => Provider.GetRequiredService<CommitService>());
-			Services.AddHostedService(Provider => Provider.GetRequiredService<ConsistencyService>());
-			Services.AddHostedService(Provider => (DowntimeService)Provider.GetRequiredService<IDowntimeService>());
-			Services.AddHostedService(Provider => Provider.GetRequiredService<IIssueService>());
-			Services.AddHostedService(Provider => (LogFileService)Provider.GetRequiredService<ILogFileService>());
-			Services.AddHostedService(Provider => (NotificationService)Provider.GetRequiredService<INotificationService>());
-			if (!Settings.DisableSchedules)
-			{
-				Services.AddHostedService(Provider => Provider.GetRequiredService<ScheduleService>());
-			}
-
-			Services.AddHostedService<MetricService>();
-			Services.AddHostedService(Provider => Provider.GetRequiredService<PerforceLoadBalancer>());
-			Services.AddHostedService<PoolUpdateService>();
-			Services.AddHostedService(Provider => Provider.GetRequiredService<SlackNotificationSink>());
-			Services.AddHostedService<ConfigService>();
-			Services.AddHostedService<TelemetryService>();
-			Services.AddHostedService(Provider => Provider.GetRequiredService<DeviceService>());
+			if (Settings.EnableBackgroundServices)
+			{
+				Services.AddHostedService(Provider => Provider.GetRequiredService<AgentService>());
+				Services.AddHostedService(Provider => Provider.GetRequiredService<AutoscaleService>());
+				Services.AddHostedService(Provider => Provider.GetRequiredService<CommitService>());
+				Services.AddHostedService(Provider => Provider.GetRequiredService<ConsistencyService>());
+				Services.AddHostedService(Provider => (DowntimeService)Provider.GetRequiredService<IDowntimeService>());
+				Services.AddHostedService(Provider => Provider.GetRequiredService<IIssueService>());
+				Services.AddHostedService(Provider => (LogFileService)Provider.GetRequiredService<ILogFileService>());
+				Services.AddHostedService(Provider => (NotificationService)Provider.GetRequiredService<INotificationService>());
+				if (!Settings.DisableSchedules)
+				{
+					Services.AddHostedService(Provider => Provider.GetRequiredService<ScheduleService>());
+				}
+
+				Services.AddHostedService<MetricService>();
+				Services.AddHostedService(Provider => Provider.GetRequiredService<PerforceLoadBalancer>());
+				Services.AddHostedService<PoolUpdateService>();
+				Services.AddHostedService(Provider => Provider.GetRequiredService<SlackNotificationSink>());
+				Services.AddHostedService<ConfigService>();
+				Services.AddHostedService<TelemetryService>();
+				Services.AddHostedService(Provider => Provider.GetRequiredService<DeviceService>());
+			}
 
 			// Task sources. Order of registration is important here; it dictates the priority in which sources are served.
 			Services.AddSingleton<JobTaskSource>();
