Horde: Change Create and Update job time indexes to descending
#rnx

diff --git a/Engine/Source/Programs/Horde/Horde.Build/Jobs/Impl/JobCollection.cs b/Engine/Source/Programs/Horde/Horde.Build/Jobs/Impl/JobCollection.cs
--- a/Engine/Source/Programs/Horde/Horde.Build/Jobs/Impl/JobCollection.cs
+++ b/Engine/Source/Programs/Horde/Horde.Build/Jobs/Impl/JobCollection.cs
@@ -362,8 +362,8 @@ namespace HordeServer.Collections.Impl
 				Jobs.Indexes.CreateOne(new CreateIndexModel<JobDocument>(Builders<JobDocument>.IndexKeys.Ascending(x => x.StreamId)));
 				Jobs.Indexes.CreateOne(new CreateIndexModel<JobDocument>(Builders<JobDocument>.IndexKeys.Ascending(x => x.Change)));
 				Jobs.Indexes.CreateOne(new CreateIndexModel<JobDocument>(Builders<JobDocument>.IndexKeys.Ascending(x => x.PreflightChange)));
-				Jobs.Indexes.CreateOne(new CreateIndexModel<JobDocument>(Builders<JobDocument>.IndexKeys.Ascending(x => x.CreateTimeUtc)));
-				Jobs.Indexes.CreateOne(new CreateIndexModel<JobDocument>(Builders<JobDocument>.IndexKeys.Ascending(x => x.UpdateTimeUtc)));
+				Jobs.Indexes.CreateOne(new CreateIndexModel<JobDocument>(Builders<JobDocument>.IndexKeys.Descending(x => x.CreateTimeUtc)));
+				Jobs.Indexes.CreateOne(new CreateIndexModel<JobDocument>(Builders<JobDocument>.IndexKeys.Descending(x => x.UpdateTimeUtc)));
 				Jobs.Indexes.CreateOne(new CreateIndexModel<JobDocument>(Builders<JobDocument>.IndexKeys.Ascending(x => x.Name)));
 				Jobs.Indexes.CreateOne(new CreateIndexModel<JobDocument>(Builders<JobDocument>.IndexKeys.Ascending(x => x.StartedByUserId)));
 				Jobs.Indexes.CreateOne(new CreateIndexModel<JobDocument>(Builders<JobDocument>.IndexKeys.Ascending(x => x.TemplateId)));
