Horde: Provide index hint to find jobs query

Adding an index requires modifying the database which isn't allowed (yet) in the canary env. So the index is reference by name instead, which of course isn't great. Might be possible to look up these during the index creation and reference them by a C# instance/type. But considering the perf impact this has, I think it's worth the trade-off.

diff --git a/Engine/Source/Programs/Horde/Horde.Build/Jobs/Impl/JobCollection.cs b/Engine/Source/Programs/Horde/Horde.Build/Jobs/Impl/JobCollection.cs
--- a/Engine/Source/Programs/Horde/Horde.Build/Jobs/Impl/JobCollection.cs
+++ b/Engine/Source/Programs/Horde/Horde.Build/Jobs/Impl/JobCollection.cs
@@ -557,7 +557,8 @@ namespace HordeServer.Collections.Impl
 			List<JobDocument> Results;
 			using (IScope Scope = GlobalTracer.Instance.BuildSpan("Jobs.Find").StartActive())
 			{
-				IFindFluent<JobDocument, JobDocument> Query = Jobs.Find<JobDocument>(Filter).SortByDescending(x => x.CreateTimeUtc!);
+				BsonString IndexHint = new (nameof(JobDocument.CreateTimeUtc) + "_1");
+				IFindFluent<JobDocument, JobDocument> Query = Jobs.Find<JobDocument>(Filter, new FindOptions { Hint = IndexHint }).SortByDescending(x => x.CreateTimeUtc!);
 				if (Index != null)
 				{
 					Query = Query.Skip(Index.Value);
