Moved Horde.Storage from git to p4

Horde.Storage is the content addressable storage component we use to store ddc cached data (formerly known as Jupiter).

This is still very much WIP and not something we recommend licenses to roll out just yet.

diff --git a/Engine/Source/Programs/Horde/HordeStorage/.editorconfig b/Engine/Source/Programs/Horde/HordeStorage/.editorconfig
new file mode 100644
--- /dev/null
+++ b/Engine/Source/Programs/Horde/HordeStorage/.editorconfig
@@ -0,0 +1,41 @@
+
+[*.{asax,ascx,aspx,cs,cshtml,css,htm,html,js,jsx,master,razor,skin,ts,tsx,vb,xaml,xamlx,xoml}]
+indent_style=space
+indent_size=4
+tab_width=4
+insert_final_newline=true
+
+[*.{appxmanifest,build,config,csproj,dbml,discomap,dtd,json,jsproj,lsproj,njsproj,nuspec,proj,props,resjson,resw,resx,StyleCop,targets,tasks,vbproj,xml,xsd}]
+indent_style=space
+indent_size=2
+tab_width=2
+insert_final_newline=true
+
+[*]
+
+# Microsoft .NET properties
+csharp_new_line_before_members_in_object_initializers=false
+csharp_preferred_modifier_order=public, private, protected, internal, new, abstract, virtual, sealed, override, static, readonly, extern, unsafe, volatile, async:suggestion
+csharp_style_var_elsewhere=false:hint
+csharp_style_var_for_built_in_types=false:hint
+csharp_style_var_when_type_is_apparent=false:hint
+dotnet_style_predefined_type_for_locals_parameters_members=true:hint
+dotnet_style_predefined_type_for_member_access=true:hint
+dotnet_style_qualification_for_event=false:hint
+dotnet_style_qualification_for_field=false:hint
+dotnet_style_qualification_for_method=false:hint
+dotnet_style_qualification_for_property=false:hint
+dotnet_style_require_accessibility_modifiers=for_non_interface_members:hint
+
+# ReSharper properties
+resharper_accessor_owner_body=accessors_with_block_body
+resharper_arguments_other=named
+resharper_arguments_skip_single=true
+resharper_braces_for_for=required
+resharper_braces_for_foreach=required
+resharper_braces_for_ifelse=required
+resharper_braces_for_while=required
+
+# ReSharper inspection severities
+resharper_check_namespace_highlighting=hint
+resharper_redundant_base_qualifier_highlighting=warning
diff --git a/Engine/Source/Programs/Horde/HordeStorage/Benchmarks/readme.md b/Engine/Source/Programs/Horde/HordeStorage/Benchmarks/readme.md
new file mode 100644
--- /dev/null
+++ b/Engine/Source/Programs/Horde/HordeStorage/Benchmarks/readme.md
@@ -0,0 +1,10 @@
+These here are benchmarks that can be run using SuperBenchmarker
+`choco install SuperBenchmarker`
+
+# Insert random blobs
+`sb -u https://dev-jupiter.northeurope.cloudapp.azure.com/api/v1/c/ue4.ddc/test/{{{name}}} -n 100 -c 16 -g 2 -t .\templates\putBlobs.txt -f .\templates\putBlobs.csv -m PUT -B -b`
+
+# Fetch random blobs
+Note that all of these may not succeed depending on which blobs were inserted in the past
+
+`sb -u https://dev-jupiter.northeurope.cloudapp.azure.com/api/v1/c/ue4.ddc/test/{{{name}}} -n 100 -c 16 -g 2 -t .\templates\putBlobs.txt -f .\templates\putBlobs.csv -B -b`
diff --git a/Engine/Source/Programs/Horde/HordeStorage/Benchmarks/templates/auth.txt b/Engine/Source/Programs/Horde/HordeStorage/Benchmarks/templates/auth.txt
new file mode 100644
--- /dev/null
+++ b/Engine/Source/Programs/Horde/HordeStorage/Benchmarks/templates/auth.txt
@@ -0,0 +1,3 @@
+Content-type: application/octet-stream
+Cache-Control: no-cache
+Authorization: Bearer eyJraWQiOiJFUGMyNzZFNTJPVjVnTDlJblpCbHBhT0RyWlgweWk3TEVLcjRnTEVYOHRRIiwiYWxnIjoiUlMyNTYifQ.eyJ2ZXIiOjEsImp0aSI6IkFULnVmbklqcmRXUEoyRkItNWRnYnY0clU0UHhKUzdOamRyQl9PT3VVM0dVVDgiLCJpc3MiOiJodHRwczovL2Rldi0zNTA5MDcub2t0YS5jb20vb2F1dGgyL2F1c3Uwd25rdzU3bnRKb1BINHg1IiwiYXVkIjoiYXBpOi8vanVwaXRlciIsImlhdCI6MTU4MzkxOTg0NCwiZXhwIjoxNTg0MDA2MjQ0LCJjaWQiOiIwb2EyNWR5d2lBTlVmSE55VzR4NiIsInNjcCI6WyJjYWNoZV9hY2Nlc3MiXSwic3ViIjoiMG9hMjVkeXdpQU5VZkhOeVc0eDYifQ.RXYIutpJCrmpr5fqUQNRoaABkIf0bVODhLH4EdmsoMdrEzOfoG8crFQSpgDYHTq4Aqrm0ezbC25x3i34pBAlAZG6Rvop69F7ngNQTc4PnoEas1eWp-d5WlktMMpOtqRV2Y0Vzxtd0RGJNM2xJCvCO05chfWWLpY6Blod7l5Dzv2uyXbAEMbXRm7UOLiJQFOQM_ikyJLHE_R3TtSlI6rpFiOmO8kdJy9uSwGocte9f0I7HZM6cvZRkXXfWVRncI1g3d3BXSqZAhVgAkZI5IHWZROBi7kvSUq3fCxDBeGveTh2fwaP7UdBjXQfEtzif6clLFUivZEGtCZkG9AXoFYMCw
\ No newline at end of file
diff --git a/Engine/Source/Programs/Horde/HordeStorage/Benchmarks/templates/putBlobs.txt b/Engine/Source/Programs/Horde/HordeStorage/Benchmarks/templates/putBlobs.txt
new file mode 100644
--- /dev/null
+++ b/Engine/Source/Programs/Horde/HordeStorage/Benchmarks/templates/putBlobs.txt
@@ -0,0 +1,5 @@
+Content-type: application/octet-stream
+Authorization Bearer eyJraWQiOiJFUGMyNzZFNTJPVjVnTDlJblpCbHBhT0RyWlgweWk3TEVLcjRnTEVYOHRRIiwiYWxnIjoiUlMyNTYifQ.eyJ2ZXIiOjEsImp0aSI6IkFULnM4azBkRFJodkZtQXlnNDFfSlItRG9jUUZHYXJhUHFCNmtCTEd5V1c4VkEiLCJpc3MiOiJodHRwczovL2Rldi0zNTA5MDcub2t0YS5jb20vb2F1dGgyL2F1c3Uwd25rdzU3bnRKb1BINHg1IiwiYXVkIjoiYXBpOi8vanVwaXRlciIsImlhdCI6MTU4MjYzNDIzMiwiZXhwIjoxNTgyNzIwNjMyLCJjaWQiOiIwb2EyNWR5d2lBTlVmSE55VzR4NiIsInNjcCI6WyJjYWNoZV9hY2Nlc3MiXSwic3ViIjoiMG9hMjVkeXdpQU5VZkhOeVc0eDYifQ.WC5h8ceOVb5F5JYoxShxUrtzPNljtMu2iAnMu8ltWVcQG1z__WGbA3QmyObUl2SBBahdQfXeiuqiEozdqrEWOFxZI5ZLbmZENipgbr6kk2eGDfVXFIRZxKUxL9AQp9fUDyalcFJteuLXsyFZA0QVSuMno1u59G74dny8qUa_26wlreWg6L0qCFtehJ0rqPdU5rrLDyoy_k4WRcEPLI20AF2r-Mi6GxgeOV8cWGq7FLLVxdYhKBQO_vHcbp-JpQe407C5PfSEVYqpO9VPON13uvxaOCi2IlB7AO_T-9icwB8j4YoxdCtLNNI4qVfl4nMJABjbsXFSTnKHw_qw0hYaaw
+X-Jupiter-Sha1: {{{sha1}}}
+
+{{{contents}}}
\ No newline at end of file
diff --git a/Engine/Source/Programs/Horde/HordeStorage/Callisto/AssemblyInfo.cs b/Engine/Source/Programs/Horde/HordeStorage/Callisto/AssemblyInfo.cs
new file mode 100644
--- /dev/null
+++ b/Engine/Source/Programs/Horde/HordeStorage/Callisto/AssemblyInfo.cs
@@ -0,0 +1,6 @@
+﻿// Copyright Epic Games, Inc. All Rights Reserved.
+
+using System.Runtime.CompilerServices;
+
+[assembly:InternalsVisibleTo("CallisoUnitTests")]
+
diff --git a/Engine/Source/Programs/Horde/HordeStorage/Callisto/Callisto.csproj b/Engine/Source/Programs/Horde/HordeStorage/Callisto/Callisto.csproj
new file mode 100644
--- /dev/null
+++ b/Engine/Source/Programs/Horde/HordeStorage/Callisto/Callisto.csproj
@@ -0,0 +1,27 @@
+﻿<Project Sdk="Microsoft.NET.Sdk.Web">
+
+  <PropertyGroup>
+    <TargetFramework>net6.0</TargetFramework>
+    <Nullable>enable</Nullable>
+    <TreatWarningsAsErrors>true</TreatWarningsAsErrors>
+    <WarningsAsErrors />
+    <LangVersion>10</LangVersion>
+  </PropertyGroup>
+
+  <ItemGroup>
+    <Compile Remove="Tests\**" />
+    <Content Remove="Tests\**" />
+    <EmbeddedResource Remove="Tests\**" />
+    <None Remove="Tests\**" />
+  </ItemGroup>
+
+  <ItemGroup>
+    <PackageReference Include="AspNetCore.HealthChecks.System" Version="5.0.1" />
+    <PackageReference Include="SerilogAnalyzer" Version="0.15.0" />
+  </ItemGroup>
+
+  <ItemGroup>
+    <ProjectReference Include="..\Jupiter.Common\Jupiter.Common.csproj" />
+  </ItemGroup>
+
+</Project>
diff --git a/Engine/Source/Programs/Horde/HordeStorage/Callisto/CallistoStartup.cs b/Engine/Source/Programs/Horde/HordeStorage/Callisto/CallistoStartup.cs
new file mode 100644
--- /dev/null
+++ b/Engine/Source/Programs/Horde/HordeStorage/Callisto/CallistoStartup.cs
@@ -0,0 +1,109 @@
+// Copyright Epic Games, Inc. All Rights Reserved.
+
+using System;
+using System.Collections.Generic;
+using System.ComponentModel.DataAnnotations;
+using System.IO;
+using Callisto.Implementation;
+using JsonSubTypes;
+using Jupiter;
+using Jupiter.Implementation;
+using Microsoft.AspNetCore.Authorization;
+using Microsoft.AspNetCore.Hosting;
+using Microsoft.AspNetCore.Mvc;
+using Microsoft.Extensions.Configuration;
+using Microsoft.Extensions.DependencyInjection;
+using Microsoft.Extensions.Options;
+
+namespace Callisto
+{
+    // ReSharper disable once ClassNeverInstantiated.Global
+    public class CallistoStartup : BaseStartup
+    {
+        public CallistoStartup(IConfiguration configuration, IWebHostEnvironment environment) : base(configuration, environment)
+        {
+        }
+
+        protected override void OnAddService(IServiceCollection services)
+        {
+            services.AddOptions<CallistoSettings>().Configure(o => Configuration.GetSection("Callisto").Bind(o)).ValidateDataAnnotations();
+
+            services.AddSingleton(serviceType: typeof(ITransactionLogs), TransactionLogFactory);
+        }
+
+        protected override void OnAddHealthChecks(IServiceCollection services, IHealthChecksBuilder healthChecks)
+        {
+            ServiceProvider provider = services.BuildServiceProvider();
+            CallistoSettings settings = provider.GetService<IOptionsMonitor<CallistoSettings>>()!.CurrentValue;
+            string? driveRoot = Path.GetPathRoot(settings!.TransactionLogRoot);
+            const long expectedFreeMegaBytes = 512;
+            healthChecks.AddDiskStorageHealthCheck(options => options.AddDrive(driveRoot, expectedFreeMegaBytes));
+        }
+
+        protected override void OnAddAuthorization(AuthorizationOptions authorizationOptions, List<string> defaultSchemes)
+        {
+            authorizationOptions.AddPolicy("TLog.read", policy =>
+            {
+                policy.AuthenticationSchemes = defaultSchemes;
+                policy.RequireClaim("transactionLog", "read", "readwrite", "full");
+            });
+
+            authorizationOptions.AddPolicy("TLog.write", policy =>
+            {
+                policy.AuthenticationSchemes = defaultSchemes;
+                policy.RequireClaim("transactionLog", "write", "readwrite", "full");
+            });
+
+            authorizationOptions.AddPolicy("TLog.delete", policy =>
+            {
+                policy.AuthenticationSchemes = defaultSchemes;
+                policy.RequireClaim("transactionLog", "delete", "full");
+            });
+
+            authorizationOptions.AddPolicy("Admin", policy =>
+            {
+                policy.AuthenticationSchemes = defaultSchemes;
+                policy.RequireClaim("Admin");
+            });
+        }
+
+        private ITransactionLogs TransactionLogFactory(IServiceProvider provider)
+        {
+            IOptionsMonitor<CallistoSettings> settings = provider.GetService<IOptionsMonitor<CallistoSettings>>()!;
+            switch (settings.CurrentValue.TransactionLogImplementation)
+            {
+                case CallistoSettings.TransactionLogImplementations.File:
+                    return new FileTransactionLogProxy(settings);
+                case CallistoSettings.TransactionLogImplementations.Memory:
+                    return new MemoryTransactionLogs();
+                default:
+                    throw new NotImplementedException();
+            }
+        }
+    }
+
+    public class CallistoSettings
+    {
+        public enum TransactionLogImplementations
+        {
+            File,
+            Memory
+        }
+
+        [Required] public TransactionLogImplementations TransactionLogImplementation { get; set; }
+
+        // Name of the current site, has to be globally unique across all deployments, used to avoid replication loops
+        [Required] public string CurrentSite { get; set; } = "";
+
+        [Required]
+        public string TransactionLogRoot
+        {
+            get { return Environment.ExpandEnvironmentVariables(_transactionLogRoot); }
+            set { _transactionLogRoot = value; }
+        }
+
+        public bool VerifySerialization { get; set; } = false;
+
+        private string _transactionLogRoot = "";
+    }
+}
diff --git a/Engine/Source/Programs/Horde/HordeStorage/Callisto/Controllers/TransactionLogController.cs b/Engine/Source/Programs/Horde/HordeStorage/Callisto/Controllers/TransactionLogController.cs
new file mode 100644
--- /dev/null
+++ b/Engine/Source/Programs/Horde/HordeStorage/Callisto/Controllers/TransactionLogController.cs
@@ -0,0 +1,237 @@
+﻿// Copyright Epic Games, Inc. All Rights Reserved.
+
+using System;
+using System.ComponentModel.DataAnnotations;
+using System.IO;
+using System.Linq;
+using System.Threading.Tasks;
+using Callisto.Implementation;
+using Jupiter;
+using Jupiter.Implementation;
+using Microsoft.AspNetCore.Authorization;
+using Microsoft.AspNetCore.Mvc;
+using Microsoft.Extensions.Options;
+using Serilog;
+
+namespace Callisto.Controllers
+{
+    [ApiController]
+    [FormatFilter]
+    [Route("api/v1/t")]
+    public class TransactionLogController : ControllerBase
+    {
+        private readonly ITransactionLogs _transactionLogs;
+        private readonly IDiagnosticContext _diagnosticContext;
+        private readonly IAuthorizationService _authorizationService;
+
+        private readonly CallistoSettings _settings;
+        private readonly ILogger _logger = Log.ForContext<TransactionLogController>();
+
+        public TransactionLogController(ITransactionLogs transactionLogs, IOptionsMonitor<CallistoSettings> options, IDiagnosticContext diagnosticContext, IAuthorizationService authorizationService)
+        {
+            _transactionLogs = transactionLogs;
+            _diagnosticContext = diagnosticContext;
+            _authorizationService = authorizationService;
+
+            _settings = options.CurrentValue;
+        }
+
+
+        [HttpGet("")]
+        [ProducesDefaultResponseType]
+        [ProducesResponseType(type: typeof(ProblemDetails), 400)]
+        [Authorize("TLog.read")]
+        public IActionResult Get()
+        {
+            NamespaceId[] namespaces = _transactionLogs.GetNamespaces();
+
+            // filter namespaces down to only the namespaces the user has access to
+            namespaces = namespaces.Where(ns =>
+            {
+                Task<AuthorizationResult> authorizationResult = _authorizationService.AuthorizeAsync(User, ns, NamespaceAccessRequirement.Name);
+                return authorizationResult.Result.Succeeded;
+            }).ToArray();
+
+            return Ok(new {
+                Logs = namespaces
+            });
+        }
+
+        [HttpGet("{ns}/{start}.{format?}")]
+        [ProducesDefaultResponseType]
+        [ProducesResponseType(type: typeof(ProblemDetails), 400)]
+        [Authorize("TLog.read")]
+        public async Task<IActionResult> Get(
+            [Required] NamespaceId ns,
+            [Required] long start,
+            [FromQuery] string? notSeenAtSite = null,
+            [FromQuery] Guid? expectedGeneration = null,
+            [FromQuery] int count = 100,
+            [FromQuery] int maxOffsetsAttempted = 100
+        )
+        {
+            AuthorizationResult authorizationResult = await _authorizationService.AuthorizeAsync(User, ns, NamespaceAccessRequirement.Name);
+
+            if (!authorizationResult.Succeeded)
+            {
+                return Forbid();
+            }
+
+            TransactionLogDescription description = _transactionLogs.Get(ns).Describe();
+            // if a given generation is expected verify we are using that generation for the namespace
+            if (expectedGeneration != null)
+            {
+                if (description.Generation != expectedGeneration)
+                {
+                    return BadRequest(new ProblemDetails
+                    {
+                        Title =
+                            $"Generation mismatch, current generation is {description.Generation} expected {expectedGeneration}",
+                        Type = "http://jupiter.epicgames.com/callisto/newGeneration"
+                    });
+                }
+            }
+
+            // try to recover by finding the next valid object in the log
+            for (int i = 0; i < maxOffsetsAttempted; i++)
+            {
+                try
+                {
+                    long currentOffset = start + i;
+                    TransactionEvents events = await _transactionLogs.Get(ns).Get(currentOffset, count, notSeenAtSite);
+                    // if the current offset was invalid we attempt to move it to find the next valid set of values but if we had some valid events before it we return those first
+                    if (events.TransactionLogMismatchFoundAt.HasValue && events.TransactionLogMismatchFoundAt.Value == currentOffset)
+                    {
+                        continue;
+                    }
+
+                    return Ok(new CallistoGetResponse(events.Events, description.Generation, events.CurrentOffset));
+                }
+                catch (FileNotFoundException e)
+                {
+                    return BadRequest(new ProblemDetails {Title = $"Unable to find log file {e.FileName}", Type = "http://jupiter.epicgames.com/callisto/unknownNamespace"});
+                }
+                catch (EventHashMismatchException)
+                {
+                    // if the event has a bad hash we skip it and search for the next event in the list
+                    continue;
+                }
+            }
+
+            return BadRequest(new ProblemDetails {Title = $"Unable to find a valid offset in {ns} tried {start} to {start+maxOffsetsAttempted}", Type = "http://jupiter.epicgames.com/callisto/transactionLogOffsetMismatch" });
+        }
+
+        [HttpGet("{ns}/sites")]
+        [ProducesDefaultResponseType]
+        [ProducesResponseType(type: typeof(ProblemDetails), 400)]
+        [Authorize("TLog.read")]
+        public async Task<IActionResult> Sites(
+            [Required] NamespaceId ns
+        )
+        {
+            AuthorizationResult authorizationResult = await _authorizationService.AuthorizeAsync(User, ns, NamespaceAccessRequirement.Name);
+
+            if (!authorizationResult.Succeeded)
+            {
+                return Forbid();
+            }
+
+            TransactionLogDescription description = _transactionLogs.Get(ns).Describe();
+            try
+            {
+                return Ok(new {Sites = description.KnownSites});
+            }
+            catch (FileNotFoundException e)
+            {
+                return BadRequest(new ProblemDetails { Title = $"Unable to find log file {e.FileName}" });
+            }
+        }
+
+
+        [HttpPost("{ns}")]
+        [Authorize("TLog.write")]
+        public async Task<IActionResult> PostNewTransaction(
+            [Required] NamespaceId ns,
+            [Required] [FromBody] TransactionEvent request
+        )
+        {
+            AuthorizationResult authorizationResult = await _authorizationService.AuthorizeAsync(User, ns, NamespaceAccessRequirement.Name);
+
+            if (!authorizationResult.Succeeded)
+            {
+                return Forbid();
+            }
+
+            if (string.IsNullOrEmpty(_settings.CurrentSite))
+            {
+                throw new Exception(
+                    "Current site setting is missing, this has to be part of the deployment and has to be globally unique");
+            }
+
+            // add the current site to the list of sites that have seen this event
+            request.Locations.Add(_settings.CurrentSite);
+            long id;
+            switch (request)
+            {
+                case AddTransactionEvent addEvent:
+                {
+                    id = await _transactionLogs.Get(ns).NewTransaction(addEvent);
+                    break;
+                }
+                case RemoveTransactionEvent removeEvent:
+                {
+                    id = await _transactionLogs.Get(ns).NewTransaction(removeEvent);
+                    break;
+                }
+                default:
+                    return BadRequest(new ProblemDetails {Title = $"Unknown request type: \"{request}\""});
+            }
+            _logger.Information("New {Event} added as {TransactionId} in {Namespace}", request, id, ns);
+            _diagnosticContext.Set("transaction-id", id);
+            return Ok(new NewTransactionResponse {Offset = id});
+
+        }
+
+
+        // this exists for admin use only
+        [HttpDelete("{ns}")]
+        [Authorize("TLog.delete")]
+        public async Task<IActionResult> DropTransactionLog(
+            [Required] NamespaceId ns
+        )
+        {
+            AuthorizationResult authorizationResult = await _authorizationService.AuthorizeAsync(User, ns, NamespaceAccessRequirement.Name);
+
+            if (!authorizationResult.Succeeded)
+            {
+                return Forbid();
+            }
+
+            await _transactionLogs.Get(ns).Drop();
+
+            return Ok();
+        }
+
+        // TODO: Compounding of the log, do we offer a endpoint for this or is it internal only?
+    }
+
+    public class CallistoGetResponse
+    {
+        public CallistoGetResponse(TransactionEvent[] events, Guid generation, long currentOffset)
+        {
+            Events = events;
+            Generation = generation;
+            CurrentOffset = currentOffset;
+        }
+
+        public long CurrentOffset { get; set; }
+
+        public Guid Generation { get; set; }
+        public TransactionEvent[] Events { get; set; }
+    }
+    public class NewTransactionResponse
+    {
+        public long Offset { get; set; }
+    }
+
+}
diff --git a/Engine/Source/Programs/Horde/HordeStorage/Callisto/Implementation/FileTransactionLog.cs b/Engine/Source/Programs/Horde/HordeStorage/Callisto/Implementation/FileTransactionLog.cs
new file mode 100644
--- /dev/null
+++ b/Engine/Source/Programs/Horde/HordeStorage/Callisto/Implementation/FileTransactionLog.cs
@@ -0,0 +1,615 @@
+﻿// Copyright Epic Games, Inc. All Rights Reserved.
+
+using System;
+using System.Collections.Generic;
+using System.Diagnostics;
+using System.IO;
+using System.Linq;
+using System.Reflection.Metadata;
+using System.Security.Cryptography;
+using System.Text;
+using System.Threading.Tasks;
+using Datadog.Trace;
+using Jupiter;
+using Jupiter.Implementation;
+using Microsoft.Extensions.Options;
+using Newtonsoft.Json;
+using Newtonsoft.Json.Bson;
+using Serilog;
+
+namespace Callisto.Implementation
+{
+    public class FileTransactionLog : ITransactionLog
+    {
+        private FileTransactionLogIndex _index;
+        private readonly FileInfo _transactionLogHandle;
+
+        private readonly object _commitLock = new object();
+        private readonly IOptionsMonitor<CallistoSettings>? _settings;
+        private readonly NamespaceId _namespace;
+
+        private const string IndexFilename = "Index.json";
+        private const string LogFilename = "Log";
+
+        private readonly ILogger _logger = Log.ForContext<FileTransactionLog>();
+
+        public FileTransactionLog(IOptionsMonitor<CallistoSettings> settings, DirectoryInfo root, NamespaceId ns) : this(root, ns)
+        {
+            _settings = settings;
+        }
+
+        public FileTransactionLog(DirectoryInfo root, NamespaceId ns)
+        {
+            _namespace = ns;
+            FileInfo indexFile = new FileInfo(Path.Combine(root.FullName, ns.ToString(), IndexFilename));
+
+            if (!indexFile.Directory?.Exists ?? false)
+            {
+                if (indexFile.Directory != null)
+                {
+                    Directory.CreateDirectory(indexFile.Directory.FullName);
+                }
+            }
+
+            if (indexFile.Exists)
+            {
+                _index = FileTransactionLogIndex.FromFile(indexFile);
+            }
+            else
+            {
+                _index = new FileTransactionLogIndex(indexFile);
+            }
+
+            _transactionLogHandle = new FileInfo(Path.Combine(root.FullName, ns.ToString(), LogFilename));
+        }
+
+
+        internal FileTransactionLogIndex IndexFile
+        {
+            get { return _index; }
+        }
+
+
+        public async Task<long> NewTransaction(AddTransactionEvent @event)
+        {
+            OpRecord eventRecord = AddOp.FromEvent(_index, @event);
+            return await CommitNewEvent(eventRecord);
+        }
+
+        public async Task<long> NewTransaction(RemoveTransactionEvent @event)
+        {
+            OpRecord eventRecord = RemoveOp.FromEvent(_index, @event);
+            return await CommitNewEvent(eventRecord);
+        }
+
+        private async Task<long> CommitNewEvent(OpRecord eventRecord)
+        {
+            // datadog instrumentation of this method
+            using Scope scope = Tracer.Instance.StartActive("log.write");
+            Span span = scope.Span;
+            span.ResourceName = $"write {_namespace}";
+            span.OperationName = "WRITE";
+
+            if (_index.Version != 1)
+                throw new InvalidOperationException($"Version mismatch of transaction log, found {_index.Version} but expected version 1");
+
+            byte[] eventBuffer = SerializeToBuffer(eventRecord);
+            int eventLength = eventBuffer.Length;
+
+            using MD5 md5 = MD5.Create();
+            byte[] md5Hash = md5.ComputeHash(eventBuffer, 0, eventLength);
+
+
+            long offset = await Task.Run(() =>
+            {
+                lock (_commitLock)
+                {
+                    using FileStream fs = _transactionLogHandle.Open(FileMode.OpenOrCreate, FileAccess.Write, FileShare.Read);
+                    // seek to end
+                    fs.Seek(0, SeekOrigin.End);
+                    long position = fs.Position;
+                    
+                    if (_settings?.CurrentValue?.VerifySerialization ?? false)
+                        VerifySerializationOfEvent(position, eventRecord.Op!, eventBuffer, md5Hash, eventLength, md5);
+
+                    OpHeader header = new OpHeader(position, md5Hash, eventLength);
+                    header.Serialize(fs);
+                    fs.Write(eventBuffer, 0, eventLength);
+
+                    return position;
+                }
+            });
+            return offset;
+        }
+
+        private byte[] SerializeToBuffer(OpRecord eventRecord)
+        {
+            using MemoryStream memoryStream = new MemoryStream();
+            eventRecord.Serialize(memoryStream);
+            
+            byte[] buffer = new byte[memoryStream.Position];
+            Array.Copy(memoryStream.GetBuffer(), buffer, memoryStream.Position);
+            return buffer;
+        }
+
+        private void VerifySerializationOfEvent(long transactionId, BaseOp op, byte[] eventBuffer, byte[] originalMd5, int eventLength, MD5 md5)
+        {
+            using MemoryStream testStream = new MemoryStream(eventLength);
+            OpHeader header = new OpHeader(0, originalMd5, eventLength);
+
+            header.Serialize(testStream);
+            testStream.Write(eventBuffer, 0, eventLength);
+
+            testStream.Seek(0, SeekOrigin.Begin);
+
+            OpHeader opHeader = OpHeader.Deserialize(testStream);
+
+            byte[] eventBuf = new byte[opHeader.Length];
+            testStream.Read(eventBuf, 0, opHeader.Length);
+
+            byte[] testMd5Hash = md5.ComputeHash(eventBuf, 0, opHeader.Length);
+            if (!opHeader.Md5hash.SequenceEqual(testMd5Hash))
+            {
+                string headerMd5 = StringUtils.FormatAsHexString(opHeader.Md5hash);
+                string calculatedMd5 = StringUtils.FormatAsHexString(testMd5Hash);
+                _logger.Warning("Verifying header {@Header} for op {@Op}, mismatching hash found {CalculatedHash} expected {ExpectedHash} {TransactionId}", opHeader, op, calculatedMd5, headerMd5, transactionId);
+
+                using MemoryStream ms = new MemoryStream(eventBuf);
+                OpRecord newOp = OpRecord.Deserialize(ms);
+                _logger.Warning("Deserialized form of op {@NewOp} . Old op was {@Op}", newOp.Op, op);
+
+                throw new Exception($"Verifying serialization failed! Expected hash {headerMd5} but calculated {calculatedMd5}");
+            }
+        }
+
+        public Task Drop()
+        {
+            _transactionLogHandle.Delete();
+            // file should never be null, its only nullable for serialization purposes
+            _index.File!.Delete();
+            _index = new FileTransactionLogIndex(_index.File);
+
+            return Task.CompletedTask;
+        }
+
+        public async Task<TransactionEvents> Get(long index, int count, string? notSeenAtSite = null)
+        {
+            // datadog instrumentation of this method
+            using Scope scope = Tracer.Instance.StartActive("log.get");
+            Span span = scope.Span;
+            span.ResourceName = $"get {_namespace}";
+            span.OperationName = "GET";
+
+            if (_index.Version != 1)
+                throw new InvalidOperationException($"Version mismatch of transaction log, found {_index.Version} but expected version 1");
+
+            OpRecord[] events = new OpRecord[count];
+
+            ulong siteFilter = 0;
+            if (notSeenAtSite != null)
+            {
+                int filteredSiteIndex = _index.ToIndex(notSeenAtSite);
+                siteFilter = 1UL << filteredSiteIndex;
+            }
+
+            _transactionLogHandle.Refresh();
+            if (!_transactionLogHandle.Exists)
+                throw new FileNotFoundException("Transaction log is empty");
+
+            // we allow parallel reads from the file, we only need to lock on writes.
+            // ReSharper disable once InconsistentlySynchronizedField
+            await using FileStream fs = _transactionLogHandle.Open(FileMode.Open, FileAccess.Read, FileShare.Read);
+            long offset = index;
+            fs.Seek(offset, SeekOrigin.Begin);
+            int foundCount = 0;
+            long? transactionLogOffsetMismatchFoundAt = null;
+            while (foundCount != count)
+            {
+                // if we have reached the end of the file we stop
+                if (fs.Length == fs.Position)
+                {
+                    break;
+                }
+
+                // if we would read outside the stream when reading the next header we stop
+                if (fs.Position + OpHeader.HeaderLength > fs.Length)
+                    break;
+
+                OpHeader opHeader = OpHeader.Deserialize(fs);
+
+                if (opHeader.Offset != offset)
+                {
+                    transactionLogOffsetMismatchFoundAt = offset;
+                    break;
+                }
+
+                long eventPosition = fs.Position;
+                byte[] eventBuf = new byte[opHeader.Length];
+                await fs.ReadAsync(eventBuf, 0, opHeader.Length);
+
+                // update the offset to our current position in the file in case we need to continue reading
+                offset = fs.Position;
+
+                if (offset - eventPosition != opHeader.Length)
+                    throw new Exception($"Read event size {offset - eventPosition} did not match expected size from header {opHeader.Length}");
+                using MD5 md5 = MD5.Create();
+                byte[] md5Hash = md5.ComputeHash(eventBuf, 0, opHeader.Length);
+                if (!opHeader.Md5hash.SequenceEqual(md5Hash))
+                {
+                    string headerMd5 = StringUtils.FormatAsHexString(opHeader.Md5hash);
+                    string calculatedMd5 = StringUtils.FormatAsHexString(md5Hash);
+                    _logger.Warning("Reading header {@Header}, mismatching hash found {CalculatedHash} expected {ExpectedHash}", opHeader, calculatedMd5, headerMd5);
+                    throw new EventHashMismatchException(opHeader.Offset, headerMd5, calculatedMd5);
+                }
+
+                await using MemoryStream eventStream = new MemoryStream(eventBuf);
+                OpRecord op = OpRecord.Deserialize(eventStream);
+
+                if (siteFilter != 0)
+                {
+                    // we have a site filter, lets check that the event has not been seen at this site (flag is set)
+                    bool hasBeenSeen = (op.Op!.Locations & siteFilter) != 0;
+                    if (hasBeenSeen)
+                    {
+                        continue;
+                    }
+                }
+
+                op.Header = opHeader;
+                events[foundCount] = op;
+                foundCount++;
+            }
+
+            // if there is more content we check the next header
+            if (fs.Length != fs.Position && fs.Position + OpHeader.HeaderLength < fs.Length)
+            {
+                // read another header and verify that it still is valid
+                OpHeader nextHeader = OpHeader.Deserialize(fs);
+
+                Debug.Assert(nextHeader != null);
+                //Debug.Assert(offset == nextHeader.Offset);
+            }
+
+            TransactionEvent[] newEvents = new TransactionEvent[foundCount];
+            for (int i = 0; i < foundCount; i++)
+            {
+                OpRecord tes = events[i];
+                Debug.Assert(condition: tes.Op != null, "tes.Op != null");
+                Debug.Assert(condition: tes.Header != null, "tes.Header != null");
+
+                newEvents[i] = tes.Op.ToTransactionEvent(tes.Header, _index);
+            }
+
+            return new TransactionEvents(newEvents, offset, transactionLogOffsetMismatchFoundAt);
+        }
+
+        internal class OpHeader
+        {
+            // Offset of the current header being written
+            public long Offset;
+
+            // Hash of the transaction event
+            public byte[] Md5hash;
+
+            // Length of transaction event
+            public int Length;
+
+
+            public static int HeaderLength
+            {
+                get { return sizeof(long) + 16 /* MD5 hash */ + sizeof(int); }
+            }
+
+            public OpHeader(in long offset, byte[] md5Hash, in int length)
+            {
+                Offset = offset;
+                Md5hash = md5Hash;
+                Length = length;
+            }
+
+            public void Serialize(Stream s)
+            {
+                using BinaryWriter writer = new BinaryWriter(s, Encoding.UTF8, leaveOpen: true);
+
+                writer.Write(Offset);
+                writer.Write(Md5hash, 0, 16);
+                writer.Write(Length);
+            }
+
+            public static OpHeader Deserialize(Stream s)
+            {
+                using BinaryReader reader = new BinaryReader(s, Encoding.UTF8, leaveOpen: true);
+
+                long offset = reader.ReadInt64();
+                byte[] md5Hash = new byte[16];
+                reader.Read(md5Hash, 0, 16);
+                int length = reader.ReadInt32();
+
+                OpHeader header = new OpHeader(offset, md5Hash, length);
+                return header;
+            }
+        }
+
+        internal enum OpType
+        {
+            Add,
+            Remove
+        };
+
+        internal class OpRecord
+        {
+            public OpRecord(OpType type, BaseOp op)
+            {
+                OpType = type;
+                Op = op;
+            }
+
+            public OpType OpType { get; set; }
+
+            public BaseOp Op { get; set; }
+
+            /// <summary>
+            /// set when a object is being read back so that we have access to the header of the op as well
+            /// </summary>
+            public OpHeader? Header { get; set; } = null;
+
+            public static OpRecord Deserialize(Stream stream)
+            {
+                int opType = stream.ReadByte();
+                OpType op = (OpType) opType;
+
+                BaseOp baseOp;
+                switch (op)
+                {
+                    case OpType.Add:
+                        baseOp = AddOp.Deserialize(stream);
+                        break;
+                    case OpType.Remove:
+                        baseOp = RemoveOp.Deserialize(stream);
+                        break;
+                    default:
+                        throw new ArgumentOutOfRangeException();
+                }
+
+                return new OpRecord(op, baseOp);
+            }
+
+            public void Serialize(Stream stream)
+            {
+                stream.WriteByte((byte)OpType);
+                Op.Serialize(stream);
+            }
+        }
+
+
+        internal abstract class BaseOp
+        {
+            protected BaseOp(string name, string bucket, ulong locations)
+            {
+                Name = name;
+                Bucket = bucket;
+                Locations = locations;
+            }
+
+            public string Name { get; set; }
+
+            public string Bucket { get; set; }
+
+            public ulong Locations { get; set; }
+
+            // 64 entries means we can encode this in a long e.g. 8 bytes
+            private const int MaxLocations = 64;
+
+            public abstract TransactionEvent ToTransactionEvent(OpHeader opHeader, FileTransactionLogIndex index);
+
+            protected List<string> ResolveLocations(FileTransactionLogIndex index)
+            {
+                List<string> locations = new List<string>();
+                for (int i = 0; i < MaxLocations; i++)
+                {
+                    bool flag = (Locations & (1UL << i)) != 0;
+                    if (flag)
+                    {
+                        locations.Add(index.NameOfIndex(i));
+                    }
+                }
+
+                return locations;
+            }
+
+            protected static ulong FlagSetFromLocations(FileTransactionLogIndex index, IEnumerable<string> eventLocations)
+            {
+                List<string> enumerable = eventLocations.ToList();
+                if (!enumerable.Any())
+                    throw new Exception("No locations set, at least the location of the current instance has to be specified");
+                ulong flags = 0;
+                foreach (int indexToSet in enumerable.Select(index.ToIndex))
+                {
+                    flags |= (1UL << indexToSet);
+                }
+
+                return flags;
+            }
+
+            public void Serialize(Stream stream)
+            {
+                using BinaryWriter writer = new BinaryWriter(stream, Encoding.UTF8, leaveOpen: true);
+                writer.Write(Name);
+                writer.Write(Bucket);
+                writer.Write(Locations);
+
+                DoSerialize(stream);
+            }
+
+            protected static (string, string, ulong) Deserialize(Stream stream)
+            {
+                using BinaryReader reader = new BinaryReader(stream, Encoding.UTF8, leaveOpen: true);
+                string name = reader.ReadString();
+                string bucket = reader.ReadString();
+                ulong flags = reader.ReadUInt64();
+                return (name, bucket, flags);
+            }
+
+            protected abstract void DoSerialize(Stream stream);
+        }
+
+
+        internal class RemoveOp : BaseOp
+        {
+            private RemoveOp(string name, string bucket, ulong locations) : base(name, bucket, locations)
+            {
+            }
+
+            public static OpRecord FromEvent(FileTransactionLogIndex index, RemoveTransactionEvent @event)
+            {
+                RemoveOp op = new RemoveOp(@event.Name, @event.Bucket, FlagSetFromLocations(index, @event.Locations));
+
+                return new OpRecord(OpType.Remove, op);
+            }
+
+            public override TransactionEvent ToTransactionEvent(OpHeader opHeader, FileTransactionLogIndex index)
+            {
+                Debug.Assert(condition: Name != null, message: nameof(Name) + " != null");
+                Debug.Assert(condition: Bucket != null, message: nameof(Bucket) + " != null");
+
+                List<string> locations = ResolveLocations(index);
+                return new RemoveTransactionEvent(Name, Bucket, locations, identifier: opHeader.Offset, nextIdentifier: opHeader.Offset + opHeader.Length + OpHeader.HeaderLength);
+            }
+
+            protected override void DoSerialize(Stream stream)
+            {
+                // nothing extra to serialize
+            }
+
+            public new static BaseOp Deserialize(Stream stream)
+            {
+                (string name, string bucket, ulong locations) = BaseOp.Deserialize(stream);
+                return new RemoveOp(name, bucket, locations);
+            }
+        }
+
+        internal class AddOp : BaseOp
+        {
+            private readonly ILogger _logger = Log.ForContext<AddOp>();
+
+            private AddOp(string name, string bucket, ulong locations, List<BlobIdentifier> blobs, byte[] metadataBuffer) : base(name, bucket, locations)
+            {
+                Blobs = blobs;
+                Metadata = metadataBuffer;
+            }
+
+            public List<BlobIdentifier> Blobs { get; set; } = null!;
+
+            public byte[] Metadata { get; set; } = new byte[0];
+
+
+            public static OpRecord FromEvent(FileTransactionLogIndex index, AddTransactionEvent @event)
+            {
+                // create a bson representation of the metadata
+                byte[]? metadataBuffer = new byte[0];
+                if (@event.Metadata != null)
+                {
+                    using MemoryStream ms = new MemoryStream();
+                    BsonDataWriter writer = new BsonDataWriter(ms);
+                    JsonSerializer serializer = JsonSerializer.CreateDefault();
+                    serializer.Serialize(writer, @event.Metadata);
+
+                    long metadataLength = ms.Position;
+                    metadataBuffer = new byte[metadataLength];
+                    Array.Copy(sourceArray: ms.GetBuffer(), metadataBuffer, metadataLength);
+                }
+
+                ulong flags = FlagSetFromLocations(index, @event.Locations);
+
+                AddOp op = new AddOp(@event.Name, @event.Bucket, flags, @event.Blobs.ToList(), metadataBuffer);
+                return new OpRecord(OpType.Add, op);
+            }
+
+            public override TransactionEvent ToTransactionEvent(OpHeader opHeader, FileTransactionLogIndex index)
+            {
+                Dictionary<string, object>? metadata = null;
+                if (Metadata != null && Metadata.Length != 0)
+                {
+                    using MemoryStream ms = new MemoryStream(Metadata);
+                    BsonDataReader reader = new BsonDataReader(ms);
+                    JsonSerializer serializer = JsonSerializer.CreateDefault();
+                    try
+                    {
+                        metadata = serializer.Deserialize<Dictionary<string, object>>(reader);
+                    }
+                    catch (JsonReaderException e)
+                    {
+                        _logger.Warning(e, "Failed to deserialize {@Metadata} into a json blob for object {Name}", Metadata, Name);
+                        metadata = null;
+                    }
+                    catch (Exception e)
+                    {
+                        _logger.Warning(e, "Failed to deserialize {@Metadata} into a json blob for object {Name} due to a unknown exception", Metadata, Name);
+                        metadata = null;
+                    }
+                }
+
+                Debug.Assert(condition: Name != null, message: nameof(Name) + " != null");
+                Debug.Assert(condition: Bucket != null, message: nameof(Bucket) + " != null");
+                Debug.Assert(condition: Blobs != null, message: nameof(Blobs) + " != null");
+
+                List<string> locations = ResolveLocations(index);
+
+                return new AddTransactionEvent(Name, Bucket, Blobs.ToArray(), locations, metadata, identifier: opHeader.Offset, nextIdentifier: opHeader.Offset + opHeader.Length + OpHeader.HeaderLength);
+            }
+
+            protected override void DoSerialize(Stream stream)
+            {
+                using BinaryWriter writer = new BinaryWriter(stream, Encoding.UTF8, leaveOpen: true);
+                writer.Write((byte)Blobs.Count);
+                foreach (BlobIdentifier blob in Blobs)
+                {
+                    writer.Write(blob.HashData,0, ContentHash.HashLength);
+                }
+                writer.Write(Metadata.Length);
+                writer.Write(Metadata, 0, Metadata.Length);
+            }
+
+            public new static BaseOp Deserialize(Stream stream)
+            {
+                (string name, string bucket, ulong locations) = BaseOp.Deserialize(stream);
+
+                using BinaryReader reader = new BinaryReader(stream, Encoding.UTF8, leaveOpen: true);
+                byte blobsCount = reader.ReadByte();
+                List<BlobIdentifier> blobs = new List<BlobIdentifier>();
+                for (int i = 0; i < blobsCount; i++)
+                {
+                    byte[] blob = new byte[ContentHash.HashLength];
+                    reader.Read(blob, 0, ContentHash.HashLength);
+                    blobs.Add(new BlobIdentifier(blob));
+                }
+
+                int metadataLength = reader.ReadInt32();
+                byte[] metadataBuffer = new byte[metadataLength];
+                reader.Read(metadataBuffer, 0, metadataLength);
+
+                return new AddOp(name, bucket, locations, blobs, metadataBuffer);
+            }
+        }
+
+        public TransactionLogDescription Describe()
+        {
+            return new TransactionLogDescription(_index.LogGeneration, _index.Locations);
+        }
+    }
+
+    public class EventHashMismatchException : Exception
+    {
+        public long Offset { get; }
+        public string HeaderMd5 { get; }
+        public string CalculatedMd5 { get; }
+
+        public EventHashMismatchException(long offset, string headerMd5, string calculatedMd5) : base(
+            $"Read transaction id {offset} expected hash {headerMd5} but calculated {calculatedMd5}")
+        {
+            Offset = offset;
+            HeaderMd5 = headerMd5;
+            CalculatedMd5 = calculatedMd5;
+        }
+    }
+}
diff --git a/Engine/Source/Programs/Horde/HordeStorage/Callisto/Implementation/FileTransactionLogIndex.cs b/Engine/Source/Programs/Horde/HordeStorage/Callisto/Implementation/FileTransactionLogIndex.cs
new file mode 100644
--- /dev/null
+++ b/Engine/Source/Programs/Horde/HordeStorage/Callisto/Implementation/FileTransactionLogIndex.cs
@@ -0,0 +1,108 @@
+﻿// Copyright Epic Games, Inc. All Rights Reserved.
+
+using System;
+using System.Collections.Generic;
+using System.IO;
+using System.Linq;
+using Datadog.Trace;
+using Newtonsoft.Json;
+
+namespace Callisto.Implementation
+{
+    public class FileTransactionLogIndex
+    {
+        [JsonProperty]
+        private readonly Dictionary<string, int> _locationToIdentifier = new Dictionary<string, int>(StringComparer.InvariantCultureIgnoreCase);
+
+        [JsonProperty]
+        private int _nextLocation = 0; 
+
+        [JsonIgnore]
+        internal FileInfo? File { get; set; }
+
+        [JsonProperty]
+        public Guid LogGeneration { get; set; }
+
+
+        [JsonProperty]
+        public int Version { get; set; }
+
+
+        [JsonIgnore]
+        public List<string> Locations
+        {
+            get { return _locationToIdentifier.Keys.ToList(); }
+        }
+
+        // ReSharper disable once UnusedMember.Local
+        public FileTransactionLogIndex()
+        {
+        }
+
+        public FileTransactionLogIndex(FileInfo indexFile, int version = 1)
+        {
+            File = indexFile;
+            // assign a new guid for this generation when we create the index file
+            LogGeneration = Guid.NewGuid();
+            Version = version;
+            Save();
+        }
+
+        private void Save()
+        {
+            if (File == null)
+            {
+                throw new ArgumentNullException();
+            }
+
+            JsonSerializer js = JsonSerializer.CreateDefault();
+            using FileStream fs = File.Open(FileMode.Create, FileAccess.Write, FileShare.None);
+            using StreamWriter sw = new StreamWriter(fs);
+            js.Serialize(sw, value: this);
+        }
+
+        public static FileTransactionLogIndex FromFile(FileInfo file)
+        {
+            JsonSerializer js = JsonSerializer.CreateDefault();
+            using StreamReader sr = file.OpenText();
+            using JsonTextReader jsonTextReader = new JsonTextReader(sr);
+            FileTransactionLogIndex? index = (FileTransactionLogIndex?) js.Deserialize(jsonTextReader, objectType: typeof(FileTransactionLogIndex));
+            index!.File = file;
+            return index;
+        }
+
+        public int ToIndex(string location)
+        {
+            lock (_locationToIdentifier)
+            {
+                if (_locationToIdentifier.TryGetValue(location, value: out int index))
+                {
+                    return index;
+                }
+
+                using Scope _ = Tracer.Instance.StartActive("update.index");
+                int newIndex = _nextLocation++;
+                _locationToIdentifier.Add(location, newIndex);
+                Save();
+                return newIndex;
+            }
+        }
+
+        public string NameOfIndex(in int i)
+        {
+            lock (_locationToIdentifier)
+            {
+                KeyValuePair<string, int>[] kvp = _locationToIdentifier.ToArray();
+                foreach ((string key, int index) in kvp)
+                {
+                    if (index == i)
+                    {
+                        return key;
+                    }
+                }
+            }
+
+            throw new Exception($"Unable to map index {i} to a location");
+        }
+    }
+}
diff --git a/Engine/Source/Programs/Horde/HordeStorage/Callisto/Implementation/FileTransactionLogProxy.cs b/Engine/Source/Programs/Horde/HordeStorage/Callisto/Implementation/FileTransactionLogProxy.cs
new file mode 100644
--- /dev/null
+++ b/Engine/Source/Programs/Horde/HordeStorage/Callisto/Implementation/FileTransactionLogProxy.cs
@@ -0,0 +1,57 @@
+﻿// Copyright Epic Games, Inc. All Rights Reserved.
+
+using System;
+using System.Collections.Concurrent;
+using System.IO;
+using System.Linq;
+using Jupiter.Implementation;
+using Microsoft.Extensions.Options;
+
+namespace Callisto.Implementation
+{
+    public class FileTransactionLogProxy : ITransactionLogs
+    {
+        private readonly IOptionsMonitor<CallistoSettings> _settings;
+        private readonly ConcurrentDictionary<NamespaceId, FileTransactionLog> _writers = new ConcurrentDictionary<NamespaceId, FileTransactionLog>();
+
+        private readonly DirectoryInfo _root;
+
+        public FileTransactionLogProxy(IOptionsMonitor<CallistoSettings> settings)
+        {
+            _settings = settings;
+            string root = settings.CurrentValue.TransactionLogRoot;
+
+            if (!Directory.Exists(root))
+            {
+                Directory.CreateDirectory(root);
+            }
+
+            _root = new DirectoryInfo(root);
+            foreach (DirectoryInfo dir in _root.EnumerateDirectories())
+            {
+                NamespaceId ns;
+                try
+                {
+                    ns = new NamespaceId(dir.Name);
+                }
+                catch (ArgumentException)
+                {
+                    // ignore any directories that are not valid namespaces
+                    continue;
+                }
+                _writers.TryAdd(ns, new FileTransactionLog(_settings, _root, ns));
+            }
+        }
+
+        public NamespaceId[] GetNamespaces()
+        {
+            return _writers.Keys.Distinct().ToArray();
+        }
+
+        public ITransactionLog Get(NamespaceId ns)
+        {
+            FileTransactionLog writer = _writers.GetOrAdd(ns, s => new FileTransactionLog(_settings, _root, ns));
+            return writer;
+        }
+    }
+}
diff --git a/Engine/Source/Programs/Horde/HordeStorage/Callisto/Implementation/ITransactionLog.cs b/Engine/Source/Programs/Horde/HordeStorage/Callisto/Implementation/ITransactionLog.cs
new file mode 100644
--- /dev/null
+++ b/Engine/Source/Programs/Horde/HordeStorage/Callisto/Implementation/ITransactionLog.cs
@@ -0,0 +1,54 @@
+﻿// Copyright Epic Games, Inc. All Rights Reserved.
+
+using System;
+using System.Collections.Generic;
+using System.Threading.Tasks;
+using Jupiter.Implementation;
+
+namespace Callisto.Implementation
+{
+    public interface ITransactionLog
+    {
+        Task<long> NewTransaction(AddTransactionEvent @event);
+
+        Task<long> NewTransaction(RemoveTransactionEvent @event);
+
+        Task<TransactionEvents> Get(long index, int count = 100, string? notSeenAtSite = null);
+        Task Drop();
+        TransactionLogDescription Describe();
+    }
+
+
+    public interface ITransactionLogs
+    {
+        NamespaceId[] GetNamespaces();
+
+        ITransactionLog Get(NamespaceId ns);
+    }
+
+    public class TransactionLogDescription
+    {
+        public TransactionLogDescription(Guid generation, List<string> knownSites)
+        {
+            Generation = generation;
+            KnownSites = knownSites;
+        }
+
+        public Guid Generation { get; set; }
+        public List<string> KnownSites { get; set; }
+    }
+
+    public class TransactionEvents
+    {
+        public TransactionEvents(TransactionEvent[] events, long currentOffset, long? transactionLogOffsetMismatchFoundAt)
+        {
+            Events = events;
+            CurrentOffset = currentOffset;
+            TransactionLogMismatchFoundAt = transactionLogOffsetMismatchFoundAt;
+        }
+
+        public TransactionEvent[] Events { get; set; }
+        public long CurrentOffset { get; set; }
+        public long? TransactionLogMismatchFoundAt { get; }
+    }
+}
diff --git a/Engine/Source/Programs/Horde/HordeStorage/Callisto/Implementation/MemoryTransactionLog.cs b/Engine/Source/Programs/Horde/HordeStorage/Callisto/Implementation/MemoryTransactionLog.cs
new file mode 100644
--- /dev/null
+++ b/Engine/Source/Programs/Horde/HordeStorage/Callisto/Implementation/MemoryTransactionLog.cs
@@ -0,0 +1,85 @@
+﻿// Copyright Epic Games, Inc. All Rights Reserved.
+
+using System;
+using System.Collections.Concurrent;
+using System.Collections.Generic;
+using System.Linq;
+using System.Threading.Tasks;
+using Jupiter.Implementation;
+
+namespace Callisto.Implementation
+{
+    // since we back this with a list we do not support as many values as the interface expects (a full long)
+    // but as this is for development testing only this should be okay
+    public class MemoryTransactionLog : ITransactionLog
+    {
+        private List<TransactionEvent> _events = new List<TransactionEvent>();
+
+        private Guid _generation = Guid.NewGuid();
+
+        private Task<long> InsertNewTransaction(TransactionEvent @event)
+        {
+            long eventId;
+            lock (_events)
+            {
+                eventId = _events.Count;
+                _events.Add(@event);
+            }
+
+            return Task.FromResult(eventId);
+        }
+
+        public async Task<long> NewTransaction(AddTransactionEvent @event)
+        {
+            return await InsertNewTransaction(@event);
+        }
+
+        public async Task<long> NewTransaction(RemoveTransactionEvent @event)
+        {
+            return await InsertNewTransaction(@event);
+        }
+
+        public Task<TransactionEvents> Get(long index, int count, string? notSeenAtSite)
+        {
+            if (!string.IsNullOrEmpty(notSeenAtSite))
+            {
+                throw new NotImplementedException("Memory transaction log does not support site filtering");
+            }
+
+            if (count != 1)
+            {
+                throw new NotImplementedException("Memory transaction log only support returning a single event");
+            }
+
+            return Task.FromResult(new TransactionEvents(new [] {_events[(int)index]}, 0, null));
+        }
+
+        public Task Drop()
+        {
+            _events = new List<TransactionEvent>();
+            _generation = Guid.NewGuid();
+
+            return Task.CompletedTask;
+        }
+
+        public TransactionLogDescription Describe()
+        {
+            return new TransactionLogDescription(_generation, _events.SelectMany((e) => e.Locations).Distinct().ToList());
+        }
+    }
+
+    class MemoryTransactionLogs : ITransactionLogs
+    {
+        private readonly ConcurrentDictionary<NamespaceId, MemoryTransactionLog> _transactionLogs = new ConcurrentDictionary<NamespaceId, MemoryTransactionLog>();
+
+        public NamespaceId[] GetNamespaces()
+        {
+            return _transactionLogs.Keys.ToArray();
+        }
+
+        public ITransactionLog Get(NamespaceId ns)
+        {
+            return _transactionLogs.GetOrAdd(ns, s => new MemoryTransactionLog());
+        }
+    }
+}
diff --git a/Engine/Source/Programs/Horde/HordeStorage/Callisto/Program.cs b/Engine/Source/Programs/Horde/HordeStorage/Callisto/Program.cs
new file mode 100644
--- /dev/null
+++ b/Engine/Source/Programs/Horde/HordeStorage/Callisto/Program.cs
@@ -0,0 +1,15 @@
+// Copyright Epic Games, Inc. All Rights Reserved.
+
+using Jupiter;
+
+namespace Callisto
+{
+    // ReSharper disable once ClassNeverInstantiated.Global
+    public class Program
+    {
+        public static int Main(string[] args)
+        {
+            return BaseProgram<CallistoStartup>.BaseMain(args);
+        }
+    }
+}
diff --git a/Engine/Source/Programs/Horde/HordeStorage/Callisto/Properties/launchSettings.json b/Engine/Source/Programs/Horde/HordeStorage/Callisto/Properties/launchSettings.json
new file mode 100644
--- /dev/null
+++ b/Engine/Source/Programs/Horde/HordeStorage/Callisto/Properties/launchSettings.json
@@ -0,0 +1,32 @@
+{
+  "iisSettings": {
+    "windowsAuthentication": false,
+    "anonymousAuthentication": true,
+    "iisExpress": {
+      "applicationUrl": "http://localhost:57002",
+      "sslPort": 0
+    }
+  },
+  "$schema": "http://json.schemastore.org/launchsettings.json",
+  "profiles": {
+    "IIS Express": {
+      "commandName": "IISExpress",
+      "launchBrowser": true,
+      "launchUrl": "api/v1/t/test/0",
+      "environmentVariables": {
+        "ASPNETCORE_ENVIRONMENT": "Development"
+      }
+    },
+    "Local Deployment": {
+      "commandName": "IISExpress",
+      "launchBrowser": true,
+      "launchUrl": "api/v1/t/test/0",
+      "environmentVariables": {
+        "Callisto__TransactionLogRoot": "%TEMP%/Callisto",
+        "Callisto__CurrentSite": "local",
+        "CORECLR_ENABLE_PROFILING": "0",
+        "ASPNETCORE_ENVIRONMENT": "Development"
+      }
+    }
+  }
+}
diff --git a/Engine/Source/Programs/Horde/HordeStorage/Callisto/Tests/CallisoFunctionalTests/CallisoFunctionalTests.csproj b/Engine/Source/Programs/Horde/HordeStorage/Callisto/Tests/CallisoFunctionalTests/CallisoFunctionalTests.csproj
new file mode 100644
--- /dev/null
+++ b/Engine/Source/Programs/Horde/HordeStorage/Callisto/Tests/CallisoFunctionalTests/CallisoFunctionalTests.csproj
@@ -0,0 +1,30 @@
+﻿<Project Sdk="Microsoft.NET.Sdk">
+
+  <PropertyGroup>
+    <TargetFramework>net6.0</TargetFramework>
+
+    <IsPackable>false</IsPackable>
+    <Nullable>enable</Nullable>
+    <TreatWarningsAsErrors>true</TreatWarningsAsErrors>
+    <WarningsAsErrors />
+    <LangVersion>10</LangVersion>
+  </PropertyGroup>
+
+  <ItemGroup>
+    <PackageReference Include="Microsoft.AspNetCore.TestHost" Version="6.0.0" />
+    <PackageReference Include="Microsoft.Extensions.DependencyInjection.Abstractions" Version="6.0.0" />
+    <PackageReference Include="Microsoft.NET.Test.Sdk" Version="17.0.0" />
+    <PackageReference Include="MSTest.TestAdapter" Version="2.2.7" />
+    <PackageReference Include="MSTest.TestFramework" Version="2.2.7" />
+    <PackageReference Include="coverlet.collector" Version="3.0.2">
+      <PrivateAssets>all</PrivateAssets>
+      <IncludeAssets>runtime; build; native; contentfiles; analyzers; buildtransitive</IncludeAssets>
+    </PackageReference>
+    <PackageReference Include="SerilogAnalyzer" Version="0.15.0" />
+  </ItemGroup>
+
+  <ItemGroup>
+    <ProjectReference Include="..\..\Callisto.csproj" />
+  </ItemGroup>
+
+</Project>
diff --git a/Engine/Source/Programs/Horde/HordeStorage/Callisto/Tests/CallisoFunctionalTests/CallistoTests.cs b/Engine/Source/Programs/Horde/HordeStorage/Callisto/Tests/CallisoFunctionalTests/CallistoTests.cs
new file mode 100644
--- /dev/null
+++ b/Engine/Source/Programs/Horde/HordeStorage/Callisto/Tests/CallisoFunctionalTests/CallistoTests.cs
@@ -0,0 +1,291 @@
+// Copyright Epic Games, Inc. All Rights Reserved.
+
+using System;
+using System.Collections.Generic;
+using System.Diagnostics;
+using System.IO;
+using System.Linq;
+using System.Net.Http;
+using System.Threading.Tasks;
+using Callisto.Controllers;
+using Callisto.Implementation;
+using Jupiter.Implementation;
+using Microsoft.AspNetCore.TestHost;
+using Microsoft.AspNetCore.Hosting;
+using Microsoft.Extensions.Configuration;
+using Microsoft.Extensions.DependencyInjection;
+using Microsoft.Extensions.Options;
+using Microsoft.VisualStudio.TestTools.UnitTesting;
+using Newtonsoft.Json.Linq;
+using Serilog;
+using Serilog.Core;
+
+namespace Callisto.FunctionalTests
+{
+    [TestClass]
+    public class CallistoTests
+    {
+        private TestServer _server = null!;
+        private HttpClient _httpClient = null!;
+        private readonly NamespaceId TestNamespace = new NamespaceId("test");
+        private readonly NamespaceId TestExistingNamespace = new NamespaceId("test-with-contents");
+        private readonly NamespaceId FullFlowNamespace = new NamespaceId("full-flow");
+
+        private IConfiguration _configuration = null!;
+
+        [TestInitialize]
+        public async Task Setup()
+        {
+            IConfigurationRoot configuration = new ConfigurationBuilder()
+                // we are not reading the base appSettings here as we want exact control over what runs in the tests
+                .AddJsonFile("appsettings.Testing.json", true)
+                .Build();
+
+            _configuration = configuration;
+
+            Logger logger = new LoggerConfiguration()
+                .ReadFrom.Configuration(configuration)
+                .CreateLogger();
+
+            TestServer server = new TestServer(new WebHostBuilder()
+                .UseConfiguration(configuration)
+                .UseEnvironment("Testing")
+                .UseSerilog(logger)
+                .UseStartup<CallistoStartup>()
+            );
+            _httpClient = server.CreateClient();
+            _server = server;
+
+            // Seed S3 storage
+            await Seed();
+        }
+
+        private async Task Seed()
+        {
+            IOptionsMonitor<CallistoSettings> settings = _server.Services.GetService<IOptionsMonitor<CallistoSettings>>()!;
+
+            FileTransactionLog writer = new FileTransactionLog(settings, root: new DirectoryInfo(settings.CurrentValue.TransactionLogRoot), TestExistingNamespace);
+            await writer.NewTransaction(new AddTransactionEvent("Foo", "Bar", blobs: new[] {new BlobIdentifier("0BEEC7B5EA3F0FDBC95D0DD47F3C5BC275DA8A33"), }, new [] {"HQ"}.ToList()));
+            await writer.NewTransaction(new AddTransactionEvent("Foo", "Bar", blobs: new[] {new BlobIdentifier("0BEEC7B5EA3F0FDBC95D0DD47F3C5BC275DA8A34"), }, new [] {"HQ"}.ToList()));
+        }
+
+        [TestCleanup]
+        public void Teardown()
+        {
+            CallistoSettings settings = new CallistoSettings();
+            _configuration!.GetSection("Callisto").Bind(settings);
+
+            Assert.IsNotNull(settings.TransactionLogRoot);
+#pragma warning disable CS8604 // Possible null reference argument.
+            string tempDir = Path.Combine(settings.TransactionLogRoot, TestNamespace.ToString());
+#pragma warning restore CS8604 // Possible null reference argument.
+            if (Directory.Exists(tempDir))
+            {
+                Directory.Delete(tempDir, true);
+            }
+
+#pragma warning disable CS8604 // Possible null reference argument.
+            string tempDir2 = Path.Combine(settings.TransactionLogRoot, TestExistingNamespace.ToString());
+#pragma warning restore CS8604 // Possible null reference argument.
+            if (Directory.Exists(tempDir2))
+            {
+                Directory.Delete(tempDir2, true);
+            }
+
+#pragma warning disable CS8604 // Possible null reference argument.
+            string tempDir3 = Path.Combine(settings.TransactionLogRoot, FullFlowNamespace.ToString());
+#pragma warning restore CS8604 // Possible null reference argument.
+            if (Directory.Exists(tempDir3))
+            {
+                Directory.Delete(tempDir3, true);
+            }
+        }
+
+        [TestMethod]
+        public async Task GetEvent()
+        {
+            const int startOffset = 0;
+            HttpResponseMessage response =
+                await _httpClient!.GetAsync($"api/v1/t/{TestExistingNamespace}/{startOffset}");
+            response.EnsureSuccessStatusCode();
+            dynamic ter = await response.Content.ReadAsAsync<dynamic>();
+            Assert.IsNotNull(ter.generation);
+
+            JArray events = ter.events;
+            Assert.AreEqual(2, events.Count);
+            AddTransactionEvent? tEvent = events[0].ToObject<AddTransactionEvent>();
+
+            Assert.AreEqual(0, tEvent!.Identifier);
+            Assert.AreEqual(70, tEvent!.NextIdentifier);
+            Assert.AreEqual("Foo", tEvent!.Name);
+            Assert.AreEqual("Bar", tEvent.Bucket);
+            Assert.IsNotNull(tEvent.Blobs);
+            Assert.AreEqual(1, tEvent.Blobs!.Length);
+            Assert.AreEqual(new BlobIdentifier("0BEEC7B5EA3F0FDBC95D0DD47F3C5BC275DA8A33"), actual: tEvent.Blobs.First());
+        }
+
+        [TestMethod]
+        public async Task GetMismatchedOffsetEvent()
+        {
+            const int startOffset = 50; // this is not a valid offset, so we should get the object at offset 70
+            HttpResponseMessage response =
+                await _httpClient!.GetAsync($"api/v1/t/{TestExistingNamespace}/{startOffset}");
+            response.EnsureSuccessStatusCode();
+            dynamic ter = await response.Content.ReadAsAsync<dynamic>();
+            Assert.IsNotNull(ter.generation);
+
+            JArray events = ter.events;
+            Assert.AreEqual(1, events.Count);
+            AddTransactionEvent? tEvent = events[0].ToObject<AddTransactionEvent>();
+
+            Assert.AreEqual(70, tEvent!.Identifier);
+            Assert.AreEqual(140, tEvent!.NextIdentifier);
+            Assert.AreEqual("Foo", tEvent!.Name);
+            Assert.AreEqual("Bar", tEvent.Bucket);
+            Assert.IsNotNull(tEvent.Blobs);
+            Assert.AreEqual(1, tEvent.Blobs!.Length);
+            Assert.AreEqual(new BlobIdentifier("0BEEC7B5EA3F0FDBC95D0DD47F3C5BC275DA8A34"), actual: tEvent.Blobs.First());
+        }
+
+
+        [TestMethod]    
+        public async Task GetNonExistentEvent()
+        {
+            // fetch the event after the second event, this will not exist
+            const int startOffset = 77;
+            HttpResponseMessage response = await _httpClient!.GetAsync($"api/v1/t/{TestExistingNamespace}/{startOffset}");
+            response.EnsureSuccessStatusCode();
+            dynamic ter = await response.Content.ReadAsAsync<dynamic>();
+            Assert.IsNotNull(ter.generation);
+            JArray events = ter.events;
+            Assert.AreEqual(0, events.Count);
+        }
+
+
+        [TestMethod]
+        public async Task AddNewEvent()
+        {
+            var request = new
+            {
+                Type = "Add",
+                Name = "foo",
+                Bucket = "bar",
+                Blobs = new[]
+                {
+                    new BlobIdentifier("0BEEC7B5EA3F0FDBC95D0DD47F3C5BC275DA8A33")
+                },
+            };
+
+            HttpResponseMessage response =
+                await _httpClient!.PostAsJsonAsync(requestUri: $"api/v1/t/{TestNamespace}", request);
+            response.EnsureSuccessStatusCode();
+            NewTransactionResponse r = await response.Content.ReadAsAsync<NewTransactionResponse>();
+            long offset = r.Offset;
+            Assert.AreEqual(0, offset);
+        }
+
+        [TestMethod]
+        public async Task AddAndRemoveEvent()
+        {
+            {
+                var request = new
+                {
+                    Type = "Add",
+                    Name = "addEvent",
+                    Bucket = "foo",
+                    Blobs = new[]
+                    {
+                        "0BEEC7B5EA3F0FDBC95D0DD47F3C5BC275DA8A33"
+                    },
+                };
+
+                HttpResponseMessage response =
+                    await _httpClient!.PostAsJsonAsync(requestUri: $"api/v1/t/{TestNamespace}", request);
+                response.EnsureSuccessStatusCode();
+                NewTransactionResponse r = await response.Content.ReadAsAsync<NewTransactionResponse>();
+                long offset = r.Offset;
+                Assert.AreEqual(0, offset);
+            }
+            {
+                var request = new
+                {
+                    Type = "Remove",
+                    Name = "addEvent",
+                    Bucket = "foo",
+                };
+
+                HttpResponseMessage response =
+                    await _httpClient!.PostAsJsonAsync(requestUri: $"api/v1/t/{TestNamespace}", request);
+                response.EnsureSuccessStatusCode();
+                NewTransactionResponse r = await response.Content.ReadAsAsync<NewTransactionResponse>();
+                long offset = r.Offset;
+                Assert.AreEqual(75, offset);
+            }
+        }
+
+        [TestMethod]
+        public async Task FullFlow()
+        {
+            {
+                var request = new
+                {
+                    Type = "Add",
+                    Name = "addEvent",
+                    Bucket = "foo",
+                    Blobs = new[]
+                    {
+                        "0BEEC7B5EA3F0FDBC95D0DD47F3C5BC275DA8A33"
+                    },
+                };
+
+                HttpResponseMessage response =
+                    await _httpClient!.PostAsJsonAsync(requestUri: $"api/v1/t/{FullFlowNamespace}", request);
+                response.EnsureSuccessStatusCode();
+                NewTransactionResponse r = await response.Content.ReadAsAsync<NewTransactionResponse>();
+                long offset = r.Offset;
+                Assert.AreEqual(0, offset);
+            }
+            {
+                var request = new
+                {
+                    Type = "Remove",
+                    Name = "addEvent",
+                    Bucket = "foo",
+                };
+
+                HttpResponseMessage response =
+                    await _httpClient!.PostAsJsonAsync(requestUri: $"api/v1/t/{FullFlowNamespace}", request);
+                response.EnsureSuccessStatusCode();
+                NewTransactionResponse r = await response.Content.ReadAsAsync<NewTransactionResponse>();
+                long offset = r.Offset;
+                Assert.AreEqual(75, offset);
+            }
+
+            {
+                const int startOffset = 0;
+                HttpResponseMessage response =
+                    await _httpClient!.GetAsync($"api/v1/t/{FullFlowNamespace}/{startOffset}");
+                response.EnsureSuccessStatusCode();
+                dynamic ter = await response.Content.ReadAsAsync<dynamic>();
+                JArray events = ter.events;
+                Assert.AreEqual(2, events.Count);
+                AddTransactionEvent? addEvent = events[0].ToObject<AddTransactionEvent>();
+
+                Assert.AreEqual("addEvent", addEvent!.Name);
+                Assert.AreEqual(0, addEvent.Identifier);
+                Assert.AreEqual(75, addEvent.NextIdentifier);
+                Assert.AreEqual("foo", addEvent.Bucket);
+                Assert.IsNotNull(addEvent.Blobs);
+                Assert.AreEqual(1, addEvent.Blobs!.Length);
+                Assert.AreEqual(new BlobIdentifier("0BEEC7B5EA3F0FDBC95D0DD47F3C5BC275DA8A33"), actual: addEvent.Blobs.First());
+
+                RemoveTransactionEvent? removeEvent = events[1].ToObject<RemoveTransactionEvent>();
+
+                Assert.AreEqual("addEvent", removeEvent!.Name);
+                Assert.AreEqual(75, removeEvent.Identifier);
+                Assert.AreEqual(125, removeEvent.NextIdentifier);
+                Assert.AreEqual("foo", removeEvent.Bucket);
+            }
+        }
+    }
+}
diff --git a/Engine/Source/Programs/Horde/HordeStorage/Callisto/Tests/CallisoUnitTests/CallisoUnitTests.csproj b/Engine/Source/Programs/Horde/HordeStorage/Callisto/Tests/CallisoUnitTests/CallisoUnitTests.csproj
new file mode 100644
--- /dev/null
+++ b/Engine/Source/Programs/Horde/HordeStorage/Callisto/Tests/CallisoUnitTests/CallisoUnitTests.csproj
@@ -0,0 +1,28 @@
+﻿<Project Sdk="Microsoft.NET.Sdk">
+
+  <PropertyGroup>
+    <TargetFramework>net6.0</TargetFramework>
+
+    <IsPackable>false</IsPackable>
+    <Nullable>enable</Nullable>
+    <TreatWarningsAsErrors>true</TreatWarningsAsErrors>
+    <WarningsAsErrors />
+    <LangVersion>10</LangVersion>
+  </PropertyGroup>
+
+  <ItemGroup>
+    <PackageReference Include="Microsoft.NET.Test.Sdk" Version="17.0.0" />
+    <PackageReference Include="MSTest.TestAdapter" Version="2.2.7" />
+    <PackageReference Include="MSTest.TestFramework" Version="2.2.7" />
+    <PackageReference Include="coverlet.collector" Version="3.0.2">
+      <PrivateAssets>all</PrivateAssets>
+      <IncludeAssets>runtime; build; native; contentfiles; analyzers; buildtransitive</IncludeAssets>
+    </PackageReference>
+    <PackageReference Include="SerilogAnalyzer" Version="0.15.0" />
+  </ItemGroup>
+
+  <ItemGroup>
+    <ProjectReference Include="..\..\Callisto.csproj" />
+  </ItemGroup>
+
+</Project>
diff --git a/Engine/Source/Programs/Horde/HordeStorage/Callisto/Tests/CallisoUnitTests/TransactionLogTests.cs b/Engine/Source/Programs/Horde/HordeStorage/Callisto/Tests/CallisoUnitTests/TransactionLogTests.cs
new file mode 100644
--- /dev/null
+++ b/Engine/Source/Programs/Horde/HordeStorage/Callisto/Tests/CallisoUnitTests/TransactionLogTests.cs
@@ -0,0 +1,294 @@
+// Copyright Epic Games, Inc. All Rights Reserved.
+
+using System;
+using System.Collections.Generic;
+using System.Diagnostics;
+using System.IO;
+using System.Linq;
+using Microsoft.VisualStudio.TestTools.UnitTesting;
+using System.Threading.Tasks;
+using Callisto.Implementation;
+using Jupiter.Implementation;
+
+namespace Calliso.UnitTests
+{
+    [TestClass]
+    public class TransactionLogTests
+    {
+        private readonly DirectoryInfo _tempPath = new DirectoryInfo(Path.GetTempPath());
+        private readonly NamespaceId _namespace = new NamespaceId("test-namespace");
+
+        [TestInitialize]
+        public void CleanupTempDir()
+        {
+            string tempDir = Path.Combine(_tempPath.FullName, _namespace.ToString());
+            if (Directory.Exists(tempDir))
+            {
+                Directory.Delete(tempDir, true);
+            }
+        }
+
+        [TestMethod]
+        public async Task MemoryTransactionLog()
+        {
+            MemoryTransactionLog log = new MemoryTransactionLog();
+            AddTransactionEvent @event = new AddTransactionEvent("Foo", "Bar",
+                blobs: new[] {new BlobIdentifier("0BEEC7B5EA3F0FDBC95D0DD47F3C5BC275DA8A33"), },
+                locations: new[] {"ARN"}.ToList());
+
+            long insertedIndex = await log.NewTransaction(@event);
+            Assert.AreEqual(0, insertedIndex);
+
+            TransactionEvent event2 = (await log.Get(0, 1, null)).Events.First();
+            Assert.AreEqual(@event.Name, event2.Name);
+
+            Assert.IsInstanceOfType(event2, expectedType: typeof(TransactionEvent));
+            AddTransactionEvent addEvent = (AddTransactionEvent) event2;
+            CollectionAssert.AreEqual(@event.Blobs, addEvent.Blobs);
+        }
+
+        [TestMethod]
+        public async Task FileTransactionLog()
+        {
+            CleanupTempDir();
+
+            FileTransactionLog log = new FileTransactionLog(_tempPath, _namespace);
+            AddTransactionEvent @event = new AddTransactionEvent("Foo", "Bar",
+                blobs: new[] {new BlobIdentifier("0BEEC7B5EA3F0FDBC95D0DD47F3C5BC275DA8A33") },
+                locations: new[] {"ARN"}.ToList());
+
+            long insertedIndex = await log.NewTransaction(@event);
+            Assert.AreEqual(0, insertedIndex);
+
+            TransactionEvent[] events = (await log.Get(0, 1)).Events;
+            TransactionEvent event2 = events[0];
+
+            Assert.AreEqual(@event.Name, event2.Name);
+            Assert.IsInstanceOfType(event2, expectedType: typeof(TransactionEvent));
+            AddTransactionEvent addEvent = (AddTransactionEvent) event2;
+            CollectionAssert.AreEqual(@event.Blobs, addEvent.Blobs);
+            Assert.IsTrue(addEvent.Locations.Contains("ARN"));
+        }
+
+        [TestMethod]
+        public async Task FileTransactionLogUnicodeEvent()
+        {
+            CleanupTempDir();
+
+            FileTransactionLog log = new FileTransactionLog(_tempPath, _namespace);
+            AddTransactionEvent @event = new AddTransactionEvent("😀", "😉",
+                blobs: new[] {new BlobIdentifier("0BEEC7B5EA3F0FDBC95D0DD47F3C5BC275DA8A33")},
+                locations: new[] {"ARN"}.ToList());
+            long insertedIndex = await log.NewTransaction(@event);
+            Assert.AreEqual(0, insertedIndex);
+
+            TransactionEvent[] events = (await log.Get(0, 1)).Events;
+            TransactionEvent event2 = events[0];
+
+            Assert.AreEqual(@event.Name, event2.Name);
+            Assert.IsInstanceOfType(event2, expectedType: typeof(TransactionEvent));
+            AddTransactionEvent addEvent = (AddTransactionEvent) event2;
+            CollectionAssert.AreEqual(@event.Blobs, addEvent.Blobs);
+            Assert.IsTrue(addEvent.Locations.Contains("ARN"));
+        }
+
+        [TestMethod]
+        public async Task FileTransactionLogMetadataEvent()
+        {
+            CleanupTempDir();
+
+            FileTransactionLog log = new FileTransactionLog(_tempPath, _namespace);
+            Dictionary<string, object> metadata = new Dictionary<string, object>
+            {
+                {"foo", "baz"},
+                {"bar", "baz"}
+            };
+            AddTransactionEvent @event = new AddTransactionEvent("Foo", "Bar",
+                blobs: new[] { new BlobIdentifier("0BEEC7B5EA3F0FDBC95D0DD47F3C5BC275DA8A33") },
+                locations: new[] {"ARN"}.ToList(), metadata);
+
+            long insertedIndex = await log.NewTransaction(@event);
+            Assert.AreEqual(0, insertedIndex);
+
+            TransactionEvent[] events = (await log.Get(0, 1)).Events;
+            TransactionEvent event2 = events[0];
+
+            Assert.AreEqual(@event.Name, event2.Name);
+            Assert.IsInstanceOfType(event2, expectedType: typeof(TransactionEvent));
+            AddTransactionEvent addEvent = (AddTransactionEvent) event2;
+            CollectionAssert.AreEqual(@event.Blobs, addEvent.Blobs);
+            Assert.IsTrue(addEvent.Locations.Contains("ARN"));
+
+            Assert.IsNotNull(addEvent.Metadata);
+            CollectionAssert.AreEqual(metadata, addEvent.Metadata);
+        }
+
+        [TestMethod]
+        public async Task FileTransactionLogMultipleEvents()
+        {
+            CleanupTempDir();
+
+            const long Index1 = 0;
+            const long Index2 = 71;
+            const long Index3 = 162;
+
+            FileTransactionLog log = new FileTransactionLog(_tempPath, _namespace);
+            AddTransactionEvent event1 = new AddTransactionEvent("Foo1", "Bar",
+                blobs: new[] {new BlobIdentifier("0BEEC7B5EA3F0FDBC95D0DD47F3C5BC275DA8A33")},
+                locations: new[] {"ARN"}.ToList(), identifier: Index1);
+            AddTransactionEvent event2 = new AddTransactionEvent( "Foo2", "Bar",
+                blobs: new[] {new BlobIdentifier("62CDB7020FF920E5AA642C3D4066950DD1F01F4D"), new BlobIdentifier("7855AA267FABC5CCD3891282DC2EAE8BAEEEC087") },
+                locations: new[] {"ARN"}.ToList(), identifier: Index2);
+            RemoveTransactionEvent event3 =
+                new RemoveTransactionEvent("Foo3", "Bar", locations: new[] {"ARN"}.ToList(), Index3);
+
+            TransactionEvent[] originalEvents = {event1, event2, event3};
+
+            long insertedIndex1 = await log.NewTransaction(event1);
+            long insertedIndex2 = await log.NewTransaction(event2);
+            long insertedIndex3 = await log.NewTransaction(event3);
+
+            const int expectedOffsetIncrement = 71;
+            Assert.AreEqual(Index1, insertedIndex1);
+            Assert.AreEqual(Index2, insertedIndex2);
+            Assert.AreEqual(Index3, insertedIndex3);
+
+            // we get the event in the middle
+            TransactionEvent[] transactionEvents = (await log.Get(expectedOffsetIncrement, 1)).Events;
+            AddTransactionEvent addTransactionEvent = (AddTransactionEvent) transactionEvents[0];
+
+            Assert.AreEqual(event2.Name, addTransactionEvent.Name);
+            CollectionAssert.AreEqual(event2.Blobs, addTransactionEvent.Blobs);
+
+            transactionEvents = (await log.Get(0, 3)).Events;
+
+            for (int i = 0; i < 3; i++)
+            {
+                TransactionEvent transactionEvent = transactionEvents[i];
+                TransactionEvent originalEvent = originalEvents[i];
+
+                Assert.AreEqual(originalEvent.Name, transactionEvent.Name);
+                Assert.AreEqual(originalEvent.Bucket, transactionEvent.Bucket);
+            }
+        }
+
+
+        [TestMethod]
+        public async Task FileTransactionLogSiteFilter()
+        {
+            CleanupTempDir();
+
+            const int Index1 = 0;
+            const int Index2 = 75;
+            const int Index3 = 149;
+
+            FileTransactionLog log = new FileTransactionLog(_tempPath, _namespace);
+            AddTransactionEvent arnEvent = new AddTransactionEvent("ARNEvent", "Bar",
+                blobs: new[] {new BlobIdentifier("0BEEC7B5EA3F0FDBC95D0DD47F3C5BC275DA8A33") },
+                locations: new[] {"ARN"}.ToList(), identifier: Index1);
+            AddTransactionEvent hqAdd = new AddTransactionEvent("HqEvent", "Bar",
+                blobs: new[] {new BlobIdentifier("62CDB7020FF920E5AA642C3D4066950DD1F01F4D")},
+                locations: new[] {"HQ"}.ToList(), identifier: Index2);
+            RemoveTransactionEvent hqRemove =
+                new RemoveTransactionEvent("Foo", "Bar", locations: new[] {"HQ"}.ToList(), identifier: Index3);
+
+            long insertedIndex1 = await log.NewTransaction(arnEvent);
+            long insertedIndex2 = await log.NewTransaction(hqAdd);
+            long insertedIndex3 = await log.NewTransaction(hqRemove);
+
+            Assert.AreEqual(Index1, insertedIndex1);
+            Assert.AreEqual(Index2, insertedIndex2);
+            Assert.AreEqual(Index3, insertedIndex3);
+
+            // we get the event in the middle
+            TransactionEvent[] transactionEvents = (await log.Get(0, 1, "ARN")).Events;
+            Assert.AreEqual(1, transactionEvents.Length);
+            Assert.IsNotNull(transactionEvents[0]);
+            AddTransactionEvent addTransactionEvent = (AddTransactionEvent) transactionEvents[0];
+
+            // this should have matched the second event even though we set a start index matching the beginning
+            Assert.AreEqual(hqAdd.Name, addTransactionEvent.Name);
+            CollectionAssert.AreEqual(hqAdd.Blobs, addTransactionEvent.Blobs);
+
+            transactionEvents = (await log.Get(0, 3, "HQ")).Events;
+
+            // this should only return 1 event because only one event has not been at HQ
+            Assert.AreEqual(1, transactionEvents.Length);
+            TransactionEvent transactionEvent = transactionEvents[0];
+
+            Assert.AreEqual(arnEvent.Name, transactionEvent.Name);
+            Assert.AreEqual(arnEvent.Bucket, transactionEvent.Bucket);
+        }
+
+
+        [TestMethod]
+        public async Task FileTransactionLogMultipleSites()
+        {
+            CleanupTempDir();
+
+            FileTransactionLog log = new FileTransactionLog(_tempPath, _namespace);
+            AddTransactionEvent arnEvent = new AddTransactionEvent("ARNEvent", "Bar",
+                blobs: new[] { new BlobIdentifier("0BEEC7B5EA3F0FDBC95D0DD47F3C5BC275DA8A33") },
+                locations: new[] { "ARN", "HQ", "USE" }.ToList());
+
+            long insertedIndex = await log.NewTransaction(arnEvent);
+
+            Assert.AreEqual(0, insertedIndex);
+
+            TransactionEvent[] transactionEvents = (await log.Get(0, 1)).Events;
+            Assert.AreEqual(1, transactionEvents.Length);
+            Assert.IsNotNull(transactionEvents[0]);
+            AddTransactionEvent addTransactionEvent = (AddTransactionEvent)transactionEvents[0];
+
+            Assert.AreEqual(arnEvent.Name, addTransactionEvent.Name);
+            Assert.AreEqual(arnEvent.Bucket, addTransactionEvent.Bucket);
+            CollectionAssert.AreEqual(new[] {"ARN", "HQ", "USE"}, addTransactionEvent.Locations);
+
+            Assert.AreEqual(0, log.IndexFile.ToIndex("ARN"));
+            Assert.AreEqual(1, log.IndexFile.ToIndex("HQ"));
+            Assert.AreEqual(2, log.IndexFile.ToIndex("USE"));
+
+            FileTransactionLog.OpRecord opRecord = Callisto.Implementation.FileTransactionLog.AddOp.FromEvent(log.IndexFile, addTransactionEvent);
+            FileTransactionLog.AddOp addOp = (FileTransactionLog.AddOp) opRecord.Op!;
+            Assert.AreEqual(7UL, addOp.Locations);
+        }
+
+        [TestMethod]
+        public async Task FileTransactionLogPartialSites()
+        {
+            CleanupTempDir();
+
+            FileTransactionLog log = new FileTransactionLog(_tempPath, _namespace);
+
+            // set the log file indexes
+            Assert.AreEqual(0, log.IndexFile.ToIndex("ARN"));
+            Assert.AreEqual(1, log.IndexFile.ToIndex("HQ"));
+            Assert.AreEqual(2, log.IndexFile.ToIndex("USE"));
+
+            // define an event that has not been seen by all indexes
+            AddTransactionEvent arnEvent = new AddTransactionEvent("ARNEvent", "Bar",
+                blobs: new[] { new BlobIdentifier("0BEEC7B5EA3F0FDBC95D0DD47F3C5BC275DA8A33") },
+                locations: new[] { "ARN", "USE" }.ToList());
+
+            long insertedIndex = await log.NewTransaction(arnEvent);
+
+            Assert.AreEqual(0, insertedIndex);
+
+            TransactionEvent[] transactionEvents = (await log.Get(0, 1)).Events;
+            Assert.AreEqual(1, transactionEvents.Length);
+            Assert.IsNotNull(transactionEvents[0]);
+            AddTransactionEvent addTransactionEvent = (AddTransactionEvent)transactionEvents[0];
+
+            Assert.AreEqual(arnEvent.Name, addTransactionEvent.Name);
+            Assert.AreEqual(arnEvent.Bucket, addTransactionEvent.Bucket);
+            CollectionAssert.AreEqual(new[] { "ARN", "USE" }, addTransactionEvent.Locations);
+
+            Assert.AreEqual(0, log.IndexFile.ToIndex("ARN"));
+            Assert.AreEqual(2, log.IndexFile.ToIndex("USE"));
+
+            FileTransactionLog.OpRecord opRecord = Callisto.Implementation.FileTransactionLog.AddOp.FromEvent(log.IndexFile, addTransactionEvent);
+            FileTransactionLog.AddOp addOp = (FileTransactionLog.AddOp)opRecord.Op!;
+            Assert.AreEqual(5UL, addOp.Locations);
+        }
+    }
+}
diff --git a/Engine/Source/Programs/Horde/HordeStorage/Callisto/appsettings.Development.json b/Engine/Source/Programs/Horde/HordeStorage/Callisto/appsettings.Development.json
new file mode 100644
--- /dev/null
+++ b/Engine/Source/Programs/Horde/HordeStorage/Callisto/appsettings.Development.json
@@ -0,0 +1,26 @@
+{
+  "Serilog": {
+    "Using": [ "Serilog.Sinks.Console" ],
+    "MinimumLevel": {
+      "Default": "Information",
+      "Override": {
+        "Microsoft": "Warning",
+        "System": "Information"
+      }
+    },
+    "WriteTo": [
+      { "Name": "Console" }
+    ],
+    "Enrich": [ "FromLogContext" ]
+  },
+
+  "Auth": {
+    "Method": "Disabled"
+  },
+
+  "Callisto": {
+    "CurrentSite": "Dev",
+    "TransactionLogRoot": "%Temp%/Callisto",
+    "TransactionLogImplementation": "File"
+  }
+}
diff --git a/Engine/Source/Programs/Horde/HordeStorage/Callisto/appsettings.Production.json b/Engine/Source/Programs/Horde/HordeStorage/Callisto/appsettings.Production.json
new file mode 100644
--- /dev/null
+++ b/Engine/Source/Programs/Horde/HordeStorage/Callisto/appsettings.Production.json
@@ -0,0 +1,30 @@
+{
+  "Serilog": {
+    "Using": ["Serilog.Sinks.Console"],
+    "MinimumLevel": {
+      "Default": "Information",
+      "Override": {
+        "Microsoft": "Warning",
+        "System": "Information"
+      }
+    },
+    "WriteTo": [
+      {
+        "Name": "Console",
+        "Args": {
+          "formatter": "Serilog.Formatting.Json.JsonFormatter"
+        }
+      }
+    ],
+    "Enrich": ["FromLogContext"]
+  },
+
+
+  "Callisto": {
+    // Current site should be specified in the deployment
+    "TransactionLogRoot": "/callisto",
+    "TransactionLogImplementation": "File"
+  },
+
+  "AllowedHosts": "*"
+}
diff --git a/Engine/Source/Programs/Horde/HordeStorage/Callisto/appsettings.Testing.json b/Engine/Source/Programs/Horde/HordeStorage/Callisto/appsettings.Testing.json
new file mode 100644
--- /dev/null
+++ b/Engine/Source/Programs/Horde/HordeStorage/Callisto/appsettings.Testing.json
@@ -0,0 +1,35 @@
+{
+  "Serilog": {
+    "Using": [ "Serilog.Sinks.Console" ],
+    "MinimumLevel": {
+      "Default": "Information",
+      "Override": {
+        "Microsoft": "Warning",
+        "System": "Information"
+      }
+    },
+    "WriteTo": [
+      { "Name": "Console" }
+    ],
+    "Enrich": [ "FromLogContext" ]
+  },
+
+  "Auth": {
+    "Method": "Disabled"
+  },
+
+  "Authorization": {
+    "NamespaceToClaim": {
+      "ue4.ddc": "*"
+    }
+  },
+
+  "Callisto": {
+    "CurrentSite": "Test",
+    "TransactionLogRoot": "./callisto-logs",
+    "transactionLogImplementation": "file",
+    "VerifySerialization":  "true"
+  },
+
+  "AllowedHosts": "*"
+}
diff --git a/Engine/Source/Programs/Horde/HordeStorage/Docs/dynamo.md b/Engine/Source/Programs/Horde/HordeStorage/Docs/dynamo.md
new file mode 100644
--- /dev/null
+++ b/Engine/Source/Programs/Horde/HordeStorage/Docs/dynamo.md
@@ -0,0 +1,84 @@
+# Dynamo DB for Jupiter
+
+Jupiter is a globally replicated service for caching of artifacts used in our build system. These artifacts are stored in S3 but we keep a index of which artifacts exists in a Dynamo DB. 
+
+
+# Tables
+We generally use global tables to allow for replication, we do not seperate into tenent tables, instead this is part of the key (namespace).
+
+# Table Layout
+
+We can use Bucket as a partition key as its intended to be the same for a chunk of objects. Id could then be a sort key. Alternativley the partition key could be a clustered key of the bucket and the name called Id (we append the two strings together to form a ID field). 
+
+| Name      | Type                                          |
+| ---       | ---                                           |
+| ID        | String                                        | Partition key
+| Namespace | String                                        |
+| Name      | String                                        |
+| Bucket    | String                                        |
+| Blobs     | Array of Sha1                                 |
+| Metadata  | Opaque json object. Serialized as string      |
+| Hash      | Sha1 as string                                |
+| LastAccess| DateTime of when this record was las accessed |
+
+# Queries
+
+## GET Record
+Most of the traffic will just be requests for fetch or check for existence of a records with a given name and bucket.
+
+## PUT Record
+A record can get replaced by another record, but we will generally not update existing records (except for last access, see below). If a existing record is modified we consider it a new record (for most intents and purposes its different).
+
+## Delete Record
+This operation will most be invoked as a result of the garbage collection query, while possible for a user to drop a individual record we expect them to most often drop a whole bucket if even deleting anything at all. 
+
+## Drop all records in Bucket
+We give users a option to drop all records in a bucket. If name is a sortkey we are unable to do a query for all buckets (a query for partition keys without a sortkey) which prevents us partitioning like that.
+
+Scan table to get back all the values of name given bucket, and then manually delete each item. 
+
+
+## Last Access updates
+Whenever the service adds new records, or fetches old ones, it tracks a list of all the rows it has accessed. This list is then used at certain intervals to update the LastAccess timestamp of each record in the list.
+
+This can use the same partition key as the is normal.
+
+## Garbage Collection
+At certain points the table will be garbage collected, with the goal of removing every record that has a last access time older then a cutoff (a week).
+
+This will need a secondary index of LastAccess were just index LastAccess, but if we have a composite key as primary key then we need a composite key for last access as well. As we are unaware of every bucket that could exist we can not partition the cleanup index on that. Thus we intend to add a random number as parition key and join on these when requesting all keys older then a given date.
+
+# Index
+
+## Name (Primary)
+| Name      | Type          |
+| ---       | ---           |
+| ID        | Partition Key |
+
+## Last Access (Secondary)
+
+# GSI
+
+| Name        |  Type           |
+| ---         | ---             |
+| RandomKey   | Partition Key   | This should scale with the number of writes we do. 1 - 10 should be sufficent. Single key should never reach 1000 writes per sec. Lean towards more then you need.
+| LastAccess  | Sort Key        | (should be epoch string)
+
+# Bucket Table
+
+This could also be a sparse object in the main table
+
+| Name        |  Type           |
+| ---         | ---             |
+| Namespace   | Partition Key   | 
+| Bucket      | Sort Key        |
+
+Lambda that reads the writes from stream and creates tenant_bucket list. 
+eventually consistent reads
+if not exist
+update bucket list with a new bucket for the tentant
+
+
+
+# Limitations to consider
+Dynamo has a max limit of 400kb per row. This means that metadata (as the only truly dynamic field) needs to be capped to make sure the full row does not exceed this limit.
diff --git a/Engine/Source/Programs/Horde/HordeStorage/Helm/Horde.Storage/README.md b/Engine/Source/Programs/Horde/HordeStorage/Helm/Horde.Storage/README.md
new file mode 100644
--- /dev/null
+++ b/Engine/Source/Programs/Horde/HordeStorage/Helm/Horde.Storage/README.md
@@ -0,0 +1,28 @@
+# Horde.Storage
+
+Horde.Storage - A Distributed storage service for use with Unreal Engine
+
+# Components 
+
+* Horde.Storage - Main service - The service implementing most of the logic
+* Callisto - Transaction - A C# service that manages and serializes the ref transactions from the ref server. Essentially a custom db for Horde.Storage focusing on a log structured approch to the transaction rather then the traditional db aspects of storage. This is a legacy component that is being faced out.
+* Jupiter.Common - Lib - Library with common shared infrastructure that is reused across all or some C# services
+
+
+# Chart Values
+
+## Globals
+
+| Name | Description | Possible Values |
+| ------------- |:-------------:| -----:|
+| cloudProvider | Which cloud provider this will be deployed on, impacts which services we interact with. | AWS or Azure |
+| awsRegion | Only used for AWS cloud, which region you expect the services to be in | AWS Region format |
+| awsRoleArn | The arn of a AWS role to assume | A valid ARN | 
+| siteName | The name of this particular site within your cluster of Jupiter nodes. Should be unique amongst all your jupiter sites and is used to understand were a record has been replicated | Any globally unique value |
+| authMethod | Which method is used to auth | Okta / JWTBearer / Disabled |
+| oktaDomain | The okta domain used to auth against | See your okta setup |
+| oktaAuthorizationServerId | The authorization server id in okta | See your okta setup for which values are valid |
+| jwtAudience | The audience expected in your JWT token, used both with Okta and JWTBearer | The JWT audience
+| jwtAuthority | The address to your JWT Authority, only used for JWT Bearer verification | A url to your JWT authority |
+
+
diff --git a/Engine/Source/Programs/Horde/HordeStorage/Horde.Storage/AssemblyInfo.cs b/Engine/Source/Programs/Horde/HordeStorage/Horde.Storage/AssemblyInfo.cs
new file mode 100644
--- /dev/null
+++ b/Engine/Source/Programs/Horde/HordeStorage/Horde.Storage/AssemblyInfo.cs
@@ -0,0 +1,6 @@
+﻿// Copyright Epic Games, Inc. All Rights Reserved.
+
+using System.Runtime.CompilerServices;
+
+[assembly: InternalsVisibleTo("Horde.Storage.Tests.Functional")]
+[assembly: InternalsVisibleTo("Horde.Storage.Tests.Unit")]
