Moved Horde.Storage from git to p4

diff --git a/Engine/Source/Programs/Horde/HordeStorage/Horde.Storage/Controllers/AdminController.cs b/Engine/Source/Programs/Horde/HordeStorage/Horde.Storage/Controllers/AdminController.cs
new file mode 100644
--- /dev/null
+++ b/Engine/Source/Programs/Horde/HordeStorage/Horde.Storage/Controllers/AdminController.cs
@@ -0,0 +1,193 @@
+﻿// Copyright Epic Games, Inc. All Rights Reserved.
+
+using System;
+using System.Collections.Generic;
+using System.ComponentModel.DataAnnotations;
+using System.Linq;
+using System.Net;
+using System.Threading;
+using System.Threading.Tasks;
+using Horde.Storage.Implementation;
+using Jupiter.Implementation;
+using Microsoft.AspNetCore.Authorization;
+using Microsoft.AspNetCore.Mvc;
+using Microsoft.Extensions.Configuration;
+using Newtonsoft.Json;
+
+namespace Horde.Storage.Controllers
+{
+    [ApiController]
+    [Route("api/v1/admin")]
+    public class AdminController : Controller
+    {
+        private readonly LastAccessService _lastAccessService;
+        private readonly IRefCleanup _refCleanup;
+        private readonly BlobCleanupService _blobCleanupService;
+        private readonly IConfiguration _configuration;
+
+        public AdminController(LastAccessService lastAccessService, IRefCleanup refCleanup, BlobCleanupService blobCleanupService, IConfiguration configuration)
+        {
+            _lastAccessService = lastAccessService;
+            _refCleanup = refCleanup;
+            _blobCleanupService = blobCleanupService;
+            _configuration = configuration;
+        }
+
+        /// <summary>
+        /// Report last access times to refs store
+        /// </summary>
+        /// <remarks>
+        /// Manually triggers a aggregation of the last access records and commits these to the refs store. This is done automatically so the only reason to use this endpoint is for debugging purposes.
+        /// </remarks>
+        /// <returns></returns>
+        [HttpPost("startLastAccessRollup")]
+        [Authorize("Admin")]
+        [ProducesResponseType(type: typeof(UpdatedRecordsResponse), 200)]
+        public async Task<IActionResult> StartLastAccessRollup()
+        {
+            Task<List<(RefRecord, DateTime)>>? updateRecordsTask = _lastAccessService.ProcessLastAccessRecords();
+            List<(RefRecord, DateTime)>? updatedRecords = null;
+            if (updateRecordsTask != null)
+            {
+                updatedRecords = await updateRecordsTask;
+            }
+
+            return Ok(new UpdatedRecordsResponse(
+                updatedRecords?.Select(tuple => new UpdatedRecordsResponse.UpdatedRecord(tuple.Item1, tuple.Item2))
+                    .ToList() ??
+                new List<UpdatedRecordsResponse.UpdatedRecord>()));
+        }
+
+        /// <summary>
+        /// Manually run the refs cleanup
+        /// </summary>
+        /// <remarks>
+        /// Manually triggers a cleanup of the refs keys based on last access time. This is done automatically so the only reason to use this endpoint is for debugging purposes.
+        /// </remarks>
+        /// <returns></returns>
+        [HttpPost("refCleanup/{ns}")]
+        public async Task<IActionResult> RefCleanup([FromRoute] [Required] NamespaceId ns)
+        {
+            List<OldRecord> records = await _refCleanup.Cleanup(ns, CancellationToken.None);
+            return Ok(new RemovedRefRecordsResponse(
+                records.Select(r => new RemovedRefRecordsResponse.RemovedRecord(r.RefName, r.Bucket))
+            ));
+        }
+
+        /// <summary>
+        /// Manually run the blob cleanup
+        /// </summary>
+        /// <remarks>
+        /// Manually triggers a cleanup of unused blobs not referenced in the transaction log. This is done automatically so the only reason to use this endpoint is for debugging purposes.
+        /// </remarks>
+        /// <returns></returns>
+        [HttpPost("blobCleanup/{ns}")]
+        public async Task<IActionResult> BlobCleanup([FromRoute] [Required] NamespaceId ns)
+        {
+            List<RemovedBlobs> records = await _blobCleanupService.Cleanup(_blobCleanupService.State, CancellationToken.None);
+            return Ok(new RemovedBlobRecords(
+                records.Select(r => r.BlobIdentifier)
+            ));
+        }
+
+        /// <summary>
+        /// Dumps all settings currently in use
+        /// </summary>
+        /// <returns></returns>
+        [HttpGet("settings")]
+        public IActionResult Settings()
+        {
+            Dictionary<string, Dictionary<string, object>> settings = new Dictionary<string, Dictionary<string, object>>();
+
+            Dictionary<string, object> ResolveSection(IConfigurationSection section)
+            {
+                Dictionary<string, object> values = new Dictionary<string, object>();
+                foreach (IConfigurationSection childSection in section.GetChildren())
+                {
+                    if (childSection.Value == null)
+                    {
+                        values.Add(childSection.Key, ResolveSection(childSection));
+                    }
+                    else
+                    {
+                        values.Add(childSection.Key, childSection.Value);
+                    }
+                }
+
+                return values;
+            }
+
+            foreach (IConfigurationSection section in _configuration.GetChildren())
+            {
+                Dictionary<string, object> values = ResolveSection(section);
+                if (values.Count != 0)
+                    settings.Add(section.Key, values);
+            }
+
+            return new JsonResult(new
+            {
+                Settings = settings
+            }, new JsonSerializerSettings
+            {
+                Formatting = Formatting.Indented
+            })
+            {
+                StatusCode = (int)HttpStatusCode.OK
+            };
+        }
+    }
+
+    public class RemovedBlobRecords
+    {
+        public BlobIdentifier[] Blobs { get; }
+
+        public RemovedBlobRecords(IEnumerable<BlobIdentifier> blobs)
+        {
+            Blobs = blobs.ToArray();
+        }
+    }
+
+    public class RemovedRefRecordsResponse
+    {
+        public class RemovedRecord
+        {
+            public RemovedRecord(KeyId name, BucketId bucket)
+            {
+                Name = name;
+                Bucket = bucket;
+            }
+
+            public KeyId Name { get; }
+            public BucketId Bucket { get; }
+        }
+
+        public RemovedRefRecordsResponse(IEnumerable<RemovedRecord> removedRecords)
+        {
+            RemovedRecords = removedRecords.ToArray();
+        }
+
+        public RemovedRecord[] RemovedRecords { get; }
+    }
+
+    public class UpdatedRecordsResponse
+    {
+        public UpdatedRecordsResponse(List<UpdatedRecord> updatedRecords)
+        {
+            UpdatedRecords = updatedRecords;
+        }
+
+        public class UpdatedRecord
+        {
+            public UpdatedRecord(RefRecord record, DateTime time)
+            {
+                Record = record;
+                Time = time;
+            }
+
+            public RefRecord Record { get; }
+            public DateTime Time { get; }
+        }
+
+        public List<UpdatedRecord> UpdatedRecords { get; }
+    }
+}
diff --git a/Engine/Source/Programs/Horde/HordeStorage/Horde.Storage/Controllers/AuthController.cs b/Engine/Source/Programs/Horde/HordeStorage/Horde.Storage/Controllers/AuthController.cs
new file mode 100644
--- /dev/null
+++ b/Engine/Source/Programs/Horde/HordeStorage/Horde.Storage/Controllers/AuthController.cs
@@ -0,0 +1,42 @@
+﻿using Microsoft.AspNetCore.Http;
+using Microsoft.AspNetCore.Mvc;
+using System;
+using System.Collections.Generic;
+using System.ComponentModel.DataAnnotations;
+using System.Linq;
+using System.Threading.Tasks;
+using Datadog.Trace;
+using Jupiter;
+using Jupiter.Implementation;
+using Microsoft.AspNetCore.Authorization;
+
+namespace Horde.Storage.Controllers
+{
+    [ApiController]
+    [Route("api/v1/auth")]
+    public class AuthController : ControllerBase
+    {
+        private readonly IAuthorizationService _authorizationService;
+
+        public AuthController(IAuthorizationService authorizationService)
+        {
+            _authorizationService = authorizationService;
+        }
+
+        [HttpGet("{ns}")]
+        [Authorize("Any")]
+        public async Task<IActionResult> Verify(
+            [FromRoute] [Required] NamespaceId ns
+            )
+        {
+            AuthorizationResult authorizationResult = await _authorizationService.AuthorizeAsync(User, ns, NamespaceAccessRequirement.Name);
+
+            if (!authorizationResult.Succeeded)
+            {
+                return Forbid();
+            }
+
+            return Ok();
+        }
+    }
+}
diff --git a/Engine/Source/Programs/Horde/HordeStorage/Horde.Storage/Controllers/CompressedBlobController.cs b/Engine/Source/Programs/Horde/HordeStorage/Horde.Storage/Controllers/CompressedBlobController.cs
new file mode 100644
--- /dev/null
+++ b/Engine/Source/Programs/Horde/HordeStorage/Horde.Storage/Controllers/CompressedBlobController.cs
@@ -0,0 +1,298 @@
+﻿// Copyright Epic Games, Inc. All Rights Reserved.
+
+using System;
+using System.Collections.Concurrent;
+using System.Collections.Generic;
+using System.ComponentModel.DataAnnotations;
+using System.IO;
+using System.Linq;
+using System.Net.Mime;
+using System.Threading.Tasks;
+using Datadog.Trace;
+using Horde.Storage.Implementation;
+using Jupiter;
+using Jupiter.Implementation;
+using Microsoft.AspNetCore.Authorization;
+using Microsoft.AspNetCore.Http;
+using Microsoft.AspNetCore.Mvc;
+using Serilog;
+
+namespace Horde.Storage.Controllers
+{
+    [ApiController]
+    [Route("api/v1/compressed-blobs")]
+    public class CompressedBlobController : ControllerBase
+    {
+        private readonly IBlobStore _storage;
+        private readonly IContentIdStore _contentIdStore;
+        private readonly IDiagnosticContext _diagnosticContext;
+        private readonly IAuthorizationService _authorizationService;
+        private readonly CompressedBufferUtils _compressedBufferUtils;
+
+        private readonly ILogger _logger = Log.ForContext<CompressedBlobController>();
+
+        public CompressedBlobController(IBlobStore storage, IContentIdStore contentIdStore, IDiagnosticContext diagnosticContext, IAuthorizationService authorizationService, CompressedBufferUtils compressedBufferUtils)
+        {
+            _storage = storage;
+            _contentIdStore = contentIdStore;
+            _diagnosticContext = diagnosticContext;
+            _authorizationService = authorizationService;
+            _compressedBufferUtils = compressedBufferUtils;
+        }
+
+
+        [HttpGet("{ns}/{id}")]
+        [Authorize("Storage.read")]
+        [ProducesResponseType(type: typeof(byte[]), 200)]
+        [ProducesResponseType(type: typeof(ValidationProblemDetails), 400)]
+        [Produces(MediaTypeNames.Application.Json, CustomMediaTypeNames.UnrealCompactBinary, CustomMediaTypeNames.UnrealCompressedBuffer)]
+
+        public async Task<IActionResult> Get(
+            [Required] NamespaceId ns,
+            [Required] BlobIdentifier id)
+        {
+            AuthorizationResult authorizationResult = await _authorizationService.AuthorizeAsync(User, ns, NamespaceAccessRequirement.Name);
+
+            if (!authorizationResult.Succeeded)
+            {
+                return Forbid();
+            }
+
+            try
+            {
+                BlobContents blobContents = await GetImpl(ns, id);
+
+                return File(blobContents.Stream, CustomMediaTypeNames.UnrealCompressedBuffer, enableRangeProcessing: true);
+            }
+            catch (BlobNotFoundException e)
+            {
+                return NotFound(new ValidationProblemDetails { Title = $"Object {e.Blob} not found" });
+            }
+            catch (ContentIdResolveException e)
+            {
+                return NotFound(new ValidationProblemDetails { Title = $"Content Id {e.ContentId} not found" });
+            }
+        }
+        
+        [HttpHead("{ns}/{id}")]
+        [Authorize("Storage.read")]
+        [ProducesDefaultResponseType]
+        public async Task<IActionResult> Head(
+            [Required] NamespaceId ns,
+            [Required] BlobIdentifier id)
+        {
+            AuthorizationResult authorizationResult = await _authorizationService.AuthorizeAsync(User, ns, NamespaceAccessRequirement.Name);
+
+            if (!authorizationResult.Succeeded)
+            {
+                return Forbid();
+            }
+            BlobIdentifier[]? chunks = await _contentIdStore.Resolve(ns, id);
+            if (chunks == null)
+            {
+                return NotFound(new ValidationProblemDetails { Title = $"Content-id {id} not found"});
+            }
+
+            Task<bool>[] tasks = new Task<bool>[chunks.Length];
+            for (int i = 0; i < chunks.Length; i++)
+            {
+                tasks[i] = _storage.Exists(ns, id);
+            }
+
+            await Task.WhenAll(tasks);
+
+            bool exists = tasks.All(task => task.Result);
+
+            if (!exists)
+            {
+                return NotFound(new ValidationProblemDetails { Title = $"Object {id} not found"});
+            }
+
+            return Ok();
+        }
+
+        [HttpPost("{ns}/exists")]
+        [Authorize("Storage.read")]
+        [ProducesDefaultResponseType]
+        public async Task<IActionResult> ExistsMultiple(
+            [Required] NamespaceId ns,
+            [Required] [FromQuery] List<BlobIdentifier> id)
+        {
+            AuthorizationResult authorizationResult = await _authorizationService.AuthorizeAsync(User, ns, NamespaceAccessRequirement.Name);
+
+            if (!authorizationResult.Succeeded)
+            {
+                return Forbid();
+            }
+
+            ConcurrentBag<BlobIdentifier> missingBlobs = new ConcurrentBag<BlobIdentifier>();
+            ConcurrentBag<BlobIdentifier> invalidContentIds = new ConcurrentBag<BlobIdentifier>();
+
+            IEnumerable<Task> tasks = id.Select(async blob =>
+            {
+                BlobIdentifier[]? chunks = await _contentIdStore.Resolve(ns, blob);
+
+                if (chunks == null)
+                {
+                    invalidContentIds.Add(blob);
+                    return;
+                }
+
+                foreach (BlobIdentifier chunk in chunks)
+                {
+                    if (!await _storage.Exists(ns, chunk))
+                    {
+                        missingBlobs.Add(blob);
+                    }
+                }
+            });
+            await Task.WhenAll(tasks);
+
+            if (invalidContentIds.Count != 0)
+                return NotFound(new ValidationProblemDetails { Title = $"Missing content ids {string.Join(" ,", invalidContentIds)}"});
+
+            return Ok(new HeadMultipleResponse { Needs = missingBlobs.ToArray()});
+        }
+
+        private async Task<BlobContents> GetImpl(NamespaceId ns, BlobIdentifier contentId)
+        {
+            BlobIdentifier[]? chunks = await _contentIdStore.Resolve(ns, contentId);
+            if (chunks == null)
+            {
+                throw new ContentIdResolveException(contentId);
+            }
+
+            using Scope _ = Tracer.Instance.StartActive("blob.combine");
+            Task<BlobContents>[] tasks = new Task<BlobContents>[chunks.Length];
+            for (int i = 0; i < chunks.Length; i++)
+            {
+                tasks[i] = _storage.GetObject(ns, chunks[i]);
+            }
+
+            MemoryStream ms = new MemoryStream();
+            foreach (Task<BlobContents> task in tasks)
+            {
+                BlobContents blob = await task;
+                await using Stream s = blob.Stream;
+                await s.CopyToAsync(ms);
+            }
+
+            ms.Seek(0, SeekOrigin.Begin);
+
+            return new BlobContents(ms, ms.Length);
+        }
+
+        [HttpPut("{ns}/{id}")]
+        [Authorize("Storage.write")]
+        [DisableRequestSizeLimit]
+        [RequiredContentType(CustomMediaTypeNames.UnrealCompressedBuffer)]
+        public async Task<IActionResult> Put(
+            [Required] NamespaceId ns,
+            [Required] BlobIdentifier id)
+        {
+            AuthorizationResult authorizationResult = await _authorizationService.AuthorizeAsync(User, ns, NamespaceAccessRequirement.Name);
+
+            if (!authorizationResult.Succeeded)
+            {
+                return Forbid();
+            }
+
+            byte[] blob;
+            try
+            {
+                blob = await RequestUtil.ReadRawBody(Request);
+            }
+            catch (BadHttpRequestException e)
+            {
+                const string msg = "Partial content transfer when reading request body.";
+                _logger.Warning(e, msg);
+                return BadRequest(msg);
+            }
+            _diagnosticContext.Set("Content-Length", Request.ContentLength ?? -1);
+
+            BlobIdentifier identifier = await PutImpl(ns, id, blob);
+            return Ok(new
+            {
+                Identifier = identifier.ToString()
+            });
+        }
+
+        private async Task<BlobIdentifier> PutImpl(NamespaceId ns, BlobIdentifier? id, byte[] content)
+        {
+            BlobIdentifier identifier;
+            {
+                using Scope _ = Tracer.Instance.StartActive("web.hash");
+                identifier = BlobIdentifier.FromBlob(content);
+            }
+            // decompress the content and generate a identifier from it to verify the identifier we got
+            byte[] decompressedContent = _compressedBufferUtils.DecompressContent(content);
+
+            BlobIdentifier identifierDecompressedPayload;
+            {
+                using Scope _ = Tracer.Instance.StartActive("web.hash");
+                identifierDecompressedPayload = BlobIdentifier.FromBlob(decompressedContent);
+            }
+
+            if (id != null && !id.Equals(identifierDecompressedPayload))
+            {
+                _logger.Debug("ID {@Id} was not the same as identifier {@Identifier} {Content}", id, identifierDecompressedPayload,
+                    content);
+
+                throw new ArgumentException("ID was not a hash of the content uploaded.", paramName: nameof(id));
+            }
+
+            // commit the mapping from the decompressed hash to the compressed hash, we run this in parallel with the blob store submit
+            // TODO: let users specify weight of the blob compared to previously submitted content ids
+            Task contentIdStoreTask = _contentIdStore.Put(ns, identifierDecompressedPayload, identifier, content.Length);
+
+            // we still commit the compressed buffer to the object store
+            await _storage.PutObject(ns, content, identifier);
+
+            await contentIdStoreTask;
+
+            return identifierDecompressedPayload;
+        }
+
+        /*[HttpDelete("{ns}/{id}")]
+        [Authorize("Storage.delete")]
+        public async Task<IActionResult> Delete(
+            [Required] string ns,
+            [Required] BlobIdentifier id)
+        {
+            AuthorizationResult authorizationResult = await _authorizationService.AuthorizeAsync(User, ns, NamespaceAccessRequirement.Name);
+
+            if (!authorizationResult.Succeeded)
+            {
+                return Forbid();
+            }
+
+            await DeleteImpl(ns, id);
+
+            return Ok();
+        }
+
+        
+        [HttpDelete("{ns}")]
+        [Authorize("Admin")]
+        public async Task<IActionResult> DeleteNamespace(
+            [Required] string ns)
+        {
+            AuthorizationResult authorizationResult = await _authorizationService.AuthorizeAsync(User, ns, NamespaceAccessRequirement.Name);
+
+            if (!authorizationResult.Succeeded)
+            {
+                return Forbid();
+            }
+
+            int deletedCount = await  _storage.DeleteNamespace(ns);
+
+            return Ok( new { Deleted = deletedCount });
+        }
+
+
+        private async Task DeleteImpl(string ns, BlobIdentifier id)
+        {
+            await _storage.DeleteObject(ns, id);
+        }*/
+    }
+}
diff --git a/Engine/Source/Programs/Horde/HordeStorage/Horde.Storage/Controllers/DDCRefController.cs b/Engine/Source/Programs/Horde/HordeStorage/Horde.Storage/Controllers/DDCRefController.cs
new file mode 100644
--- /dev/null
+++ b/Engine/Source/Programs/Horde/HordeStorage/Horde.Storage/Controllers/DDCRefController.cs
@@ -0,0 +1,973 @@
+// Copyright Epic Games, Inc. All Rights Reserved.
+
+using System;
+using System.Collections.Generic;
+using System.ComponentModel.DataAnnotations;
+using System.IO;
+using System.Linq;
+using System.Net.Mime;
+using System.Threading.Tasks;
+using Dasync.Collections;
+using Datadog.Trace;
+using Horde.Storage.Implementation;
+using Jupiter;
+using Jupiter.Implementation;
+using Jupiter.Utils;
+using Microsoft.AspNetCore.Authorization;
+using Microsoft.AspNetCore.Http;
+using Microsoft.AspNetCore.Http.Extensions;
+using Microsoft.AspNetCore.Mvc;
+using Microsoft.AspNetCore.Mvc.Infrastructure;
+using Microsoft.Extensions.Configuration;
+using Microsoft.Extensions.Options;
+using Newtonsoft.Json;
+using Serilog;
+
+namespace Horde.Storage.Controllers
+{
+    public class RefRequest
+    {
+        [JsonConstructor]
+        public RefRequest(BlobIdentifier[] blobReferences, Dictionary<string, object>? metadata, ContentHash contentHash)
+        {
+            BlobReferences = blobReferences;
+            Metadata = metadata;
+            ContentHash = contentHash;
+        }
+
+        public Dictionary<string, object>? Metadata { get; }
+
+        [Required] public BlobIdentifier[] BlobReferences { get; }
+
+        [Required] public ContentHash ContentHash { get; }
+    }
+
+    [ApiController]
+    [FormatFilter]
+    [Route("api/v1/c")]
+    public class DDCRefController : ControllerBase
+    {
+        private readonly IRefsStore _refsStore;
+        private readonly IDiagnosticContext _diagnosticContext;
+        private readonly IAuthorizationService _authorizationService;
+        private readonly IOptionsMonitor<JupiterSettings> _jupiterSettings;
+
+        private readonly ILogger _logger = Log.ForContext<DDCRefController>();
+        private readonly IDDCRefService _ddcRefService;
+
+        public DDCRefController(IDDCRefService ddcRefService, IRefsStore refsStore, IDiagnosticContext diagnosticContext, IAuthorizationService authorizationService, IOptionsMonitor<JupiterSettings> jupiterSettings)
+        {
+            _refsStore = refsStore;
+            _ddcRefService = ddcRefService;
+            _diagnosticContext = diagnosticContext;
+            _authorizationService = authorizationService;
+            _jupiterSettings = jupiterSettings;
+        }
+
+        /// <summary>
+        /// Returns all the known namespace the token has access to
+        /// </summary>
+        /// <returns></returns>
+        [HttpGet("ddc")]
+        [ProducesDefaultResponseType]
+        [ProducesResponseType(type: typeof(ProblemDetails), 400)]
+        [Authorize("Cache.read")]
+        public async Task<IActionResult> GetNamespaces()
+        {
+            NamespaceId[] namespaces = await _refsStore.GetNamespaces().ToArrayAsync();
+
+            if (!ShouldDoAuth())
+            {
+                // filter namespaces down to only the namespaces the user has access to
+                namespaces = namespaces.Where(ns =>
+                {
+                    Task<AuthorizationResult> authorizationResult = _authorizationService.AuthorizeAsync(User, ns, NamespaceAccessRequirement.Name);
+                    return authorizationResult.Result.Succeeded;
+                }).ToArray();
+            }
+
+            return Ok(new GetNamespacesResponse(namespaces));
+        }
+
+        /// <summary>
+        /// Returns a refs key
+        /// </summary>
+        /// <remarks>If you use the .raw format specifier the raw object uploaded will be returned. If using json or another structured object it is instead base64 encoded.</remarks>
+        /// <param name="ns">Namespace. Each namespace is completely separated from each other. Use for different types of data that is never expected to be similar (between two different games for instance). Example: `uc4.ddc`</param>
+        /// <param name="bucket">The category/type of record you are caching. Is a clustered key together with the actual key, but all records in the same bucket can be dropped easily. Example: `terrainTexture` </param>
+        /// <param name="key">The unique name of this particular key. `iAmAVeryValidKey`</param>
+        /// <param name="fields">The fields to include in the response, omit this to include everything.</param>
+        /// <param name="format">Optional specifier to set which output format is used json/raw</param>
+        /// <returns>A json blob with the base 64 encoded object, or if octet-stream is used just the raw object</returns>
+        [HttpGet("ddc/{ns}/{bucket}/{key}.{format?}", Order = 500)]
+        [ProducesResponseType(type: typeof(RefResponse), 200)]
+        [ProducesResponseType(type: typeof(ValidationProblemDetails), 400)]
+        [Produces(MediaTypeNames.Application.Json, MediaTypeNames.Application.Octet, CustomMediaTypeNames.UnrealCompactBinary)]
+        [Authorize("Cache.read")]
+        public async Task<IActionResult> Get(
+            [FromRoute] [Required] NamespaceId ns,
+            [FromRoute] [Required] BucketId bucket,
+            [FromRoute] [Required] KeyId key,
+            [FromQuery] string[] fields,
+            [FromRoute] string? format = null)
+        {
+            if (!ShouldDoAuth())
+            {
+                using (Scope _ = Tracer.Instance.StartActive("authorize"))
+                {
+                    AuthorizationResult authorizationResult = await _authorizationService.AuthorizeAsync(User, ns, NamespaceAccessRequirement.Name);
+
+                    if (!authorizationResult.Succeeded)
+                    {
+                        return Forbid();
+                    }
+                }
+            }
+
+            try
+            {
+                bool isRaw = format?.ToLowerInvariant() == "raw";
+                // if the raw format is used only the blob itself will be output so no need to fetch anything else
+                if (isRaw)
+                {
+                    fields = new[] {"blob"};
+                }
+                (RefResponse record, BlobContents? blob) = await _ddcRefService.Get(ns, bucket, key, fields);
+                
+                Response.Headers[CommonHeaders.HashHeaderName] = record.ContentHash.ToString();
+
+                if (blob != null)
+                {
+                    Response.Headers["Content-Length"] = blob.Length.ToString();
+   
+                    if (isRaw)
+                    {
+                        using (Scope _ = Tracer.Instance.StartActive("body.write"))
+                        {
+                            const int BufferSize = 64 * 1024;
+                            var outputStream = Response.Body;
+                            Response.ContentLength = blob.Length;
+                            Response.ContentType = MediaTypeNames.Application.Octet;
+                            Response.StatusCode = StatusCodes.Status200OK;
+                            try
+                            {
+                                await StreamCopyOperation.CopyToAsync(blob.Stream, outputStream, count: null, bufferSize: BufferSize,
+                                    cancel: Response.HttpContext.RequestAborted);
+                            }
+                            catch (OperationCanceledException e)
+                            {
+                                _logger.Error(e, "Copy operation cancelled due to http request being cancelled.");
+                            }
+                            catch (Exception e)
+                            {
+                                // catch all exceptions because the unhandled exception filter will try to update the response which we have already started writing
+                                _logger.Error(e, "Exception while writing response");
+                            }
+                        }
+
+                        return new EmptyResult();
+                    }
+
+                    await using Stream blobStream = blob.Stream;
+                    // convert to byte array in preparation for json serialization
+                    record.Blob = await blobStream.ToByteArray();
+                }
+
+
+                return Ok(record);
+
+            }
+            catch (NamespaceNotFoundException e)
+            {
+                return BadRequest(new ProblemDetails {Title = $"Namespace {e.Namespace} did not exist"});
+            }
+            catch (RefRecordNotFoundException e)
+            {
+                return BadRequest(new ProblemDetails { Title = $"Object {e.Bucket} {e.Key} did not exist" });
+            }
+            catch (BlobNotFoundException e)
+            {
+                return BadRequest(new ProblemDetails { Title = $"Object {e.Blob} in {e.Ns} not found" });
+            }
+        }
+
+        private static bool IsServedOnHighPerfHttpPort(HostString host, IEnumerable<int> highPerfPorts)
+        {
+            foreach (int port in highPerfPorts)
+            {
+                if (host.Port.HasValue && host.Port.Value == port)
+                    return true;
+            }
+            
+            return false;
+        }
+
+        [ApiExplorerSettings(IgnoreApi = true)]
+        [NonAction]
+        public async Task FastGet(HttpContext httpContext, Func<Task> next)
+        {
+            bool isHighPerf = IsServedOnHighPerfHttpPort(httpContext.Request.Host, _jupiterSettings.CurrentValue.DisableAuthOnPorts);
+
+            if (isHighPerf && httpContext.Request.Path.StartsWithSegments("/api/v1/c/ddc", out var remaining))
+            {
+                string[] parts = remaining.ToString().Split("/");
+                if (parts.Length == 4)
+                {
+                    NamespaceId ns = new NamespaceId(parts[1]);
+                    BucketId bucket = new BucketId(parts[2]);
+                    string keyAndFormat = parts[3];
+                    if (httpContext.Request.Method == "GET")
+                    {
+                        KeyId key;
+                        // assume raw key if nothing is specified
+                        string format = "raw";
+                        string[] keyParts = keyAndFormat.Split(".");
+                        if (keyParts.Length == 2)
+                        {
+                            key = new KeyId(keyParts[0]);
+                            format = keyParts[1];
+                        }
+                        else
+                        {
+                            key = new KeyId(keyAndFormat);
+                        }
+
+                        if (format != "raw")
+                        {
+                            // if not requesting the raw format we can not be used so continue searching middlewares
+                            await next.Invoke();
+                            return;
+                        }
+
+                        httpContext.Response.Headers["X-Jupiter-ImplUsed"] = "ReqMid";
+                        try
+                        {
+                            (RefResponse refRes, BlobContents? tempBlob) = await _ddcRefService.Get(ns, bucket, key, fields: new[] {"blob"});
+                            BlobContents blob = tempBlob!;
+                            
+                            httpContext.Response.StatusCode = StatusCodes.Status200OK;
+                            httpContext.Response.Headers["Content-Type"] = MediaTypeNames.Application.Octet;
+                            httpContext.Response.Headers["Content-Length"] = Convert.ToString(blob.Length);
+                            httpContext.Response.Headers[CommonHeaders.HashHeaderName] = refRes.ContentHash.ToString();
+
+                            using (Scope _ = Tracer.Instance.StartActive("body.write"))
+                            {
+                                await blob.Stream.CopyToAsync(httpContext.Response.Body);
+                            }
+                        }
+                        catch (NamespaceNotFoundException e)
+                        {
+                            httpContext.Response.StatusCode = StatusCodes.Status400BadRequest;
+                            httpContext.Response.Headers["Content-Type"] = "application/json";
+                            await httpContext.Response.WriteAsync("{\"title\": \"Namespace " + e.Namespace + " did not exist\", \"status\": 400}");
+                        }
+                        catch (RefRecordNotFoundException e)
+                        {
+                            httpContext.Response.StatusCode = StatusCodes.Status400BadRequest;
+                            httpContext.Response.Headers["Content-Type"] = "application/json";
+                            await httpContext.Response.WriteAsync("{\"title\": \"Object " + e.Bucket + " " + e.Key + " did not exist\", \"status\": 400}");
+                        }
+                        catch (BlobNotFoundException e)
+                        {
+                            httpContext.Response.StatusCode = StatusCodes.Status400BadRequest;
+                            httpContext.Response.Headers["Content-Type"] = "application/json";
+                            await httpContext.Response.WriteAsync("{\"title\": \"Object " + e.Blob + " in " + e.Ns + " not found\", \"status\": 400}");
+                        }
+
+                        return;
+                    }
+                    else if (httpContext.Request.Method == "HEAD")
+                    {
+                        KeyId key = new KeyId(keyAndFormat);
+                        httpContext.Response.Headers["X-Jupiter-ImplUsed"] = "ReqMid";
+                        try
+                        {
+                            RefRecord record = await _ddcRefService.Exists(ns, bucket, key);
+                            
+                            httpContext.Response.StatusCode = StatusCodes.Status204NoContent;
+                            httpContext.Response.Headers[CommonHeaders.HashHeaderName] = record.ContentHash.ToString();
+                        }
+                        catch (NamespaceNotFoundException e)
+                        {
+                            httpContext.Response.StatusCode = StatusCodes.Status400BadRequest;
+                            httpContext.Response.Headers["Content-Type"] = "application/json";
+                            await httpContext.Response.WriteAsync($"{{\"title\": \"Namespace {e.Namespace} did not exist\", \"status\": 400}}");
+                        }
+                        catch (RefRecordNotFoundException e)
+                        {
+                            httpContext.Response.StatusCode = StatusCodes.Status400BadRequest;
+                            httpContext.Response.Headers["Content-Type"] = "application/json";
+                            await httpContext.Response.WriteAsync($"{{\"title\": \"Object {e.Bucket} {e.Key} in namespace {e.Namespace} did not exist\", \"status\": 400}}");
+                        }
+                        catch (MissingBlobsException e)
+                        {
+                            httpContext.Response.StatusCode = StatusCodes.Status400BadRequest;
+                            httpContext.Response.Headers["Content-Type"] = "application/json";
+                            await httpContext.Response.WriteAsync($"{{\"title\": \"Blobs {e.Blobs} from object {e.Bucket} {e.Key} in namespace {e.Namespace} did not exist\", \"status\": 400}}");
+                        }
+
+                        return;
+                    }
+
+                }
+            }
+            else if (isHighPerf && httpContext.Request.Path.StartsWithSegments("/api/v1/c/ddc-rpc/batchGet", out _))
+            {
+                string requestBody;
+
+                using (StreamReader sr = new StreamReader(httpContext.Request.Body))
+                {
+                    using Scope _ = Tracer.Instance.StartActive("body.read");
+                    requestBody = await sr.ReadToEndAsync();
+                }
+
+                BatchGetOp batch;
+                using (Scope _ = Tracer.Instance.StartActive("json.deserialize"))
+                {
+                    batch = JsonConvert.DeserializeObject<BatchGetOp>(requestBody)!;
+                }
+                
+                BatchWriter writer = new BatchWriter();
+                httpContext.Response.StatusCode = StatusCodes.Status200OK;
+                await writer.WriteToStream(httpContext.Response.Body, batch.Namespace!.Value, batch.Operations.Select(op => new Tuple<BatchWriter.OpVerb, BucketId, KeyId>(op.Verb, op.Bucket!.Value, op.Key!.Value)).ToList(),
+                    async (verb, ns, bucket, key, fields) =>
+                    {
+                        BatchWriter.OpState opState;
+                        RefResponse? refResponse = null;
+                        BlobContents? blob = null;
+                        try
+                        {
+                            switch (verb)
+                            {
+                                case BatchWriter.OpVerb.GET:
+                                    (refResponse, blob) = await _ddcRefService.Get(ns, bucket, key, fields);
+                                    opState = BatchWriter.OpState.OK;
+                                    break;
+                                case BatchWriter.OpVerb.HEAD:
+                                    await _ddcRefService.Exists(ns, bucket, key);
+                                    opState = BatchWriter.OpState.Exists;
+                                    refResponse = null;
+                                    blob = null;
+                                    break;
+                                default:
+                                    throw new ArgumentOutOfRangeException(nameof(verb), verb, null);
+                            }
+                        }
+                        catch (RefRecordNotFoundException)
+                        {
+                            opState = BatchWriter.OpState.NotFound;
+                        }
+                        catch (BlobNotFoundException)
+                        {
+                            opState = BatchWriter.OpState.NotFound;
+                        }
+                        catch (MissingBlobsException)
+                        {
+                            opState = BatchWriter.OpState.NotFound;
+                        }
+                        catch (NamespaceNotFoundException)
+                        {
+                            opState = BatchWriter.OpState.NotFound;
+                        }
+                        catch (Exception e)
+                        {
+                            _logger.Error(e, "Unknown exception when executing batch get");
+
+                            // we want to make sure that we always continue to write the results even when we get errors
+                            opState = BatchWriter.OpState.Failed;
+                        }
+
+                        return new Tuple<ContentHash?, BlobContents?, BatchWriter.OpState>(refResponse?.ContentHash, blob, opState);
+                    });
+
+                return;
+            }
+
+            await next.Invoke();
+        }
+
+        /// <summary>
+        /// Checks if a refs key exists
+        /// </summary>
+        /// <param name="ns">Namespace. Each namespace is completely separated from each other. Use for different types of data that is never expected to be similar (between two different games for instance). Example: `uc4.ddc`</param>
+        /// <param name="bucket">The category/type of record you are caching. Is a clustered key together with the actual key, but all records in the same bucket can be dropped easily. Example: `terrainTexture` </param>
+        /// <param name="key">The unique name of this particular key. `iAmAVeryValidKey`</param>
+        /// <returns>200 if it existed, 400 otherwise</returns>
+        [HttpHead("ddc/{ns}/{bucket}/{key}", Order = 500)]
+        [ProducesResponseType(type: typeof(RefResponse), 200)]
+        [ProducesResponseType(type: typeof(ValidationProblemDetails), 400)]
+        [Authorize("Cache.read")]
+        public async Task<IActionResult> Head(
+            [FromRoute] [Required] NamespaceId ns,
+            [FromRoute] [Required] BucketId bucket,
+            [FromRoute] [Required] KeyId key)
+        {
+            if (!ShouldDoAuth())
+            {
+                AuthorizationResult authorizationResult = await _authorizationService.AuthorizeAsync(User, ns, NamespaceAccessRequirement.Name);
+
+                if (!authorizationResult.Succeeded)
+                {
+                    return Forbid();
+                }
+            }
+
+            try
+            {
+                RefRecord record = await _ddcRefService.Exists(ns, bucket, key);
+                Response.Headers[CommonHeaders.HashHeaderName] = record.ContentHash.ToString();
+            }
+            catch (NamespaceNotFoundException e)
+            {
+                return BadRequest(new ProblemDetails {Title = $"Namespace {e.Namespace} did not exist"});
+            }
+            catch (RefRecordNotFoundException e)
+            {
+                return BadRequest(new ProblemDetails {Title = $"Object {e.Bucket} {e.Key} in namespace {e.Namespace} did not exist"});
+            }
+            catch (MissingBlobsException e)
+            {
+                return BadRequest(new ProblemDetails { Title = $"Blobs {e.Blobs} from object {e.Bucket} {e.Key} in namespace {e.Namespace} did not exist" });
+            }
+
+            return NoContent();
+        }
+
+        /// <summary>
+        /// Insert a new refs key.
+        /// </summary>
+        /// <remarks>The raw object can also be sent as a octet-stream in which case the refRequest should not be sent.</remarks>
+        /// <param name="ns">Namespace. Each namespace is completely separated from each other. Use for different types of data that is never expected to be similar (between two different games for instance). Example: `uc4.ddc`</param>
+        /// <param name="bucket">The category/type of record you are caching. Is a clustered key together with the actual key, but all records in the same bucket can be dropped easily. Example: `terrainTexture` </param>
+        /// <param name="key">The unique name of this particular key. `iAmAVeryValidKey`</param>
+        /// <param name="refRequest">Json object containing which blobs you have already inserted and metadata about them. Instead of sending this you can also send a octet-stream of the object you want to refs.</param>
+        /// <returns>The transaction id of the created object</returns>
+        [HttpPut("ddc/{ns}/{bucket}/{key}", Order = 500)]
+        [Consumes(MediaTypeNames.Application.Json)]
+        [ProducesResponseType(type: typeof(PutRequestResponse), 200)]
+        [ProducesResponseType(type: typeof(ValidationProblemDetails), 400)]
+        [Authorize("Cache.write")]
+        public async Task<IActionResult> PutStructured(
+            [FromRoute] [Required] NamespaceId ns,
+            [FromRoute] [Required] BucketId bucket,
+            [FromRoute] [Required] KeyId key,
+            [FromBody] RefRequest refRequest)
+        {
+            if (!ShouldDoAuth())
+            {
+                AuthorizationResult authorizationResult = await _authorizationService.AuthorizeAsync(User, ns, NamespaceAccessRequirement.Name);
+
+                if (!authorizationResult.Succeeded)
+                {
+                    return Forbid();
+                }
+            }
+
+            try
+            {
+                _diagnosticContext.Set("Content-Length", Request.ContentLength ?? -1);
+
+                long transactionId = await _ddcRefService.PutIndirect(ns, bucket, key, refRequest);
+                return Ok(new PutRequestResponse(transactionId));
+            }
+            catch (HashMismatchException e)
+            {
+                return BadRequest(new ProblemDetails
+                {
+                    Title = $"Incorrect hash, got hash \"{e.NewHash}\" but hash of content was determined to be \"{e.ExpectedHash}\""
+                });
+            }
+            catch (MissingBlobsException e)
+            {
+                return BadRequest(new ProblemDetails { Title = "Some blobs were not uploaded", Detail = string.Join(" ", e.Blobs.Select(identifier => identifier.ToString())) });
+            }
+        }
+
+
+        // that structured data will match the route above first
+        [RequiredContentType(MediaTypeNames.Application.Octet)]
+        [HttpPut("ddc/{ns}/{bucket}/{key}", Order = 500)]
+        // TODO: Investigate if we can resolve the conflict between this and the other put endpoint in open api
+        [ApiExplorerSettings(IgnoreApi = true)]
+        [DisableRequestSizeLimit]
+        [Authorize("Cache.write")]
+        public async Task<IActionResult> PutBlob(
+            [FromRoute] [Required] NamespaceId ns,
+            [FromRoute] [Required] BucketId bucket,
+            [FromRoute] [Required] KeyId key)
+        {
+            if (!ShouldDoAuth())
+            {
+                AuthorizationResult authorizationResult = await _authorizationService.AuthorizeAsync(User, ns, NamespaceAccessRequirement.Name);
+
+                if (!authorizationResult.Succeeded)
+                {
+                    return Forbid();
+                }
+            }
+
+            byte[] blob;
+            try
+            {
+                blob = await RequestUtil.ReadRawBody(Request);
+            }
+            catch (BadHttpRequestException e)
+            {
+                const string msg = "Partial content transfer when reading request body.";
+                _logger.Warning(e, msg);
+                return BadRequest(msg);
+            }
+            _logger.Debug("Received PUT for {Namespace} {Bucket} {Key}", ns, bucket, key);
+            _diagnosticContext.Set("Content-Length", Request.ContentLength ?? -1);
+
+            ContentHash headerHash;
+            ContentHash blobHash = ContentHash.FromBlob(blob);
+            ContentHash headerVerificationHash = blobHash;
+            if (Request.Headers.ContainsKey(CommonHeaders.HashHeaderName))
+            {
+                headerHash = new ContentHash(Request.Headers[CommonHeaders.HashHeaderName]);
+            }
+            else if (Request.Headers.ContainsKey(CommonHeaders.HashHeaderSHA1Name))
+            {
+                headerHash = new ContentHash(Request.Headers[CommonHeaders.HashHeaderSHA1Name]);
+                byte[] sha1Hash = Sha1Utils.GetSHA1(blob);
+                headerVerificationHash = new ContentHash(sha1Hash);
+            }
+            else
+            {
+                return BadRequest(new ProblemDetails
+                {
+                    Title = $"Missing expected header {CommonHeaders.HashHeaderName} or {CommonHeaders.HashHeaderSHA1Name}"
+                });
+            }
+
+            if (!headerVerificationHash.Equals(headerHash))
+            {
+                return BadRequest(new ProblemDetails
+                {
+                    Title = $"Incorrect hash, got hash \"{headerHash}\" but hash of content was determined to be \"{headerVerificationHash}\""
+                });
+            }
+
+            PutRequestResponse response = await _ddcRefService.Put(ns, bucket, key, blobHash, blob);
+            return Ok(response);
+
+        }
+
+        /// <summary>
+        /// Drop all refs records in the namespace
+        /// </summary>
+        /// <param name="ns">Namespace. Each namespace is completely separated from each other. Use for different types of data that is never expected to be similar (between two different games for instance)</param>
+        [HttpDelete("ddc/{ns}", Order = 500)]
+        [ProducesResponseType(204)]
+        [Authorize("admin")]
+        public async Task<IActionResult> DeleteNamespace(
+            [FromRoute] [Required] NamespaceId ns
+        )
+        {
+            AuthorizationResult authorizationResult = await _authorizationService.AuthorizeAsync(User, ns, NamespaceAccessRequirement.Name);
+
+            if (!authorizationResult.Succeeded)
+            {
+                return Forbid();
+            }
+
+            try
+            {
+                await _refsStore.DropNamespace(ns);
+            }
+            catch (NamespaceNotFoundException e)
+            {
+                return BadRequest(new ProblemDetails {Title = $"Namespace {e.Namespace} did not exist"});
+            }
+
+
+            return NoContent();
+        }
+
+        /// <summary>
+        /// Drop all refs records in the bucket
+        /// </summary>
+        /// <param name="ns">Namespace. Each namespace is completely separated from each other. Use for different types of data that is never expected to be similar (between two different games for instance)</param>
+        /// <param name="bucket">The category/type of record you are caching. Is a clustered key together with the actual key, but all records in the same bucket can be dropped easily.</param>
+        [HttpDelete("ddc/{ns}/{bucket}", Order = 500)]
+        [ProducesResponseType(204)]
+        [Authorize("Cache.delete")]
+        public async Task<IActionResult> DeleteBucket(
+            [FromRoute] [Required] NamespaceId ns,
+            [FromRoute] [Required] BucketId bucket)
+        {
+            AuthorizationResult authorizationResult = await _authorizationService.AuthorizeAsync(User, ns, NamespaceAccessRequirement.Name);
+
+            if (!authorizationResult.Succeeded)
+            {
+                return Forbid();
+            }
+
+            try
+            {
+                await _refsStore.DeleteBucket(ns, bucket);
+            }
+            catch (NamespaceNotFoundException e)
+            {
+                return BadRequest(new ProblemDetails {Title = $"Namespace {e.Namespace} did not exist"});
+            }
+
+
+            return NoContent();
+        }
+
+        /// <summary>
+        /// Delete a individual refs key
+        /// </summary>
+        /// <param name="ns">Namespace. Each namespace is completely separated from each other. Use for different types of data that is never expected to be similar (between two different games for instance)</param>
+        /// <param name="bucket">The category/type of record you are caching. Is a clustered key together with the actual key, but all records in the same bucket can be dropped easily.</param>
+        /// <param name="key">The unique name of this particular key</param>
+        [HttpDelete("ddc/{ns}/{bucket}/{key}", Order = 500)]
+        [ProducesResponseType(204)]
+        [ProducesResponseType(400)]
+        [Authorize("Cache.delete")]
+        public async Task<IActionResult> Delete(
+            [FromRoute] [Required] NamespaceId ns,
+            [FromRoute] [Required] BucketId bucket,
+            [FromRoute] [Required] KeyId key)
+        {
+            AuthorizationResult authorizationResult = await _authorizationService.AuthorizeAsync(User, ns, NamespaceAccessRequirement.Name);
+
+            if (!authorizationResult.Succeeded)
+            {
+                return Forbid();
+            }
+
+            try
+            {
+                long deleteCount = await _ddcRefService.Delete(ns, bucket, key);
+
+                if (deleteCount == 0)
+                    return BadRequest("Deleted 0 records, most likely the object did not exist");
+            }
+            catch (NamespaceNotFoundException e)
+            {
+                return BadRequest(new ProblemDetails {Title = $"Namespace {e.Namespace} did not exist"});
+            }
+            return NoContent();
+        }
+
+
+        // ReSharper disable UnusedAutoPropertyAccessor.Global
+        // ReSharper disable once ClassNeverInstantiated.Global
+        public class BatchOp
+        {
+            // ReSharper disable InconsistentNaming
+            public enum Operation
+            {
+                INVALID,
+                GET,
+                PUT,
+                DELETE,
+                HEAD
+            }
+            // ReSharper restore InconsistentNaming
+
+            [Required] public NamespaceId? Namespace { get; set; }
+
+            [Required]
+            public BucketId? Bucket { get; set; }
+
+            [Required] public KeyId? Id { get; set; }
+
+            [Required] public Operation Op { get; set; }
+
+            // used for put ops
+            public ContentHash? ContentHash { get; set; }
+            public byte[]? Content { get; set; }
+        }
+
+        public class BatchCall
+        {
+            public BatchOp[]? Operations { get; set; }
+        }
+
+        [HttpPost("ddc-rpc")]
+        public async Task<IActionResult> Post([FromBody] BatchCall batch)
+        {
+            string OpToPolicy(BatchOp.Operation op)
+            {
+                switch (op)
+                {
+                    case BatchOp.Operation.GET:
+                    case BatchOp.Operation.HEAD:
+                        return "Cache.read";
+                    case BatchOp.Operation.PUT:
+                        return "Cache.write";
+                    case BatchOp.Operation.DELETE:
+                        return "Cache.delete";
+                    default:
+                        throw new ArgumentOutOfRangeException(nameof(op), op, null);
+                }
+            }
+
+            if (batch.Operations == null)
+            {
+                throw new ArgumentNullException();
+            }
+
+            Task<object?>[] tasks = new Task<object?>[batch.Operations.Length];
+            for (int index = 0; index < batch.Operations.Length; index++)
+            {
+                BatchOp op = batch.Operations[index];
+
+                AuthorizationResult authorizationResultNamespace = await _authorizationService.AuthorizeAsync(User, op.Namespace, NamespaceAccessRequirement.Name);
+                AuthorizationResult authorizationResultOp = await _authorizationService.AuthorizeAsync(User, op, OpToPolicy(op.Op));
+
+                bool authorized = authorizationResultNamespace.Succeeded && authorizationResultOp.Succeeded;
+
+                if (!authorized)
+                {
+                    tasks[index] = Task.FromResult((object?)new ProblemDetails { Title = "Forbidden" });
+                }
+                else
+                {
+                    tasks[index] = ProcessOp(op).ContinueWith(task => task.IsFaulted ? (object?)new ProblemDetails { Title = "Exception thrown", Detail = task!.Exception!.ToString()} : task.Result);
+                }
+            }
+
+            await Task.WhenAll(tasks);
+
+            object?[] results = tasks.Select(t => t.Result).ToArray();
+
+            return Ok(results);
+        }
+
+        private Task<object?> ProcessOp(BatchOp op)
+        {
+            if (op.Namespace == null)
+            {
+                throw new ArgumentNullException("namespace");
+            }
+
+            if (op.Bucket == null)
+            {
+                throw new ArgumentNullException("bucket");
+            }
+
+            if (op.Id == null)
+            {
+                throw new ArgumentNullException("id");
+            }
+
+            switch (op.Op)
+            {
+                case BatchOp.Operation.INVALID:
+                    throw new ArgumentOutOfRangeException();
+                case BatchOp.Operation.GET:
+                    // TODO: support field filtering
+                    return _ddcRefService.Get(op.Namespace.Value, op.Bucket.Value, op.Id.Value, new string[0]).ContinueWith(t => 
+                    {
+                        // we are serializing this to json, so we just stream the blob into memory to return it.
+                        (RefResponse response, BlobContents? blob) = t.Result;
+                        if (blob != null)
+                        {
+                            using BlobContents blobContents = blob;
+                            response.Blob = blobContents.Stream.ToByteArray().Result;
+                        }
+                        return (object?)response;
+                    });
+                case BatchOp.Operation.HEAD:
+                    return _ddcRefService.Exists(op.Namespace.Value, op.Bucket.Value, op.Id.Value).ContinueWith(t => t.Result == null ? (object?)null : op.Id);
+                case BatchOp.Operation.PUT:
+                {
+                    if (op.ContentHash == null)
+                    {
+                        throw new ArgumentNullException("ContentHash");
+                    }
+
+                    if (op.Content == null)
+                    {
+                        throw new ArgumentNullException("Content");
+                    }
+
+                    ContentHash blobHash = ContentHash.FromBlob(op.Content);
+                    if (!blobHash.Equals(op.ContentHash))
+                    {
+                        throw new HashMismatchException(blobHash, op.ContentHash);
+                    }
+                    return _ddcRefService.Put(op.Namespace.Value, op.Bucket.Value, op.Id.Value, blobHash, op.Content).ContinueWith(t => (object?)t.Result);
+                }
+                case BatchOp.Operation.DELETE:
+                    return _ddcRefService.Delete(op.Namespace.Value, op.Bucket.Value, op.Id.Value).ContinueWith(t => (object?)null);
+                default:
+                    throw new ArgumentOutOfRangeException();
+            }
+        }
+
+        public class BatchGetOp
+        {
+            public class GetOp
+            {
+                public BucketId? Bucket { get; set; } = null!;
+                public KeyId? Key { get; set; } = null!;
+                public BatchWriter.OpVerb Verb { get; set; } = BatchWriter.OpVerb.GET;
+            }
+
+            public NamespaceId? Namespace { get; set; } = null!;
+
+            public GetOp[] Operations { get; set; } = null!;
+        }
+        
+        /// <summary>
+        /// Custom batch get rpc to allows us to fetch multiple objects and returning them as a custom binary stream avoiding base64 encoding.
+        /// Objects are returned in the order they have been fetched, not the order they were requested in.
+        /// </summary>
+        /// <param name="batch">Spec for which objects to return</param>
+        /// <returns></returns>
+        [HttpPost("ddc-rpc/batchGet")]
+        [Authorize("Cache.read")]
+        public async Task<IActionResult> BatchGet([FromBody] BatchGetOp batch)
+        {
+            if (batch.Namespace == null)
+            {
+                throw new ArgumentNullException();
+            }
+
+            if (!ShouldDoAuth())
+            {
+                AuthorizationResult authorizationResult = await _authorizationService.AuthorizeAsync(User, batch.Namespace, NamespaceAccessRequirement.Name);
+
+                if (!authorizationResult.Succeeded)
+                {
+                    return Forbid();
+                }
+            }
+
+            BatchWriter writer = new BatchWriter();
+            await writer.WriteToStream(Response.Body, batch.Namespace.Value!, batch.Operations.Select(op => new Tuple<BatchWriter.OpVerb, BucketId, KeyId>(op.Verb, op.Bucket!.Value, op.Key!.Value)).ToList(),
+                async (verb, ns, bucket, key, fields) =>
+                {
+                    BatchWriter.OpState opState;
+                    RefResponse? refResponse = null;
+                    BlobContents? blob = null;
+                    try
+                    {
+                        switch (verb)
+                        {
+                            case BatchWriter.OpVerb.GET:
+                                (refResponse, blob) = await _ddcRefService.Get(ns, bucket, key, fields);
+                                opState = BatchWriter.OpState.OK;
+                                break;
+                            case BatchWriter.OpVerb.HEAD:
+                                await _ddcRefService.Exists(ns, bucket, key);
+                                opState = BatchWriter.OpState.Exists;
+                                refResponse = null;
+                                blob = null;
+                                break;
+                            default:
+                                throw new ArgumentOutOfRangeException(nameof(verb), verb, null);
+                        }
+                    }
+                    catch (RefRecordNotFoundException)
+                    {
+                        opState = BatchWriter.OpState.NotFound;
+                    }
+                    catch (BlobNotFoundException)
+                    {
+                        opState = BatchWriter.OpState.NotFound;
+                    }
+                    catch (MissingBlobsException)
+                    {
+                        opState = BatchWriter.OpState.NotFound;
+                    }
+                    catch (NamespaceNotFoundException)
+                    {
+                        opState = BatchWriter.OpState.NotFound;
+                    }
+                    catch (Exception e)
+                    {
+                        _logger.Error(e, "Unknown exception when executing batch get");
+
+                        // we want to make sure that we always continue to write the results even when we get errors
+                        opState = BatchWriter.OpState.Failed;
+                    }
+
+                    return new Tuple<ContentHash?, BlobContents?, BatchWriter.OpState>(refResponse?.ContentHash, blob, opState);
+                });
+
+            // we have already set the result by writing to Response
+            return new EmptyResult();
+        }
+
+        private bool ShouldDoAuth()
+        {
+            foreach (int port in _jupiterSettings.CurrentValue.DisableAuthOnPorts)
+            {
+                if (port == HttpContext.Connection.LocalPort)
+                    return true;
+            }
+
+            return false;
+        }
+    }
+
+    public class HashMismatchException : Exception
+    {
+        public ContentHash NewHash { get; }
+        public ContentHash ExpectedHash { get; }
+
+        public HashMismatchException(ContentHash newHash, ContentHash expectedHash)
+        {
+            NewHash = newHash;
+            ExpectedHash = expectedHash;
+        }
+    }
+
+    public class RefRecordNotFoundException : Exception
+    {
+        public RefRecordNotFoundException(NamespaceId ns, BucketId bucket, KeyId key)
+        {
+            Namespace = ns;
+            Bucket = bucket;
+            Key = key;
+        }
+
+        public NamespaceId Namespace { get; }
+        public BucketId Bucket { get; }
+        public KeyId Key { get; }
+    }
+
+    public class MissingBlobsException : Exception
+    {
+        public MissingBlobsException(NamespaceId ns, BucketId bucket, KeyId key, BlobIdentifier[] blobs)
+        {
+            Namespace = ns;
+            Bucket = bucket;
+            Key = key;
+            Blobs = blobs;
+        }
+
+        public NamespaceId Namespace { get; }
+        public BucketId Bucket { get; }
+        public KeyId Key { get; }
+        public BlobIdentifier[] Blobs { get; }
+    }
+
+    public class PutRequestResponse
+    {
+        public PutRequestResponse(long transactionId)
+        {
+            TransactionId = transactionId;
+        }
+
+        public long TransactionId { get; }
+    }
+
+    public class GetNamespacesResponse
+    {
+        [JsonConstructor]
+        public GetNamespacesResponse(NamespaceId[] namespaces)
+        {
+            Namespaces = namespaces;
+        }
+
+        public NamespaceId[] Namespaces { get; set; }
+    }
+}
diff --git a/Engine/Source/Programs/Horde/HordeStorage/Horde.Storage/Controllers/DebugController.cs b/Engine/Source/Programs/Horde/HordeStorage/Horde.Storage/Controllers/DebugController.cs
new file mode 100644
--- /dev/null
+++ b/Engine/Source/Programs/Horde/HordeStorage/Horde.Storage/Controllers/DebugController.cs
@@ -0,0 +1,91 @@
+﻿// Copyright Epic Games, Inc. All Rights Reserved.
+
+using System;
+using System.Net.Mime;
+using System.Threading.Tasks;
+using Microsoft.AspNetCore.Authorization;
+using Microsoft.AspNetCore.Http;
+using Microsoft.AspNetCore.Mvc;
+
+namespace Horde.Storage.Controllers
+{
+    [ApiController]
+    [Route("api/v1/c/_debug")]
+    public class DebugController : Controller
+    {
+        private static readonly byte[] OneKilobyteData = new byte[1024];
+        private static readonly byte[] OneMegabyteData = new byte[1024 * 1024];
+        private static readonly byte[] SixteenMegabyteData = new byte[16 * 1024 * 1024];
+        private static bool s_testBuffersFilled = false;
+        
+        /// <summary>
+        /// Return bytes of specified length with auth, used for testing only
+        /// </summary>
+        /// <returns></returns>
+        [HttpGet("getBytes")]
+        [Authorize("Any")]
+        public IActionResult GetBytes([FromQuery] int length = 1)
+        {
+            return GenerateByteResponse(length);
+        }
+        
+        /// <summary>
+        /// Return bytes of specified length without auth, used for testing only
+        /// </summary>
+        /// <returns></returns>
+        [HttpGet("getBytesWithoutAuth")]
+        public IActionResult GetBytesWithoutAuth([FromQuery] int length = 1)
+        {
+            return GenerateByteResponse(length);
+        }
+        
+        [ApiExplorerSettings(IgnoreApi = true)]
+        [NonAction]
+        public async Task FastGetBytes(HttpContext httpContext, Func<Task> next)
+        {
+            async Task Send(byte[] data)
+            {
+                httpContext.Response.StatusCode = StatusCodes.Status200OK;
+                httpContext.Response.Headers["Content-Type"] = MediaTypeNames.Application.Octet;
+                httpContext.Response.Headers["Content-Length"] = Convert.ToString(data.Length);
+                await httpContext.Response.Body.WriteAsync(data);
+            }
+
+            if (!s_testBuffersFilled)
+            {
+                Array.Fill(OneKilobyteData, (byte)'k');
+                Array.Fill(OneMegabyteData, (byte)'m');
+                Array.Fill(SixteenMegabyteData, (byte)'s');
+                s_testBuffersFilled = true;
+            }
+            
+            if (httpContext.Request.Path.StartsWithSegments("/api/v1/c/_debug/fastBytes/1byte"))
+            {
+                await Send(new byte[] {0x4A});
+            }
+            else if (httpContext.Request.Path.StartsWithSegments("/api/v1/c/_debug/fastBytes/1kilobyte"))
+            {
+                await Send(OneKilobyteData);
+            }
+            else if (httpContext.Request.Path.StartsWithSegments("/api/v1/c/_debug/fastBytes/1megabyte"))
+            {
+                await Send(OneMegabyteData);
+            }
+            else if (httpContext.Request.Path.StartsWithSegments("/api/v1/c/_debug/fastBytes/16megabyte"))
+            {
+                await Send(SixteenMegabyteData);
+            }
+            else
+            {
+                await next.Invoke();    
+            }
+        }
+
+        private FileContentResult GenerateByteResponse(int length)
+        {
+            byte[] generatedData = new byte[length];
+            Array.Fill(generatedData, (byte)'J');
+            return File(generatedData, "application/octet-stream");
+        }
+    }
+}
diff --git a/Engine/Source/Programs/Horde/HordeStorage/Horde.Storage/Controllers/ObjectController.cs b/Engine/Source/Programs/Horde/HordeStorage/Horde.Storage/Controllers/ObjectController.cs
new file mode 100644
--- /dev/null
+++ b/Engine/Source/Programs/Horde/HordeStorage/Horde.Storage/Controllers/ObjectController.cs
@@ -0,0 +1,254 @@
+﻿// Copyright Epic Games, Inc. All Rights Reserved.
+
+using System;
+using System.Collections.Concurrent;
+using System.Collections.Generic;
+using System.ComponentModel.DataAnnotations;
+using System.Linq;
+using System.Net.Mime;
+using System.Threading.Tasks;
+using async_enumerable_dotnet;
+using Datadog.Trace;
+using Horde.Storage.Implementation;
+using Jupiter;
+using Jupiter.Implementation;
+using Jupiter.Utils;
+using Microsoft.AspNetCore.Authorization;
+using Microsoft.AspNetCore.Http;
+using Microsoft.AspNetCore.Mvc;
+using Serilog;
+
+namespace Horde.Storage.Controllers
+{
+    [ApiController]
+    [Route("api/v1/objects", Order = 0)]
+    [Produces(CustomMediaTypeNames.UnrealCompactBinary, MediaTypeNames.Application.Json)]
+    public class ObjectController : ControllerBase
+    {
+        private readonly IBlobStore _storage;
+        private readonly IDiagnosticContext _diagnosticContext;
+        private readonly IAuthorizationService _authorizationService;
+        private readonly IReferenceResolver _referenceResolver;
+
+        private readonly ILogger _logger = Log.ForContext<ObjectController>();
+
+        public ObjectController(IBlobStore storage, IDiagnosticContext diagnosticContext, IAuthorizationService authorizationService, IReferenceResolver referenceResolver)
+        {
+            _storage = storage;
+            _diagnosticContext = diagnosticContext;
+            _authorizationService = authorizationService;
+            _referenceResolver = referenceResolver;
+        }
+
+
+        [HttpGet("{ns}/{id}")]
+        [Authorize("Storage.read")]
+        [ProducesDefaultResponseType]
+        public async Task<IActionResult> Get(
+            [Required] NamespaceId ns,
+            [Required] BlobIdentifier id)
+        {
+            AuthorizationResult authorizationResult = await _authorizationService.AuthorizeAsync(User, ns, NamespaceAccessRequirement.Name);
+
+            if (!authorizationResult.Succeeded)
+            {
+                return Forbid();
+            }
+
+            try
+            {
+                BlobContents blobContents = await _storage.GetObject(ns, id);
+
+                return File(blobContents.Stream, CustomMediaTypeNames.UnrealCompactBinary);
+            }
+            catch (BlobNotFoundException e)
+            {
+                return NotFound(new ValidationProblemDetails {Title = $"Object {e.Blob} not found"});
+            }
+        }
+
+        [HttpHead("{ns}/{id}")]
+        [Authorize("Storage.read")]
+        [ProducesDefaultResponseType]
+        public async Task<IActionResult> Head(
+            [Required] NamespaceId ns,
+            [Required] BlobIdentifier id)
+        {
+            AuthorizationResult authorizationResult = await _authorizationService.AuthorizeAsync(User, ns, NamespaceAccessRequirement.Name);
+
+            if (!authorizationResult.Succeeded)
+            {
+                return Forbid();
+            }
+
+            bool exists = await _storage.Exists(ns, id);
+
+            if (!exists)
+            {
+                return NotFound(new ValidationProblemDetails {Title = $"Object {id} not found"});
+            }
+
+            return Ok();
+        }
+
+        [HttpPost("{ns}/exists")]
+        [Authorize("Storage.read")]
+        [ProducesDefaultResponseType]
+        public async Task<IActionResult> ExistsMultiple(
+            [Required] NamespaceId ns,
+            [Required] [FromQuery] List<BlobIdentifier> id)
+        {
+            AuthorizationResult authorizationResult = await _authorizationService.AuthorizeAsync(User, ns, NamespaceAccessRequirement.Name);
+
+            if (!authorizationResult.Succeeded)
+            {
+                return Forbid();
+            }
+
+            ConcurrentBag<BlobIdentifier> missingBlobs = new ConcurrentBag<BlobIdentifier>();
+
+            IEnumerable<Task> tasks = id.Select(async blob =>
+            {
+                if (!await _storage.Exists(ns, blob))
+                    missingBlobs.Add(blob);
+            });
+            await Task.WhenAll(tasks);
+
+            return Ok(new HeadMultipleResponse {Needs = missingBlobs.ToArray()});
+        }
+
+        [HttpPut("{ns}/{id}")]
+        [Authorize("Storage.write")]
+        [RequiredContentType(CustomMediaTypeNames.UnrealCompactBinary)]
+        public async Task<IActionResult> Put(
+            [Required] NamespaceId ns,
+            [Required] BlobIdentifier id)
+        {
+            AuthorizationResult authorizationResult = await _authorizationService.AuthorizeAsync(User, ns, NamespaceAccessRequirement.Name);
+
+            if (!authorizationResult.Succeeded)
+            {
+                return Forbid();
+            }
+
+            byte[] blob;
+            try
+            {
+                blob = await RequestUtil.ReadRawBody(Request);
+            }
+            catch (BadHttpRequestException e)
+            {
+                const string msg = "Partial content transfer when reading request body.";
+                _logger.Warning(e, msg);
+                return BadRequest(msg);
+            }
+
+            _diagnosticContext.Set("Content-Length", Request.ContentLength ?? -1);
+
+            BlobIdentifier blobHash;
+            {
+                using Scope _ = Tracer.Instance.StartActive("web.hash");
+                blobHash = BlobIdentifier.FromBlob(blob);
+            }
+
+            if (!id.Equals(blobHash))
+            {
+                _logger.Debug("ID {@Id} was not the same as identifier {@Identifier} {Content}", id, blobHash,
+                    blob);
+
+                throw new ArgumentException("ID was not a hash of the content uploaded.", paramName: nameof(id));
+            }
+
+            await _storage.PutObject(ns, blob, blobHash);
+            BlobIdentifier identifier = blobHash;
+            return Ok(new
+            {
+                Identifier = identifier
+            });
+        }
+
+        [HttpGet("{ns}/{id}/references")]
+        [Authorize("Storage.read")]
+        public async Task<IActionResult> ResolveReferences(
+            [Required] NamespaceId ns,
+            [Required] BlobIdentifier id,
+            [FromQuery] int depth = 1)
+        {
+            AuthorizationResult authorizationResult = await _authorizationService.AuthorizeAsync(User, ns, NamespaceAccessRequirement.Name);
+
+            if (!authorizationResult.Succeeded)
+            {
+                return Forbid();
+            }
+
+            BlobContents blob = await _storage.GetObject(ns, id);
+            byte[] blobContents = await blob.Stream.ToByteArray();
+            CompactBinaryObject compactBinaryObject = CompactBinaryObject.Load(blobContents);
+
+            try
+            {
+                BlobIdentifier[] references = await _referenceResolver.ResolveReferences(ns, compactBinaryObject).ToArrayAsync();
+                return Ok(new ResolvedReferencesResult(references));
+            }
+            catch (PartialReferenceResolveException e)
+            {
+                return BadRequest(new ValidationProblemDetails {Title = $"Object {id} is missing blobs", Detail = $"Following blobs are missing: {string.Join(",", e.UnresolvedReferences)}"});
+            }
+        }
+
+
+        [HttpDelete("{ns}/{id}")]
+        [Authorize("Storage.delete")]
+        public async Task<IActionResult> Delete(
+            [Required] NamespaceId ns,
+            [Required] BlobIdentifier id)
+        {
+            AuthorizationResult authorizationResult = await _authorizationService.AuthorizeAsync(User, ns, NamespaceAccessRequirement.Name);
+
+            if (!authorizationResult.Succeeded)
+            {
+                return Forbid();
+            }
+
+            await _storage.DeleteObject(ns, id);
+
+            return Ok( new DeletedResponse
+            {
+                DeletedCount = 1
+            });
+        }
+
+
+        [HttpDelete("{ns}")]
+        [Authorize("Admin")]
+        public async Task<IActionResult> DeleteNamespace(
+            [Required] NamespaceId ns)
+        {
+            AuthorizationResult authorizationResult = await _authorizationService.AuthorizeAsync(User, ns, NamespaceAccessRequirement.Name);
+
+            if (!authorizationResult.Succeeded)
+            {
+                return Forbid();
+            }
+
+            await _storage.DeleteNamespace(ns);
+
+            return Ok();
+        }
+    }
+
+    public class DeletedResponse
+    {
+        public int DeletedCount { get; set; }
+    }
+
+    public class ResolvedReferencesResult
+    {
+        public ResolvedReferencesResult(BlobIdentifier[] references)
+        {
+            References = references;
+        }
+
+        public BlobIdentifier[] References { get; set; }
+    }
+}
diff --git a/Engine/Source/Programs/Horde/HordeStorage/Horde.Storage/Controllers/ProblemTypes.cs b/Engine/Source/Programs/Horde/HordeStorage/Horde.Storage/Controllers/ProblemTypes.cs
new file mode 100644
--- /dev/null
+++ b/Engine/Source/Programs/Horde/HordeStorage/Horde.Storage/Controllers/ProblemTypes.cs
@@ -0,0 +1,9 @@
+﻿// Copyright Epic Games, Inc. All Rights Reserved.
+
+namespace Horde.Storage.Controllers
+{
+    public class ProblemTypes
+    {
+        public static string UseSnapshot { get; set; } = "http://jupiter.epicgames.com/replication/useSnapshot";
+    }
+}
