Moved Horde.Storage from git to p4

diff --git a/Engine/Source/Programs/Horde/HordeStorage/Horde.Storage/Controllers/ReferencesController.cs b/Engine/Source/Programs/Horde/HordeStorage/Horde.Storage/Controllers/ReferencesController.cs
new file mode 100644
--- /dev/null
+++ b/Engine/Source/Programs/Horde/HordeStorage/Horde.Storage/Controllers/ReferencesController.cs
@@ -0,0 +1,536 @@
+// Copyright Epic Games, Inc. All Rights Reserved.
+
+using System;
+using System.Collections.Concurrent;
+using System.Collections.Generic;
+using System.ComponentModel.DataAnnotations;
+using System.Linq;
+using System.Net.Mime;
+using System.Text;
+using System.Threading.Tasks;
+using Dasync.Collections;
+using Datadog.Trace;
+using Horde.Storage.Implementation;
+using Jupiter;
+using Jupiter.Implementation;
+using Jupiter.Utils;
+using Microsoft.AspNetCore.Authorization;
+using Microsoft.AspNetCore.Http;
+using Microsoft.AspNetCore.Http.Extensions;
+using Microsoft.AspNetCore.Mvc;
+using Microsoft.Extensions.Options;
+using Serilog;
+
+namespace Horde.Storage.Controllers
+{
+    [ApiController]
+    [FormatFilter]
+    [Produces(MediaTypeNames.Application.Json, MediaTypeNames.Application.Octet, CustomMediaTypeNames.UnrealCompactBinary)]
+    [Route("api/v1/refs")]
+    public class ReferencesController : ControllerBase
+    {
+        private readonly IDiagnosticContext _diagnosticContext;
+        private readonly IAuthorizationService _authorizationService;
+        private readonly IOptionsMonitor<JupiterSettings> _jupiterSettings;
+        private readonly FormatResolver _formatResolver;
+
+        private readonly ILogger _logger = Log.ForContext<ReferencesController>();
+        private readonly IObjectService _objectService;
+        private readonly IBlobStore _blobStore;
+
+        public ReferencesController(IObjectService objectService, IBlobStore blobStore, IDiagnosticContext diagnosticContext, IAuthorizationService authorizationService, IOptionsMonitor<JupiterSettings> jupiterSettings, FormatResolver formatResolver)
+        {
+            _objectService = objectService;
+            _blobStore = blobStore;
+            _diagnosticContext = diagnosticContext;
+            _authorizationService = authorizationService;
+            _jupiterSettings = jupiterSettings;
+            _formatResolver = formatResolver;
+        }
+
+        /// <summary>
+        /// Returns all the known namespace the token has access to
+        /// </summary>
+        /// <returns></returns>
+        [HttpGet("")]
+        [ProducesDefaultResponseType]
+        [ProducesResponseType(type: typeof(ProblemDetails), 400)]
+        [Authorize("Cache.read")]
+        public async Task<IActionResult> GetNamespaces()
+        {
+            NamespaceId[] namespaces = await _objectService.GetNamespaces().ToArrayAsync();
+
+            if (ShouldDoAuth())
+            {
+                // filter namespaces down to only the namespaces the user has access to
+                namespaces = namespaces.Where(ns =>
+                {
+                    Task<AuthorizationResult> authorizationResult = _authorizationService.AuthorizeAsync(User, ns, NamespaceAccessRequirement.Name);
+                    return authorizationResult.Result.Succeeded;
+                }).ToArray();
+            }
+
+            return Ok(new GetNamespacesResponse(namespaces));
+        }
+
+        /// <summary>
+        /// Returns a refs key
+        /// </summary>
+        /// <param name="ns">Namespace. Each namespace is completely separated from each other. Use for different types of data that is never expected to be similar (between two different games for instance). Example: `uc4.ddc`</param>
+        /// <param name="bucket">The category/type of record you are caching. Is a clustered key together with the actual key, but all records in the same bucket can be dropped easily. Example: `terrainTexture` </param>
+        /// <param name="key">The unique name of this particular key. `iAmAVeryValidKey`</param>
+        /// <param name="fields">The fields to include in the response, omit this to include everything.</param>
+        /// <param name="format">Optional specifier to set which output format is used json/raw/cb</param>
+        [HttpGet("{ns}/{bucket}/{key}.{format?}", Order = 500)]
+        [Authorize("Object.read")]
+        public async Task<IActionResult> Get(
+            [FromRoute] [Required] NamespaceId ns,
+            [FromRoute] [Required] BucketId bucket,
+            [FromRoute] [Required] KeyId key,
+            [FromQuery] string[] fields,
+            [FromRoute] string? format = null)
+        {
+            if (ShouldDoAuth())
+            {
+                using (Scope _ = Tracer.Instance.StartActive("authorize"))
+                {
+                    AuthorizationResult authorizationResult = await _authorizationService.AuthorizeAsync(User, ns, NamespaceAccessRequirement.Name);
+
+                    if (!authorizationResult.Succeeded)
+                    {
+                        return Forbid();
+                    }
+                }
+            }
+
+            try
+            {
+                // TODO: Do we want to keep fields? we have no real need for it now
+                (ObjectRecord objectRecord, BlobContents? blob) = await _objectService.Get(ns, bucket, key, fields);
+
+                if (!objectRecord.IsFinalized)
+                {
+                    // we do not considered un-finalized objects are valid
+                    return BadRequest(new ProblemDetails { Title = $"Object {objectRecord.Bucket} {objectRecord.Name} is not finalized." });
+                }
+
+                Response.Headers[CommonHeaders.HashHeaderName] = objectRecord.BlobIdentifier.ToString();
+                Response.Headers[CommonHeaders.LastAccessHeaderName] = objectRecord.LastAccess.ToString();
+
+                async Task WriteBody(BlobContents blobContents, string contentType)
+                {
+                    using Scope scope = Tracer.Instance.StartActive("body.write");
+                    long contentLength = blobContents.Length;
+                    scope.Span.SetTag("content-length", contentLength.ToString());
+                    const int BufferSize = 64 * 1024;
+                    var outputStream = Response.Body;
+                    Response.ContentLength = contentLength;
+                    Response.ContentType = contentType;
+                    Response.StatusCode = StatusCodes.Status200OK;
+                    await StreamCopyOperation.CopyToAsync(blobContents.Stream, outputStream, count: null, bufferSize: BufferSize, cancel: Response.HttpContext.RequestAborted);
+                }
+
+                string responseType = _formatResolver.GetResponseType(Request, format, CustomMediaTypeNames.UnrealCompactBinary);
+
+                switch (responseType)
+                {
+                    case CustomMediaTypeNames.UnrealCompactBinary:
+                    {
+                        // for compact binary we can just serialize our internal object
+                        await WriteBody(blob, CustomMediaTypeNames.UnrealCompactBinary);
+
+                        break;
+                    }
+                    case MediaTypeNames.Application.Octet:
+                    {
+                        byte[] blobMemory = await blob.Stream.ToByteArray();
+                        ReadOnlyMemory<byte> localMemory = new ReadOnlyMemory<byte>(blobMemory);
+                        CompactBinaryObject cb = CompactBinaryObject.Load(ref localMemory);
+                        List<CompactBinaryField> compactBinaryFields = cb.GetFields().ToList();
+                        // detect if this object was uploaded as a raw object, in which case we deconstruct the generated compact binary
+                        if (compactBinaryFields.Count == 1)
+                        {
+                            CompactBinaryField field = compactBinaryFields[0];
+                            if (field.IsBinaryAttachment() && string.IsNullOrEmpty(field.Name))
+                            {
+                                // this is a very simple object, we just fetch the blob it references and return that
+                                BlobContents referencedBlobContents = await _blobStore.GetObject(ns, field.AsBinaryAttachment()!);
+                                await WriteBody(referencedBlobContents, MediaTypeNames.Application.Octet);
+                                break;
+                            }
+                        }
+
+                        // this doesn't look like the generated compact binary so we just return the payload
+                        await WriteBody(new BlobContents(blobMemory), MediaTypeNames.Application.Octet);
+                        break;
+                    }
+                    case MediaTypeNames.Application.Json:
+                    {
+                        byte[] blobMemory;
+                        {
+                            using Scope scope = Tracer.Instance.StartActive("json.readblob");
+                            blobMemory = await blob.Stream.ToByteArray();
+                        }
+                        ReadOnlyMemory<byte> localMemory = new ReadOnlyMemory<byte>(blobMemory);
+                        CompactBinaryObject cb = CompactBinaryObject.Load(ref localMemory);
+                        string s = cb.ToJson();
+                        await WriteBody(new BlobContents(Encoding.UTF8.GetBytes(s)), MediaTypeNames.Application.Json);
+                        break;
+
+                    }
+                    default:
+                        throw new NotImplementedException($"Unknown expected response type {responseType}");
+                }
+                
+                // this result is ignored as we write to the body explicitly
+                return new EmptyResult();
+
+            }
+            catch (NamespaceNotFoundException e)
+            {
+                return NotFound(new ProblemDetails {Title = $"Namespace {e.Namespace} did not exist"});
+            }
+            catch (ObjectNotFoundException e)
+            {
+                return NotFound(new ProblemDetails { Title = $"Object {e.Bucket} {e.Key} did not exist" });
+            }
+            catch (BlobNotFoundException e)
+            {
+                return NotFound(new ProblemDetails { Title = $"Object {e.Blob} in {e.Ns} not found" });
+            }
+        }
+
+
+        /// <summary>
+        /// Checks if a object exists
+        /// </summary>
+        /// <param name="ns">Namespace. Each namespace is completely separated from each other. Use for different types of data that is never expected to be similar (between two different games for instance). Example: `uc4.ddc`</param>
+        /// <param name="bucket">The category/type of record you are caching. Is a clustered key together with the actual key, but all records in the same bucket can be dropped easily. Example: `terrainTexture` </param>
+        /// <param name="key">The unique name of this particular key. `iAmAVeryValidKey`</param>
+        /// <returns>200 if it existed, 400 otherwise</returns>
+        [HttpHead("{ns}/{bucket}/{key}", Order = 500)]
+        [ProducesResponseType(type: typeof(RefResponse), 200)]
+        [ProducesResponseType(type: typeof(ValidationProblemDetails), 400)]
+        [Authorize("Object.read")]
+        public async Task<IActionResult> Head(
+            [FromRoute] [Required] NamespaceId ns,
+            [FromRoute] [Required] BucketId bucket,
+            [FromRoute] [Required] KeyId key)
+        {
+            if (ShouldDoAuth())
+            {
+                using (Scope _ = Tracer.Instance.StartActive("authorize"))
+                {
+                    AuthorizationResult authorizationResult = await _authorizationService.AuthorizeAsync(User, ns, NamespaceAccessRequirement.Name);
+
+                    if (!authorizationResult.Succeeded)
+                    {
+                        return Forbid();
+                    }
+                }
+            }
+
+            try
+            {
+                (ObjectRecord record, BlobContents _) = await _objectService.Get(ns, bucket, key, new string[] {"blobIdentifier"});
+                Response.Headers[CommonHeaders.HashHeaderName] = record.BlobIdentifier.ToString();
+            }
+            catch (NamespaceNotFoundException e)
+            {
+                return NotFound(new ProblemDetails {Title = $"Namespace {e.Namespace} did not exist"});
+            }
+            catch (ObjectNotFoundException e)
+            {
+                return NotFound(new ProblemDetails {Title = $"Object {e.Bucket} {e.Key} in namespace {e.Namespace} did not exist"});
+            }
+            catch (MissingBlobsException e)
+            {
+                return NotFound(new ProblemDetails { Title = $"Blobs {e.Blobs} from object {e.Bucket} {e.Key} in namespace {e.Namespace} did not exist" });
+            }
+
+            return Ok();
+        }
+
+        [HttpPost("{ns}/exists")]
+        [Authorize("Object.read")]
+        [ProducesDefaultResponseType]
+        public async Task<IActionResult> ExistsMultiple(
+            [FromRoute] [Required] NamespaceId ns,
+            [FromRoute] [Required] BucketId bucket,
+            [FromQuery] [Required] List<KeyId> names)
+        {
+            AuthorizationResult authorizationResult = await _authorizationService.AuthorizeAsync(User, ns, NamespaceAccessRequirement.Name);
+
+            if (!authorizationResult.Succeeded)
+            {
+                return Forbid();
+            }
+
+            ConcurrentBag<KeyId> missingObject = new ConcurrentBag<KeyId>();
+
+            IEnumerable<Task> tasks = names.Select(async name =>
+            {
+                try
+                {
+                    await _objectService.Get(ns, bucket, name, new string[] {"blobIdentifier"});
+                }
+                catch (ObjectNotFoundException)
+                {
+                    missingObject.Add(name);
+                }
+            });
+            await Task.WhenAll(tasks);
+
+            return Ok(new { Needs = missingObject.ToArray()});
+        }
+
+
+        [HttpPut("{ns}/{bucket}/{key}.{format?}", Order = 500)]
+        [DisableRequestSizeLimit]
+        [Authorize("Object.write")]
+        public async Task<IActionResult> PutObject(
+            [FromRoute] [Required] NamespaceId ns,
+            [FromRoute] [Required] BucketId bucket,
+            [FromRoute] [Required] KeyId key)
+        {
+            if (ShouldDoAuth())
+            {
+                using (Scope _ = Tracer.Instance.StartActive("authorize"))
+                {
+                    AuthorizationResult authorizationResult = await _authorizationService.AuthorizeAsync(User, ns, NamespaceAccessRequirement.Name);
+
+                    if (!authorizationResult.Succeeded)
+                    {
+                        return Forbid();
+                    }
+                }
+            }
+
+            byte[] blob;
+            try
+            {
+                blob = await RequestUtil.ReadRawBody(Request);
+            }
+            catch (BadHttpRequestException e)
+            {
+                const string msg = "Partial content transfer when reading request body.";
+                _logger.Warning(e, msg);
+                return BadRequest(msg);
+            }
+            _diagnosticContext.Set("Content-Length", Request.ContentLength ?? -1);
+
+            BlobIdentifier headerHash;
+            BlobIdentifier blobHeader = BlobIdentifier.FromBlob(blob);
+            if (Request.Headers.ContainsKey(CommonHeaders.HashHeaderName))
+            {
+                headerHash = new BlobIdentifier(Request.Headers[CommonHeaders.HashHeaderName]);
+            }
+            else
+            {
+                return BadRequest(new ProblemDetails
+                {
+                    Title = $"Missing expected header {CommonHeaders.HashHeaderName}"
+                });
+            }
+
+            if (!blobHeader.Equals(headerHash))
+            {
+                return BadRequest(new ProblemDetails
+                {
+                    Title = $"Incorrect hash, got hash \"{headerHash}\" but hash of content was determined to be \"{blobHeader}\""
+                });
+            }
+
+            CompactBinaryObject payloadObject;
+            switch (Request.ContentType)
+            {
+                case MediaTypeNames.Application.Json:
+                {
+                    // TODO: define a scheme for how a json object specifies references
+
+                    await _blobStore.PutObject(ns, blob, blobHeader);
+
+                    // TODO: convert the json object into a compact binary instead
+                    CompactBinaryWriter compactBinaryWriter = new CompactBinaryWriter();
+                    compactBinaryWriter.BeginObject();
+                    compactBinaryWriter.AddBinaryAttachment(blobHeader);
+                    compactBinaryWriter.EndObject();
+
+                    blob = compactBinaryWriter.Save();
+                    payloadObject = CompactBinaryObject.Load(blob);
+                    blobHeader = BlobIdentifier.FromBlob(blob);
+                    break;
+                }
+                case CustomMediaTypeNames.UnrealCompactBinary:
+                {
+                    payloadObject = CompactBinaryObject.Load(blob);
+                    break;
+                }
+                case MediaTypeNames.Application.Octet:
+                {
+                    await _blobStore.PutObject(ns, blob, blobHeader);
+
+                    CompactBinaryWriter compactBinaryWriter = new CompactBinaryWriter();
+                    compactBinaryWriter.BeginObject();
+                    compactBinaryWriter.AddBinaryAttachment(blobHeader);
+                    compactBinaryWriter.EndObject();
+
+                    blob = compactBinaryWriter.Save();
+                    payloadObject = CompactBinaryObject.Load(blob);
+                    blobHeader = BlobIdentifier.FromBlob(blob);
+                    break;
+                }
+                default:
+                    throw new Exception($"Unknown request type {Request.ContentType}, if submitting a blob please use {MediaTypeNames.Application.Octet}");
+            }
+
+            PutObjectResult result = await _objectService.Put(ns, bucket, key, blobHeader, payloadObject);
+            return Ok(new PutObjectResponse(result.MissingReferences));
+        }
+
+        [HttpPost("{ns}/{bucket}/{key}/finalize/{hash}.{format?}")]
+        [Authorize("Object.write")]
+        public async Task<IActionResult> FinalizeObject(
+            [FromRoute] [Required] NamespaceId ns,
+            [FromRoute] [Required] BucketId bucket,
+            [FromRoute] [Required] KeyId key,
+            [FromRoute] [Required] BlobIdentifier hash)
+        {
+            if (ShouldDoAuth())
+            {
+                using (Scope _ = Tracer.Instance.StartActive("authorize"))
+                {
+                    AuthorizationResult authorizationResult = await _authorizationService.AuthorizeAsync(User, ns, NamespaceAccessRequirement.Name);
+
+                    if (!authorizationResult.Succeeded)
+                    {
+                        return Forbid();
+                    }
+                }
+            }
+
+            BlobIdentifier[] missingReferences = await _objectService.Finalize(ns, bucket, key, hash);
+            return Ok(new PutObjectResponse(missingReferences));
+        }
+
+
+        
+        /// <summary>
+        /// Drop all refs records in the namespace
+        /// </summary>
+        /// <param name="ns">Namespace. Each namespace is completely separated from each other. Use for different types of data that is never expected to be similar (between two different games for instance)</param>
+        [HttpDelete("{ns}", Order = 500)]
+        [ProducesResponseType(204)]
+        [Authorize("admin")]
+        public async Task<IActionResult> DeleteNamespace(
+            [FromRoute] [Required] NamespaceId ns
+        )
+        {
+            AuthorizationResult authorizationResult = await _authorizationService.AuthorizeAsync(User, ns, NamespaceAccessRequirement.Name);
+
+            if (!authorizationResult.Succeeded)
+            {
+                return Forbid();
+            }
+
+            try
+            {
+                await _objectService.DropNamespace(ns);
+            }
+            catch (NamespaceNotFoundException e)
+            {
+                return NotFound(new ProblemDetails {Title = $"Namespace {e.Namespace} did not exist"});
+            }
+
+
+            return NoContent();
+        }
+
+        /// <summary>
+        /// Drop all refs records in the bucket
+        /// </summary>
+        /// <param name="ns">Namespace. Each namespace is completely separated from each other. Use for different types of data that is never expected to be similar (between two different games for instance)</param>
+        /// <param name="bucket">The category/type of record you are caching. Is a clustered key together with the actual key, but all records in the same bucket can be dropped easily.</param>
+        [HttpDelete("{ns}/{bucket}", Order = 500)]
+        [ProducesResponseType(204)]
+        [Authorize("Object.delete")]
+        public async Task<IActionResult> DeleteBucket(
+            [FromRoute] [Required] NamespaceId ns,
+            [FromRoute] [Required] BucketId bucket)
+        {
+            AuthorizationResult authorizationResult = await _authorizationService.AuthorizeAsync(User, ns, NamespaceAccessRequirement.Name);
+
+            if (!authorizationResult.Succeeded)
+            {
+                return Forbid();
+            }
+
+            long countOfDeletedRecords = 0;
+            try
+            {
+                countOfDeletedRecords = await _objectService.DeleteBucket(ns, bucket);
+            }
+            catch (NamespaceNotFoundException e)
+            {
+                return NotFound(new ProblemDetails {Title = $"Namespace {e.Namespace} did not exist"});
+            }
+
+
+            return Ok(new { CountOfDeletedRecords = countOfDeletedRecords});
+        }
+
+        /// <summary>
+        /// Delete a individual refs key
+        /// </summary>
+        /// <param name="ns">Namespace. Each namespace is completely separated from each other. Use for different types of data that is never expected to be similar (between two different games for instance)</param>
+        /// <param name="bucket">The category/type of record you are caching. Is a clustered key together with the actual key, but all records in the same bucket can be dropped easily.</param>
+        /// <param name="key">The unique name of this particular key</param>
+        [HttpDelete("{ns}/{bucket}/{key}", Order = 500)]
+        [ProducesResponseType(204)]
+        [ProducesResponseType(400)]
+        [Authorize("Object.delete")]
+        public async Task<IActionResult> Delete(
+            [FromRoute] [Required] NamespaceId ns,
+            [FromRoute] [Required] BucketId bucket,
+            [FromRoute] [Required] KeyId key)
+        {
+            AuthorizationResult authorizationResult = await _authorizationService.AuthorizeAsync(User, ns, NamespaceAccessRequirement.Name);
+
+            if (!authorizationResult.Succeeded)
+            {
+                return Forbid();
+            }
+
+            try
+            {
+                long deleteCount = await _objectService.Delete(ns, bucket, key);
+                return Ok(new { DeletedCount = deleteCount } );
+            }
+            catch (NamespaceNotFoundException e)
+            {
+                return NotFound(new ProblemDetails {Title = $"Namespace {e.Namespace} did not exist"});
+            }
+        }
+
+        private bool ShouldDoAuth()
+        {
+            foreach (int port in _jupiterSettings.CurrentValue.DisableAuthOnPorts)
+            {
+                if (port == HttpContext.Connection.LocalPort)
+                    return false;
+            }
+
+            return true;
+        }
+    }
+
+    public class PutObjectResponse
+    {
+        public PutObjectResponse(BlobIdentifier[] missingReferences)
+        {
+            Needs = missingReferences;
+        }
+
+        public BlobIdentifier[] Needs { get; set; }
+    }
+}
diff --git a/Engine/Source/Programs/Horde/HordeStorage/Horde.Storage/Controllers/ReplicationController.cs b/Engine/Source/Programs/Horde/HordeStorage/Horde.Storage/Controllers/ReplicationController.cs
new file mode 100644
--- /dev/null
+++ b/Engine/Source/Programs/Horde/HordeStorage/Horde.Storage/Controllers/ReplicationController.cs
@@ -0,0 +1,207 @@
+﻿// Copyright Epic Games, Inc. All Rights Reserved.
+
+using System;
+using System.Collections.Generic;
+using System.ComponentModel.DataAnnotations;
+using System.Linq;
+using System.Threading.Tasks;
+using Horde.Storage.Implementation;
+using Jupiter;
+using Jupiter.Implementation;
+using Microsoft.AspNetCore.Authorization;
+using Microsoft.AspNetCore.Mvc;
+
+namespace Horde.Storage.Controllers
+{
+    [ApiController]
+    [FormatFilter]
+    [Route("api/v1/g")]
+    [Authorize("Admin")]
+    public class ReplicationController : ControllerBase
+    {
+        private readonly ReplicationService _replicationService;
+        private readonly IAuthorizationService _authorizationService;
+
+        public ReplicationController(ReplicationService replicationService, IAuthorizationService authorizationService)
+        {
+            _replicationService = replicationService;
+            _authorizationService = authorizationService;
+        }
+
+
+        [HttpGet("{ns}")]
+        [ProducesDefaultResponseType]
+        public async Task<IActionResult> Get(
+            [Required] NamespaceId ns
+        )
+        {
+            AuthorizationResult authorizationResult = await _authorizationService.AuthorizeAsync(User, ns, NamespaceAccessRequirement.Name);
+
+            if (!authorizationResult.Succeeded)
+            {
+                return Forbid();
+            }
+
+            IEnumerable<IReplicator>? replicators = _replicationService.GetReplicators(ns);
+
+            if (replicators == null)
+            {
+                return BadRequest(new ValidationProblemDetails
+                {
+                    Title = $"No replication configured for namespace {ns}"
+                });
+            }
+
+            return Ok(new 
+            {
+                Replicators = replicators.Select(replicator => new ReplicatorStateResponse
+                {
+                    Name = replicator.Info.ReplicatorName,
+                    Namespace = replicator.Info.NamespaceToReplicate,
+                    Offset = replicator.Info.State.ReplicatorOffset ?? 0,
+                    Generation = replicator.Info.State.ReplicatingGeneration ?? Guid.Empty,
+
+                    LastReplicationRun = replicator.Info.LastRun,
+                    CountOfRunningReplications = replicator.Info.CountOfRunningReplications,
+                })
+            });
+
+        }
+
+
+        [HttpPost("{ns}/{replicatorName}/{offset}")]
+        [ProducesDefaultResponseType]
+        public async Task<IActionResult> SetReplicationOffset(
+            [Required] NamespaceId ns,
+            [Required] string replicatorName,
+            [Required] long offset
+        )
+        {
+            AuthorizationResult authorizationResult = await _authorizationService.AuthorizeAsync(User, ns, NamespaceAccessRequirement.Name);
+
+            if (!authorizationResult.Succeeded)
+            {
+                return Forbid();
+            }
+
+            IEnumerable<IReplicator>? replicators = _replicationService.GetReplicators(ns);
+
+            if (replicators == null)
+            {
+                return BadRequest(new ValidationProblemDetails
+                {
+                    Title = $"No replication configured for namespace {ns}"
+                });
+            }
+
+            IReplicator? replicator = replicators.FirstOrDefault(replicator => string.Equals(replicator.Info.ReplicatorName, replicatorName, StringComparison.InvariantCultureIgnoreCase));
+            if (replicator == null)
+            {
+                return BadRequest(new ValidationProblemDetails
+                {
+                    Title = $"No replication found in namespace {ns} with name {replicatorName}"
+                });
+            }
+
+            replicator.SetReplicationOffset(offset);
+
+            return Ok();
+        }
+
+        [HttpPost("refs/{ns}/{replicatorName}")]
+        [ProducesDefaultResponseType]
+        public async Task<IActionResult> SetReplicationState(
+            [Required] NamespaceId ns,
+            [Required] string replicatorName,
+            [Required] [FromBody] NewReplicationState replicationState
+        )
+        {
+            AuthorizationResult authorizationResult = await _authorizationService.AuthorizeAsync(User, ns, NamespaceAccessRequirement.Name);
+
+            if (!authorizationResult.Succeeded)
+            {
+                return Forbid();
+            }
+
+            IEnumerable<IReplicator>? replicators = _replicationService.GetReplicators(ns);
+
+            if (replicators == null)
+            {
+                return BadRequest(new ValidationProblemDetails
+                {
+                    Title = $"No replication configured for namespace {ns}"
+                });
+            }
+
+            IReplicator? replicator = replicators.FirstOrDefault(replicator => string.Equals(replicator.Info.ReplicatorName, replicatorName, StringComparison.InvariantCultureIgnoreCase));
+            if (replicator == null)
+            {
+                return BadRequest(new ValidationProblemDetails
+                {
+                    Title = $"No replication found in namespace {ns} with name {replicatorName}"
+                });
+            }
+
+            if (replicator is not RefsReplicator refsReplicator)
+            {
+                return BadRequest(new ValidationProblemDetails
+                {
+                    Title = $"Replicator in namespace {ns} with name {replicatorName} was not a refs replicator"
+                });
+            }
+
+            refsReplicator.SetRefState(replicationState.LastBucket, replicationState.LastEvent);
+            return Ok();
+        }
+
+
+        [HttpDelete("{ns}")]
+        [ProducesDefaultResponseType]
+        public async Task<IActionResult> Delete(
+            [Required] NamespaceId ns
+        )
+        {
+            AuthorizationResult authorizationResult = await _authorizationService.AuthorizeAsync(User, ns, NamespaceAccessRequirement.Name);
+
+            if (!authorizationResult.Succeeded)
+            {
+                return Forbid();
+            }
+
+            List<IReplicator> replicators = _replicationService.GetReplicators(ns).ToList();
+
+            if (!replicators.Any())
+            {
+                return BadRequest(new ValidationProblemDetails
+                {
+                    Title = $"No replication configured for namespace {ns}"
+                });
+            }
+
+            foreach (IReplicator replicator in replicators)
+            {
+                await replicator.DeleteState();
+            }
+
+
+            return Ok();
+
+        }
+    }
+
+    public class NewReplicationState
+    {
+        public string? LastBucket { get; set; }
+        public Guid? LastEvent { get; set; }
+    }
+
+    public class ReplicatorStateResponse
+    {
+        public string? Name { get; set; }
+        public NamespaceId? Namespace { get; set; }
+        public long? Offset { get; set; }
+        public Guid? Generation { get; set; }
+        public DateTime LastReplicationRun { get; set; }
+        public int CountOfRunningReplications { get; set; }
+    }
+}
diff --git a/Engine/Source/Programs/Horde/HordeStorage/Horde.Storage/Controllers/ReplicationLogController.cs b/Engine/Source/Programs/Horde/HordeStorage/Horde.Storage/Controllers/ReplicationLogController.cs
new file mode 100644
--- /dev/null
+++ b/Engine/Source/Programs/Horde/HordeStorage/Horde.Storage/Controllers/ReplicationLogController.cs
@@ -0,0 +1,173 @@
+﻿using System;
+using System.Collections.Generic;
+using System.ComponentModel.DataAnnotations;
+using System.Threading;
+using System.Threading.Tasks;
+using async_enumerable_dotnet;
+using Horde.Storage.Implementation;
+using Horde.Storage.Implementation.TransactionLog;
+using Jupiter;
+using Jupiter.Implementation;
+using Microsoft.AspNetCore.Authorization;
+using Microsoft.AspNetCore.Mvc;
+using Microsoft.Extensions.DependencyInjection;
+using Microsoft.Extensions.Options;
+using Newtonsoft.Json;
+
+namespace Horde.Storage.Controllers
+{
+    [ApiController]
+    [Route("api/v1/replication-log")]
+    public class ReplicationLogController : ControllerBase
+    {
+        private readonly IServiceProvider _provider;
+        private readonly IAuthorizationService _authorizationService;
+        private readonly IReplicationLog _replicationLog;
+        private readonly IOptionsMonitor<SnapshotSettings> _snapshotSettings;
+
+        public ReplicationLogController(IServiceProvider provider, IAuthorizationService authorizationService, IReplicationLog replicationLog, IOptionsMonitor<SnapshotSettings> snapshotSettings)
+        {
+            _provider = provider;
+            _authorizationService = authorizationService;
+            _replicationLog = replicationLog;
+            _snapshotSettings = snapshotSettings;
+        }
+
+
+        [HttpGet("snapshots/{ns}")]
+        [ProducesDefaultResponseType]
+        [ProducesResponseType(type: typeof(ProblemDetails), 400)]
+        [Authorize("replication-log.read")]
+        public async Task<IActionResult> GetSnapshots(
+            [Required] NamespaceId ns
+        )
+        {
+            AuthorizationResult authorizationResult = await _authorizationService.AuthorizeAsync(User, ns, NamespaceAccessRequirement.Name);
+
+            if (!authorizationResult.Succeeded)
+            {
+                return Forbid();
+            }
+
+            return Ok(new ReplicationLogSnapshots(await _replicationLog.GetSnapshots(ns).ToListAsync()));
+        }
+
+        
+        [HttpPost("snapshots/{ns}/create")]
+        [ProducesDefaultResponseType]
+        [ProducesResponseType(type: typeof(ProblemDetails), 400)]
+        [Authorize("Admin")]
+        public async Task<IActionResult> CreateSnapshot(
+            [Required] NamespaceId ns
+        )
+        {
+            AuthorizationResult authorizationResult = await _authorizationService.AuthorizeAsync(User, ns, NamespaceAccessRequirement.Name);
+
+            if (!authorizationResult.Succeeded)
+            {
+                return Forbid();
+            }
+
+            ReplicationLogSnapshotBuilder builder = ActivatorUtilities.CreateInstance<ReplicationLogSnapshotBuilder>(_provider);
+            BlobIdentifier snapshotBlob = await builder.BuildSnapshot(ns, _snapshotSettings.CurrentValue.SnapshotStorageNamespace, CancellationToken.None);
+            return Ok(new { SnapshotBlobId = snapshotBlob });
+        }
+
+        [HttpGet("incremental/{ns}")]
+        [ProducesDefaultResponseType]
+        [ProducesResponseType(type: typeof(ProblemDetails), 400)]
+        [Authorize("replication-log.read")]
+        public async Task<IActionResult> GetIncrementalEvents(
+            [Required] NamespaceId ns,
+            [FromQuery] string? lastBucket,
+            [FromQuery] Guid? lastEvent,
+            [FromQuery] int count = 100
+        )
+        {
+            AuthorizationResult authorizationResult = await _authorizationService.AuthorizeAsync(User, ns, NamespaceAccessRequirement.Name);
+
+            if (!authorizationResult.Succeeded)
+            {
+                return Forbid();
+            }
+
+            if ((lastBucket == null && lastEvent.HasValue) || (lastBucket != null && !lastEvent.HasValue))
+            {
+                return BadRequest(new ProblemDetails
+                {
+                    Title = $"Both bucket and event has to be specified, or omit both.",
+                });
+            }
+
+            try
+            {
+                IAsyncEnumerable<ReplicationLogEvent> events = _replicationLog.Get(ns, lastBucket, lastEvent);
+
+                List<ReplicationLogEvent> l = await events.Take(count).ToListAsync();
+                return Ok(new ReplicationLogEvents(l));
+            }
+            catch (IncrementalLogNotAvailableException)
+            {
+                // failed to resume from the incremental log, check for a snapshot instead
+                SnapshotInfo? snapshot = await _replicationLog.GetLatestSnapshot(ns);
+                if (snapshot != null)
+                {
+                    // no log file is available
+                    return BadRequest(new ProblemDetails
+                    {
+                        Title = $"Log file is not available, use snapshot {snapshot.SnapshotBlob} instead",
+                        Type = ProblemTypes.UseSnapshot,
+                        Extensions = { { "SnapshotId", snapshot.SnapshotBlob } }
+                    });
+                }
+
+                // if no snapshot is available we just give up, they can always reset the replication to the default behavior by not sending in lastBucket and lastEvent
+
+                return BadRequest(new ProblemDetails
+                {
+                    Title = $"No snapshot or bucket found for namespace \"{ns}\"",
+                });
+            }
+            catch (NamespaceNotFoundException)
+            {
+                return NotFound(new ProblemDetails
+                {
+                    Title = $"Namespace {ns} was not found",
+                });
+            }
+
+        }
+    }
+
+    public class ReplicationLogSnapshots
+    {
+        public ReplicationLogSnapshots()
+        {
+            Snapshots = new List<SnapshotInfo>();
+        }
+
+        [JsonConstructor]
+        public ReplicationLogSnapshots(List<SnapshotInfo> snapshots)
+        {
+            Snapshots = snapshots;
+        }
+
+        public List<SnapshotInfo> Snapshots { get; set; }
+    }
+
+    public class ReplicationLogEvents
+    {
+        public ReplicationLogEvents()
+        {
+            Events = new List<ReplicationLogEvent>();
+        }
+
+        [JsonConstructor]
+        public ReplicationLogEvents(List<ReplicationLogEvent> events)
+        {
+            Events = events;
+        }
+
+        public List<ReplicationLogEvent> Events { get; set; }
+    }
+}
diff --git a/Engine/Source/Programs/Horde/HordeStorage/Horde.Storage/Controllers/StatusController.cs b/Engine/Source/Programs/Horde/HordeStorage/Horde.Storage/Controllers/StatusController.cs
new file mode 100644
--- /dev/null
+++ b/Engine/Source/Programs/Horde/HordeStorage/Horde.Storage/Controllers/StatusController.cs
@@ -0,0 +1,85 @@
+﻿// Copyright Epic Games, Inc. All Rights Reserved.
+
+using System;
+using System.Collections.Generic;
+using System.ComponentModel.DataAnnotations;
+using System.IO;
+using System.Linq;
+using System.Reflection;
+using System.Threading;
+using System.Threading.Tasks;
+using Datadog.Trace.DuckTyping;
+using Horde.Storage.Implementation;
+using Jupiter.Implementation;
+using Microsoft.AspNetCore.Authorization;
+using Microsoft.AspNetCore.Mvc;
+using Microsoft.Extensions.Options;
+using StatsdClient;
+
+namespace Horde.Storage.Controllers
+{
+    [ApiController]
+    [Route("api/v1/status")]
+    public class StatusController : Controller
+    {
+        private readonly VersionFile _versionFile;
+        private readonly IOptionsMonitor<ReplicationSettings> _replicationSettings;
+
+        public StatusController(VersionFile versionFile, IOptionsMonitor<ReplicationSettings> replicationSettings)
+        {
+            _versionFile = versionFile;
+            _replicationSettings = replicationSettings;
+        }
+
+        /// <summary>
+        /// Fetch information about Jupiter
+        /// </summary>
+        /// <remarks>
+        /// General information about the Jupiter service, which version it is running and more.
+        /// </remarks>
+        /// <returns></returns>
+        [HttpGet("")]
+        [Authorize("Any")]
+        [ProducesResponseType(type: typeof(StatusResponse), 200)]
+        public IActionResult Status()
+        {
+            IEnumerable<AssemblyMetadataAttribute> attrs = typeof(StatusController).Assembly.GetCustomAttributes<AssemblyMetadataAttribute>();
+
+            string srcControlIdentifier = "Unknown";
+            AssemblyMetadataAttribute? gitHashAttribute = attrs.FirstOrDefault(attr => attr.Key == "GitHash");
+            if (gitHashAttribute?.Value != null)
+                srcControlIdentifier = gitHashAttribute.Value;
+
+            AssemblyMetadataAttribute? p4ChangeAttribute = attrs.FirstOrDefault(attr => attr.Key == "PerforceChangelist");
+            if (p4ChangeAttribute?.Value != null)
+                srcControlIdentifier = p4ChangeAttribute.Value;
+            return Ok(new StatusResponse(_versionFile.VersionString ?? "Unknown", srcControlIdentifier, GetCapabilities(), _replicationSettings.CurrentValue.CurrentSite));
+        }
+
+        private string[] GetCapabilities()
+        {
+            return new string[]
+            {
+                "transactionlog",
+                "ddc",
+                "tree"
+            };
+        }
+    }
+
+    public class StatusResponse
+    {
+        public StatusResponse(string version, string gitHash, string[] capabilities, string siteIdentifier)
+        {
+            Version = version;
+            GitHash = gitHash;
+            Capabilities = capabilities;
+            SiteIdentifier = siteIdentifier;
+        }
+
+        public string Version { get; set; }
+        public string GitHash { get; set; }
+        public string[] Capabilities { get; set; }
+        public string SiteIdentifier { get; set; }
+    }
+}
diff --git a/Engine/Source/Programs/Horde/HordeStorage/Horde.Storage/Controllers/StorageController.cs b/Engine/Source/Programs/Horde/HordeStorage/Horde.Storage/Controllers/StorageController.cs
new file mode 100644
--- /dev/null
+++ b/Engine/Source/Programs/Horde/HordeStorage/Horde.Storage/Controllers/StorageController.cs
@@ -0,0 +1,445 @@
+﻿// Copyright Epic Games, Inc. All Rights Reserved.
+
+using System;
+using System.Collections.Concurrent;
+using System.Collections.Generic;
+using System.ComponentModel.DataAnnotations;
+using System.IO;
+using System.Linq;
+using System.Net.Mime;
+using System.Threading.Tasks;
+using Datadog.Trace;
+using Horde.Storage.Implementation;
+using Jupiter;
+using Jupiter.Implementation;
+using Microsoft.AspNetCore.Authorization;
+using Microsoft.AspNetCore.Http;
+using Microsoft.AspNetCore.Mvc;
+using Microsoft.Extensions.Options;
+using Serilog;
+
+namespace Horde.Storage.Controllers
+{
+    [ApiController]
+    [Route("api/v1/s", Order = 1)]
+    [Route("api/v1/blobs", Order = 0)]
+    public class StorageController : ControllerBase
+    {
+        private readonly IBlobStore _storage;
+        private readonly IDiagnosticContext _diagnosticContext;
+        private readonly IAuthorizationService _authorizationService;
+        private readonly HordeStorageSettings _settings;
+
+        private readonly ILogger _logger = Log.ForContext<StorageController>();
+
+        public StorageController(IBlobStore storage, IOptions<HordeStorageSettings> settings, IDiagnosticContext diagnosticContext, IAuthorizationService authorizationService)
+        {
+            _storage = storage;
+            _diagnosticContext = diagnosticContext;
+            _authorizationService = authorizationService;
+            _settings = settings.Value;
+        }
+
+
+        [HttpGet("{ns}/{id}")]
+        [Authorize("Storage.read")]
+        [ProducesDefaultResponseType]
+        public async Task<IActionResult> Get(
+            [Required] NamespaceId ns,
+            [Required] BlobIdentifier id)
+        {
+            AuthorizationResult authorizationResult = await _authorizationService.AuthorizeAsync(User, ns, NamespaceAccessRequirement.Name);
+
+            if (!authorizationResult.Succeeded)
+            {
+                return Forbid();
+            }
+
+            try
+            {
+                BlobContents blobContents = await GetImpl(ns, id);
+
+                return File(blobContents.Stream, MediaTypeNames.Application.Octet, enableRangeProcessing: true);
+            }
+            catch (BlobNotFoundException e)
+            {
+                return NotFound(new ValidationProblemDetails { Title = $"Blob {e.Blob} not found"});
+            }
+        }
+        
+        [HttpHead("{ns}/{id}")]
+        [Authorize("Storage.read")]
+        [ProducesDefaultResponseType]
+        public async Task<IActionResult> Head(
+            [Required] NamespaceId ns,
+            [Required] BlobIdentifier id)
+        {
+            AuthorizationResult authorizationResult = await _authorizationService.AuthorizeAsync(User, ns, NamespaceAccessRequirement.Name);
+
+            if (!authorizationResult.Succeeded)
+            {
+                return Forbid();
+            }
+            bool exists = await _storage.Exists(ns, id);
+
+            if (!exists)
+            {
+                return NotFound(new ValidationProblemDetails { Title = $"Blob {id} not found"});
+            }
+
+            return Ok();
+        }
+
+        [HttpPost("{ns}/exists")]
+        [Authorize("Storage.read")]
+        [ProducesDefaultResponseType]
+        public async Task<IActionResult> ExistsMultiple(
+            [Required] NamespaceId ns,
+             [Required] [FromQuery] List<BlobIdentifier> id)
+        {
+            AuthorizationResult authorizationResult = await _authorizationService.AuthorizeAsync(User, ns, NamespaceAccessRequirement.Name);
+
+            if (!authorizationResult.Succeeded)
+            {
+                return Forbid();
+            }
+
+            ConcurrentBag<BlobIdentifier> missingBlobs = new ConcurrentBag<BlobIdentifier>();
+
+            IEnumerable<Task> tasks = id.Select(async blob =>
+            {
+                if (!await _storage.Exists(ns, blob))
+                    missingBlobs.Add(blob);
+            });
+            await Task.WhenAll(tasks);
+
+            return Ok(new HeadMultipleResponse { Needs = missingBlobs.ToArray()});
+        }
+
+        [HttpPost("{ns}/exist")]
+        [Authorize("Storage.read")]
+        [ProducesDefaultResponseType]
+        public async Task<IActionResult> ExistsBody(
+            [Required] NamespaceId ns,
+            [FromBody] BlobIdentifier[] bodyIds)
+        {
+            AuthorizationResult authorizationResult = await _authorizationService.AuthorizeAsync(User, ns, NamespaceAccessRequirement.Name);
+
+            if (!authorizationResult.Succeeded)
+            {
+                return Forbid();
+            }
+
+            ConcurrentBag<BlobIdentifier> missingBlobs = new ConcurrentBag<BlobIdentifier>();
+
+            IEnumerable<Task> tasks = bodyIds.Select(async blob =>
+            {
+                if (!await _storage.Exists(ns, blob))
+                    missingBlobs.Add(blob);
+            });
+            await Task.WhenAll(tasks);
+
+            return Ok(new HeadMultipleResponse { Needs = missingBlobs.ToArray()});
+        }
+
+        private async Task<BlobContents> GetImpl(NamespaceId ns, BlobIdentifier blob)
+        {
+            return await _storage.GetObject(ns, blob);
+        }
+
+        [HttpPut("{ns}/{id}")]
+        [Authorize("Storage.write")]
+        [RequiredContentType(MediaTypeNames.Application.Octet)]
+        [DisableRequestSizeLimit]
+        public async Task<IActionResult> Put(
+            [Required] NamespaceId ns,
+            [Required] BlobIdentifier id)
+        {
+            AuthorizationResult authorizationResult = await _authorizationService.AuthorizeAsync(User, ns, NamespaceAccessRequirement.Name);
+
+            if (!authorizationResult.Succeeded)
+            {
+                return Forbid();
+            }
+
+            BlobIdentifier identifier;
+            // blob is small enough to fit into memory so we just stream it into memory
+            if (Request.ContentLength is < Int32.MaxValue)
+            {
+                byte[] blob;
+                try
+                {
+                    blob = await RequestUtil.ReadRawBody(Request);
+                }
+                catch (BadHttpRequestException e)
+                {
+                    const string msg = "Partial content transfer when reading request body.";
+                    _logger.Warning(e, msg);
+                    return BadRequest(msg);
+                }
+
+                identifier = await PutImpl(ns, id, blob);
+            }
+            else
+            {
+                FileInfo? tempFile = null;
+
+                Stream s = Request.Body;
+                try
+                {
+                    // stream to a temporary file on disk if the stream is not seekable
+                    if (!Request.Body.CanSeek)
+                    {
+                        tempFile = new FileInfo(Path.GetTempFileName());
+
+                        {
+                            await using FileStream fs = tempFile.OpenWrite();
+                            await Request.Body.CopyToAsync(fs);
+                        }
+
+                        s = tempFile.OpenRead();
+                    }
+                    identifier = await PutImpl(ns, id, s);
+                }
+                finally
+                {
+                    s.Close();
+                    if (tempFile != null && tempFile.Exists)
+                        tempFile.Delete();
+                }
+            }
+            _diagnosticContext.Set("Content-Length", Request.ContentLength ?? -1);
+
+            return Ok(new
+            {
+                Identifier = identifier.ToString()
+            });
+        }
+
+        private async Task<BlobIdentifier> PutImpl(NamespaceId ns, BlobIdentifier id, byte[] content)
+        {
+            BlobIdentifier identifier;
+            {
+                using Scope _ = Tracer.Instance.StartActive("web.hash");
+                identifier = BlobIdentifier.FromBlob(content);
+            }
+
+            if (!id.Equals(identifier))
+            {
+                _logger.Debug("ID {@Id} was not the same as identifier {@Identifier} {Content}", id, identifier,
+                    content);
+
+                throw new ArgumentException("ID was not a hash of the content uploaded.", paramName: nameof(id));
+            }
+
+            await _storage.PutObject(ns, content, identifier);
+            return identifier;
+        }
+
+        private async Task<BlobIdentifier> PutImpl(NamespaceId ns, BlobIdentifier id, Stream content)
+        {
+            BlobIdentifier identifier;
+            {
+                using Scope _ = Tracer.Instance.StartActive("web.hash");
+                identifier = await BlobIdentifier.FromStream(content);
+            }
+
+            if (!id.Equals(identifier))
+            {
+                _logger.Debug("ID {@Id} was not the same as identifier {@Identifier} {Content}", id, identifier,
+                    content);
+
+                throw new ArgumentException("ID was not a hash of the content uploaded.", paramName: nameof(id));
+            }
+
+            // seek back to the beginning
+            content.Position = 0;
+            await _storage.PutObject(ns, content, identifier);
+            return identifier;
+        }
+
+        [HttpDelete("{ns}/{id}")]
+        [Authorize("Storage.delete")]
+        public async Task<IActionResult> Delete(
+            [Required] NamespaceId ns,
+            [Required] BlobIdentifier id)
+        {
+            AuthorizationResult authorizationResult = await _authorizationService.AuthorizeAsync(User, ns, NamespaceAccessRequirement.Name);
+
+            if (!authorizationResult.Succeeded)
+            {
+                return Forbid();
+            }
+
+            await DeleteImpl(ns, id);
+
+            return NoContent();
+        }
+
+        
+        [HttpDelete("{ns}")]
+        [Authorize("Admin")]
+        public async Task<IActionResult> DeleteNamespace(
+            [Required] NamespaceId ns)
+        {
+            AuthorizationResult authorizationResult = await _authorizationService.AuthorizeAsync(User, ns, NamespaceAccessRequirement.Name);
+
+            if (!authorizationResult.Succeeded)
+            {
+                return Forbid();
+            }
+
+            await  _storage.DeleteNamespace(ns);
+
+            return NoContent();
+        }
+
+
+        private async Task DeleteImpl(NamespaceId ns, BlobIdentifier id)
+        {
+            await _storage.DeleteObject(ns, id);
+        }
+
+        // ReSharper disable UnusedAutoPropertyAccessor.Global
+        // ReSharper disable once ClassNeverInstantiated.Global
+        public class BatchOp
+        {
+            // ReSharper disable once InconsistentNaming
+            public enum Operation
+            {
+                INVALID,
+                GET,
+                PUT,
+                DELETE,
+                HEAD
+            }
+
+            [Required] public NamespaceId? Namespace { get; set; }
+
+            public BlobIdentifier? Id { get; set; }
+
+            [Required] public Operation Op { get; set; }
+
+            public byte[]? Content { get; set; }
+        }
+
+        public class BatchCall
+        {
+            public BatchOp[]? Operations { get; set; }
+        }
+        // ReSharper restore UnusedAutoPropertyAccessor.Global
+
+        [HttpPost("")]
+        public async Task<IActionResult> Post([FromBody] BatchCall batch)
+        {
+            string OpToPolicy(BatchOp.Operation op)
+            {
+                switch (op)
+                {
+                    case BatchOp.Operation.GET:
+                    case BatchOp.Operation.HEAD:
+                        return "Storage.read";
+                    case BatchOp.Operation.PUT:
+                        return "Storage.write";
+                    case BatchOp.Operation.DELETE:
+                        return "Storage.delete";
+                    default:
+                        throw new ArgumentOutOfRangeException(nameof(op), op, null);
+                }
+            }
+
+            if (batch?.Operations == null)
+            {
+                throw new ArgumentNullException();
+            }
+
+            Task<object?>[] tasks = new Task<object?>[batch.Operations.Length];
+            for (int index = 0; index < batch.Operations.Length; index++)
+            {
+                BatchOp op = batch.Operations[index];
+                if (op.Namespace == null)
+                {
+                    throw new ArgumentNullException("namespace");
+                }
+
+                AuthorizationResult authorizationResultNamespace = await _authorizationService.AuthorizeAsync(User, op.Namespace, NamespaceAccessRequirement.Name);
+                AuthorizationResult authorizationResultOp = await _authorizationService.AuthorizeAsync(User, OpToPolicy(op.Op));
+
+                if (!authorizationResultNamespace.Succeeded || !authorizationResultOp.Succeeded)
+                {
+                    return Forbid();
+                }
+
+                switch (op.Op)
+                {
+                    case BatchOp.Operation.INVALID:
+                        throw new ArgumentOutOfRangeException();
+                    case BatchOp.Operation.GET:
+                        if (op.Id == null)
+                        {
+                            throw new ArgumentNullException("id");
+                        }
+
+                        tasks[index] = GetImpl(op.Namespace.Value, op.Id).ContinueWith(t =>
+                        {
+                            // TODO: This is very allocation heavy but given that the end result is a json object we can not really stream this anyway
+                            using BlobContents blobContents = t.Result;
+
+                            using MemoryStream ms = new MemoryStream();
+                            blobContents.Stream.CopyTo(ms);
+                            ms.Seek(0, SeekOrigin.Begin);
+                            string str = Convert.ToBase64String(ms.ToArray());
+                            return (object?) str;
+                        });
+                        break;
+                    case BatchOp.Operation.HEAD:
+                        if (op.Id == null)
+                        {
+                            throw new ArgumentNullException("id");
+                        }
+
+                        tasks[index] = _storage.Exists(op.Namespace.Value, op.Id)
+                            .ContinueWith(t => t.Result ? (object?) null : op.Id);
+                        break;
+                    case BatchOp.Operation.PUT:
+                    {
+                        if (op.Content == null)
+                        {
+                            return BadRequest();
+                        }
+
+                        if (op.Id == null)
+                        {
+                            return BadRequest();
+                        }
+
+                        tasks[index] = PutImpl(op.Namespace.Value, op.Id, op.Content).ContinueWith(t => (object?) t.Result);
+                        break;
+                    }
+                    case BatchOp.Operation.DELETE:
+                        if (op.Id == null)
+                        {
+                            throw new ArgumentNullException("id");
+                        }
+
+                        tasks[index] = DeleteImpl(op.Namespace.Value, op.Id).ContinueWith(t => (object?) null);
+                        break;
+                    default:
+                        throw new ArgumentOutOfRangeException();
+                }
+            }
+
+            await Task.WhenAll(tasks);
+
+            object?[] results = tasks.Select(t => t.Result).ToArray();
+
+            return Ok(results);
+        }
+
+    }
+
+    public class HeadMultipleResponse
+    {
+        public BlobIdentifier[] Needs { get; set; } = null!;
+    }
+}
diff --git a/Engine/Source/Programs/Horde/HordeStorage/Horde.Storage/Endpoints/DDCEndpoints.cs b/Engine/Source/Programs/Horde/HordeStorage/Horde.Storage/Endpoints/DDCEndpoints.cs
new file mode 100644
--- /dev/null
+++ b/Engine/Source/Programs/Horde/HordeStorage/Horde.Storage/Endpoints/DDCEndpoints.cs
@@ -0,0 +1,102 @@
+// Copyright Epic Games, Inc. All Rights Reserved.
+
+using System;
+using System.Net.Mime;
+using System.Threading.Tasks;
+using Datadog.Trace;
+using Horde.Storage.Controllers;
+using Horde.Storage.Implementation;
+using Jupiter;
+using Jupiter.Implementation;
+using Microsoft.AspNetCore.Authorization;
+using Microsoft.AspNetCore.Http;
+using Microsoft.AspNetCore.Routing;
+
+namespace Horde.Storage
+{
+    public class DDCEndpoints
+    {
+        private readonly IDDCRefService _ddcRefService;
+        private readonly IAuthorizationService _authorizationService;
+
+        public DDCEndpoints(IDDCRefService ddcRefService, IAuthorizationService authorizationService)
+        {
+            _ddcRefService = ddcRefService;
+            _authorizationService = authorizationService;
+        }
+
+        public string GetRawRoute { get; set; } = "/api/v1/c/ddc/{ns}/{bucket}/{key}.raw";
+
+        public async Task GetRaw(HttpContext httpContext)
+        {
+            string? ns = httpContext.GetRouteValue("ns") as string;
+            if (string.IsNullOrEmpty(ns))
+            {
+                // TODO: Error message
+                return;
+            }
+            
+            using (Scope _ = Tracer.Instance.StartActive("authorize"))
+            {
+                {
+                    AuthorizationResult authorizationResult = await _authorizationService.AuthorizeAsync(httpContext.User, ns, NamespaceAccessRequirement.Name);
+
+                    if (!authorizationResult.Succeeded)
+                    {
+                        httpContext.Response.StatusCode = StatusCodes.Status401Unauthorized;
+                        return;
+                    }
+                }
+            }
+
+            string? bucket = httpContext.GetRouteValue("bucket") as string;
+            if (string.IsNullOrEmpty(bucket))
+            {
+                // TODO: Error message
+                return;
+            }
+
+            string? key = httpContext.GetRouteValue("key") as string;
+            if (string.IsNullOrEmpty(key))
+            {
+                // TODO: Error message
+                return;
+            }
+
+            httpContext.Response.Headers["X-Jupiter-ImplUsed"] = "Endpoint";
+            try
+            {
+                (RefResponse refRes, BlobContents? tempBlob) = await _ddcRefService.Get(new NamespaceId(ns), new BucketId(bucket), new KeyId(key), fields: new[] {"blob"});
+                BlobContents blob = tempBlob!;
+
+                httpContext.Response.StatusCode = StatusCodes.Status200OK;
+                httpContext.Response.Headers["Content-Type"] = MediaTypeNames.Application.Octet;
+                httpContext.Response.Headers["Content-Length"] = Convert.ToString(blob.Length);
+                httpContext.Response.Headers[CommonHeaders.HashHeaderName] = refRes.ContentHash.ToString();
+
+                using (Scope _ = Tracer.Instance.StartActive("body.write"))
+                {
+                    await blob.Stream.CopyToAsync(httpContext.Response.Body);
+                }
+            }
+            catch (NamespaceNotFoundException e)
+            {
+                httpContext.Response.StatusCode = StatusCodes.Status400BadRequest;
+                httpContext.Response.Headers["Content-Type"] = "application/json";
+                await httpContext.Response.WriteAsync("{\"title\": \"Namespace " + e.Namespace + " did not exist\", \"status\": 400}");
+            }
+            catch (RefRecordNotFoundException e)
+            {
+                httpContext.Response.StatusCode = StatusCodes.Status400BadRequest;
+                httpContext.Response.Headers["Content-Type"] = "application/json";
+                await httpContext.Response.WriteAsync("{\"title\": \"Object " + e.Bucket + " " + e.Key + " did not exist\", \"status\": 400}");
+            }
+            catch (BlobNotFoundException e)
+            {
+                httpContext.Response.StatusCode = StatusCodes.Status400BadRequest;
+                httpContext.Response.Headers["Content-Type"] = "application/json";
+                await httpContext.Response.WriteAsync("{\"title\": \"Object " + e.Blob + " in " + e.Ns + " not found\", \"status\": 400}");
+            }
+        }
+    }
+}
diff --git a/Engine/Source/Programs/Horde/HordeStorage/Horde.Storage/Horde.Storage.Startup.cs b/Engine/Source/Programs/Horde/HordeStorage/Horde.Storage/Horde.Storage.Startup.cs
new file mode 100644
--- /dev/null
+++ b/Engine/Source/Programs/Horde/HordeStorage/Horde.Storage/Horde.Storage.Startup.cs
@@ -0,0 +1,667 @@
+// Copyright Epic Games, Inc. All Rights Reserved.
+
+using System;
+using System.Collections.Generic;
+using System.IO;
+using System.Linq;
+using Amazon.DAX;
+using Amazon.DynamoDBv2;
+using Amazon.Extensions.NETCore.Setup;
+using Amazon.Runtime;
+using Amazon.S3;
+using Amazon.SecretsManager;
+using Cassandra;
+using Horde.Storage.Controllers;
+using Horde.Storage.Implementation;
+using Horde.Storage.Implementation.LeaderElection;
+using Jupiter;
+using Jupiter.Common.Implementation;
+using Jupiter.Utils;
+using Microsoft.AspNetCore.Authorization;
+using Microsoft.AspNetCore.Builder;
+using Microsoft.AspNetCore.Hosting;
+using Microsoft.AspNetCore.Routing;
+using Microsoft.Extensions.Configuration;
+using Microsoft.Extensions.DependencyInjection;
+using Microsoft.Extensions.Options;
+using StatsdClient;
+
+namespace Horde.Storage
+{
+    // ReSharper disable once ClassNeverInstantiated.Global
+    public class HordeStorageStartup : BaseStartup
+    {
+        public HordeStorageStartup(IConfiguration configuration, IWebHostEnvironment environment) : base(configuration,
+            environment)
+        {
+            string? ddAgentHost = System.Environment.GetEnvironmentVariable("DD_AGENT_HOST");
+            if (!string.IsNullOrEmpty(ddAgentHost))
+            {
+                _logger.Information("Initializing Dogstatsd to connect to: {DatadogAgentHost}", ddAgentHost);
+                StatsdConfig dogstatsdConfig = new StatsdConfig
+                {
+                    StatsdServerName = ddAgentHost,
+                    StatsdPort = 8125,
+                };
+
+                DogStatsd.Configure(dogstatsdConfig);
+            }
+        }
+
+        protected override void OnAddAuthorization(AuthorizationOptions authorizationOptions, List<string> defaultSchemes)
+        {
+            authorizationOptions.AddPolicy("Cache.read", policy =>
+            {
+                policy.AuthenticationSchemes = defaultSchemes;
+                policy.RequireClaim("Cache", "read", "readwrite", "full");
+            });
+
+            authorizationOptions.AddPolicy("Cache.write", policy =>
+            {
+                policy.AuthenticationSchemes = defaultSchemes;
+                policy.RequireClaim("Cache", "write", "readwrite", "full");
+            });
+
+            authorizationOptions.AddPolicy("Cache.delete", policy =>
+            {
+                policy.AuthenticationSchemes = defaultSchemes;
+                policy.RequireClaim("Cache", "delete", "full");
+            });
+
+            authorizationOptions.AddPolicy("Object.read", policy =>
+            {
+                policy.AuthenticationSchemes = defaultSchemes;
+                policy.RequireClaim("Cache", "read", "readwrite", "full");
+            });
+
+            authorizationOptions.AddPolicy("Object.write", policy =>
+            {
+                policy.AuthenticationSchemes = defaultSchemes;
+                policy.RequireClaim("Cache", "write", "readwrite", "full");
+            });
+
+            authorizationOptions.AddPolicy("Object.delete", policy =>
+            {
+                policy.AuthenticationSchemes = defaultSchemes;
+                policy.RequireClaim("Cache", "delete", "full");
+            });
+
+            authorizationOptions.AddPolicy("Storage.read", policy =>
+            {
+                policy.AuthenticationSchemes = defaultSchemes;
+                policy.RequireClaim("Storage", "read", "readwrite", "full");
+            });
+
+            authorizationOptions.AddPolicy("Storage.write", policy =>
+            {
+                policy.AuthenticationSchemes = defaultSchemes;
+                policy.RequireClaim("Storage", "write", "readwrite", "full");
+            });
+
+            authorizationOptions.AddPolicy("Storage.delete", policy =>
+            {
+                policy.AuthenticationSchemes = defaultSchemes;
+                policy.RequireClaim("Storage", "delete", "full");
+            });
+
+            authorizationOptions.AddPolicy("replication-log.read", policy =>
+            {
+                policy.AuthenticationSchemes = defaultSchemes;
+                policy.RequireClaim("transactionLog", "read", "readwrite", "full");
+            });
+
+            authorizationOptions.AddPolicy("replication-log.write", policy =>
+            {
+                policy.AuthenticationSchemes = defaultSchemes;
+                policy.RequireClaim("transactionLog", "write", "readwrite", "full");
+            });
+
+
+            authorizationOptions.AddPolicy("Admin", policy =>
+            {
+                policy.AuthenticationSchemes = defaultSchemes;
+                policy.RequireClaim("Admin");
+            });
+
+            // A policy that grants any authenticated user access
+            authorizationOptions.AddPolicy("Any", policy =>
+            {
+                policy.RequireAuthenticatedUser();
+            });
+        }
+
+        protected override void OnUseEndpoints(IWebHostEnvironment env, IEndpointRouteBuilder endpoints)
+        {
+            HordeStorageSettings settings = endpoints.ServiceProvider.GetService<IOptionsMonitor<HordeStorageSettings>>()!.CurrentValue!;
+
+            if (settings.UseNewDDCEndpoints)
+            {
+                DDCEndpoints ddcEndpoints = endpoints.ServiceProvider.GetService<DDCEndpoints>()!;
+
+                endpoints.Map(ddcEndpoints.GetRawRoute, ddcEndpoints.GetRaw).RequireAuthorization("Cache.read");
+            }
+        }
+
+        protected override void OnAddService(IServiceCollection services)
+        {
+            // For requests served on the high-perf HTTP port there is not authn/authz.
+            // Suppressing the check below removes the warning for that case.
+            services.Configure<RouteOptions>(options =>
+            {
+                options.SuppressCheckForUnhandledSecurityMetadata = true;
+            });
+
+            services.AddHttpClient();
+
+            services.AddOptions<HordeStorageSettings>().Bind(Configuration.GetSection("Horde.Storage")).ValidateDataAnnotations();
+            services.AddOptions<MongoSettings>().Bind(Configuration.GetSection("Mongo")).ValidateDataAnnotations();
+            services.AddOptions<CosmosSettings>().Bind(Configuration.GetSection("Cosmos")).ValidateDataAnnotations();
+            services.AddOptions<DynamoDbSettings>().Bind(Configuration.GetSection("DynamoDb")).ValidateDataAnnotations();
+
+            services.AddOptions<CallistoTransactionLogSettings>().Bind(Configuration.GetSection("Callisto")).ValidateDataAnnotations();
+
+            services.AddOptions<MemoryCacheBlobSettings>().Bind(Configuration.GetSection("Cache.Blob")).ValidateDataAnnotations();
+            services.AddOptions<MemoryCacheRefSettings>().Bind(Configuration.GetSection("Cache.Db")).ValidateDataAnnotations();
+
+            services.AddOptions<S3Settings>().Bind(Configuration.GetSection("S3")).ValidateDataAnnotations();
+            services.AddOptions<AzureSettings>().Bind(Configuration.GetSection("Azure")).ValidateDataAnnotations();
+            services.AddOptions<FilesystemSettings>().Bind(Configuration.GetSection("Filesystem")).ValidateDataAnnotations();
+
+            services.AddOptions<ConsistencyCheckSettings>().Bind(Configuration.GetSection("ConsistencyCheck")).ValidateDataAnnotations();
+
+            services.AddOptions<GCSettings>().Configure(o => Configuration.GetSection("GC").Bind(o)).ValidateDataAnnotations();
+            services.AddOptions<ReplicationSettings>().Configure(o => Configuration.GetSection("Replication").Bind(o)).ValidateDataAnnotations();
+            services.AddOptions<ServiceCredentialSettings>().Configure(o => Configuration.GetSection("ServiceCredentials").Bind(o)).ValidateDataAnnotations();
+            services.AddOptions<SnapshotSettings>().Configure(o => Configuration.GetSection("Snapshot").Bind(o)).ValidateDataAnnotations();
+
+            services.AddOptions<ScyllaSettings>().Configure(o => Configuration.GetSection("Scylla").Bind(o)).ValidateDataAnnotations();
+
+            services.AddOptions<KubernetesLeaderElectionSettings>().Configure(o => Configuration.GetSection("Kubernetes").Bind(o)).ValidateDataAnnotations();
+
+
+            services.AddSingleton(typeof(OodleCompressor), CreateOodleCompressor);
+            services.AddSingleton(typeof(CompressedBufferUtils), CreateCompressedBufferUtils);
+
+            services.AddSingleton<AWSCredentials>(provider =>
+            {
+                AWSCredentialsSettings awsSettings = provider.GetService<IOptionsMonitor<AWSCredentialsSettings>>()!.CurrentValue;
+
+                return AWSCredentialsHelper.GetCredentials(awsSettings, "Horde.Storage");
+            });
+            services.AddSingleton(typeof(IAmazonDynamoDB), AddDynamo);
+            
+            services.AddSingleton<BlobCleanupService>();
+            services.AddHostedService<BlobCleanupService>(p => p.GetService<BlobCleanupService>()!);
+
+            services.AddSingleton(serviceType: typeof(IDDCRefService), typeof(DDCRefService));
+            services.AddSingleton(serviceType: typeof(DDCEndpoints));
+
+            services.AddSingleton(serviceType: typeof(IObjectService), typeof(ObjectService));
+            services.AddSingleton(serviceType: typeof(IReferencesStore), ObjectStoreFactory);
+            services.AddSingleton(serviceType: typeof(IReferenceResolver), typeof(ReferenceResolver));
+            services.AddSingleton(serviceType: typeof(IContentIdStore), ContentIdStoreFactory);
+
+            services.AddSingleton(serviceType: typeof(ITransactionLogWriter), TransactionLogWriterFactory);
+            services.AddSingleton(serviceType: typeof(IAmazonS3), CreateS3);
+            services.AddSingleton(serviceType: typeof(IBlobStore), StorageBackendFactory);
+
+            services.AddSingleton<AmazonS3Store>();
+            services.AddSingleton<FileSystemStore>();
+            services.AddSingleton<AzureBlobStore>();
+            services.AddSingleton<MemoryCacheBlobStore>();
+
+            services.AddSingleton(serviceType: typeof(IScyllaSessionManager), ScyllaFactory);
+
+            services.AddSingleton(serviceType: typeof(IReplicationLog), ReplicationLogWriterFactory);
+            
+            LastAccessTrackerRefRecord lastAccessTracker = new LastAccessTrackerRefRecord();
+            services.AddSingleton(serviceType: typeof(ILastAccessCache<RefRecord>), lastAccessTracker);
+            services.AddSingleton(serviceType: typeof(ILastAccessTracker<RefRecord>), lastAccessTracker);
+
+            services.AddSingleton(serviceType: typeof(ILeaderElection), CreateLeaderElection);
+
+            services.AddSingleton(Configuration);
+            services.AddSingleton<DynamoNamespaceStore>();
+            services.AddSingleton(RefStoreFactory);
+
+            services.AddSingleton<FormatResolver>();
+
+            services.AddSingleton(serviceType: typeof(ISecretResolver), typeof(SecretResolver));
+            services.AddSingleton(typeof(IAmazonSecretsManager), CreateAWSSecretsManager);
+
+            services.AddSingleton<LastAccessService>();
+            services.AddHostedService<LastAccessService>(p => p.GetService<LastAccessService>()!);
+
+            services.AddSingleton<IServiceCredentials, ServiceCredentials>(p => ActivatorUtilities.CreateInstance<ServiceCredentials>(p));
+            
+            services.AddSingleton<ReplicationService>();
+            services.AddHostedService<ReplicationService>(p => p.GetService<ReplicationService>()!);
+
+            services.AddSingleton<ReplicationSnapshotService>();
+            services.AddHostedService<ReplicationSnapshotService>(p => p.GetService<ReplicationSnapshotService>()!);
+
+            services.AddSingleton<ConsistencyCheckService>();
+            services.AddHostedService<ConsistencyCheckService>(p => p.GetService<ConsistencyCheckService>()!);
+        
+            services.AddTransient(typeof(IRefCleanup), typeof(RefCleanup));
+
+            services.AddTransient(typeof(VersionFile), typeof(VersionFile));
+
+            services.AddSingleton<RefCleanupService>();
+            services.AddHostedService<RefCleanupService>(p => p.GetService<RefCleanupService>()!);
+
+            // This will technically create a new settings object, but as we do not dynamically update it the settings will reflect what we actually want
+            HordeStorageSettings settings = services.BuildServiceProvider().GetService<IOptionsMonitor<HordeStorageSettings>>()!.CurrentValue!;
+
+            if (settings.LeaderElectionImplementation == HordeStorageSettings.LeaderElectionImplementations.Kubernetes)
+            {
+                // add the kubernetes leader instance under its actual type as well to make it easier to find
+                services.AddSingleton<KubernetesLeaderElection>(p => (KubernetesLeaderElection)p.GetService<ILeaderElection>()!);
+                services.AddHostedService<KubernetesLeaderElection>(p => p.GetService<KubernetesLeaderElection>()!);
+            }
+        }
+
+        private IAmazonSecretsManager CreateAWSSecretsManager(IServiceProvider provider)
+        {
+            AWSCredentials awsCredentials = provider.GetService<AWSCredentials>()!;
+            AWSOptions awsOptions = Configuration.GetAWSOptions();
+            awsOptions.Credentials = awsCredentials;
+
+            IAmazonSecretsManager serviceClient = awsOptions.CreateServiceClient<IAmazonSecretsManager>();
+            return serviceClient;
+        }
+
+        private CompressedBufferUtils CreateCompressedBufferUtils(IServiceProvider provider)
+        {
+            return ActivatorUtilities.CreateInstance<CompressedBufferUtils>(provider);
+        }
+
+        private object CreateOodleCompressor(IServiceProvider provider)
+        {
+            OodleCompressor compressor = new OodleCompressor();
+            compressor.InitializeOodle();
+            return compressor;
+        }
+
+        private object ContentIdStoreFactory(IServiceProvider provider)
+        {
+            HordeStorageSettings settings = provider.GetService<IOptionsMonitor<HordeStorageSettings>>()!.CurrentValue!;
+            switch (settings.ContentIdStoreImplementation)
+            {
+                case HordeStorageSettings.ContentIdStoreImplementations.Memory:
+                    return ActivatorUtilities.CreateInstance<MemoryContentIdStore>(provider);
+                case HordeStorageSettings.ContentIdStoreImplementations.Scylla:
+                    return ActivatorUtilities.CreateInstance<ScyllaContentIdStore>(provider);
+                default:
+                    throw new ArgumentOutOfRangeException();
+            }
+        }
+
+        private IScyllaSessionManager ScyllaFactory(IServiceProvider provider)
+        {
+            ScyllaSettings settings = provider.GetService<IOptionsMonitor<ScyllaSettings>>()!.CurrentValue!;
+
+            Serilog.Extensions.Logging.SerilogLoggerProvider serilogLoggerProvider = new();
+            Diagnostics.AddLoggerProvider(serilogLoggerProvider);
+
+            const string DefaultKeyspaceName = "jupiter";
+            // Configure the builder with your cluster's contact points
+            var cluster = Cluster.Builder()
+                .AddContactPoints(settings.ContactPoints)
+                .WithLoadBalancingPolicy(new DefaultLoadBalancingPolicy(settings.LocalDatacenterName))
+                .WithPoolingOptions(PoolingOptions.Create().SetMaxConnectionsPerHost(HostDistance.Local, settings.MaxConnectionForLocalHost))
+                .WithDefaultKeyspace(DefaultKeyspaceName)
+                .WithExecutionProfiles(options => 
+                    options.WithProfile("default", builder => builder.WithConsistencyLevel(ConsistencyLevel.LocalOne)))
+                .Build();
+
+            Dictionary<string, string> replicationStrategy = ReplicationStrategies.CreateSimpleStrategyReplicationProperty(2);
+            if (settings.KeyspaceReplicationStrategy != null)
+            {
+                replicationStrategy = settings.KeyspaceReplicationStrategy;
+                _logger.Information("Scylla Replication strategy for replicated keyspace is set to {@ReplicationStrategy}", replicationStrategy);
+            }
+            ISession replicatedSession = cluster.ConnectAndCreateDefaultKeyspaceIfNotExists(replicationStrategy);
+
+            replicatedSession.Execute(new SimpleStatement("CREATE TYPE IF NOT EXISTS blob_identifier (hash blob)"));
+            replicatedSession.UserDefinedTypes.Define(UdtMap.For<ScyllaBlobIdentifier>("blob_identifier", DefaultKeyspaceName));
+
+            string localKeyspaceName = $"jupiter_local_{settings.LocalKeyspaceSuffix}";
+
+            Dictionary<string, string> replicationStrategyLocal = ReplicationStrategies.CreateSimpleStrategyReplicationProperty(2);
+            if (settings.LocalKeyspaceReplicationStrategy != null)
+            {
+                replicationStrategyLocal = settings.LocalKeyspaceReplicationStrategy;
+                _logger.Information("Scylla Replication strategy for local keyspace is set to {@ReplicationStrategy}", replicationStrategyLocal);
+            }
+
+            // the local keyspace is never replicated so we do not support controlling how the replication strategy is setup
+            replicatedSession.CreateKeyspaceIfNotExists(localKeyspaceName, replicationStrategyLocal);
+            ISession localSession = cluster.Connect(localKeyspaceName);
+
+            localSession.Execute(new SimpleStatement("CREATE TYPE IF NOT EXISTS blob_identifier (hash blob)"));
+            localSession.UserDefinedTypes.Define(UdtMap.For<ScyllaBlobIdentifier>("blob_identifier", localKeyspaceName));
+
+            return new ScyllaSessionManager(replicatedSession, localSession);
+        }
+
+        private IReferencesStore ObjectStoreFactory(IServiceProvider provider)
+        {
+            HordeStorageSettings settings = provider.GetService<IOptionsMonitor<HordeStorageSettings>>()!.CurrentValue!;
+            switch (settings.ReferencesDbImplementation)
+            {
+                case HordeStorageSettings.ReferencesDbImplementations.Memory:
+                    return ActivatorUtilities.CreateInstance<MemoryReferencesStore>(provider);
+                case HordeStorageSettings.ReferencesDbImplementations.Scylla:
+                    return ActivatorUtilities.CreateInstance<ScyllaReferencesStore>(provider);
+                default:
+                    throw new ArgumentOutOfRangeException();
+            }
+        }
+
+        protected override void OnConfigureAppEarly(IApplicationBuilder app, IWebHostEnvironment env)
+        {
+            HordeStorageSettings settings = app.ApplicationServices.GetService<IOptionsMonitor<HordeStorageSettings>>()!.CurrentValue!;
+            // Register two methods as middlewares for benchmarking purposes
+            // By circumventing the full ASP.NET stack we can measure the performance difference between these calls.
+            if (!settings.UseNewDDCEndpoints)
+            {
+                DDCRefController ddcRefController = ActivatorUtilities.CreateInstance<DDCRefController>(app.ApplicationServices);
+                DebugController debugController = ActivatorUtilities.CreateInstance<DebugController>(app.ApplicationServices);
+                app.Use(ddcRefController.FastGet);
+                app.Use(debugController.FastGetBytes);
+            }
+        }
+
+        private ILeaderElection CreateLeaderElection(IServiceProvider provider)
+        {
+            HordeStorageSettings settings = provider.GetService<IOptionsMonitor<HordeStorageSettings>>()!.CurrentValue!;
+            if (settings.LeaderElectionImplementation == HordeStorageSettings.LeaderElectionImplementations.Kubernetes)
+            {
+                return ActivatorUtilities.CreateInstance<KubernetesLeaderElection>(provider);
+            }
+            else
+            {
+                // hard coded leader election that assumes it is always the leader
+                return ActivatorUtilities.CreateInstance<LeaderElectionStub>(provider, true);
+            }
+        }
+
+        private IAmazonS3 CreateS3(IServiceProvider provider)
+        {
+            HordeStorageSettings settings = provider.GetService<IOptionsMonitor<HordeStorageSettings>>()!.CurrentValue!;
+
+            if (settings.StorageImplementations == null)
+            {
+                throw new ArgumentException("No storage implementation set");
+            }
+            
+            bool isS3InUse = settings.StorageImplementations.Any(x =>
+                String.Equals(x, HordeStorageSettings.StorageBackendImplementations.S3.ToString(), StringComparison.CurrentCultureIgnoreCase));
+
+            if (isS3InUse)
+            {
+                S3Settings s3Settings = provider.GetService<IOptionsMonitor<S3Settings>>()!.CurrentValue!;
+                AWSCredentials awsCredentials = provider.GetService<AWSCredentials>()!;
+                AWSOptions awsOptions = Configuration.GetAWSOptions();
+                if (s3Settings.ConnectionString.ToUpper() != "AWS")
+                {
+                    awsOptions.DefaultClientConfig.ServiceURL = s3Settings.ConnectionString;
+                }
+
+                awsOptions.Credentials = awsCredentials;
+
+                IAmazonS3 serviceClient = awsOptions.CreateServiceClient<IAmazonS3>();
+
+                if (serviceClient.Config is AmazonS3Config c)
+                {
+                    if (s3Settings.ForceAWSPathStyle)
+                    {
+                        c.ForcePathStyle = true;
+                    }
+
+                }
+                else
+                {
+                    throw new NotImplementedException();
+                }
+
+                return serviceClient;
+            }
+
+            return null!;
+        }
+
+        private IAmazonDynamoDB AddDynamo(IServiceProvider provider)
+        {
+            IOptionsMonitor<HordeStorageSettings> settings = provider.GetService<IOptionsMonitor<HordeStorageSettings>>()!;
+
+            bool hasDynamo = settings.CurrentValue.RefDbImplementation == HordeStorageSettings.RefDbImplementations.DynamoDb ||
+                             settings.CurrentValue.TreeRootStoreImplementation == HordeStorageSettings.TreeRootStoreImplementations.DynamoDb ||
+                             settings.CurrentValue.TreeStoreImplementation == HordeStorageSettings.TreeStoreImplementations.DynamoDb;
+            if (!hasDynamo)
+            {
+                return null!;
+            }
+
+            AWSCredentials awsCredentials = provider.GetService<AWSCredentials>()!;
+
+            DynamoDbSettings dynamoDbSettings = provider.GetService<IOptionsMonitor<DynamoDbSettings>>()!.CurrentValue;
+
+            AWSOptions awsOptions = Configuration.GetAWSOptions();
+            if (dynamoDbSettings.ConnectionString.ToUpper() != "AWS")
+                awsOptions.DefaultClientConfig.ServiceURL = dynamoDbSettings.ConnectionString;
+
+            awsOptions.Credentials = awsCredentials;
+
+            IAmazonDynamoDB serviceClient = awsOptions.CreateServiceClient<IAmazonDynamoDB>();
+
+            if (!string.IsNullOrEmpty(dynamoDbSettings.DaxEndpoint))
+            {
+                (string daxHost, int daxPort) = DynamoDbSettings.ParseDaxEndpointAsHostPort(dynamoDbSettings.DaxEndpoint);
+                DaxClientConfig daxConfig = new (daxHost, daxPort)
+                {
+                    AwsCredentials = awsOptions.Credentials
+                };
+                
+                serviceClient = new ClusterDaxClient(daxConfig);
+            }
+
+            return serviceClient;
+        }
+
+        private ITransactionLogWriter TransactionLogWriterFactory(IServiceProvider provider)
+        {
+            HordeStorageSettings settings = provider.GetService<IOptionsMonitor<HordeStorageSettings>>()!.CurrentValue;
+            switch (settings.TransactionLogWriterImplementation)
+            {
+                case HordeStorageSettings.TransactionLogWriterImplementations.Callisto:
+                    return ActivatorUtilities.CreateInstance<CallistoTransactionLogWriter>(provider);
+                case HordeStorageSettings.TransactionLogWriterImplementations.Memory:
+                    return ActivatorUtilities.CreateInstance<MemoryTransactionLogWriter>(provider);
+                default:
+                    throw new NotImplementedException();
+            }
+        }
+
+        private IReplicationLog ReplicationLogWriterFactory(IServiceProvider provider)
+        {
+            HordeStorageSettings settings = provider.GetService<IOptionsMonitor<HordeStorageSettings>>()!.CurrentValue;
+            switch (settings.ReplicationLogWriterImplementation)
+            {
+                case HordeStorageSettings.ReplicationLogWriterImplementations.Scylla:
+                    return ActivatorUtilities.CreateInstance<ScyllaReplicationLog>(provider);
+                case HordeStorageSettings.ReplicationLogWriterImplementations.Memory:
+                    return ActivatorUtilities.CreateInstance<MemoryReplicationLog>(provider);
+                default:
+                    throw new NotImplementedException();
+            }
+        }
+
+        private HordeStorageSettings.StorageBackendImplementations[] GetStorageImplementationEnums(string[]? storageImpls)
+        {
+            storageImpls ??= new[] {HordeStorageSettings.StorageBackendImplementations.Memory.ToString()};
+            return storageImpls.Select(x =>
+            {
+                if (Enum.TryParse(typeof(HordeStorageSettings.StorageBackendImplementations), x, true, out var backendType) && backendType is HordeStorageSettings.StorageBackendImplementations type)
+                {
+                    return type;
+                }
+
+                throw new ArgumentException($"Unable to find storage implementation for specified value '{x}'");
+            }).ToArray();
+        }
+        
+        private IBlobStore StorageBackendFactory(IServiceProvider provider)
+        {
+            IOptionsMonitor<HordeStorageSettings> settings = provider.GetService<IOptionsMonitor<HordeStorageSettings>>()!;
+            IOptionsMonitor<GCSettings> gcSettings = provider.GetService<IOptionsMonitor<GCSettings>>()!;
+            BlobCleanupService blobCleanupService = provider.GetService<BlobCleanupService>()!;
+
+            List<IBlobStore> storageBackends = GetStorageImplementationEnums(settings.CurrentValue.StorageImplementations)
+                    .Select<HordeStorageSettings.StorageBackendImplementations, IBlobStore?>(impl =>
+                    {
+                        switch (impl)
+                        {
+                            case HordeStorageSettings.StorageBackendImplementations.S3: return provider.GetService<AmazonS3Store>();
+                            case HordeStorageSettings.StorageBackendImplementations.Azure: return provider.GetService<AzureBlobStore>();
+                            case HordeStorageSettings.StorageBackendImplementations.Memory: return provider.GetService<MemoryCacheBlobStore>();
+                            case HordeStorageSettings.StorageBackendImplementations.FileSystem:
+                                FileSystemStore store = ActivatorUtilities.CreateInstance<FileSystemStore>(provider)!;
+                                blobCleanupService.RegisterCleanup(store);
+                                return store;
+                            default: throw new NotImplementedException();
+                        }
+                    })
+                    .Select(x => x ?? throw new ArgumentException("Unable to resolve all stores!"))
+                    .ToList();
+
+            // Only use the hierarchical backend if more than one backend is specified 
+            IBlobStore blobStore = storageBackends.Count == 1 ? storageBackends[0] : new HierarchicalBlobStore(storageBackends);
+            
+            if (gcSettings.CurrentValue.CleanOldBlobs)
+            {
+                OrphanBlobCleanup orphanBlobCleanup = ActivatorUtilities.CreateInstance<OrphanBlobCleanup>(provider, blobStore);
+                blobCleanupService.RegisterCleanup(orphanBlobCleanup);    
+            }
+
+            return blobStore;
+        }
+
+        private static IRefsStore RefStoreFactory(IServiceProvider provider)
+        {
+            HordeStorageSettings settings = provider.GetService<IOptionsMonitor<HordeStorageSettings>>()!.CurrentValue;
+
+            IRefsStore refsStore = settings.RefDbImplementation switch
+            {
+                HordeStorageSettings.RefDbImplementations.DynamoDb => ActivatorUtilities.CreateInstance<DynamoDbRefsStore>(provider),
+                HordeStorageSettings.RefDbImplementations.Mongo => ActivatorUtilities.CreateInstance<MongoRefsStore>(provider),
+                HordeStorageSettings.RefDbImplementations.Cosmos => ActivatorUtilities.CreateInstance<CosmosRefsStore>(provider),
+                _ => throw new NotImplementedException()
+            };
+
+            MemoryCacheRefSettings memoryCacheSettings = provider.GetService<IOptionsMonitor<MemoryCacheRefSettings>>()!.CurrentValue;
+
+            if (memoryCacheSettings.Enabled)
+            {
+                refsStore = ActivatorUtilities.CreateInstance<CachedRefStore>(provider, refsStore);
+            }
+
+            return refsStore;
+        }
+
+        protected override void OnAddHealthChecks(IServiceCollection services, IHealthChecksBuilder healthChecks)
+        {
+            ServiceProvider provider = services.BuildServiceProvider();
+            HordeStorageSettings settings = provider.GetService<IOptionsMonitor<HordeStorageSettings>>()!.CurrentValue;
+
+            switch (settings.RefDbImplementation)
+            {
+                case HordeStorageSettings.RefDbImplementations.Mongo:
+                case HordeStorageSettings.RefDbImplementations.Cosmos:
+                    MongoSettings mongoSettings = provider.GetService<IOptionsMonitor<MongoSettings>>()!.CurrentValue;
+                    healthChecks.AddMongoDb(mongoSettings.ConnectionString, tags: new[] {"services"});
+                    break;
+                case HordeStorageSettings.RefDbImplementations.DynamoDb:
+                    // dynamo health tests disabled for now as they do not support specifying credentials that are automatically renewed
+                    break;
+                    /*
+                    AWSCredentials awsCredentials = provider.GetService<AWSCredentials>()!;
+                    IAmazonDynamoDB amazonDynamoDb = provider.GetService<IAmazonDynamoDB>()!;
+
+                    // if the region endpoint is null, e.g. we are using a local dynamo then we do not create health checks as they do not support this
+                    if (amazonDynamoDb.Config.RegionEndpoint != null)
+                    {
+                        healthChecks.AddDynamoDb(options =>
+                        {
+                            ImmutableCredentials credentials = awsCredentials.GetCredentials();
+                            options.RegionEndpoint = amazonDynamoDb.Config.RegionEndpoint;
+                            options.AccessKey = credentials.AccessKey;
+                            options.SecretKey = credentials.SecretKey;
+                        }, tags: new[] {"services"});
+                    }
+                    break;*/
+                default:
+                    throw new ArgumentOutOfRangeException();
+            }
+
+            foreach (HordeStorageSettings.StorageBackendImplementations impl in GetStorageImplementationEnums(settings.StorageImplementations))
+            {
+                switch (impl)
+                {
+                    case HordeStorageSettings.StorageBackendImplementations.S3:
+
+                        // The S3 health checks are disabled as they do not support dynamic credentials similar to the issue with dynamo
+                        /*AWSCredentials awsCredentials = provider.GetService<AWSCredentials>();
+                        S3Settings s3Settings = provider.GetService<IOptionsMonitor<S3Settings>>()!.CurrentValue;
+                        healthChecks.AddS3(options =>
+                        {
+                            options.BucketName = s3Settings.BucketName;
+                            options.S3Config = provider.GetService<IAmazonS3>().Config as AmazonS3Config;
+                            options.Credentials = awsCredentials;
+                        }, tags: new[] {"services"});*/
+                        break;
+                    case HordeStorageSettings.StorageBackendImplementations.Azure:
+                        AzureSettings azureSettings = provider.GetService<IOptionsMonitor<AzureSettings>>()!.CurrentValue;
+                        healthChecks.AddAzureBlobStorage(azureSettings.ConnectionString, tags: new[] {"services"});
+                        break;
+
+                    case HordeStorageSettings.StorageBackendImplementations.FileSystem:
+                        FilesystemSettings filesystemSettings = provider.GetService<IOptionsMonitor<FilesystemSettings>>()!.CurrentValue;
+                        healthChecks.AddDiskStorageHealthCheck(options =>
+                        {
+                            string? driveRoot = Path.GetPathRoot(filesystemSettings.RootDir);
+                            options.AddDrive(driveRoot);
+                        });
+                        break;
+                    case HordeStorageSettings.StorageBackendImplementations.Memory:
+                        break;
+                    default:
+                        throw new ArgumentOutOfRangeException("Unhandled storage impl " + impl);
+                }
+            }
+
+            healthChecks.AddCheck<LastAccessServiceCheck>("LastAccessServiceCheck", tags: new[] {"services"});
+
+            healthChecks.AddCheck<ReplicatorServiceCheck>("ReplicatorServiceCheck", tags: new[] { "services" });
+            healthChecks.AddCheck<ReplicationSnapshotServiceCheck>("ReplicationSnapshotServiceCheck", tags: new[] { "services" });
+
+            GCSettings gcSettings = provider.GetService<IOptionsMonitor<GCSettings>>()!.CurrentValue;
+
+            if (gcSettings.CleanOldRefRecords)
+            {
+                healthChecks.AddCheck<RefCleanupServiceCheck>("RefCleanupCheck", tags: new[] {"services"});
+            }
+
+            if (gcSettings.CleanOldBlobs)
+            {
+                healthChecks.AddCheck<BlobCleanupServiceCheck>("BlobStoreCheck", tags: new[] {"services"});
+            }
+
+            if (settings.LeaderElectionImplementation == HordeStorageSettings.LeaderElectionImplementations.Kubernetes)
+            {
+                healthChecks.AddCheck<KubernetesLeaderServiceCheck>("KubernetesLeaderService", tags: new[] {"services"});
+            }
+        }
+    }
+}
