Moved Horde.Storage from git to p4

diff --git a/Engine/Source/Programs/Horde/HordeStorage/Horde.Storage/Horde.Storage.csproj b/Engine/Source/Programs/Horde/HordeStorage/Horde.Storage/Horde.Storage.csproj
new file mode 100644
--- /dev/null
+++ b/Engine/Source/Programs/Horde/HordeStorage/Horde.Storage/Horde.Storage.csproj
@@ -0,0 +1,98 @@
+﻿<Project Sdk="Microsoft.NET.Sdk.Web">
+
+  <!-- Embed the git hash as a specific assembly metadata for easier reference -->
+  <Target Name="WriteGitHash" BeforeTargets="CoreCompile">
+    <PropertyGroup>
+      <GitAssemblyInfoFile>$(IntermediateOutputPath)GitAssemblyInfo.cs</GitAssemblyInfoFile>
+    </PropertyGroup>
+
+    <ItemGroup>
+      <Compile Include="$(GitAssemblyInfoFile)" />
+    </ItemGroup>
+
+    <ItemGroup>
+      <AssemblyAttributes Include="AssemblyMetadata">
+        <_Parameter1>GitHash</_Parameter1>
+        <_Parameter2>$(SourceRevisionId)</_Parameter2>
+      </AssemblyAttributes>
+
+      <AssemblyAttributes Include="AssemblyMetadata">
+        <_Parameter1>PerforceChangelist</_Parameter1>
+        <_Parameter2>$(UE_P4_CHANGELIST)</_Parameter2>
+      </AssemblyAttributes>
+
+    </ItemGroup>
+    <WriteCodeFragment Language="C#" OutputFile="$(GitAssemblyInfoFile)" AssemblyAttributes="@(AssemblyAttributes)" />
+  </Target>
+
+  <PropertyGroup>
+    <TargetFramework>net6.0</TargetFramework>
+    <Nullable>enable</Nullable>
+    <TreatWarningsAsErrors>true</TreatWarningsAsErrors>
+    <WarningsAsErrors />
+    <LangVersion>10</LangVersion>
+
+    <!--- SourceLink -->
+    <PublishRepositoryUrl>false</PublishRepositoryUrl>
+    <EmbedUntrackedSources>false</EmbedUntrackedSources>
+  </PropertyGroup>
+
+  <PropertyGroup Condition="'$(Configuration)|$(Platform)'=='Debug|AnyCPU'">
+    <DocumentationFile>Horde.Storage.xml</DocumentationFile>
+    <NoWarn>1701;1702;1591</NoWarn>
+  </PropertyGroup>
+
+  <PropertyGroup Condition="'$(Configuration)|$(Platform)'=='Release|AnyCPU'">
+    <DocumentationFile>Horde.Storage.xml</DocumentationFile>
+    <NoWarn>1701;1702;1591</NoWarn>
+  </PropertyGroup>
+
+  <ItemGroup>
+    <Compile Remove="Tests\**" />
+    <Content Remove="Tests\**" />
+    <EmbeddedResource Remove="Tests\**" />
+    <None Remove="Tests\**" />
+  </ItemGroup>
+
+  <ItemGroup>
+    <PackageReference Include="akarnokd.async-enumerable-dotnet" Version="0.0.4" />
+    <PackageReference Include="AspNetCore.HealthChecks.Aws.S3" Version="5.0.1" />
+    <PackageReference Include="AspNetCore.HealthChecks.AzureStorage" Version="5.0.1" />
+    <PackageReference Include="AspNetCore.HealthChecks.DynamoDb" Version="5.0.1" />
+    <PackageReference Include="AspNetCore.HealthChecks.MongoDb" Version="5.0.1" />
+    <PackageReference Include="AspNetCore.HealthChecks.System" Version="5.0.1" />
+    <PackageReference Include="AsyncEnumerator" Version="4.0.2" />
+    <PackageReference Include="AWSSDK.DAX.Client" Version="2.0.2" />
+    <PackageReference Include="AWSSDK.DynamoDBv2" Version="3.7.0.16" />
+    <PackageReference Include="AWSSDK.Extensions.NETCore.Setup" Version="3.7.0.1" />
+    <PackageReference Include="AWSSDK.S3" Version="3.7.0.17" />
+    <PackageReference Include="Azure.Storage.Blobs" Version="12.8.0" />
+    <PackageReference Include="Destructurama.Attributed" Version="2.0.0" />
+    <PackageReference Include="KubernetesClient" Version="4.0.10" />
+    <PackageReference Include="Microsoft.Extensions.Caching.Memory" Version="6.0.0" />
+    <PackageReference Include="Microsoft.Extensions.Logging.Debug" Version="6.0.0" />
+    <PackageReference Include="MongoDB.Driver" Version="2.11.6" />
+    <PackageReference Include="Nito.AsyncEx" Version="5.1.0" />
+    <PackageReference Include="SerilogAnalyzer" Version="0.15.0" />
+    <PackageReference Include="Crc32.NET" Version="1.2.0" />
+    <PackageReference Include="K4os.Compression.LZ4" Version="1.2.6" />
+
+    <PackageReference Include="CassandraCSharpDriver" Version="3.16.3" />
+
+    <!-- Embeds git information into the build process -->
+    <PackageReference Include="Microsoft.SourceLink.GitHub" Version="1.1.1">
+      <PrivateAssets>all</PrivateAssets>
+    </PackageReference>
+  </ItemGroup>
+
+  <ItemGroup>
+    <ProjectReference Include="..\Jupiter.Common\Jupiter.Common.csproj" />
+  </ItemGroup>
+
+  <ItemGroup>
+    <None Include="..\version.yaml">
+      <CopyToOutputDirectory>PreserveNewest</CopyToOutputDirectory>
+    </None>
+  </ItemGroup>
+
+</Project>
diff --git a/Engine/Source/Programs/Horde/HordeStorage/Horde.Storage/Implementation/Blob/AmazonS3Store.cs b/Engine/Source/Programs/Horde/HordeStorage/Horde.Storage/Implementation/Blob/AmazonS3Store.cs
new file mode 100644
--- /dev/null
+++ b/Engine/Source/Programs/Horde/HordeStorage/Horde.Storage/Implementation/Blob/AmazonS3Store.cs
@@ -0,0 +1,235 @@
+// Copyright Epic Games, Inc. All Rights Reserved.
+
+using System;
+using System.Collections.Generic;
+using System.IO;
+using System.Net;
+using System.Threading.Tasks;
+using Amazon.S3;
+using Amazon.S3.Model;
+using Jupiter.Common.Implementation;
+using Jupiter.Implementation;
+using Microsoft.Extensions.Options;
+using Serilog;
+
+namespace Horde.Storage.Implementation
+{
+    public class AmazonS3Store : IBlobStore
+    {
+        private readonly ILogger _logger = Log.ForContext<AmazonS3Store>();
+        private readonly IAmazonS3 _amazonS3;
+        private readonly S3Settings _settings;
+        private readonly HashSet<string> _bucketAccessPolicyApplied = new HashSet<string>();
+        private readonly HashSet<string> _bucketExistenceChecked = new HashSet<string>();
+
+        public AmazonS3Store(IAmazonS3 amazonS3, IOptionsMonitor<S3Settings> settings)
+        {
+            _amazonS3 = amazonS3;
+            _settings = settings.CurrentValue;
+        }
+
+        // interact with a S3 compatible object store
+        public async Task<BlobIdentifier> PutObject(NamespaceId ns, ReadOnlyMemory<byte> content, BlobIdentifier objectName)
+        {
+            await using MemoryStream stream = new MemoryStream(content.ToArray());
+            return await PutObject(ns, stream, objectName);
+        }
+
+        public async Task<BlobIdentifier> PutObject(NamespaceId ns, Stream stream, BlobIdentifier objectName)
+        {
+            string bucketName = GetBucketName(ns);
+            if (_settings.CreateBucketIfMissing)
+            {
+                if (!_bucketExistenceChecked.Contains(bucketName))
+                {
+                    bool bucketExist = await _amazonS3.DoesS3BucketExistAsync(bucketName);
+                    if (!bucketExist)
+                    {
+                        var putBucketRequest = new PutBucketRequest
+                        {
+                            BucketName = bucketName,
+                            UseClientRegion = true
+                        };
+
+                        await _amazonS3.PutBucketAsync(putBucketRequest);
+                    }
+                    _bucketExistenceChecked.Add(bucketName);
+                }
+            }
+
+            if (_settings.SetBucketPolicies && !_bucketAccessPolicyApplied.Contains(bucketName))
+            {
+                // block all public access to the bucket
+                try
+                {
+                    await _amazonS3.PutPublicAccessBlockAsync(new PutPublicAccessBlockRequest
+                    {
+                        BucketName = bucketName,
+                        PublicAccessBlockConfiguration = new PublicAccessBlockConfiguration()
+                        {
+                            RestrictPublicBuckets = true,
+                            BlockPublicAcls = true,
+                            BlockPublicPolicy = true,
+                            IgnorePublicAcls = true,
+                        }
+                    });
+
+                    _bucketAccessPolicyApplied.Add(bucketName);
+                }
+                catch (AmazonS3Exception e)
+                {
+                    // if a conflicting operation is being applied to the public access block we just ignore it, as it will get reset the next time we run
+                    if (e.StatusCode != HttpStatusCode.Conflict)
+                    {
+                        throw;
+                    }
+                }
+            }
+            PutObjectRequest request = new PutObjectRequest
+            {
+                BucketName = bucketName,
+                Key = objectName.AsS3Key(),
+                InputStream = new InstrumentedStream(stream, "s3-put"),
+            };
+
+            try
+            {
+                await _amazonS3.PutObjectAsync(request);
+            }
+            catch (AmazonS3Exception e)
+            {
+                // if the same object is added twice S3 will raise a error, as we are content addressed we can just accept whichever of the objects so we can ignore that error
+                if (e.StatusCode == HttpStatusCode.Conflict)
+                {
+                    return objectName;
+                }
+                throw;
+            }
+
+            return objectName;
+        }
+
+        public async Task<BlobIdentifier> PutObject(NamespaceId ns, byte[] content, BlobIdentifier objectName)
+        {
+            await using MemoryStream stream = new MemoryStream(content);
+            return await PutObject(ns, stream, objectName);
+        }
+
+        private string GetBucketName(NamespaceId ns)
+        {
+            return $"{_settings.BucketName}";
+        }
+
+
+        public async Task<BlobContents> GetObject(NamespaceId ns, BlobIdentifier blob)
+        {
+            string bucketName = GetBucketName(ns);
+            GetObjectResponse response;
+            try
+            {
+                response = await _amazonS3.GetObjectAsync(bucketName, blob.AsS3Key());
+            }
+            catch (AmazonS3Exception e)
+            {
+                if (e.ErrorCode == "NoSuchKey")
+                {
+                    throw new BlobNotFoundException(ns, blob);
+                }
+
+                if (e.ErrorCode == "NoSuchBucket")
+                {
+                    throw new NamespaceNotFoundException(ns);
+                }
+                throw;
+            }
+            return new BlobContents(new InstrumentedStream(response.ResponseStream, "s3-get"), response.ContentLength);
+        }
+
+        public async Task<bool> Exists(NamespaceId ns, BlobIdentifier blobIdentifier)
+        {
+            string bucketName = GetBucketName(ns);
+            try
+            {
+                await _amazonS3.GetObjectMetadataAsync(bucketName, blobIdentifier.AsS3Key());
+            }
+            catch (AmazonS3Exception e)
+            {
+                // if the object does not exist we get a not found status code
+                if (e.StatusCode == HttpStatusCode.NotFound)
+                {
+                    return false;
+                }
+
+                // something else happened, lets just process it as usual
+            }
+
+            return true;
+        }
+
+        public async Task DeleteNamespace(NamespaceId ns)
+        {
+            string bucketName = GetBucketName(ns);
+            try
+            {
+                await _amazonS3.DeleteBucketAsync(bucketName);
+            }
+            catch (AmazonS3Exception e)
+            {
+                // if the bucket does not exist we get a not found status code
+                if (e.StatusCode == HttpStatusCode.NotFound)
+                {
+                    // deleting a none existent bucket is a success
+                    return;
+                }
+
+                // something else happened, lets just process it as usual
+            }
+        }
+
+        public async IAsyncEnumerable<BlobIdentifier> ListOldObjects(NamespaceId ns, DateTime cutoff)
+        {
+            ListObjectsV2Request request = new ListObjectsV2Request
+            {
+                BucketName = GetBucketName(ns)
+            };
+
+            ListObjectsV2Response response;
+            do
+            {
+                response = await _amazonS3.ListObjectsV2Async(request);
+                foreach (S3Object obj in response.S3Objects)
+                {
+                    string key = obj.Key;
+                    string identifierString = key.Substring(key.LastIndexOf("/", StringComparison.Ordinal)+1);
+
+                    // if the object was updated after the cutoff time we do not consider it old
+                    if (obj.LastModified > cutoff)
+                    {
+                        continue;
+                    }
+
+                    yield return new BlobIdentifier(identifierString);
+                }
+
+                request.ContinuationToken = response.NextContinuationToken;
+            } while (response.IsTruncated);
+
+        }
+
+        public async Task DeleteObject(NamespaceId ns, BlobIdentifier blobIdentifier)
+        {
+            string bucketName = GetBucketName(ns);
+            await _amazonS3.DeleteObjectAsync(bucketName, blobIdentifier.AsS3Key());
+        }
+    }
+
+    public static class BlobIdentifierExtensions
+    {
+        public static string AsS3Key(this BlobIdentifier blobIdentifier)
+        {
+            string s = blobIdentifier.ToString();
+            string prefix = s.Substring(0, 4);
+            return $"{prefix}/{s}";
+        }
+    }
+}
diff --git a/Engine/Source/Programs/Horde/HordeStorage/Horde.Storage/Implementation/Blob/AzureBlobStore.cs b/Engine/Source/Programs/Horde/HordeStorage/Horde.Storage/Implementation/Blob/AzureBlobStore.cs
new file mode 100644
--- /dev/null
+++ b/Engine/Source/Programs/Horde/HordeStorage/Horde.Storage/Implementation/Blob/AzureBlobStore.cs
@@ -0,0 +1,174 @@
+﻿// Copyright Epic Games, Inc. All Rights Reserved.
+
+using System;
+using System.Collections.Generic;
+using System.Globalization;
+using System.IO;
+using System.Threading.Tasks;
+using Azure;
+using Microsoft.Extensions.Options;
+using Azure.Storage.Blobs;
+using Azure.Storage.Blobs.Models;
+using Jupiter.Implementation;
+using Serilog;
+
+namespace Horde.Storage.Implementation
+{
+    public class AzureBlobStore : IBlobStore
+    {
+        private readonly ILogger _logger = Log.ForContext<AzureBlobStore>();
+        private readonly string _connectionString;
+
+        private const string LastTouchedKey = "Io_LastTouched";
+        private const string NamespaceKey = "Io_Namespace";
+
+        public AzureBlobStore(IOptionsMonitor<AzureSettings> settings)
+        {
+            _connectionString = settings.CurrentValue.ConnectionString;
+        }
+
+        public async Task<BlobIdentifier> PutObject(NamespaceId ns, ReadOnlyMemory<byte> content, BlobIdentifier blobIdentifier)
+        {
+            // TODO: this is not ideal as we copy the buffer, but there is no upload from memory available so we would need this copy anyway
+            await using MemoryStream stream = new MemoryStream(content.ToArray());
+            return await PutObject(ns, stream, blobIdentifier);
+        }
+
+        public async Task<BlobIdentifier> PutObject(NamespaceId ns, Stream content, BlobIdentifier blobIdentifier)
+        {
+            string @namespace = SanitizeNamespace(ns);
+            _logger.Debug("Checking if Azure container with name {Name} exists", @namespace);
+            BlobContainerClient container = new BlobContainerClient(_connectionString, @namespace);
+            Dictionary<string, string> metadata = new Dictionary<string, string> { { NamespaceKey, ns.ToString() } };
+            await container.CreateIfNotExistsAsync(metadata:metadata);
+
+            _logger.Debug("Fetching blob reference with name {ObjectName}", blobIdentifier);
+            try
+            {
+
+                await container.GetBlobClient(blobIdentifier.ToString()).UploadAsync(content);
+            }
+            catch (RequestFailedException e)
+            {
+                if (e.Status == 409)
+                {
+                    // the object already existed, that is fine, no need to do anything
+                    return blobIdentifier;
+                }
+
+                throw;
+            }
+            finally
+            {
+                // we touch the blob so that the last access time is always refreshed even if we didnt actually mutate it to make sure the gc knows this is a active blob
+                // see delete operation in Leda blob store cleanup
+                await TouchBlob(container.GetBlobClient(blobIdentifier.ToString()));
+                _logger.Debug("Upload of blob {ObjectName} completed", blobIdentifier);
+            }
+
+            return blobIdentifier;
+        }
+
+        public async Task<BlobIdentifier> PutObject(NamespaceId ns, byte[] content, BlobIdentifier blobIdentifier)
+        {
+            await using MemoryStream stream = new MemoryStream(content);
+            return await PutObject(ns, stream, blobIdentifier);
+        }
+
+        private async Task TouchBlob(BlobClient blob)
+        {
+            Dictionary<string, string> metadata = new Dictionary<string, string>
+            {
+                {
+                    LastTouchedKey, DateTime.Now.ToString(CultureInfo.InvariantCulture)
+                }
+            };
+            // we update the metadata, we don''t really care about the field being specified as we just want to touch the blob to update its last modified date property which we can not set
+            await blob.SetMetadataAsync(metadata);
+        }
+
+        public async Task<BlobContents> GetObject(NamespaceId ns, BlobIdentifier blobIdentifier)
+        {
+            string fixedNamespace = SanitizeNamespace(ns);
+            BlobContainerClient container = new BlobContainerClient(_connectionString, fixedNamespace);
+            if (!await container.ExistsAsync())
+            {
+                throw new InvalidOperationException($"Container {fixedNamespace} did not exist");
+            }
+
+            try
+            {
+                BlobClient blob = container.GetBlobClient(blobIdentifier.ToString());
+                Response<BlobDownloadInfo> blobInfo = await blob.DownloadAsync();
+                return new BlobContents(blobInfo.Value.Content, blobInfo.Value.ContentLength);
+            }
+            catch (RequestFailedException e)
+            {
+                if (e.Status == 404)
+                    throw new BlobNotFoundException(ns, blobIdentifier);
+
+                throw;
+            }
+        }
+
+        public async Task<bool> Exists(NamespaceId ns, BlobIdentifier blobIdentifier)
+        {
+            BlobContainerClient container = new BlobContainerClient(_connectionString, SanitizeNamespace(ns));
+            if (!await container.ExistsAsync())
+            {
+                return false;
+            }
+
+            BlobClient blob = container.GetBlobClient(blobIdentifier.ToString());
+            return await blob.ExistsAsync();
+        }
+
+        public async Task DeleteObject(NamespaceId ns, BlobIdentifier blobIdentifier)
+        {
+            string fixedNamespace = SanitizeNamespace(ns);
+            BlobContainerClient container = new BlobContainerClient(_connectionString, fixedNamespace);
+            if (!await container.ExistsAsync())
+            {
+                throw new InvalidOperationException($"Container {fixedNamespace} did not exist");
+            }
+
+            await container.DeleteBlobAsync(blobIdentifier.ToString());
+        }
+
+        public async Task DeleteNamespace(NamespaceId ns)
+        {
+            string fixedNamespace = SanitizeNamespace(ns);
+            BlobContainerClient container = new BlobContainerClient(_connectionString, fixedNamespace);
+            if (await container.ExistsAsync())
+            {
+                // we can only delete it if the container exists
+                await container.DeleteAsync();
+            }
+        }
+
+        public async IAsyncEnumerable<BlobIdentifier> ListOldObjects(NamespaceId ns, DateTime cutoff)
+        {
+            string fixedNamespace = SanitizeNamespace(ns);
+            BlobContainerClient container = new BlobContainerClient(_connectionString, fixedNamespace);
+            bool exists = await container.ExistsAsync();
+            if (!exists)
+                yield break;
+
+            await foreach (var item in container.GetBlobsAsync(BlobTraits.Metadata))
+            {
+                // ignore any recent blob
+                if ((item.Properties?.LastModified ?? DateTime.Now) > cutoff)
+                {
+                    continue;
+                }
+
+                yield return new BlobIdentifier(item.Name);
+            }
+        }
+
+        private static string SanitizeNamespace(NamespaceId ns)
+        {
+            return ns.ToString().Replace(".", "-").ToLower();
+        }
+    }
+}
diff --git a/Engine/Source/Programs/Horde/HordeStorage/Horde.Storage/Implementation/Blob/ConsistencyCheckService.cs b/Engine/Source/Programs/Horde/HordeStorage/Horde.Storage/Implementation/Blob/ConsistencyCheckService.cs
new file mode 100644
--- /dev/null
+++ b/Engine/Source/Programs/Horde/HordeStorage/Horde.Storage/Implementation/Blob/ConsistencyCheckService.cs
@@ -0,0 +1,124 @@
+﻿// Copyright Epic Games, Inc. All Rights Reserved.
+
+using System;
+using System.IO;
+using System.Threading;
+using System.Threading.Tasks;
+using Cassandra;
+using Datadog.Trace;
+using Jupiter;
+using Jupiter.Implementation;
+using Microsoft.Extensions.DependencyInjection;
+using Microsoft.Extensions.Options;
+using Serilog;
+
+namespace Horde.Storage.Implementation
+{
+    // ReSharper disable once ClassNeverInstantiated.Global
+    public class ConsistencyCheckService : PollingService<ConsistencyCheckService.ConsistencyState>, IDisposable
+    {
+        private readonly IOptionsMonitor<ConsistencyCheckSettings> _settings;
+        private readonly IServiceProvider _provider;
+        private readonly ILeaderElection _leaderElection;
+        private readonly IRefsStore _refsStore;
+        private readonly ILogger _logger = Log.ForContext<ConsistencyCheckService>();
+
+        public class ConsistencyState
+        {
+        }
+
+        public override bool ShouldStartPolling()
+        {
+            return _settings.CurrentValue.Enabled;
+        }
+
+        public ConsistencyCheckService(IOptionsMonitor<ConsistencyCheckSettings> settings, IServiceProvider provider, ILeaderElection leaderElection, IRefsStore refsStore) :
+            base(serviceName: nameof(ConsistencyCheckService), TimeSpan.FromSeconds(settings.CurrentValue.ConsistencyCheckPollFrequencySeconds), new ConsistencyState())
+        {
+            _settings = settings;
+            _provider = provider;
+            _leaderElection = leaderElection;
+            _refsStore = refsStore;
+        }
+
+        public override async Task<bool> OnPoll(ConsistencyState state, CancellationToken cancellationToken)
+        {
+            if (!_settings.CurrentValue.Enabled)
+            {
+                _logger.Information("Skipped running consistency check as it is disabled");
+                return false;
+            }
+
+            if (!_leaderElection.IsThisInstanceLeader())
+            {
+                _logger.Information("Skipped running consistency check because this instance was not the leader");
+                return false;
+            }
+
+            using Scope scope = Tracer.Instance.StartActive("consistency_check.run");
+
+            await RunConsistencyCheck();
+
+            return true;
+        }
+
+        private async Task RunConsistencyCheck()
+        {
+            //TODO: This isn't really specific to the S3 store, but it should likely only run on the backing store and not the cache layers which we do not have a way of fetching right now.
+            AmazonS3Store? s3Store;
+            try
+            {
+                s3Store = _provider.GetService<AmazonS3Store>();
+            }
+            catch (Exception)
+            {
+                s3Store = null;
+            }
+
+            if (s3Store == null)
+            {
+                _logger.Warning("No S3 Store found so will not run any consistency check");
+                return;
+            }
+
+            long countOfBlobsChecked = 0;
+            DateTime objectCutoff = DateTime.Now;
+            // technically this does not need to be run per namespace but per s3 bucket
+            await foreach (NamespaceId ns in _refsStore.GetNamespaces())
+            {
+                await foreach (BlobIdentifier blob in s3Store.ListOldObjects(ns, objectCutoff))
+                {
+                    if (countOfBlobsChecked % 100 == 0)
+                        _logger.Information("Consistency check running on S3 blob store, count of blobs processed so far: {CountOfBlobs}", countOfBlobsChecked);
+                    Interlocked.Increment(ref countOfBlobsChecked);
+                    
+                    BlobContents contents = await s3Store.GetObject(ns, blob);
+                    await using Stream s = contents.Stream;
+
+                    BlobIdentifier newHash = await BlobIdentifier.FromStream(s);
+                    if (!blob.Equals(newHash))
+                    {
+                        _logger.Error("Mismatching hash for S3 object {Blob} in {Namespace}, new hash has {NewHash}", blob, ns, newHash);
+                    }
+                }
+                
+            }
+        }
+
+        protected override Task OnStopping(ConsistencyState state)
+        {
+            return Task.CompletedTask;
+        }
+
+        public void Dispose()
+        {
+
+        }
+    }
+
+    public class ConsistencyCheckSettings
+    {
+        public bool Enabled { get; set; } = false;
+        public double ConsistencyCheckPollFrequencySeconds { get; set; } = TimeSpan.FromHours(2).TotalSeconds;
+    }
+}
diff --git a/Engine/Source/Programs/Horde/HordeStorage/Horde.Storage/Implementation/Blob/FileSystemStore.cs b/Engine/Source/Programs/Horde/HordeStorage/Horde.Storage/Implementation/Blob/FileSystemStore.cs
new file mode 100644
--- /dev/null
+++ b/Engine/Source/Programs/Horde/HordeStorage/Horde.Storage/Implementation/Blob/FileSystemStore.cs
@@ -0,0 +1,311 @@
+﻿using System;
+using System.Collections.Generic;
+using System.IO;
+using System.Linq;
+using System.Threading;
+using System.Threading.Tasks;
+using Dasync.Collections;
+using Jupiter;
+using Jupiter.Implementation;
+using Microsoft.Extensions.Options;
+using Serilog;
+
+namespace Horde.Storage.Implementation
+{
+    public class FileSystemStore : IBlobStore, IBlobCleanup 
+    {
+        private readonly ILogger _logger = Log.ForContext<FileSystemStore>();
+        private IOptionsMonitor<FilesystemSettings> _settings;
+
+        internal const int DefaultBufferSize = 4096;
+
+        public FileSystemStore(IOptionsMonitor<FilesystemSettings> settings)
+        {
+            _settings = settings;
+        }
+
+        public static FileInfo GetFilesystemPath(string rootDir, NamespaceId ns, BlobIdentifier blob)
+        {
+            const int CountOfCharactersPerDirectory = 2;
+            string objectName = blob.ToString();
+            string firstPart = objectName.Substring(0, CountOfCharactersPerDirectory);
+            string secondPart = objectName.Substring(CountOfCharactersPerDirectory, CountOfCharactersPerDirectory);
+            string fileName = objectName;
+
+            return new FileInfo(Path.Combine(rootDir, ns.ToString(), firstPart, secondPart, fileName));
+        }
+
+        public FileInfo GetFilesystemPath(NamespaceId ns, BlobIdentifier objectName)
+        {
+            return GetFilesystemPath(_settings.CurrentValue.RootDir, ns, objectName);
+        }
+
+        public DirectoryInfo GetFilesystemPath(NamespaceId ns)
+        {
+            return new DirectoryInfo(Path.Combine(_settings.CurrentValue.RootDir, ns.ToString()));
+        }
+
+        public async Task<BlobIdentifier> PutObject(NamespaceId ns, ReadOnlyMemory<byte> content, BlobIdentifier blobIdentifier)
+        {
+            FileInfo filePath = GetFilesystemPath(ns, blobIdentifier);
+            filePath.Directory?.Create();
+
+            await using FileStream fs = new FileStream(filePath.FullName, FileMode.Create, FileAccess.Write, FileShare.Read, DefaultBufferSize, FileOptions.Asynchronous | FileOptions.SequentialScan);
+            CancellationToken cancellationToken = default(CancellationToken);
+            await fs.WriteAsync(content, cancellationToken);
+            await fs.FlushAsync(cancellationToken);
+
+            UpdateLastWriteTime(filePath.FullName, DateTime.UnixEpoch);
+            return blobIdentifier;
+        }
+
+
+        public async Task<BlobIdentifier> PutObject(NamespaceId ns, Stream content, BlobIdentifier blobIdentifier)
+        {
+            FileInfo filePath = GetFilesystemPath(ns, blobIdentifier);
+            filePath.Directory?.Create();
+
+            await using FileStream fs = new FileStream(filePath.FullName, FileMode.Create, FileAccess.Write, FileShare.Read, DefaultBufferSize, FileOptions.Asynchronous | FileOptions.SequentialScan);
+            CancellationToken cancellationToken = default(CancellationToken);
+            await content.CopyToAsync(fs, cancellationToken);
+            await fs.FlushAsync(cancellationToken);
+
+            UpdateLastWriteTime(filePath.FullName, DateTime.UnixEpoch);
+            return blobIdentifier;
+        }
+
+        public async Task<BlobIdentifier> PutObject(NamespaceId ns, byte[] content, BlobIdentifier blobIdentifier)
+        {
+            FileInfo filePath = GetFilesystemPath(ns, blobIdentifier);
+            filePath.Directory?.Create();
+
+            await using FileStream fs = new FileStream(filePath.FullName, FileMode.Create, FileAccess.Write, FileShare.Read, DefaultBufferSize, FileOptions.Asynchronous | FileOptions.SequentialScan);
+            CancellationToken cancellationToken = default(CancellationToken);
+            await fs.WriteAsync(content, cancellationToken);
+            await fs.FlushAsync(cancellationToken);
+
+            UpdateLastWriteTime(filePath.FullName, DateTime.UnixEpoch);
+            return blobIdentifier;
+        }
+
+        public Task<BlobContents> GetObject(NamespaceId ns, BlobIdentifier blob)
+        {
+            FileInfo filePath = GetFilesystemPath(ns, blob);
+
+            if (!filePath.Exists)
+                throw new BlobNotFoundException(ns, blob);
+
+            FileStream fs = new FileStream(filePath.FullName, FileMode.Open, FileAccess.Read, FileShare.Read, DefaultBufferSize, FileOptions.Asynchronous | FileOptions.SequentialScan);
+            
+            UpdateLastWriteTime(filePath.FullName, DateTime.UtcNow);
+            return Task.FromResult(new BlobContents(fs, fs.Length));
+        }
+
+        /// <summary>
+        /// Update the last modified/write time that is used for determining when file was last accessed
+        /// 
+        /// Using access time is tricky as many file systems disable that for performance reasons.
+        /// A new blob written to disk should be set with the oldest possible write time.
+        /// This will ensure sorting of least recently accessed files during garbage collection works as intended.
+        /// The write time update will happen async without any waiting to prevent blocking the critical path
+        /// as it's best-effort only.
+        /// 
+        /// See <seealso cref="GetLeastRecentlyAccessedObjects" />.
+        /// </summary>
+        /// <param name="filePath"></param>
+        /// <param name="lastAccessed">Time the file was last accessed</param>
+        private void UpdateLastWriteTime(string filePath, DateTime lastAccessed)
+        {
+            Task.Run(() =>
+            {
+                try
+                {
+                    File.SetLastWriteTimeUtc(filePath, lastAccessed);
+                }
+                catch (FileNotFoundException)
+                {
+                    // it is okay if the file does not exist anymore, that just means it got gced
+                }
+            }).ContinueWith(task =>
+            {
+                if (task.Exception != null)
+                {
+                    _logger.Error(task.Exception, "Exception when updating write time for {FilePath}", filePath);
+                }
+            });
+        }
+
+        public Task<bool> Exists(NamespaceId ns, BlobIdentifier blob)
+        {
+            FileInfo filePath = GetFilesystemPath(ns, blob);
+
+            return Task.FromResult(filePath.Exists);
+        }
+
+        public Task DeleteObject(NamespaceId ns, BlobIdentifier objectName)
+        {
+            FileInfo filePath = GetFilesystemPath(ns, objectName);
+            filePath.Delete();
+            return Task.CompletedTask;
+        }
+
+        public Task DeleteNamespace(NamespaceId ns)
+        {
+            DirectoryInfo namespaceDirectory = GetFilesystemPath(ns);
+            if (namespaceDirectory.Exists)
+                namespaceDirectory.Delete(true);
+
+            return Task.CompletedTask;
+        }
+
+        public IAsyncEnumerable<BlobIdentifier> ListOldObjects(NamespaceId ns, DateTime cutoff)
+        {
+            return DoListOldObjects(ns, cutoff).ToAsyncEnumerable();
+        }
+
+        private IEnumerable<BlobIdentifier> DoListOldObjects(NamespaceId ns, DateTime cutoff)
+        {
+            DirectoryInfo di = new DirectoryInfo(Path.Combine(_settings.CurrentValue.RootDir, ns.ToString()));
+            if (!di.Exists)
+            {
+                yield break;
+            }
+
+            foreach (FileSystemInfo? file in di.EnumerateFileSystemInfos("*", SearchOption.AllDirectories))
+            {
+                if (file.Attributes.HasFlag(FileAttributes.Directory))
+                {
+                    continue;
+                }
+
+                if (file.LastWriteTime > cutoff)
+                {
+                    _logger.Information("Ignoring {Blob} as it is new. Last modified: {LastModified} Cutoff: {CutOff}", file.Name, file.LastWriteTime, cutoff);
+                    continue;
+                }
+                _logger.Information("Old {Blob} found. Last modified: {LastModified} Cutoff: {CutOff}", file.Name, file.LastWriteTime, cutoff);
+
+                yield return new BlobIdentifier(file.Name);
+            }
+        }
+        
+        public async Task<List<RemovedBlobs>> Cleanup(CancellationToken cancellationToken)
+        {
+            return await CleanupInternal(cancellationToken);
+        }
+
+        /// <summary>
+        /// Clean the store from expired files
+        /// 
+        /// Uses the configuration to remove the least recently accessed (modified) blobs
+        /// </summary>
+        /// <param name="cancellationToken">Cancellation token</param>
+        /// <param name="batchSize">Number of files to scan for clean up. A higher number is recommended since blob store can contain many small blobs</param>
+        /// <returns></returns>
+        public async Task<List<RemovedBlobs>> CleanupInternal(CancellationToken cancellationToken, int batchSize = 100000)
+        {
+            long size = await CalculateDiskSpaceUsed();
+            ulong maxSizeBytes = _settings.CurrentValue.MaxSizeBytes;
+            long triggerSize = (long) (maxSizeBytes * _settings.CurrentValue.TriggerThresholdPercentage);
+            long targetSize = (long) (maxSizeBytes * _settings.CurrentValue.TargetThresholdPercentage); // Target to shrink to if triggered
+            List<RemovedBlobs> blobsRemoved = new List<RemovedBlobs>(batchSize);
+
+            if (size < triggerSize)
+            {
+                return blobsRemoved;
+            }
+
+            // Perform a maximum of 5 clean up runs
+            for (int i = 0; i < 5; i++)
+            {
+                size = await CalculateDiskSpaceUsed();
+                if (size <= targetSize)
+                {
+                    return blobsRemoved;
+                }
+                
+                FileInfo[] fileInfos = await GetLeastRecentlyAccessedObjects(maxResults: batchSize);
+                if (fileInfos.Length == 0)
+                {
+                    return blobsRemoved;
+                }
+                
+                long totalBytesDeleted = 0;
+                foreach (FileInfo fi in fileInfos)
+                {
+                    try
+                    {
+                        totalBytesDeleted += fi.Length;
+                        fi.Delete();
+                        blobsRemoved.Add(new RemovedBlobs(new BlobIdentifier(fi.Name)));
+
+                        long currentSize = size - totalBytesDeleted;
+                        if (currentSize <= targetSize || cancellationToken.IsCancellationRequested)
+                        {
+                            return blobsRemoved;
+                        }
+                    }
+                    catch (FileNotFoundException)
+                    {
+                        // if the file was gced while running we can just ignore it
+                    }
+                }
+            }
+            
+            return blobsRemoved;
+        }
+
+        /// <summary>
+        /// Get least recently accessed objects
+        /// Assumes files on disk have the their last access timestamp up-to-date
+        /// </summary>
+        /// <param name="ns">Namespace, if set to null all namespaces will be scanned</param>
+        /// <param name="maxResults">Max results to return. Note that the entire namespace will be scanned no matter what.</param>
+        /// <returns>Enumerable of least recently accessed objects as FileInfos</returns>
+        public async Task<FileInfo[]> GetLeastRecentlyAccessedObjects(NamespaceId? ns = null, int maxResults = 10000)
+        {
+            string path = ns != null ? Path.Combine(_settings.CurrentValue.RootDir, ns.ToString()!) : _settings.CurrentValue.RootDir;
+            DirectoryInfo di = new DirectoryInfo(path);
+            if (!di.Exists)
+            {
+                return Array.Empty<FileInfo>();
+            }
+            
+            return await Task.Run(() => di.EnumerateFiles("*", SearchOption.AllDirectories).OrderBy(x => x.LastWriteTime).Take(maxResults).ToArray());
+        }
+        
+        /// <summary>
+        /// Calculate the total size of blobs on disk for given namespace
+        /// </summary>
+        /// <param name="ns">Namespace, if set to null the total size of all namespaces will be returned</param>
+        /// <returns>Total size of blobs in bytes</returns>
+        public async Task<long> CalculateDiskSpaceUsed(NamespaceId? ns = null)
+        {
+            string path = ns != null ? Path.Combine(_settings.CurrentValue.RootDir, ns.ToString()!) : _settings.CurrentValue.RootDir;
+            DirectoryInfo di = new DirectoryInfo(path);
+            if (!di.Exists)
+            {
+                return 0;
+            }
+            
+            return await Task.Run(() => di.EnumerateFiles("*", SearchOption.AllDirectories).Sum(x =>
+            {
+                try
+                {
+                    return x.Length;
+                }
+                catch (FileNotFoundException)
+                {
+                    // if the file has been gced we just ignore it
+                    return 0;
+                }
+            }));
+        }
+        
+        public  IAsyncEnumerable<NamespaceId> ListNamespaces()
+        {
+            DirectoryInfo di = new DirectoryInfo(_settings.CurrentValue.RootDir);
+            return di.GetDirectories().Select(x => new NamespaceId(x.Name)).ToAsyncEnumerable();
+        }
+    }
+}
diff --git a/Engine/Source/Programs/Horde/HordeStorage/Horde.Storage/Implementation/Blob/HierarchicalBlobStore.cs b/Engine/Source/Programs/Horde/HordeStorage/Horde.Storage/Implementation/Blob/HierarchicalBlobStore.cs
new file mode 100644
--- /dev/null
+++ b/Engine/Source/Programs/Horde/HordeStorage/Horde.Storage/Implementation/Blob/HierarchicalBlobStore.cs
@@ -0,0 +1,260 @@
+﻿// Copyright Epic Games, Inc. All Rights Reserved.
+
+using System;
+using System.Collections.Generic;
+using System.IO;
+using System.Linq;
+using System.Threading.Tasks;
+using Datadog.Trace;
+using Jupiter.Implementation;
+using Jupiter.Utils;
+
+namespace Horde.Storage.Implementation
+{
+    public class HierarchicalBlobStoreException : Exception
+    {
+        public HierarchicalBlobStoreException(string? message) : base(message)
+        {
+        }
+    }
+    
+    /// <summary>
+    /// A hierarchical blob store
+    /// Accepts a sequence of other blob stores and will return a result as soon as one is found when traversing.
+    /// </summary>
+    internal class HierarchicalBlobStore : IBlobStore
+    {
+        public IBlobStore[] BlobStores { get; }
+        
+        public HierarchicalBlobStore(IEnumerable<IBlobStore> blobStores)
+        {
+            this.BlobStores = blobStores.ToArray();
+        }
+
+        public async Task<BlobIdentifier> PutObject(NamespaceId ns, byte[] blob, BlobIdentifier identifier)
+        {
+            BlobIdentifier lastIdentifier = null!;
+            foreach (IBlobStore store in BlobStores)
+            {
+                using Scope scope = Tracer.Instance.StartActive("HierarchicalStore.PutObject");
+                scope.Span.ResourceName = identifier.ToString();
+                scope.Span.SetTag("BlobStore", store.GetType().Name);
+                lastIdentifier = await store.PutObject(ns, blob, identifier);
+            }
+
+            return lastIdentifier;
+        }
+
+        public async Task<BlobIdentifier> PutObject(NamespaceId ns, ReadOnlyMemory<byte> blob, BlobIdentifier identifier)
+        {
+            BlobIdentifier lastIdentifier = null!;
+            foreach (IBlobStore store in BlobStores)
+            {
+                using Scope scope = Tracer.Instance.StartActive("HierarchicalStore.PutObject");
+                scope.Span.ResourceName = identifier.ToString();
+                scope.Span.SetTag("BlobStore", store.GetType().Name);
+                lastIdentifier = await store.PutObject(ns, blob, identifier);
+            }
+            return lastIdentifier;
+        }
+
+        public async Task<BlobIdentifier> PutObject(NamespaceId ns, Stream content, BlobIdentifier identifier)
+        {
+            BlobIdentifier? lastIdentifier = null;
+            foreach (IBlobStore store in BlobStores)
+            {
+                using Scope scope = Tracer.Instance.StartActive("HierarchicalStore.PutObject");
+                scope.Span.ResourceName = identifier.ToString();
+                scope.Span.SetTag("BlobStore", store.GetType().Name);
+                try
+                {
+                    lastIdentifier = await store.PutObject(ns, content, identifier);
+                }
+                catch (BlobToLargeException)
+                {
+                    // if the object is to large to cache we hope for the other stores to cache it
+                }
+            }
+
+            if (lastIdentifier == null)
+            {
+                throw new Exception($"Failed to cache blob {identifier}, the object might be to large to fit into any of the available stores");
+            }
+            return lastIdentifier;
+        }
+
+        public async Task<BlobContents> GetObject(NamespaceId ns, BlobIdentifier blob)
+        {
+            bool seenBlobNotFound = false;
+            bool seenNamespaceNotFound = false;
+            int numStoreMisses = 0;
+            BlobContents? blobContents = null;
+            foreach (IBlobStore store in BlobStores)
+            {
+                using Scope scope = Tracer.Instance.StartActive("HierarchicalStore.GetObject");
+                scope.Span.SetTag("BlobStore", store.GetType().Name);
+                scope.Span.SetTag("ObjectFound", false.ToString());
+                try
+                {
+                    blobContents = await store.GetObject(ns, blob);
+                    scope.Span.SetTag("ObjectFound", true.ToString());
+                    break;
+                }
+                catch (BlobNotFoundException)
+                {
+                    seenBlobNotFound = true;
+                    numStoreMisses++;
+                }
+                catch (NamespaceNotFoundException)
+                {
+                    seenNamespaceNotFound = true;
+                }
+            }
+
+            if (seenBlobNotFound && blobContents == null)
+            {
+                throw new BlobNotFoundException(ns, blob);
+            }
+
+            if (seenNamespaceNotFound && blobContents == null)
+            {
+                throw new NamespaceNotFoundException(ns);
+            }
+            
+            if (blobContents == null)
+            {
+                // Should not happen but exists to safeguard against the null pointer
+                throw new HierarchicalBlobStoreException("blobContents is null");
+            }
+
+            if (numStoreMisses >= 1)
+            {
+                using Scope _ = Tracer.Instance.StartActive("HierarchicalStore.Populate");
+                await using MemoryStream tempStream = new MemoryStream();
+                await blobContents.Stream.CopyToAsync(tempStream);
+                byte[] data = tempStream.ToArray();
+                
+                // Don't populate the last store, as that is where we got the hit
+                for (int i = 0; i < numStoreMisses; i++)
+                {
+                    var blobStore = BlobStores[i];
+                    using Scope scope = Tracer.Instance.StartActive("HierarchicalStore.PopulateStore");
+                    scope.Span.SetTag("BlobStore", blobStore.GetType().Name);
+                    // Populate each store traversed that did not have the content found lower in the hierarchy
+                    await blobStore.PutObject(ns, data, blob);
+                }
+
+                blobContents = new BlobContents(new MemoryStream(data), data.Length);
+            }
+            
+            return blobContents;
+        }
+
+        public async Task<bool> Exists(NamespaceId ns, BlobIdentifier blob)
+        {
+            foreach (IBlobStore store in BlobStores)
+            {
+                using Scope scope = Tracer.Instance.StartActive("HierarchicalStore.ObjectExists");
+                scope.Span.SetTag("BlobStore", store.GetType().Name);
+                if (await store.Exists(ns, blob))
+                {
+                    scope.Span.SetTag("ObjectFound", true.ToString());
+                    return true;
+                }
+                scope.Span.SetTag("ObjectFound", false.ToString());
+            }
+
+            return false;
+        }
+
+        public async Task DeleteObject(NamespaceId ns, BlobIdentifier blob)
+        {
+            bool blobNotFound = false;
+            bool deletedAtLeastOnce = false;
+            foreach (IBlobStore store in BlobStores)
+            {
+                try
+                {
+                    using Scope scope = Tracer.Instance.StartActive("HierarchicalStore.DeleteObject");
+                    scope.Span.SetTag("BlobStore", store.GetType().Name);
+                    await store.DeleteObject(ns, blob);
+                    deletedAtLeastOnce = true;
+                }
+                catch (NamespaceNotFoundException)
+                {
+                    // Ignore
+                }
+                catch (BlobNotFoundException)
+                {
+                    blobNotFound = true;
+                }
+            }
+
+            if (deletedAtLeastOnce)
+                return;
+
+            if (blobNotFound)
+                throw new BlobNotFoundException(ns, blob);
+
+            throw new NamespaceNotFoundException(ns);
+        }
+
+        public async Task DeleteNamespace(NamespaceId ns)
+        {
+            bool deletedAtLeastOnce = false;
+            foreach (IBlobStore store in BlobStores)
+            {
+                using Scope scope = Tracer.Instance.StartActive("HierarchicalStore.DeleteNamespace");
+                scope.Span.SetTag("BlobStore", store.GetType().Name);
+                try
+                {
+                    await store.DeleteNamespace(ns);
+                    deletedAtLeastOnce = true;
+                }
+                catch (NamespaceNotFoundException)
+                {
+                    // Ignore
+                }
+            }
+
+            if (deletedAtLeastOnce)
+                return;
+
+            throw new NamespaceNotFoundException(ns);
+        }
+
+        public async IAsyncEnumerable<BlobIdentifier> ListOldObjects(NamespaceId ns, DateTime cutoff)
+        {
+            bool blobIdentifierAdded = false;
+            HashSet<BlobIdentifier> uniqueBlobIds = new HashSet<BlobIdentifier>();
+            foreach (IBlobStore store in BlobStores)
+            {
+                try
+                {
+                    using Scope scope = Tracer.Instance.StartActive("HierarchicalStore.ListOldObjects");
+                    scope.Span.SetTag("BlobStore", store.GetType().Name);
+                    await foreach (BlobIdentifier blobId in store.ListOldObjects(ns, cutoff))
+                    {
+                        uniqueBlobIds.Add(blobId);
+                    }
+
+                    blobIdentifierAdded = true;
+                }
+                catch (NamespaceNotFoundException)
+                {
+                    // Ignore
+                }
+            }
+
+            if (!blobIdentifierAdded)
+            {
+                throw new NamespaceNotFoundException(ns);
+            }
+
+            foreach (BlobIdentifier blobId in uniqueBlobIds)
+            {
+                yield return blobId;
+            }
+        }
+    }
+}
diff --git a/Engine/Source/Programs/Horde/HordeStorage/Horde.Storage/Implementation/Blob/IBlobStore.cs b/Engine/Source/Programs/Horde/HordeStorage/Horde.Storage/Implementation/Blob/IBlobStore.cs
new file mode 100644
--- /dev/null
+++ b/Engine/Source/Programs/Horde/HordeStorage/Horde.Storage/Implementation/Blob/IBlobStore.cs
@@ -0,0 +1,130 @@
+﻿// Copyright Epic Games, Inc. All Rights Reserved.
+
+using System;
+using System.Collections.Concurrent;
+using System.Collections.Generic;
+using System.IO;
+using System.Linq;
+using System.Threading.Tasks;
+using async_enumerable_dotnet;
+using Dasync.Collections;
+using Datadog.Trace;
+using Jupiter.Implementation;
+
+namespace Horde.Storage.Implementation
+{
+    public interface IBlobStore
+    {
+        Task<BlobIdentifier> PutObject(NamespaceId ns, byte[] blob, BlobIdentifier identifier);
+        Task<BlobIdentifier> PutObject(NamespaceId ns, ReadOnlyMemory<byte> blob, BlobIdentifier identifier);
+        Task<BlobIdentifier> PutObject(NamespaceId ns, Stream content, BlobIdentifier identifier);
+
+        Task<BlobContents> GetObject(NamespaceId ns, BlobIdentifier blob);
+        Task<bool> Exists(NamespaceId ns, BlobIdentifier blob);
+
+        // Delete a object
+        Task DeleteObject(NamespaceId ns, BlobIdentifier blob);
+
+
+        // delete the whole namespace
+        Task DeleteNamespace(NamespaceId ns);
+
+        IAsyncEnumerable<BlobIdentifier> ListOldObjects(NamespaceId ns, DateTime cutoff);
+    }
+
+    public class BlobNotFoundException : Exception
+    {
+        public NamespaceId Ns { get; }
+        public BlobIdentifier Blob { get; }
+
+        public BlobNotFoundException(NamespaceId ns, BlobIdentifier blob) : base($"No Blob in Namespace {ns} with id {blob}")
+        {
+            Ns = ns;
+            Blob = blob;
+        }
+    }
+
+    public class NamespaceNotFoundException : Exception
+    {
+        public NamespaceId Namespace { get; }
+
+        public NamespaceNotFoundException(NamespaceId @namespace) : base($"Could not find namespace {@namespace}")
+        {
+            Namespace = @namespace;
+        }
+
+    }
+
+    public class BlobToLargeException : Exception
+    {
+        public BlobIdentifier Blob { get; }
+
+        public BlobToLargeException(BlobIdentifier blob) : base($"Blob {blob} was to large to cache")
+        {
+            Blob = blob;
+        }
+
+    }
+
+    public static class BlobStoreUtils
+    {
+        public static async Task<BlobIdentifier[]> FilterOutKnownBlobs(this IBlobStore blobStore, NamespaceId ns, BlobIdentifier[] blobs)
+        {
+            var tasks = blobs.Select(async blobIdentifier => new { BlobIdentifier = blobIdentifier, Exists = await blobStore.Exists(ns, blobIdentifier) });
+            var blobResults = await Task.WhenAll(tasks);
+            var filteredBlobs = blobResults.Where(ac => !ac.Exists).Select(ac => ac.BlobIdentifier);
+            return filteredBlobs.ToArray();
+        }
+
+        public static async Task<BlobIdentifier[]> FilterOutKnownBlobs(this IBlobStore blobStore, NamespaceId ns, IAsyncEnumerable<BlobIdentifier> blobs)
+        {
+            ConcurrentBag<BlobIdentifier> missingBlobs = new ConcurrentBag<BlobIdentifier>();
+
+            try
+            {
+                await blobs.ParallelForEachAsync(async identifier =>
+                {
+                    bool exists = await blobStore.Exists(ns, identifier);
+
+                    if (!exists)
+                    {
+                        missingBlobs.Add(identifier);
+                    }
+                });
+            }
+            catch (ParallelForEachException e)
+            {
+                if (e.InnerException is PartialReferenceResolveException)
+                    throw e.InnerException;
+
+                throw;
+            }
+
+            return missingBlobs.ToArray();
+        }
+
+
+        public static async Task<BlobContents> GetObjects(this IBlobStore blobStore, NamespaceId ns, BlobIdentifier[] blobs)
+        {
+            using Scope _ = Tracer.Instance.StartActive("blob.combine");
+            Task<BlobContents>[] tasks = new Task<BlobContents>[blobs.Length];
+            for (int i = 0; i < blobs.Length; i++)
+            {
+                tasks[i] = blobStore.GetObject(ns, blobs[i]);
+            }
+
+            MemoryStream ms = new MemoryStream();
+            foreach (Task<BlobContents> task in tasks)
+            {
+                BlobContents blob = await task;
+                await using Stream s = blob.Stream;
+                await s.CopyToAsync(ms);
+            }
+
+            ms.Seek(0, SeekOrigin.Begin);
+
+            return new BlobContents(ms, ms.Length);
+        }
+    }
+
+}
diff --git a/Engine/Source/Programs/Horde/HordeStorage/Horde.Storage/Implementation/Blob/MemoryBlobStore.cs b/Engine/Source/Programs/Horde/HordeStorage/Horde.Storage/Implementation/Blob/MemoryBlobStore.cs
new file mode 100644
--- /dev/null
+++ b/Engine/Source/Programs/Horde/HordeStorage/Horde.Storage/Implementation/Blob/MemoryBlobStore.cs
@@ -0,0 +1,169 @@
+﻿// Copyright Epic Games, Inc. All Rights Reserved.
+
+using System;
+using System.Collections.Concurrent;
+using System.Collections.Generic;
+using System.IO;
+using System.Linq;
+using System.Threading.Tasks;
+using Jupiter;
+using Jupiter.Implementation;
+using Jupiter.Utils;
+
+namespace Horde.Storage.Implementation
+{
+    internal class MemoryBlobStore : IBlobStore
+    {
+        private class BlobContainer
+        {
+            public BlobIdentifier BlobIdentifier { get; }
+            public byte[] Contents { get; }
+            public DateTime LastModified { get; set; }
+
+            public BlobContainer(BlobIdentifier blobIdentifier, byte[] contents)
+            {
+                BlobIdentifier = blobIdentifier;
+                Contents = contents;
+                LastModified = DateTime.Now;
+            }
+        }
+
+        private readonly ConcurrentDictionary<NamespaceId, ConcurrentDictionary<BlobIdentifier, BlobContainer>> _blobs = new ConcurrentDictionary<NamespaceId, ConcurrentDictionary<BlobIdentifier, BlobContainer>>();
+        
+        /// <summary>
+        /// Throw an exception when putting an object that already exists
+        ///
+        /// Primarily used for testing.
+        /// </summary>
+        private readonly bool throwOnOverwrite;
+
+        public MemoryBlobStore(bool throwOnOverwrite = false)
+        {
+            this.throwOnOverwrite = throwOnOverwrite;
+        }
+
+        public Task<BlobIdentifier> PutObject(NamespaceId ns, byte[] blob, BlobIdentifier? identifier = null)
+        {
+            // we do not split the blob into smaller parts when storing in memory, this is only for test purposes
+            // so there is no need to add that complexity
+
+            if (identifier == null)
+            {
+                identifier = BlobIdentifier.FromBlob(blob);
+            }
+
+            var namespaceContainer = _blobs.GetOrAdd(ns, new ConcurrentDictionary<BlobIdentifier, BlobContainer>());
+            if (namespaceContainer.ContainsKey(identifier))
+            {
+                throw new Exception($"Blob {identifier} already exists in {ns}");
+            }
+            namespaceContainer.TryAdd(identifier, new BlobContainer(identifier, blob.ToArray()));
+
+            return Task.FromResult(identifier);
+        }
+
+        public async Task<BlobIdentifier> PutObject(NamespaceId ns, ReadOnlyMemory<byte> blob, BlobIdentifier? identifier = null)
+        {
+            return await PutObject(ns, blob: blob.ToArray(), identifier);
+        }
+
+        public async Task<BlobIdentifier> PutObject(NamespaceId ns, Stream blob, BlobIdentifier? identifier = null)
+        {
+            return await PutObject(ns, blob: await blob.ToByteArray(), identifier);
+        }
+
+        public Task<BlobContents> GetObject(NamespaceId ns, BlobIdentifier blob)
+        {
+            if (!_blobs.TryGetValue(ns, value: out var namespaceContainer))
+            {
+                throw new NamespaceNotFoundException(ns);
+            }
+
+            if (!namespaceContainer.TryGetValue(blob, value: out BlobContainer? blobContainer))
+            {
+                throw new BlobNotFoundException(ns, blob);
+            }
+
+            byte[] content = blobContainer.Contents;
+            return Task.FromResult(new BlobContents(new MemoryStream(content), content.LongLength));
+        }
+
+        public Task DeleteObject(NamespaceId ns, BlobIdentifier blob)
+        {
+            if (!_blobs.TryGetValue(ns, value: out var namespaceContainer))
+            {
+                throw new NamespaceNotFoundException(ns);
+            }
+
+            if (!namespaceContainer.TryRemove(blob, out _))
+            {
+                throw new BlobNotFoundException(ns, blob);
+            }
+
+            return Task.CompletedTask;
+        }
+
+        public Task<bool> Exists(NamespaceId ns, BlobIdentifier blob)
+        {
+            if (!_blobs.TryGetValue(ns, value: out var namespaceContainer))
+            {
+                throw new NamespaceNotFoundException(ns);
+            }
+
+            return Task.FromResult(namespaceContainer.ContainsKey(blob));
+        }
+
+        public Task DeleteNamespace(NamespaceId ns)
+        {
+            _blobs.Remove(ns, out var nsBlobs);
+            if (nsBlobs == null)
+            {
+                throw new NamespaceNotFoundException(ns);
+            }
+            
+            return Task.CompletedTask;
+        }
+
+        public async IAsyncEnumerable<BlobIdentifier> ListOldObjects(NamespaceId ns, DateTime cutoff)
+        {
+            if (!_blobs.TryGetValue(ns, value: out var namespaceContainer))
+            {
+                throw new NamespaceNotFoundException(ns);
+            }
+
+            await Task.CompletedTask;
+            foreach (BlobContainer blobContainer in namespaceContainer.Values)
+            {
+                if (blobContainer.LastModified > cutoff)
+                {
+                    continue;
+                }
+                yield return blobContainer.BlobIdentifier;
+            }
+        }
+
+        internal IEnumerable<BlobIdentifier> GetIdentifiers(NamespaceId ns)
+        {
+            if (!_blobs.TryGetValue(ns, value: out var namespaceContainer))
+            {
+                throw new NamespaceNotFoundException(ns);
+            }
+
+            return namespaceContainer.Values.Select(container => container.BlobIdentifier);
+        }
+
+        // only for unit tests to update the last modified time
+        internal void SetLastModifiedTime(NamespaceId ns, BlobIdentifier blob, DateTime modifiedTime)
+        {
+            if (!_blobs.TryGetValue(ns, value: out var namespaceContainer))
+            {
+                throw new Exception($"Namespace {ns} not found");
+            }
+
+            if (namespaceContainer.TryGetValue(blob, value: out BlobContainer? blobContainer))
+            {
+                blobContainer.LastModified = modifiedTime;
+            }
+        }
+    }
+}
diff --git a/Engine/Source/Programs/Horde/HordeStorage/Horde.Storage/Implementation/Blob/MemoryCacheBlobStore.cs b/Engine/Source/Programs/Horde/HordeStorage/Horde.Storage/Implementation/Blob/MemoryCacheBlobStore.cs
new file mode 100644
--- /dev/null
+++ b/Engine/Source/Programs/Horde/HordeStorage/Horde.Storage/Implementation/Blob/MemoryCacheBlobStore.cs
@@ -0,0 +1,121 @@
+﻿// Copyright Epic Games, Inc. All Rights Reserved.
+
+using System;
+using System.Collections.Generic;
+using System.IO;
+using System.Reflection;
+using System.Threading;
+using System.Threading.Tasks;
+using Dasync.Collections;
+using Jupiter.Implementation;
+using Jupiter.Utils;
+using Microsoft.Extensions.Caching.Memory;
+using Microsoft.Extensions.Options;
+using Serilog;
+using Serilog.Context;
+using ILogger = Serilog.ILogger;
+
+namespace Horde.Storage.Implementation
+{
+    /// <summary>
+    /// In-memory only store implementing the IBlobStore interface using Microsoft's MemoryCache
+    /// Primarily intended for use together with HierarchicalBlobStore to speed up reads from other stores, such as Amazon S3 and disk.
+    /// </summary>
+    // ReSharper disable once ClassNeverInstantiated.Global
+    public class MemoryCacheBlobStore : IBlobStore
+    {
+        private readonly IOptionsMonitor<MemoryCacheBlobSettings> _memoryOptions;
+
+        private readonly MemoryCache _memoryCache;
+
+        private readonly ILogger _logger = Log.ForContext<MemoryCacheBlobStore>();
+        private readonly FieldInfo _sizeField;
+        private Timer _cacheInfoTimer;
+
+        public MemoryCacheBlobStore(IOptionsMonitor<MemoryCacheBlobSettings> memoryOptions)
+        {
+            _memoryOptions = memoryOptions;
+            _memoryCache = new MemoryCache(memoryOptions.CurrentValue);
+            _sizeField = typeof(MemoryCache).GetField("_cacheSize", BindingFlags.NonPublic | BindingFlags.Instance) ?? throw new Exception("Unable to find internal cache size field, did the memory cache get refactored?");
+            _cacheInfoTimer = new Timer(LogCacheInfo, null, TimeSpan.Zero, TimeSpan.FromMinutes(5));
+        }
+
+        private void LogCacheInfo(object? state)
+        {
+            LogContext.Reset();
+            _logger.Information("Memory Cache Blobs {Size} {Count}", GetMemoryCacheSize(), _memoryCache.Count);
+        }
+
+        public Task<BlobIdentifier> PutObject(NamespaceId ns, byte[] blob, BlobIdentifier identifier)
+        {
+            using ICacheEntry entry = _memoryCache.CreateEntry(BuildKey(ns, identifier));
+            entry.Value = blob;
+            entry.Size = blob.Length;
+            
+            if (_memoryOptions.CurrentValue.EnableSlidingExpiry)
+            {
+                entry.SlidingExpiration = TimeSpan.FromMinutes(_memoryOptions.CurrentValue.SlidingExpirationMinutes);
+            }
+
+            return Task.FromResult(identifier);
+        }
+
+        public Task<BlobIdentifier> PutObject(NamespaceId ns, ReadOnlyMemory<byte> blob, BlobIdentifier identifier)
+        {
+            return PutObject(ns, blob: blob.ToArray(), identifier);
+        }
+
+        public async Task<BlobIdentifier> PutObject(NamespaceId ns, Stream blob, BlobIdentifier identifier)
+        {
+            if (blob.Length > int.MaxValue)
+                throw new BlobToLargeException(identifier);
+
+            return await PutObject(ns, blob: await blob.ToByteArray(), identifier);
+        }
+
+        public Task<BlobContents> GetObject(NamespaceId ns, BlobIdentifier identifier)
+        {
+            bool exists = _memoryCache.TryGetValue(BuildKey(ns, identifier), out object? dataObj);
+            if (!exists)
+            {
+                throw new BlobNotFoundException(ns, identifier);
+            }
+
+            byte[] data = (dataObj as byte[])!;
+            return Task.FromResult(new BlobContents(new MemoryStream(data), data.LongLength));
+        }
+
+        public Task DeleteObject(NamespaceId ns, BlobIdentifier blobIdentifier)
+        {
+            _memoryCache.Remove(BuildKey(ns, blobIdentifier));
+            return Task.CompletedTask;
+        }
+
+        public Task<bool> Exists(NamespaceId ns, BlobIdentifier blobIdentifier)
+        {
+            return Task.FromResult(_memoryCache.TryGetValue(BuildKey(ns, blobIdentifier), out _));
+        }
+
+        public Task DeleteNamespace(NamespaceId ns)
+        {
+            // Silently ignore this request as we cannot implement this using MemoryCache
+            return Task.CompletedTask;
+        }
+
+        public IAsyncEnumerable<BlobIdentifier> ListOldObjects(NamespaceId ns, DateTime cutoff)
+        {
+            // Silently ignore this request as we cannot implement this using MemoryCache
+            return Array.Empty<BlobIdentifier>().ToAsyncEnumerable();
+        }
+        
+        private string BuildKey(NamespaceId ns, BlobIdentifier blob)
+        {
+            return $"{ns}.{blob}";
+        }
+
+        internal long? GetMemoryCacheSize()
+        {
+            return (long?)_sizeField.GetValue(_memoryCache);
+        }
+    }
+}
diff --git a/Engine/Source/Programs/Horde/HordeStorage/Horde.Storage/Implementation/ContentId/CompressedBufferUtils.cs b/Engine/Source/Programs/Horde/HordeStorage/Horde.Storage/Implementation/ContentId/CompressedBufferUtils.cs
new file mode 100644
--- /dev/null
+++ b/Engine/Source/Programs/Horde/HordeStorage/Horde.Storage/Implementation/ContentId/CompressedBufferUtils.cs
@@ -0,0 +1,325 @@
+﻿// Copyright Epic Games, Inc. All Rights Reserved.
+
+using System;
+using System.Buffers.Binary;
+using System.IO;
+using Blake3;
+using Datadog.Trace;
+using Force.Crc32;
+using Jupiter;
+using Jupiter.Utils;
+using K4os.Compression.LZ4;
+
+namespace Horde.Storage.Implementation
+{
+    public class CompressedBufferUtils
+    {
+        private readonly OodleCompressor _oodleCompressor;
+
+        public CompressedBufferUtils(OodleCompressor oodleCompressor)
+        {
+            _oodleCompressor = oodleCompressor;
+        }
+
+        public class Header
+        {
+            public enum CompressionMethod : byte
+            {
+                // Header is followed by one uncompressed block. 
+                None = 0,
+                // Header is followed by an array of compressed block sizes then the compressed blocks. 
+                Oodle=3,
+                LZ4 = 4,
+            }
+
+            public static uint ExpectedMagic = 0xb7756362; // <dot>ucb
+            public static uint HeaderLength = 64;
+
+            // A magic number to identify a compressed buffer. Always 0xb7756362.
+            public uint Magic;
+            // A CRC-32 used to check integrity of the buffer. Uses the polynomial 0x04c11db7.
+            public uint Crc32;
+            // The method used to compress the buffer. Affects layout of data following the header. 
+            public CompressionMethod Method;
+            public byte CompressionLevel;
+            public byte CompressionMethodUsed;
+
+            // The power of two size of every uncompressed block except the last. Size is 1 << BlockSizeExponent. 
+            public byte BlockSizeExponent;
+
+            // The number of blocks that follow the header. 
+            public uint BlockCount;
+
+            // The total size of the uncompressed data. 
+            public ulong TotalRawSize;
+
+            // The total size of the compressed data including the header. 
+            public ulong TotalCompressedSize;
+
+            /** The hash of the uncompressed data. */
+            public byte[] RawHash = Array.Empty<byte>();
+
+            public void ByteSwap()
+            {
+                Magic = BinaryPrimitives.ReverseEndianness(Magic);
+                Crc32 = BinaryPrimitives.ReverseEndianness(Crc32);
+                BlockCount = BinaryPrimitives.ReverseEndianness(BlockCount);
+                TotalRawSize = BinaryPrimitives.ReverseEndianness(TotalRawSize);
+                TotalCompressedSize = BinaryPrimitives.ReverseEndianness(TotalCompressedSize);
+            }
+        }
+
+        public Header ExtractHeader(byte[] content)
+        {
+            // the header is always stored big endian
+            bool needsByteSwap = BitConverter.IsLittleEndian;
+            if (content.Length < Header.HeaderLength)
+                throw new ArgumentOutOfRangeException(nameof(content), $"Content was less then {Header.HeaderLength} bytes and thus is not a compressed buffer");
+
+            using MemoryStream ms = new MemoryStream(content);
+            using BinaryReader reader = new BinaryReader(ms);
+
+            Header header = new Header
+            {
+                Magic = reader.ReadUInt32(),
+                Crc32 = reader.ReadUInt32(),
+                Method = (Header.CompressionMethod) reader.ReadByte(),
+                CompressionLevel = reader.ReadByte(),
+                CompressionMethodUsed = reader.ReadByte(),
+                BlockSizeExponent = reader.ReadByte(),
+                BlockCount = reader.ReadUInt32(),
+                TotalRawSize = reader.ReadUInt64(),
+                TotalCompressedSize = reader.ReadUInt64()
+            };
+            byte[] hash = reader.ReadBytes(32); // a full blake3 hash
+            header.RawHash = hash;
+
+            if (needsByteSwap)
+                header.ByteSwap();
+
+            if (ms.Position != Header.HeaderLength)
+                throw new Exception($"Read {ms.Position} bytes but expected to read {Header.HeaderLength}");
+
+            if (header.Magic != Header.ExpectedMagic)
+                throw new InvalidMagicException(header.Magic, Header.ExpectedMagic);
+
+            // calculate the crc from the start of the method field (skipping magic which is a constant and the crc field itself)
+            const int methodOffset = sizeof(uint) + sizeof(uint);
+
+            // none compressed objects have no extra blocks
+            uint blocksByteUsed = header.Method != Header.CompressionMethod.None ? header.BlockCount * (uint)sizeof(uint) : 0;
+            uint calculatedCrc = Crc32Algorithm.Compute(content, methodOffset, (int)(Header.HeaderLength - methodOffset + blocksByteUsed));
+            
+            if (header.Crc32 != calculatedCrc)
+                throw new InvalidHashException(header.Crc32, calculatedCrc);
+
+            return header;
+        }
+
+
+        public byte[] DecompressContent(byte[] content)
+        {
+            Header header = ExtractHeader(content);
+
+            if (content.LongLength < (long)header.TotalCompressedSize)
+                throw new Exception($"Expected buffer to be {header.TotalCompressedSize} but it was {content.LongLength}");
+
+            ulong decompressedPayloadOffset = 0;
+            byte[] decompressedPayload = new byte[header.TotalRawSize];
+
+            ReadOnlySpan<byte> memory = new ReadOnlySpan<byte>(content);
+            memory = memory.Slice((int)Header.HeaderLength);
+
+            bool willHaveBlocks = header.Method != Header.CompressionMethod.None;
+            if (willHaveBlocks)
+            {
+                uint[] compressedBlockSizes = new uint[header.BlockCount];
+                for (int i = 0; i < header.BlockCount; i++)
+                {
+                    uint compressedBlockSize = BinaryPrimitives.ReadUInt32BigEndian(memory);
+                    compressedBlockSizes[i] = compressedBlockSize;
+                    memory = memory.Slice(sizeof(uint));
+                }
+
+                ulong blockSize = 1ul << header.BlockSizeExponent;
+                ulong compressedOffset = 0;
+
+                foreach (uint compressedBlockSize in compressedBlockSizes)
+                {
+                    ulong rawBlockSize = Math.Min(header.TotalRawSize - decompressedPayloadOffset, blockSize);
+                    ReadOnlySpan<byte> compressedPayload = memory.Slice((int)compressedOffset, (int)compressedBlockSize);
+                    Span<byte> targetSpan = new Span<byte>(decompressedPayload, (int)decompressedPayloadOffset, (int)rawBlockSize);
+
+                    int writtenBytes;
+                    // if a block has the same raw and compressed size its uncompressed and we should not attempt to decompress it
+                    if (rawBlockSize == compressedBlockSize)
+                    {
+                        writtenBytes = (int)rawBlockSize;
+                        compressedPayload.CopyTo(targetSpan);
+                    }
+                    else
+                    {
+                        writtenBytes = DecompressPayload(compressedPayload, header, rawBlockSize, targetSpan);
+                    }
+
+                    decompressedPayloadOffset += (uint)writtenBytes;
+                    compressedOffset += compressedBlockSize;
+                }
+            }
+            else
+            {
+                // if no compression is applied there are no extra blocks and just a single chunk that is uncompressed
+                Span<byte> targetSpan = new Span<byte>(decompressedPayload, 0, (int)header.TotalRawSize);
+                DecompressPayload(memory, header, header.TotalRawSize, targetSpan);
+                decompressedPayloadOffset = header.TotalRawSize;
+            }
+           
+            if (header.TotalRawSize != decompressedPayloadOffset)
+                throw new Exception("Did not decompress the full payload");
+
+            {
+                using Scope _ = Tracer.Instance.StartActive("cb.hash");
+                using Hasher hasher = Hasher.New();
+                hasher.UpdateWithJoin(decompressedPayload);
+                Hash blake3Hash = hasher.Finalize();
+
+                Span<byte> hash = blake3Hash.AsSpan();
+                if (!new ByteArrayComparer().Equals(hash.ToArray(), header.RawHash))
+                    throw new Exception($"Payload was expected to be {StringUtils.FormatAsHexString(header.RawHash)} but was {StringUtils.FormatAsHexString(hash.ToArray())}");
+            }
+
+
+            return decompressedPayload;
+        }
+
+        private int DecompressPayload(ReadOnlySpan<byte> compressedPayload, Header header, ulong rawBlockSize, Span<byte> target)
+        {
+             switch (header.Method)
+            {
+                case Header.CompressionMethod.None:
+                    compressedPayload.CopyTo(target);
+                    return compressedPayload.Length;
+                case Header.CompressionMethod.Oodle:
+                {
+                    // we have separate enums from the official oodle api to make sure they are stable enums so we need to translate to the ones expected by the api
+                    OodleLZ_Compressor compressor = OodleUtils.ToOodleApiCompressor((OoodleCompressorMethod)header.CompressionMethodUsed);
+                    OodleLZ_CompressionLevel compressionLevel = OodleUtils.ToOodleApiCompressionLevel((OoodleCompressionLevel)header.CompressionLevel);
+                    long writtenBytes = _oodleCompressor.Decompress(compressor,compressedPayload.ToArray(), compressionLevel, (long)rawBlockSize, out byte[] result);
+                    if (writtenBytes == 0)
+                    {
+                        throw new Exception("Failed to run oodle decompress");
+                    }
+                    result.CopyTo(target);
+                    return (int)writtenBytes;
+                }
+                case Header.CompressionMethod.LZ4:
+                {
+                    int writtenBytes = LZ4Codec.Decode(compressedPayload, target);
+                    return writtenBytes;
+                }
+                default:
+                    throw new ArgumentOutOfRangeException(nameof(header.Method), header.Method, null);
+            }
+        }
+    }
+
+    internal class InvalidHashException : Exception
+    {
+        public InvalidHashException(uint headerCrc32, uint calculatedCrc) : base($"Header specified crc \"{headerCrc32}\" but calculated hash was \"{calculatedCrc}\"")
+        {
+            
+        }
+    }
+
+    internal class InvalidMagicException : Exception
+    {
+        public InvalidMagicException(uint headerMagic, uint expectedMagic) : base($"Header magic \"{headerMagic}\" was incorrect, expected to be {expectedMagic}")
+        {
+        }
+    }
+
+    // from OodleDataCompression.h , we define our own enums for oodle compressions used and convert to the ones expected in the oodle api
+    public enum OoodleCompressorMethod: byte
+    {
+        NotSet = 0,
+        Selkie = 1,
+        Mermaid = 2,
+        Kraken  = 3,
+        Leviathan = 4
+    }
+
+    public enum OoodleCompressionLevel : sbyte
+    {
+        HyperFast4 = -4,
+        HyperFast3 = -3,
+        HyperFast2 = -2,
+        HyperFast1 = -1,
+        None = 0,
+        SuperFast = 1,
+        VeryFast = 2,
+        Fast = 3,
+        Normal = 4,
+        Optimal1 = 5,
+        Optimal2 = 6,
+        Optimal3 = 7,
+        Optimal4 = 8,
+    }
+
+    public static class OodleUtils
+    {
+        public static OodleLZ_CompressionLevel ToOodleApiCompressionLevel(OoodleCompressionLevel compressionLevel)
+        {
+            switch (compressionLevel)
+            {
+                case OoodleCompressionLevel.HyperFast4:
+                    return OodleLZ_CompressionLevel.OodleLZ_CompressionLevel_HyperFast4;
+                case OoodleCompressionLevel.HyperFast3:
+                    return OodleLZ_CompressionLevel.OodleLZ_CompressionLevel_HyperFast3;
+                case OoodleCompressionLevel.HyperFast2:
+                    return OodleLZ_CompressionLevel.OodleLZ_CompressionLevel_HyperFast2;
+                case OoodleCompressionLevel.HyperFast1:
+                    return OodleLZ_CompressionLevel.OodleLZ_CompressionLevel_HyperFast1;
+                case OoodleCompressionLevel.None:
+                    return OodleLZ_CompressionLevel.OodleLZ_CompressionLevel_None;
+                case OoodleCompressionLevel.SuperFast:
+                    return OodleLZ_CompressionLevel.OodleLZ_CompressionLevel_SuperFast;
+                case OoodleCompressionLevel.VeryFast:
+                    return OodleLZ_CompressionLevel.OodleLZ_CompressionLevel_VeryFast;
+                case OoodleCompressionLevel.Fast:
+                    return OodleLZ_CompressionLevel.OodleLZ_CompressionLevel_Fast;
+                case OoodleCompressionLevel.Normal:
+                    return OodleLZ_CompressionLevel.OodleLZ_CompressionLevel_Normal;
+                case OoodleCompressionLevel.Optimal1:
+                    return OodleLZ_CompressionLevel.OodleLZ_CompressionLevel_Optimal1;
+                case OoodleCompressionLevel.Optimal2:
+                    return OodleLZ_CompressionLevel.OodleLZ_CompressionLevel_Optimal2;
+                case OoodleCompressionLevel.Optimal3:
+                    return OodleLZ_CompressionLevel.OodleLZ_CompressionLevel_Optimal3;
+                case OoodleCompressionLevel.Optimal4:
+                    return OodleLZ_CompressionLevel.OodleLZ_CompressionLevel_Optimal4;
+                default:
+                    throw new ArgumentOutOfRangeException(nameof(compressionLevel), compressionLevel, null);
+            }
+        }
+
+        public static OodleLZ_Compressor ToOodleApiCompressor(OoodleCompressorMethod compressor)
+        {
+            switch (compressor)
+            {
+                case OoodleCompressorMethod.NotSet:
+                    return OodleLZ_Compressor.OodleLZ_Compressor_None;
+                case OoodleCompressorMethod.Selkie:
+                    return OodleLZ_Compressor.OodleLZ_Compressor_Selkie;
+                case OoodleCompressorMethod.Mermaid:
+                    return OodleLZ_Compressor.OodleLZ_Compressor_Mermaid;
+                case OoodleCompressorMethod.Kraken:
+                    return OodleLZ_Compressor.OodleLZ_Compressor_Kraken;
+                case OoodleCompressorMethod.Leviathan:
+                    return OodleLZ_Compressor.OodleLZ_Compressor_Leviathan;
+                default:
+                    throw new ArgumentOutOfRangeException(nameof(compressor), compressor, null);
+            }
+        }
+
+    }
+}
diff --git a/Engine/Source/Programs/Horde/HordeStorage/Horde.Storage/Implementation/ContentId/IContentIdStore.cs b/Engine/Source/Programs/Horde/HordeStorage/Horde.Storage/Implementation/ContentId/IContentIdStore.cs
new file mode 100644
--- /dev/null
+++ b/Engine/Source/Programs/Horde/HordeStorage/Horde.Storage/Implementation/ContentId/IContentIdStore.cs
@@ -0,0 +1,47 @@
+// Copyright Epic Games, Inc. All Rights Reserved.
+
+using System;
+using System.Threading.Tasks;
+using Jupiter.Implementation;
+
+namespace Horde.Storage.Implementation
+{
+    public interface IContentIdStore
+    {
+        /// <summary>
+        /// Resolve a content id from its hash into the actual blob (that can in turn be chunked into a set of blobs)
+        /// </summary>
+        /// <param name="ns">The namespace to operate in</param>
+        /// <param name="contentId">The identifier for the content id</param>
+        /// <returns></returns>
+        Task<BlobIdentifier[]?> Resolve(NamespaceId ns, BlobIdentifier contentId);
+
+        /// <summary>
+        /// Add a mapping from contentId to blobIdentifier
+        /// </summary>
+        /// <param name="ns">The namespace to operate in</param>
+        /// <param name="contentId">The contentId</param>
+        /// <param name="blobIdentifier">The blob the content id maps to</param>
+        /// <param name="contentWeight">Weight of this identifier compared to previous mappings, used to determine which is more important, lower weight is considered a better fit</param>
+        /// <returns></returns>
+        Task Put(NamespaceId ns, BlobIdentifier contentId, BlobIdentifier blobIdentifier, int contentWeight);
+    }
+
+    public class InvalidContentIdException : Exception
+    {
+        public InvalidContentIdException(BlobIdentifier contentId) : base($"Unknown content id {contentId}")
+        {
+            
+        }
+    }
+
+    public class ContentIdResolveException : Exception
+    {
+        public BlobIdentifier ContentId { get; }
+
+        public ContentIdResolveException(BlobIdentifier contentId): base($"Unable to find any mapping of contentId {contentId} that has all blobs present")
+        {
+            ContentId = contentId;
+        }
+    }
+}
diff --git a/Engine/Source/Programs/Horde/HordeStorage/Horde.Storage/Implementation/ContentId/MemoryContentIdStore.cs b/Engine/Source/Programs/Horde/HordeStorage/Horde.Storage/Implementation/ContentId/MemoryContentIdStore.cs
new file mode 100644
--- /dev/null
+++ b/Engine/Source/Programs/Horde/HordeStorage/Horde.Storage/Implementation/ContentId/MemoryContentIdStore.cs
@@ -0,0 +1,80 @@
+﻿// Copyright Epic Games, Inc. All Rights Reserved.
+
+using System;
+using System.Collections.Concurrent;
+using System.Collections.Generic;
+using System.Linq;
+using System.Threading.Tasks;
+using Cassandra;
+using Cassandra.Mapping;
+using Jupiter.Implementation;
+
+namespace Horde.Storage.Implementation
+{
+    public class MemoryContentIdStore : IContentIdStore
+    {
+        private readonly IBlobStore _blobStore;
+        private readonly ConcurrentDictionary<NamespaceId, ConcurrentDictionary<BlobIdentifier, SortedList<int, BlobIdentifier[]>>> _contentIds = new ConcurrentDictionary<NamespaceId, ConcurrentDictionary<BlobIdentifier, SortedList<int, BlobIdentifier[]>>>();
+        
+        public MemoryContentIdStore(IBlobStore blobStore)
+        {
+            _blobStore = blobStore;
+        }
+
+        public async Task<BlobIdentifier[]?> Resolve(NamespaceId ns, BlobIdentifier contentId)
+        {
+            if (_contentIds.TryGetValue(ns, out ConcurrentDictionary<BlobIdentifier, SortedList<int, BlobIdentifier[]>>? contentIdsForNamespace))
+            {
+                if (contentIdsForNamespace.TryGetValue(contentId, out SortedList<int, BlobIdentifier[]>? contentIdMappings))
+                {
+                    foreach ((int weight, BlobIdentifier[] blobs) in contentIdMappings)
+                    {
+                        BlobIdentifier[] missingBlobs = await _blobStore.FilterOutKnownBlobs(ns, blobs);
+                        if (missingBlobs.Length == 0)
+                            return blobs;
+                        // blobs are missing continue testing with the next content id in the weighted list as that might exist
+                    }
+                }
+            }
+
+            // if no content id is found, but we have a blob that matches the content id (so a unchunked and uncompressed version of the data) we use that instead
+            if (await _blobStore.Exists(ns, contentId))
+                return new[] { contentId };
+
+            return null;
+        }
+
+        public Task Put(NamespaceId ns, BlobIdentifier contentId, BlobIdentifier blobIdentifier, int contentWeight)
+        {
+            _contentIds.AddOrUpdate(ns, (_) =>
+            {
+                ConcurrentDictionary<BlobIdentifier, SortedList<int, BlobIdentifier[]>> dict = new()
+                {
+                    [contentId] = new SortedList<int, BlobIdentifier[]>
+                    {
+                        {contentWeight, new BlobIdentifier[] { blobIdentifier } }
+                    }
+                };
+
+                return dict;
+            }, (_, dict) =>
+            {
+                dict.AddOrUpdate(contentId, (_) =>
+                {
+                    return new SortedList<int, BlobIdentifier[]>
+                    {
+                        { contentWeight, new BlobIdentifier[] { blobIdentifier } }
+                    };
+                }, (_, mappings) =>
+                {
+                    mappings[contentWeight] = new[] { blobIdentifier };
+                    return mappings;
+                });
+
+                return dict;
+            });
+
+            return Task.CompletedTask;
+        }
+    }
+}
diff --git a/Engine/Source/Programs/Horde/HordeStorage/Horde.Storage/Implementation/ContentId/ScyllaContentIdStore.cs b/Engine/Source/Programs/Horde/HordeStorage/Horde.Storage/Implementation/ContentId/ScyllaContentIdStore.cs
new file mode 100644
--- /dev/null
+++ b/Engine/Source/Programs/Horde/HordeStorage/Horde.Storage/Implementation/ContentId/ScyllaContentIdStore.cs
@@ -0,0 +1,103 @@
+// Copyright Epic Games, Inc. All Rights Reserved.
+
+using System;
+using System.Collections.Generic;
+using System.Linq;
+using System.Threading.Tasks;
+using Cassandra;
+using Cassandra.Mapping;
+using Datadog.Trace;
+using Jupiter.Implementation;
+
+namespace Horde.Storage.Implementation
+{
+    public class ScyllaContentIdStore : IContentIdStore
+    {
+        private readonly ISession _session;
+        private readonly IBlobStore _blobStore;
+        private Mapper _mapper;
+
+        public ScyllaContentIdStore(IScyllaSessionManager scyllaSessionManager, IBlobStore blobStore)
+        {
+            _session = scyllaSessionManager.GetSessionForReplicatedKeyspace();
+            _blobStore = blobStore;
+
+            _mapper = new Mapper(_session);
+
+            _session.Execute(new SimpleStatement(@"CREATE TABLE IF NOT EXISTS content_id (
+                content_id frozen<blob_identifier>,
+                content_weight int, 
+                chunks set<frozen<blob_identifier>>, 
+                PRIMARY KEY ((content_id), content_weight)
+            );"
+            ));
+        }
+
+
+        public async Task<BlobIdentifier[]?> Resolve(NamespaceId ns, BlobIdentifier contentId)
+        {
+            using Scope scope = Tracer.Instance.StartActive("ScyllaContentIdStore.ResolveContentId");
+            scope.Span.ResourceName = contentId.ToString();
+
+            Task<bool> blobStoreExistsTask = _blobStore.Exists(ns, contentId);
+
+            // lower content_weight means its a better candidate to resolve to
+            foreach (ScyllaContentId? resolvedContentId in await _mapper.FetchAsync<ScyllaContentId>("WHERE content_id = ? ORDER BY content_weight DESC", new ScyllaBlobIdentifier(contentId)))
+            {
+                if (resolvedContentId == null)
+                    throw new InvalidContentIdException(contentId);
+                
+                BlobIdentifier[] blobs = resolvedContentId.Chunks.Select(b => b.AsBlobIdentifier()).ToArray();
+
+                {
+                    using Scope _ = Tracer.Instance.StartActive("ScyllaContentIdStore.FindMissingBlobs");
+
+                    BlobIdentifier[] missingBlobs = await _blobStore.FilterOutKnownBlobs(ns, blobs);
+                    if (missingBlobs.Length == 0)
+                        return blobs;
+                }
+                // blobs are missing continue testing with the next content id in the weighted list as that might exist
+            }
+
+            // if no content id is found, but we have a blob that matches the content id (so a unchunked and uncompressed version of the data) we use that instead
+            bool contentIdBlobExists = await blobStoreExistsTask;
+
+            if (contentIdBlobExists)
+                return new[] { contentId };
+
+            // unable to resolve the content id
+            return null;
+        }
+
+        public async Task Put(NamespaceId ns, BlobIdentifier contentId, BlobIdentifier blobIdentifier, int contentWeight)
+        {
+            await _mapper.UpdateAsync<ScyllaContentId>("SET chunks = ? WHERE content_id = ? AND content_weight = ?", new [] {new ScyllaBlobIdentifier(blobIdentifier)}, new ScyllaBlobIdentifier(contentId), contentWeight);
+        }
+    }
+
+    [Cassandra.Mapping.Attributes.Table("content_id")]
+    public class ScyllaContentId
+    {
+        public ScyllaContentId()
+        {
+
+        }
+
+        public ScyllaContentId(BlobIdentifier contentId, BlobIdentifier[] chunks, int contentWeight)
+        {
+            ContentId = new ScyllaBlobIdentifier(contentId);
+            ContentWeight = contentWeight;
+            Chunks = chunks.Select(b => new ScyllaBlobIdentifier(b)).ToArray();
+        }
+
+        [Cassandra.Mapping.Attributes.PartitionKey]
+        [Cassandra.Mapping.Attributes.Column("content_id")]
+        public ScyllaBlobIdentifier? ContentId { get; set; }
+
+        [Cassandra.Mapping.Attributes.ClusteringKey]
+        [Cassandra.Mapping.Attributes.Column("content_weight")]
+        public int? ContentWeight { get; set; }
+
+        public ScyllaBlobIdentifier[] Chunks { get; set; } = Array.Empty<ScyllaBlobIdentifier>();
+    }
+}
