Moved Horde.Storage from git to p4

diff --git a/Engine/Source/Programs/Horde/HordeStorage/Horde.Storage/Implementation/DDCRefService.cs b/Engine/Source/Programs/Horde/HordeStorage/Horde.Storage/Implementation/DDCRefService.cs
new file mode 100644
--- /dev/null
+++ b/Engine/Source/Programs/Horde/HordeStorage/Horde.Storage/Implementation/DDCRefService.cs
@@ -0,0 +1,306 @@
+// Copyright Epic Games, Inc. All Rights Reserved.
+
+using System;
+using System.Collections.Generic;
+using System.IO;
+using System.Linq;
+using System.Threading.Tasks;
+using Datadog.Trace;
+using Horde.Storage.Controllers;
+using Jupiter.Implementation;
+using Microsoft.Extensions.Options;
+using Serilog;
+
+namespace Horde.Storage.Implementation
+{
+    public class RefResponse
+    {
+        public RefResponse(string name, DateTime? lastAccessTime, ContentHash contentHash, BlobIdentifier[] blobIdentifiers, Dictionary<string, object>? metadata)
+        {
+            Name = name;
+            LastAccessTime = lastAccessTime;
+            Metadata = metadata;
+            ContentHash = contentHash;
+            BlobIdentifiers = blobIdentifiers;
+        }
+
+        public string Name { get; }
+        public DateTime? LastAccessTime { get; }
+        public Dictionary<string, object>? Metadata { get; }
+
+        public ContentHash ContentHash { get; }
+        public BlobIdentifier[] BlobIdentifiers { get; }
+
+        public byte[]? Blob { get; set; }
+    }
+
+    public interface IDDCRefService
+    {
+        Task<(RefResponse, BlobContents?)> Get(NamespaceId ns, BucketId bucket, KeyId key, string[] fields);
+
+        Task<RefRecord> Exists(NamespaceId ns, BucketId bucket, KeyId key);
+
+        Task<PutRequestResponse> Put(NamespaceId ns, BucketId bucket, KeyId key, ContentHash blobHash, byte[] blob);
+
+        /// <summary>
+        /// Adds a ref key for which all blobs are already expected to be uploaded
+        /// </summary>
+        /// <param name="ns"></param>
+        /// <param name="bucket"></param>
+        /// <param name="key"></param>
+        /// <param name="refRequest"></param>
+        /// <returns></returns>
+        Task<long> PutIndirect(NamespaceId ns, BucketId bucket, KeyId key, RefRequest refRequest);
+
+        Task<long> Delete(NamespaceId ns, BucketId bucket, KeyId key);
+    }
+
+    public class DDCRefService : IDDCRefService
+    {
+        private readonly IBlobStore _blobStore;
+        private readonly IRefsStore _refsStore;
+        private readonly ITransactionLogWriter _transactionLog;
+        private readonly ILastAccessTracker<RefRecord> _lastAccessTracker;
+        private readonly IDiagnosticContext _diagnosticContext;
+        private readonly HordeStorageSettings _settings;
+        private readonly ILogger _logger = Log.ForContext<DDCRefService>();
+
+        public DDCRefService(IBlobStore blobStore, IRefsStore refsStore, ITransactionLogWriter transactionLog, ILastAccessTracker<RefRecord> lastAccessTracker,
+            IDiagnosticContext diagnosticContext, IOptionsMonitor<HordeStorageSettings> settings)
+        {
+            _blobStore = blobStore;
+            _refsStore = refsStore;
+            _transactionLog = transactionLog;
+            _lastAccessTracker = lastAccessTracker;
+            _diagnosticContext = diagnosticContext;
+            _settings = settings.CurrentValue;
+        }
+
+        private static bool HasFieldFilter(string[] fields)
+        {
+            return fields.Length != 0 && fields.Any(field => !string.IsNullOrEmpty(field));
+        }
+
+        public async Task<(RefResponse, BlobContents?)> Get(NamespaceId ns, BucketId bucket, KeyId key, string[] fields) { 
+            bool needsBlob = !HasFieldFilter(fields) || fields.Contains("blob", StringComparer.InvariantCultureIgnoreCase);
+            bool needsLastAccess = !HasFieldFilter(fields) || fields.Contains("lastAccessTime", StringComparer.InvariantCultureIgnoreCase);
+            bool needsMetadata = !HasFieldFilter(fields) || fields.Contains("metadata", StringComparer.InvariantCultureIgnoreCase);
+
+            IRefsStore.ExtraFieldsFlag flags = IRefsStore.ExtraFieldsFlag.None;
+            if (needsLastAccess)
+                flags |= IRefsStore.ExtraFieldsFlag.LastAccess;
+            if (needsMetadata)
+                flags |= IRefsStore.ExtraFieldsFlag.Metadata;
+
+            string resource = $"{ns}.{bucket}.{key}";
+            RefRecord? record;
+            using (Scope scope = Tracer.Instance.StartActive("ref.get"))
+            {
+                scope.Span.ResourceName = resource;
+                record = await _refsStore.Get(ns, bucket, key, flags);
+                if (record == null)
+                {
+                    throw new RefRecordNotFoundException(ns, bucket, key);
+                }
+            }
+
+            // its not critical that this finishes, so we just log errors in case it fails but never await
+            {
+                Task _ = _lastAccessTracker.TrackUsed(record).ContinueWith(task =>
+                {
+                    if (task.Exception != null)
+                    {
+                        _logger.Error(task.Exception, "Exception when tracking last access record");
+                    }
+                });
+            }
+
+            BlobContents? maybeBlob = null;
+
+            if (needsBlob)
+            {
+                using Scope scope = Tracer.Instance.StartActive("blob.get");
+                scope.Span.ResourceName = resource;
+                scope.Span.SetTag("BlobCount", record.Blobs.Length.ToString());
+
+                if (record.Blobs.Length == 1)
+                {
+                    // stream directly to the output if the blob is not chunked
+                    maybeBlob = await _blobStore.GetObject(ns, record.Blobs[0]);
+                }
+                else
+                {
+                    maybeBlob = await _blobStore.GetObjects(ns, record.Blobs);
+                }
+                // add content length to the current span so we can determine if the request was working on a large payload
+                scope?.Span.SetTag("Content-Length", maybeBlob.Length.ToString());
+
+                _diagnosticContext.Set("Content-Length", maybeBlob.Length);
+            }
+
+            RefResponse response = new RefResponse(
+                name: $"{record.Namespace}.{record.Bucket}.{record.RefName}",
+                contentHash: record.ContentHash,
+                lastAccessTime: record.LastAccessTime,
+                blobIdentifiers: record.Blobs,
+                metadata: record.Metadata
+            );
+            return (response, maybeBlob);
+        }
+
+        public async Task<RefRecord> Exists(
+            NamespaceId ns,
+            BucketId bucket,
+            KeyId key)
+        {
+
+            RefRecord? record = await _refsStore.Get(ns, bucket, key, IRefsStore.ExtraFieldsFlag.None);
+            if (record == null)
+            {
+                throw new RefRecordNotFoundException(ns, bucket, key);
+            }
+
+            // we have to verify the blobs are available locally, as the record of the key is replicated a head of the content
+            BlobIdentifier[] unknownBlobs = await _blobStore.FilterOutKnownBlobs(ns, record.Blobs);
+            if (unknownBlobs.Length != 0)
+            {
+                throw new MissingBlobsException(ns, bucket, key, unknownBlobs);
+            }
+
+            return record;
+        }
+
+
+        public async Task<PutRequestResponse> Put(
+            NamespaceId ns,
+            BucketId bucket,
+            KeyId key,
+            ContentHash blobHash,
+            byte[] blob)
+        {
+            BlobIdentifier[] blobReferences;
+            Memory<byte>[]? slices;
+            if (!_settings.MaxSingleBlobSize.HasValue || blob.Length < _settings.MaxSingleBlobSize)
+            {
+                // we do not need to chunk, lets just use the hash we have already calculated
+                blobReferences = new[] { BlobIdentifier.FromContentHash(blobHash) };
+
+                // we should not use the sliced version since we can just pass the whole byte[] instead which if faster
+                slices = null;
+            }
+            else
+            {
+                _logger.Debug("Blob was to large, starting partitioning");
+                // partition the blob
+                Memory<byte> memory = blob.AsMemory();
+
+                int countOfSlices = (int) (blob.LongLength / _settings.MaxSingleBlobSize);
+                // if there was any re-meaning buffer we need to allocate another slice to fill that.
+                countOfSlices += blob.LongLength % _settings.MaxSingleBlobSize != 0 ? 1 : 0;
+
+                blobReferences = new BlobIdentifier[countOfSlices];
+                slices = new Memory<byte>[countOfSlices];
+                int maxSingleBlobSize = _settings.MaxSingleBlobSize.Value;
+                for (int i = 0; i < countOfSlices; i++)
+                {
+                    int blobStart = maxSingleBlobSize * i;
+                    int length = Math.Min(maxSingleBlobSize, val2: blob.Length - blobStart);
+                    slices[i] = memory.Slice(blobStart, length);
+                    BlobIdentifier hashIdentifier = BlobIdentifier.FromBlob(slices[i]);
+                    blobReferences[i] = hashIdentifier;
+                }
+                _logger.Debug("Blobs partitioned into {PartitionCount} slices with {Identifiers}", countOfSlices, blobReferences);
+            }
+            int countOfBlobs = blobReferences.Length;
+            if (countOfBlobs == 0)
+                throw new ArgumentException("No blobs found when determining partitioning");
+
+            // as the body is just a binary blob we are unable to receive metadata for this route
+            // TODO: We could accept metadata from query parameters
+            RefRecord record = new RefRecord(ns, bucket, key, metadata: null, blobs: blobReferences,
+                lastAccessTime: DateTime.Now, contentHash: blobHash);
+
+            // we have to insert (and wait for flush) into the transaction log before we attempt to upload the blobs to avoid race conditions when cleaning up
+            long transactionId = await _transactionLog.Add(ns, record.ToAddTransactionEvent());
+
+            Task blobInsertTask;
+            if (countOfBlobs == 1)
+            {
+                // use the faster byte[] path if we do not need to partition the blob
+                blobInsertTask = _blobStore.PutObject(ns, blob, blobReferences[0]);
+            }
+            else
+            {
+                if (slices == null)
+                    throw new Exception("Slices was never set when uploading partitioned blob.");
+
+                Task[] insertTasks = new Task[countOfBlobs];
+                for (int i = 0; i < countOfBlobs; i++)
+                {
+                    insertTasks[i] = _blobStore.PutObject(ns, slices[i], blobReferences[i]);
+                }
+
+                blobInsertTask = Task.WhenAll(insertTasks);
+            }
+
+            Task addRefStore = _refsStore.Add(record);
+            await Task.WhenAll(addRefStore, blobInsertTask);
+            return new PutRequestResponse(transactionId);
+        }
+
+        /// <inheritdoc/>
+        public async Task<long> PutIndirect(
+            NamespaceId ns,
+            BucketId bucket,
+            KeyId key,
+            RefRequest refRequest)
+        {
+            if (refRequest.BlobReferences == null)
+            {
+                throw new Exception("BlobReferences missing from request");
+            }
+
+            // TODO: should they just inline the blobs instead so we can submit them like we do for the blob put?
+
+            BlobIdentifier[] unknownBlobs = await _blobStore.FilterOutKnownBlobs(ns, refRequest.BlobReferences);
+            if (unknownBlobs.Length != 0)
+            {
+               throw new MissingBlobsException(ns, bucket, key, unknownBlobs);
+            }
+
+            await using BlobContents blobContents = refRequest.BlobReferences.Length == 1
+                ?
+                // avoid reallocating the stream if its a single object that is not chunked
+                await _blobStore.GetObject(ns, refRequest.BlobReferences[0])
+                : await _blobStore.GetObjects(ns, refRequest.BlobReferences);
+
+            // TODO: This is downloading the blob from the store just to verify that the content hash is correct
+            // this is quite expensive and slow, the blob is content hashed so we know the hash is correct for it
+            // but if chunked we need to verify the concatenated blob is correct
+            await using MemoryStream ms = new MemoryStream();
+            await blobContents.Stream.CopyToAsync(ms);
+            ContentHash actualHash = ContentHash.FromBlob(ms.ToArray());
+            ContentHash fullHash = refRequest.ContentHash;
+
+            if (!fullHash.Equals(actualHash))
+            {
+                throw new HashMismatchException(fullHash, actualHash);
+            }
+
+            RefRecord record = new RefRecord(ns, bucket, key,
+                refRequest.BlobReferences, DateTime.Now, actualHash, refRequest.Metadata);
+            long transactionId = await _transactionLog.Add(ns, record.ToAddTransactionEvent());
+            await _refsStore.Add(record);
+
+            return transactionId;
+        }
+
+        public async Task<long> Delete(NamespaceId ns, BucketId bucket, KeyId key)
+        {
+            long deleteCount = await _refsStore.Delete(ns, bucket, key);
+            if (deleteCount != 0)
+                await _transactionLog.Delete(ns, bucket, key);
+            return deleteCount;
+        }
+    }
+}
diff --git a/Engine/Source/Programs/Horde/HordeStorage/Horde.Storage/Implementation/GC/BlobCleanupService.cs b/Engine/Source/Programs/Horde/HordeStorage/Horde.Storage/Implementation/GC/BlobCleanupService.cs
new file mode 100644
--- /dev/null
+++ b/Engine/Source/Programs/Horde/HordeStorage/Horde.Storage/Implementation/GC/BlobCleanupService.cs
@@ -0,0 +1,87 @@
+﻿// Copyright Epic Games, Inc. All Rights Reserved.
+
+using System;
+using System.Collections.Generic;
+using System.Threading;
+using System.Threading.Tasks;
+using Jupiter;
+using Microsoft.Extensions.Options;
+using Datadog.Trace;
+using Serilog;
+
+namespace Horde.Storage.Implementation
+{
+    public class BlobCleanupService : PollingService<BlobCleanupService.BlobCleanupState>
+    {
+        private readonly IOptionsMonitor<GCSettings> _settings;
+        private volatile bool _alreadyPolling;
+        private readonly ILogger _logger = Log.ForContext<BlobCleanupService>();
+
+        public override bool ShouldStartPolling()
+        {
+            return _settings.CurrentValue.BlobCleanupServiceEnabled;
+        }
+
+        public class BlobCleanupState
+        {
+            public List<IBlobCleanup> BlobCleanups { get; } = new List<IBlobCleanup>();
+        }
+
+        public BlobCleanupService(IOptionsMonitor<GCSettings> settings) :
+            base(serviceName: nameof(RefCleanupService), settings.CurrentValue.BlobCleanupPollFrequency,
+                new BlobCleanupState())
+        {
+            _settings = settings;
+        }
+
+        public void RegisterCleanup(IBlobCleanup cleanup)
+        {
+            State.BlobCleanups.Add(cleanup);
+        }
+
+        public override async Task<bool> OnPoll(BlobCleanupState state, CancellationToken cancellationToken)
+        {
+            if (_alreadyPolling)
+                return false;
+
+            _alreadyPolling = true;
+            try
+            {
+                await Cleanup(state, cancellationToken);
+                return true;
+            }
+            finally
+            {
+                _alreadyPolling = false;
+            }
+        }
+
+        public async Task<List<RemovedBlobs>> Cleanup(BlobCleanupState state, CancellationToken cancellationToken)
+        {
+            List<RemovedBlobs> removedBlobs = new List<RemovedBlobs>();
+            
+            foreach (IBlobCleanup blobCleanup in state.BlobCleanups)
+            {
+                string type = blobCleanup.GetType().ToString();
+                _logger.Information("Blob cleanup running for {BlobCleanup}", type);
+                using Scope scope = Tracer.Instance.StartActive("gc.blob");
+                scope.Span.SetTag("type", type);
+
+                _logger.Information("Attempting to run Blob Cleanup {BlobCleanup}. ", type);
+                try
+                {
+                    List<RemovedBlobs> tempBlobs = await blobCleanup.Cleanup(cancellationToken);
+                    _logger.Information("Ran blob cleanup {BlobCleanup}. Deleted {CountBlobRecords}", type, tempBlobs.Count);
+                    removedBlobs.AddRange(tempBlobs);
+                }
+                catch (Exception e)
+                {
+                    _logger.Error(e, "Exception running Blob Cleanup {BlobCleanup} .", type);
+                }
+            }
+
+            return removedBlobs;
+        }
+    }
+}
+
diff --git a/Engine/Source/Programs/Horde/HordeStorage/Horde.Storage/Implementation/GC/CacheCleanup.cs b/Engine/Source/Programs/Horde/HordeStorage/Horde.Storage/Implementation/GC/CacheCleanup.cs
new file mode 100644
--- /dev/null
+++ b/Engine/Source/Programs/Horde/HordeStorage/Horde.Storage/Implementation/GC/CacheCleanup.cs
@@ -0,0 +1,59 @@
+﻿// Copyright Epic Games, Inc. All Rights Reserved.
+
+using System;
+using System.Collections.Generic;
+using System.Threading;
+using System.Threading.Tasks;
+using Datadog.Trace;
+using Jupiter.Implementation;
+using Microsoft.Extensions.Options;
+using Serilog;
+
+namespace Horde.Storage.Implementation
+{
+    public interface IRefCleanup
+    {
+        Task<List<OldRecord>> Cleanup(NamespaceId ns, CancellationToken cancellationToken);
+    }
+
+    public class RefCleanup : IRefCleanup
+    {
+        private readonly IOptionsMonitor<GCSettings> _settings;
+        private readonly IRefsStore _refs;
+        private readonly ITransactionLogWriter _transactionLog;
+        private readonly ILogger _logger = Log.ForContext<RefCleanup>();
+
+        public RefCleanup(IOptionsMonitor<GCSettings> settings, IRefsStore refs, ITransactionLogWriter transactionLog)
+        {
+            _settings = settings;
+            _refs = refs;
+            _transactionLog = transactionLog;
+        }
+
+        public async Task<List<OldRecord>> Cleanup(NamespaceId ns, CancellationToken cancellationToken)
+        {
+            using Scope scope = Tracer.Instance.StartActive("gc.refs.namespace");
+            scope.Span.ResourceName = ns.ToString();
+            List<OldRecord> deletedRecords = new List<OldRecord>();
+            await foreach (OldRecord record in _refs.GetOldRecords(ns, _settings.CurrentValue.LastAccessCutoff).WithCancellation(cancellationToken))
+            {
+                // delete the old record from the ref refs
+                Task storeDelete = _refs.Delete(record.Namespace, record.Bucket, record.RefName);
+                // insert a delete event into the transaction log
+                Task transactionLogDelete = _transactionLog.Delete(record.Namespace, record.Bucket, record.RefName);
+
+                try
+                {
+                    Task.WaitAll(storeDelete, transactionLogDelete);
+                }
+                catch (Exception e)
+                {
+                    _logger.Warning(e, "Exception when attempting to delete record {Record} in {Namespace}", record, ns);
+                }
+                deletedRecords.Add(record);
+            }
+
+            return deletedRecords;
+        }
+    }
+}
diff --git a/Engine/Source/Programs/Horde/HordeStorage/Horde.Storage/Implementation/GC/OrphanBlobCleanup.cs b/Engine/Source/Programs/Horde/HordeStorage/Horde.Storage/Implementation/GC/OrphanBlobCleanup.cs
new file mode 100644
--- /dev/null
+++ b/Engine/Source/Programs/Horde/HordeStorage/Horde.Storage/Implementation/GC/OrphanBlobCleanup.cs
@@ -0,0 +1,246 @@
+﻿// Copyright Epic Games, Inc. All Rights Reserved.
+
+using System;
+using System.Collections.Concurrent;
+using System.Collections.Generic;
+using System.Linq;
+using System.Threading;
+using System.Threading.Tasks;
+using async_enumerable_dotnet;
+using Dasync.Collections;
+using Jupiter;
+using Jupiter.Implementation;
+using Microsoft.Extensions.Options;
+using RestSharp;
+using RestSharp.Serializers.NewtonsoftJson;
+using Serilog;
+
+namespace Horde.Storage.Implementation
+{
+    public class RemovedBlobs
+    {
+        public RemovedBlobs(BlobIdentifier blobIdentifier)
+        {
+            BlobIdentifier = blobIdentifier;
+        }
+
+        public BlobIdentifier BlobIdentifier { get; set; }
+    }
+
+    public interface IBlobCleanup
+    {
+        Task<List<RemovedBlobs>> Cleanup(CancellationToken none);
+    }
+
+    public class OrphanBlobCleanup : IBlobCleanup
+    {
+        private readonly IBlobStore _blobStore;
+        private readonly ILeaderElection _leaderElection;
+        private readonly IRestClient _client;
+        private readonly ILogger _logger = Log.ForContext<OrphanBlobCleanup>();
+
+        // ReSharper disable once UnusedMember.Global
+        public OrphanBlobCleanup(IBlobStore blobStore, ILeaderElection leaderElection, IOptionsMonitor<CallistoTransactionLogSettings> callistoSettings, IServiceCredentials serviceCredentials)
+        {
+            _blobStore = blobStore;
+            _leaderElection = leaderElection;
+
+            _client = new RestClient(callistoSettings.CurrentValue.ConnectionString)
+            {
+                Authenticator = serviceCredentials.GetAuthenticator()
+            }.UseSerializer(() => new JsonNetSerializer());
+        }
+
+        internal OrphanBlobCleanup(IBlobStore blobStore, ILeaderElection leaderElection, IRestClient callistoClient)
+        {
+            _blobStore = blobStore;
+            _leaderElection = leaderElection;
+            _client = callistoClient;
+        }
+
+        private struct GCRootState
+        {
+            public long TransactionLogPointer { get; }
+            public HashSet<BlobIdentifier> GcRoots { get; }
+            public Guid TransactionLogGeneration { get; set; }
+
+            public GCRootState(long transactionLogPointer, Guid transactionLogGeneration, HashSet<BlobIdentifier> gcRoots)
+            {
+                TransactionLogPointer = transactionLogPointer;
+                TransactionLogGeneration = transactionLogGeneration;
+                GcRoots = gcRoots;
+            }
+
+        }
+
+        public async Task<List<RemovedBlobs>> Cleanup(CancellationToken cancellationToken)
+        {
+            if (!_leaderElection.IsThisInstanceLeader())
+            {
+                _logger.Information("Skipped orphan blob cleanup run as this instance is not the leader");
+                return new List<RemovedBlobs>();
+            }
+
+            List<NamespaceId> namespaces = await ListNamespaces().ToListAsync();
+            ConcurrentDictionary<NamespaceId, GCRootState> perNamespaceRoots = new();
+            await namespaces.ParallelForEachAsync(async ns =>
+            {
+                GCRootState gcRootState = await DetermineGCRoots(ns, _client, cancellationToken);
+                perNamespaceRoots.TryAdd(ns, gcRootState);
+            }, cancellationToken);
+
+            List<BlobIdentifier> blobsToRemove = new List<BlobIdentifier>();
+            List<RemovedBlobs> removedBlobs = new List<RemovedBlobs>();
+
+            // enumerate all namespaces, and check if the old blob is valid in any of them to allow for a blob store to just store them in a single pile if it wants to
+            foreach (NamespaceId @namespace in namespaces)
+            {
+                // only consider blobs that have been around for 60 minutes
+                // this due to cases were blobs are uploaded first, and the GC roots specified after the fact for DDC keys
+                await foreach (BlobIdentifier blob in _blobStore.ListOldObjects(@namespace, DateTime.Now.AddMinutes(-60)).WithCancellation(cancellationToken))
+                {
+                    // TODO: This could be turned into a parallel for as each iteration does not mutate the state
+                    // there is no parallel for each for async enumerable yet
+                    // we could manually build one, see [REDACTED_PROFILE_URL]/corefx/issues/34233 but lets wait until we know the time this takes is an issue
+
+                    bool blobFoundInRootSet = perNamespaceRoots.Values.Any(gcRootState => gcRootState.GcRoots.Contains(blob));
+
+                    if (blobFoundInRootSet)
+                    {
+                        continue;
+                    }
+
+                    bool blobFound = false;
+                    // we enumerate the namespaces again as we need to check each transaction log to determine if its been added to any of them
+                    await namespaces.ParallelForEachAsync(async ns =>
+                    {
+                        GCRootState gcRootState = perNamespaceRoots[ns];
+                        bool hasBeenAdded = await HasBlobBeenAddedToGcSet(ns, blob, gcRootState.TransactionLogPointer, gcRootState.TransactionLogGeneration,
+                            cancellationToken);
+                        if (hasBeenAdded)
+                        {
+                            blobFound = true;
+                        }
+                    }, cancellationToken);
+
+                    if (!blobFound)
+                    {
+                        blobsToRemove.Add(blob);
+                    }
+                }
+            }
+
+            foreach (BlobIdentifier blob in blobsToRemove)
+            {
+                if (cancellationToken.IsCancellationRequested)
+                    break;
+                _logger.Information("Attempting to GC Orphan blob {Blob}", blob);
+                bool deleted = true;
+                foreach (NamespaceId ns in namespaces)
+                {
+                    try
+                    {
+                        await _blobStore.DeleteObject(ns, blob);
+                    }
+                    catch (Exception e)
+                    {
+                        _logger.Warning("Failed to delete blob {Blob} from {Namespace} due to {Error}", blob, ns, e.Message);
+                        deleted = false;
+                    }
+                }
+                if (deleted)
+                {
+                    removedBlobs.Add(new RemovedBlobs(blob));
+                }
+            }
+
+            return removedBlobs;
+        }
+
+        public virtual async IAsyncEnumerable<NamespaceId> ListNamespaces()
+        {
+            RestRequest request = new RestRequest("api/v1/t/");
+
+            IRestResponse<ListNamespaceResponse> response = await _client.ExecuteGetAsync<ListNamespaceResponse>(request);
+            if (!response.IsSuccessful)
+            {
+                throw response.ToException("Unable to list namespaces from callisto.");
+            }
+
+            foreach (NamespaceId ns in response.Data.Logs)
+            {
+                yield return ns;
+            }
+        }
+
+        public class ListNamespaceResponse
+        {
+            public NamespaceId[] Logs { get; set; } = null!;
+        }
+
+
+        private async Task<bool> HasBlobBeenAddedToGcSet(NamespaceId ns, BlobIdentifier blob, long transactionLogPointer, Guid transactionLogGeneration, CancellationToken cancellationToken)
+        {
+            // request all events since the transaction log pointer and make sure this blob has not been added since
+
+            CallistoReader reader = new CallistoReader(_client, ns);
+            await foreach (bool found in reader.GetOps(transactionLogPointer, transactionLogGeneration, maxOffsetsAttempted: 1000000).Any(transactionEvent =>
+            {
+                if (transactionEvent is AddTransactionEvent addTransactionEvent)
+                {
+                    return addTransactionEvent.Blobs?.Any(identifier => identifier.Equals(blob)) ?? false;
+                }
+
+                return false;
+            })
+                .WithCancellation(cancellationToken))
+            {
+                return found;
+            }
+
+            return false;
+        }
+
+        private async Task<GCRootState> DetermineGCRoots(NamespaceId ns, IRestClient client, CancellationToken cancellationToken)
+        {
+            Dictionary<string, BlobIdentifier[]> events = new Dictionary<string, BlobIdentifier[]>();
+            // Build the hashset from reading the entire callisto log
+            const int prefetchCount = 300; // 300 because the chunk size is usually 100 for callisto gets, thus we attempt to read 3x rows ahead
+            // prefetch events so that we do can just processes them directly without having to wait as much
+
+            CallistoReader reader = new CallistoReader(client, ns);
+            OpsEnumerationState enumerationState = new OpsEnumerationState();
+            long offsetToContinueAt = 0;
+            await foreach (TransactionEvent transactionEvent in reader.GetOps(cancellationToken: cancellationToken, maxOffsetsAttempted: 1000000, enumerationState: enumerationState).Prefetch(prefetchCount).WithCancellation(cancellationToken))
+            {
+                if (transactionEvent == null)
+                    continue;
+                
+                string key = $"{transactionEvent.Bucket}.{transactionEvent.Name}";
+                switch (transactionEvent)
+                {
+                    case AddTransactionEvent add:
+                        if (add.Blobs == null)
+                            throw new Exception("Expected to find blobs in a add transaction event");
+
+                        if (!events.TryAdd(key, add.Blobs))
+                        {
+                            _logger.Debug("Event with {Key} already added. This indicates a replace operation which should not happen often.", key);
+                        }
+                        break;
+                    case RemoveTransactionEvent remove:
+                        events.Remove(key);
+                        break;
+
+                    default:
+                        throw new NotImplementedException("Unknown transaction event type: " + transactionEvent);
+                }
+
+                offsetToContinueAt = transactionEvent.NextIdentifier!.Value;
+            }
+
+            HashSet<BlobIdentifier> gcRoots = new HashSet<BlobIdentifier>(events.Values.SelectMany(e => e));
+            return new GCRootState(offsetToContinueAt, enumerationState.ReplicatingGeneration, gcRoots);
+        }
+    }
+}
diff --git a/Engine/Source/Programs/Horde/HordeStorage/Horde.Storage/Implementation/GC/RefCleanupService.cs b/Engine/Source/Programs/Horde/HordeStorage/Horde.Storage/Implementation/GC/RefCleanupService.cs
new file mode 100644
--- /dev/null
+++ b/Engine/Source/Programs/Horde/HordeStorage/Horde.Storage/Implementation/GC/RefCleanupService.cs
@@ -0,0 +1,86 @@
+﻿// Copyright Epic Games, Inc. All Rights Reserved.
+
+using System;
+using System.Collections.Generic;
+using System.Threading;
+using System.Threading.Tasks;
+using Datadog.Trace;
+using Jupiter;
+using Jupiter.Implementation;
+using Microsoft.Extensions.Options;
+using Serilog;
+
+namespace Horde.Storage.Implementation
+{
+    public class RefCleanupService : PollingService<RefCleanupService.RefCleanupState>
+    {
+        private readonly IOptionsMonitor<GCSettings> _settings;
+        private readonly ILeaderElection _leaderElection;
+        private volatile bool _alreadyPolling;
+        private readonly ILogger _logger = Log.ForContext<RefCleanupService>();
+
+
+        public class RefCleanupState
+        {
+            public RefCleanupState(IRefsStore refs, IRefCleanup refCleanup)
+            {
+                Refs = refs;
+                RefCleanup = refCleanup;
+            }
+
+            public IRefsStore Refs { get; }
+            public IRefCleanup RefCleanup { get; }
+        }
+
+        public RefCleanupService(IOptionsMonitor<GCSettings> settings, IRefsStore store, IRefCleanup refCleanup, ILeaderElection leaderElection) : 
+            base(serviceName: nameof(RefCleanupService), settings.CurrentValue.RefCleanupPollFrequency, new RefCleanupState(store, refCleanup))
+        {
+            _settings = settings;
+            _leaderElection = leaderElection;
+        }
+
+        public override bool ShouldStartPolling()
+        {
+            return _settings.CurrentValue.CleanOldRefRecords;
+        }
+
+        public override async Task<bool> OnPoll(RefCleanupState state, CancellationToken cancellationToken)
+        {
+            if (_alreadyPolling)
+                return false;
+
+            _alreadyPolling = true;
+            try
+            {
+                if (!_leaderElection.IsThisInstanceLeader())
+                {
+                    _logger.Information("Skipped ref cleanup run as this instance was not the leader");
+                    return false;
+                }
+                await foreach (NamespaceId ns in state.Refs.GetNamespaces().WithCancellation(cancellationToken))
+                {
+                    using Scope scope = Tracer.Instance.StartActive("gc.refs");
+                    scope.Span.ResourceName = ns.ToString();
+
+                    _logger.Information("Attempting to run Refs Cleanup of {Namespace}. ", ns);
+                    try
+                    {
+                        List<OldRecord> oldRecords = await state.RefCleanup.Cleanup(ns, cancellationToken);
+                        _logger.Information("Ran Refs Cleanup of {Namespace}. Deleted {CountRefRecords}", ns, oldRecords.Count);
+                    }
+                    catch (Exception e)
+                    {
+                        _logger.Error("Error running Refs Cleanup of {Namespace}. {Exception}", ns, e);
+                    }
+                }
+
+                return true;
+
+            }
+            finally
+            {
+                _alreadyPolling = false;
+            }
+        }
+    }
+}
diff --git a/Engine/Source/Programs/Horde/HordeStorage/Horde.Storage/Implementation/IReferenceResolver.cs b/Engine/Source/Programs/Horde/HordeStorage/Horde.Storage/Implementation/IReferenceResolver.cs
new file mode 100644
--- /dev/null
+++ b/Engine/Source/Programs/Horde/HordeStorage/Horde.Storage/Implementation/IReferenceResolver.cs
@@ -0,0 +1,160 @@
+// Copyright Epic Games, Inc. All Rights Reserved.
+
+using System;
+using System.Collections.Concurrent;
+using System.Collections.Generic;
+using System.Reflection.Metadata;
+using System.Threading.Tasks;
+using Datadog.Trace;
+using Jupiter.Implementation;
+using Jupiter.Utils;
+
+namespace Horde.Storage.Implementation
+{
+    public interface IReferenceResolver
+    {
+        IAsyncEnumerable<BlobIdentifier> ResolveReferences(NamespaceId ns, CompactBinaryObject cb);
+    }
+
+    public class ReferenceResolver : IReferenceResolver
+    {
+        private readonly IBlobStore _blobStore;
+        private readonly IContentIdStore _contentIdStore;
+
+        public ReferenceResolver(IBlobStore blobStore, IContentIdStore contentIdStore)
+        {
+            _blobStore = blobStore;
+            _contentIdStore = contentIdStore;
+        }
+
+        public async IAsyncEnumerable<BlobIdentifier> ResolveReferences(NamespaceId ns, CompactBinaryObject cb)
+        {
+            // TODO: This is cacheable and we should store the result somewhere
+            Queue<CompactBinaryObject> objectsToVisit = new Queue<CompactBinaryObject>();
+            objectsToVisit.Enqueue(cb);
+            List<BlobIdentifier> unresolvedReferences = new List<BlobIdentifier>();
+
+            List<Task<(BlobIdentifier, BlobIdentifier[]?)>> pendingContentIdResolves = new();
+            List<Task<CompactBinaryObject>> pendingCompactBinaryAttachments = new();
+
+            while(pendingCompactBinaryAttachments.Count != 0 || pendingContentIdResolves.Count != 0 || objectsToVisit.Count != 0)
+            {
+                if (objectsToVisit.TryDequeue(out CompactBinaryObject? parent))
+                {
+                    // enumerate all fields in the compact binary and start to resolve their dependencies
+                    foreach (CompactBinaryField field in parent.GetAllFields())
+                    {
+                        if (!field.IsAttachment())
+                            continue;
+
+                        BlobIdentifier? blobIdentifier = field.AsAttachment();
+                        if (blobIdentifier == null)
+                            continue;
+
+                        if (field.IsBinaryAttachment())
+                        {
+                            pendingContentIdResolves.Add(ResolveContentId(ns, blobIdentifier));
+                        }
+                        else
+                        {
+                            // the identifier as is points to a blob
+                            yield return blobIdentifier;
+                        }
+
+                        // if its a reference to another compact binary we need to fetch it and resolve it
+                        if (!field.IsCompactBinaryAttachment())
+                            continue;
+
+                        pendingCompactBinaryAttachments.Add(ParseCompactBinaryAttachment(ns, blobIdentifier));
+                    }
+                }
+
+                List<Task<(BlobIdentifier, BlobIdentifier[]?)>> contentIdResolvesToRemove = new();
+
+                foreach (Task<(BlobIdentifier, BlobIdentifier[]?)> pendingContentIdResolveTask in pendingContentIdResolves)
+                {
+                    // check for any content id resolve that has finished and return those blobs it found
+                    if (pendingContentIdResolveTask.IsCompleted)
+                    {
+                        (BlobIdentifier id, BlobIdentifier[]? resolvedBlobs) = await pendingContentIdResolveTask;
+                        if (resolvedBlobs != null)
+                        {
+                            foreach (BlobIdentifier b in resolvedBlobs)
+                            {
+                                yield return b;
+                            }
+                        }
+                        else
+                        {
+                            unresolvedReferences.Add(id);
+                        }
+                        contentIdResolvesToRemove.Add(pendingContentIdResolveTask);
+                    }
+                }
+
+                // cleanup finished tasks
+                foreach (Task<(BlobIdentifier, BlobIdentifier[]?)> finishedTask in contentIdResolvesToRemove)
+                {
+                    pendingContentIdResolves.Remove(finishedTask);
+                }
+
+                // check for any compact binary attachment fetches and add those to the objects we are handling
+                List<Task<CompactBinaryObject>> finishedCompactBinaryResolves = new();
+                foreach (Task<CompactBinaryObject> pendingCompactBinaryAttachment in pendingCompactBinaryAttachments)
+                {
+                    if (pendingCompactBinaryAttachment.IsCompleted)
+                    {
+                        try
+                        {
+                            CompactBinaryObject childBinaryObject = await pendingCompactBinaryAttachment;
+                            objectsToVisit.Enqueue(childBinaryObject);
+                        }
+                        catch (BlobNotFoundException e)
+                        {
+                            unresolvedReferences.Add(e.Blob);
+                        }
+                        finishedCompactBinaryResolves.Add(pendingCompactBinaryAttachment);
+                    }
+                }
+
+                // cleanup finished tasks
+                foreach (Task<CompactBinaryObject> finishedTask in finishedCompactBinaryResolves)
+                {
+                    pendingCompactBinaryAttachments.Remove(finishedTask);
+                }
+            }
+
+            if (unresolvedReferences.Count != 0)
+            {
+                throw new PartialReferenceResolveException(unresolvedReferences);
+            }
+        }
+
+        private async Task<CompactBinaryObject> ParseCompactBinaryAttachment(NamespaceId ns, BlobIdentifier blobIdentifier)
+        {
+            BlobContents contents = await _blobStore.GetObject(ns, blobIdentifier);
+            byte[] data = await contents.Stream.ToByteArray();
+            CompactBinaryObject childBinaryObject = CompactBinaryObject.Load(data);
+
+            return childBinaryObject;
+        }
+
+        private async Task<(BlobIdentifier, BlobIdentifier[]?)> ResolveContentId(NamespaceId ns, BlobIdentifier blobIdentifier)
+        {
+            using Scope scope = Tracer.Instance.StartActive("ReferenceResolver.ResolveContentId");
+            scope.Span.ResourceName = blobIdentifier.ToString();
+            BlobIdentifier[]? resolvedBlobs = await _contentIdStore.Resolve(ns, blobIdentifier);
+            return (blobIdentifier, resolvedBlobs);
+        }
+    }
+
+    public class PartialReferenceResolveException : Exception
+    {
+        public List<BlobIdentifier> UnresolvedReferences { get; }
+
+        public PartialReferenceResolveException(List<BlobIdentifier> unresolvedReferences) : base($"References missing: {string.Join(',', unresolvedReferences)}")
+        {
+            UnresolvedReferences = unresolvedReferences;
+        }
+    }
+}
diff --git a/Engine/Source/Programs/Horde/HordeStorage/Horde.Storage/Implementation/Kestrel/AsciiString.cs b/Engine/Source/Programs/Horde/HordeStorage/Horde.Storage/Implementation/Kestrel/AsciiString.cs
new file mode 100644
--- /dev/null
+++ b/Engine/Source/Programs/Horde/HordeStorage/Horde.Storage/Implementation/Kestrel/AsciiString.cs
@@ -0,0 +1,57 @@
+﻿using System;
+using System.Text;
+
+namespace Horde.Storage.Implementation.Kestrel
+{
+    public readonly struct AsciiString : IEquatable<AsciiString>
+    {
+        private readonly byte[] _data;
+
+        public AsciiString(string s) => _data = Encoding.ASCII.GetBytes(s);
+
+        public AsciiString(byte[] b) => _data = b;
+
+        public int Length => _data.Length;
+        public byte[] Data => _data;
+
+        public ReadOnlySpan<byte> AsSpan() => _data;
+
+        public static implicit operator ReadOnlySpan<byte>(AsciiString str) => str._data;
+        public static implicit operator byte[] (AsciiString str) => str._data;
+
+        public static implicit operator AsciiString(string str) => new AsciiString(str);
+
+        public override string ToString() => Encoding.ASCII.GetString(_data);
+        public static explicit operator string(AsciiString str) => str.ToString();
+
+        public bool Equals(AsciiString other) => ReferenceEquals(_data, other._data) || SequenceEqual(_data, other._data);
+        private bool SequenceEqual(byte[] data1, byte[] data2) => new Span<byte>(data1).SequenceEqual(data2);
+
+        public static bool operator ==(AsciiString a, AsciiString b) => a.Equals(b);
+        public static bool operator !=(AsciiString a, AsciiString b) => !a.Equals(b);
+        public override bool Equals(object? other) => (other is AsciiString) && Equals((AsciiString)other);
+
+        public static AsciiString operator +(AsciiString a, AsciiString b)
+        {
+            var result = new byte[a.Length + b.Length];
+            a._data.CopyTo(result, 0);
+            b._data.CopyTo(result, a.Length);
+            return new AsciiString(result);
+        }
+
+        public override int GetHashCode()
+        {
+            // Copied from x64 version of string.GetLegacyNonRandomizedHashCode()
+            // [REDACTED_PROFILE_URL]/coreclr/blob/master/src/mscorlib/src/System/String.Comparison.cs
+            var data = _data;
+            int hash1 = 5381;
+            int hash2 = hash1;
+            foreach (int b in data)
+            {
+                hash1 = ((hash1 << 5) + hash1) ^ b;
+            }
+            return hash1 + (hash2 * 1566083941);
+        }
+
+    }
+}
diff --git a/Engine/Source/Programs/Horde/HordeStorage/Horde.Storage/Implementation/Kestrel/KestrelDispatcher.HttpConnection.cs b/Engine/Source/Programs/Horde/HordeStorage/Horde.Storage/Implementation/Kestrel/KestrelDispatcher.HttpConnection.cs
new file mode 100644
--- /dev/null
+++ b/Engine/Source/Programs/Horde/HordeStorage/Horde.Storage/Implementation/Kestrel/KestrelDispatcher.HttpConnection.cs
@@ -0,0 +1,230 @@
+﻿// Copyright Epic Games, Inc. All Rights Reserved.
+
+using System;
+using System.Buffers;
+using System.IO;
+using System.IO.Pipelines;
+using System.Net.Sockets;
+using System.Text;
+using System.Threading.Tasks;
+using Jupiter;
+using Microsoft.AspNetCore.Connections;
+using Microsoft.AspNetCore.Server.Kestrel.Core.Internal.Http;
+using Microsoft.Extensions.DependencyInjection;
+using Serilog;
+
+namespace Horde.Storage.Implementation.Kestrel
+{
+    public partial class KestrelDispatcher : BaseHttpConnection
+    {
+        private readonly ILogger _logger = Log.ForContext<KestrelDispatcher>();
+        private State _state;
+        private HttpParser<ParsingAdapter> Parser { get; } = new HttpParser<ParsingAdapter>();
+        
+        public KestrelDispatcher(IServiceProvider serviceProvider, PipeReader reader, PipeWriter writer, Socket socket) : base(serviceProvider, reader, writer, socket)
+        {
+            _ddcRefService = serviceProvider.GetService<DDCRefService>()!;
+        }
+        
+        // Used for tests only
+        public KestrelDispatcher(IDDCRefService ddcRefService, PipeReader reader, PipeWriter writer, Socket socket) : base(null!, reader, writer, socket)
+        {
+            _ddcRefService = ddcRefService;
+        }
+
+        /// <inheritdoc/>
+        public override async Task ExecuteAsync(ConnectionContext connection)
+        {
+            try
+            {
+                await ProcessRequestsAsync();
+
+                Reader.Complete();
+            }
+            catch (Exception ex)
+            {
+                Reader.Complete(ex);
+
+                if (!(ex is ConnectionResetException))
+                {
+                    throw;  
+                }
+            }
+            finally
+            {
+                Writer.Complete();
+            }
+        }
+
+        private async Task ProcessRequestsAsync()
+        {
+            while (true)
+            {
+                _requestBodyLength = null;
+                _requestBodyData = null;
+                _requestBodyDataStream = null;
+                var readResult = await Reader.ReadAsync();
+                var buffer = readResult.Buffer;
+                var isCompleted = readResult.IsCompleted;
+
+                if (buffer.IsEmpty && isCompleted)
+                {
+                    return;
+                }
+
+                while (true)
+                {
+                    if (!ParseHttpRequest(ref buffer, isCompleted))
+                    {
+                        return;
+                    }
+
+                    if (_state == State.Body)
+                    {
+                        await ProcessRequestAsync(ref buffer);
+
+                        _state = State.StartLine;
+
+                        if (!buffer.IsEmpty)
+                        {
+                            // More input data to parse
+                            continue;
+                        }
+                    }
+
+                    // No more input or incomplete data, Advance the Reader
+                    Reader.AdvanceTo(buffer.Start, buffer.End);
+                    break;
+                }
+
+                await Writer.FlushAsync();
+            }
+        }
+
+        private bool ParseHttpRequest(ref ReadOnlySequence<byte> buffer, bool isCompleted)
+        {
+            var reader = new SequenceReader<byte>(buffer);
+            var state = _state;
+
+            if (state == State.StartLine)
+            {
+                if (Parser.ParseRequestLine(new ParsingAdapter(this), ref reader))
+                {
+                    state = State.Headers;
+                }
+            }
+
+            if (state == State.Headers)
+            {
+                var success = Parser.ParseHeaders(new ParsingAdapter(this), ref reader);
+
+                if (success)
+                {
+                    state = State.Body;
+                }
+            }
+
+            if (state != State.Body && isCompleted)
+            {
+                ThrowUnexpectedEndOfData();
+            }
+
+            _state = state;
+
+            if (state == State.Body)
+            {
+                if (_requestMethod == HttpMethod.Put && _requestBodyLength != null && _requestBodyDataStream != null)
+                {
+                    // Remove the status line and headers
+                    buffer = buffer.Slice(reader.Position);
+                    // Body starts here in buffer
+
+                    // Dumb byte-by-byte reading
+                    var temp = new byte[1];
+                    while (reader.TryRead(out byte v))
+                    {
+                        temp[0] = v;
+                        _requestBodyDataStream.Write(temp);
+                    }
+                    buffer = buffer.Slice(reader.Position, 0);
+                }
+                else
+                {
+                    // Mark the rest as consumed
+                    buffer = buffer.Slice(reader.Position, 0);
+                }
+            }
+            else
+            {
+                // In-complete request read, consumed is current position and examined is the remaining.
+                buffer = buffer.Slice(reader.Position);
+            }
+
+            return true;
+        }
+
+        public void OnStaticIndexedHeader(int index)
+        {
+        }
+
+        public void OnStaticIndexedHeader(int index, ReadOnlySpan<byte> value)
+        {
+        }
+
+        public void OnHeader(ReadOnlySpan<byte> name, ReadOnlySpan<byte> value)
+        {
+            if (name.Length == 14 && name.SequenceEqual(HeaderContentLength))
+            {
+                _requestBodyLength = Convert.ToInt64(Encoding.ASCII.GetString(value));
+                _requestBodyData = new byte[_requestBodyLength.Value];
+                _requestBodyDataStream = new MemoryStream(_requestBodyData);
+            }
+        }
+        
+        public void OnHeadersComplete(bool endStream)
+        {
+        }
+        
+        public void OnStartLine(HttpVersionAndMethod versionAndMethod, TargetOffsetPathLength targetPath, Span<byte> startLine)
+        {
+            _requestMethod = versionAndMethod.Method;
+            _requestType = GetRequestType(startLine.Slice(targetPath.Offset, targetPath.Length));
+            _requestPath = new AsciiString(startLine.Slice(targetPath.Offset, targetPath.Length).ToArray());
+        }
+        
+        private static void ThrowUnexpectedEndOfData()
+        {
+            throw new InvalidOperationException("Unexpected end of data!");
+        }
+
+        private enum State
+        {
+            StartLine,
+            Headers,
+            Body
+        }
+
+        private struct ParsingAdapter : IHttpRequestLineHandler, IHttpHeadersHandler
+        {
+            public KestrelDispatcher RequestHandler;
+
+            public ParsingAdapter(KestrelDispatcher requestHandler)
+                => RequestHandler = requestHandler;
+
+            public void OnStaticIndexedHeader(int index) 
+                => RequestHandler.OnStaticIndexedHeader(index);
+
+            public void OnStaticIndexedHeader(int index, ReadOnlySpan<byte> value)
+                => RequestHandler.OnStaticIndexedHeader(index, value);
+
+            public void OnHeader(ReadOnlySpan<byte> name, ReadOnlySpan<byte> value)
+                => RequestHandler.OnHeader(name, value);
+
+            public void OnHeadersComplete(bool endStream)
+                => RequestHandler.OnHeadersComplete(endStream);
+
+            public void OnStartLine(HttpVersionAndMethod versionAndMethod, TargetOffsetPathLength targetPath, Span<byte> startLine)
+                => RequestHandler.OnStartLine(versionAndMethod, targetPath, startLine);
+        }
+    }
+}
diff --git a/Engine/Source/Programs/Horde/HordeStorage/Horde.Storage/Implementation/Kestrel/KestrelDispatcher.cs b/Engine/Source/Programs/Horde/HordeStorage/Horde.Storage/Implementation/Kestrel/KestrelDispatcher.cs
new file mode 100644
--- /dev/null
+++ b/Engine/Source/Programs/Horde/HordeStorage/Horde.Storage/Implementation/Kestrel/KestrelDispatcher.cs
@@ -0,0 +1,267 @@
+﻿// Copyright Epic Games, Inc. All Rights Reserved.
+
+using System;
+using System.Buffers;
+using System.IO;
+using System.IO.Pipelines;
+using System.Net.Sockets;
+using System.Threading;
+using System.Threading.Tasks;
+using Horde.Storage.Controllers;
+using Jupiter.Implementation;
+using Microsoft.AspNetCore.Server.Kestrel.Core.Internal.Http;
+
+namespace Horde.Storage.Implementation.Kestrel
+{
+    public partial class KestrelDispatcher
+    {
+        private RequestType _requestType;
+        private HttpMethod _requestMethod;
+        private AsciiString _requestPath;
+        private long? _requestBodyLength = null;
+        private byte[]? _requestBodyData = null;
+        private MemoryStream? _requestBodyDataStream = null;
+        private readonly IDDCRefService _ddcRefService;
+
+        private enum RequestType
+        {
+            Ping,
+            HealthLive,
+            HealthReady,
+            DdcGet,
+            DdcHead,
+            DdcPut,
+            NotRecognized
+        }
+
+        public static class Paths
+        {
+            public static readonly AsciiString Ping = "/ping";
+            public static readonly AsciiString HealthLive = "/health/live";
+            public static readonly AsciiString HealthReady = "/health/ready";
+            public static readonly AsciiString DdcGet = "/api/v1/c/ddc";
+            public static readonly AsciiString DdcHead = "/api/v1/c/ddc";
+            public static readonly AsciiString DdcPut = "/api/v1/c/ddc";
+        }
+
+        private RequestType GetRequestType(ReadOnlySpan<byte> path)
+        {
+            if (path.Length == 5 && path.SequenceEqual(Paths.Ping)) { return RequestType.Ping; }
+            if (path.Length == 12 && path.SequenceEqual(Paths.HealthLive)) { return RequestType.HealthLive; }
+            if (path.Length == 13 && path.SequenceEqual(Paths.HealthReady)) { return RequestType.HealthReady; }
+            if (_requestMethod == HttpMethod.Get && path.Length > 13 && path[..13].SequenceEqual(Paths.DdcGet)) { return RequestType.DdcGet; }
+            if (_requestMethod == HttpMethod.Head && path.Length > 13 && path[..13].SequenceEqual(Paths.DdcHead)) { return RequestType.DdcHead; }
+            if (_requestMethod == HttpMethod.Put && path.Length > 13 && path[..13].SequenceEqual(Paths.DdcPut)) { return RequestType.DdcPut; }
+            
+            return RequestType.NotRecognized;
+        }
+        
+        private Task ProcessRequestAsync(ref ReadOnlySequence<byte> buffer) => _requestType switch
+        {
+            RequestType.Ping => Ping(Writer, Socket),
+            RequestType.HealthLive => HealthLive(Writer, Socket),
+            RequestType.HealthReady => HealthReady(Writer, Socket),
+            RequestType.DdcGet => DdcGet(_requestPath, Writer, Socket),
+            RequestType.DdcHead => DdcHead(_requestPath, Writer, Socket),
+            RequestType.DdcPut => DdcPut(buffer, _requestPath, Writer, Socket),
+            _ => Default(_requestPath, Writer, Socket)
+        };
+        
+        private static readonly AsciiString Crlf = "\r\n";
+        private static readonly AsciiString Eoh = "\r\n\r\n"; // End Of Headers
+        private static readonly AsciiString Http11NotFound = "HTTP/1.1 404 Not Found\r\n";
+        private static readonly AsciiString Http11NotFoundBody = "404 Not Found";
+        private static readonly AsciiString Http11Ok = "HTTP/1.1 200 OK\r\n";
+        private static readonly AsciiString Http11NoContent = "HTTP/1.1 204 No Content\r\n";
+        private static readonly AsciiString Http11BadRequest = "HTTP/1.1 400 Bad Request\r\n";
+        private static readonly AsciiString Http11InternalServerError = "HTTP/1.1 500 Internal Server Error\r\n";
+        private static readonly AsciiString PingBody = "pong";
+        private static readonly AsciiString HealthLiveBody = "Healthy";
+        private static readonly AsciiString HeaderContentLength = "Content-Length";
+        public static readonly AsciiString HeaderIoHash = "X-Jupiter-IoHash";
+        private static readonly AsciiString HeaderIoHashRes = HeaderIoHash + ": ";
+        private static readonly AsciiString HeaderServer = "Server: K";
+        private static readonly AsciiString HeaderContentLengthRes = HeaderContentLength + ": ";
+        private static readonly AsciiString HeaderContentTypeText = "Content-Type: text/plain";
+        private static readonly AsciiString HeaderContentTypeJson = "Content-Type: application/json";
+        private static readonly AsciiString HeaderContentTypeOctetStream = "Content-Type: application/octet-stream";
+
+        private static readonly AsciiString _defaultPreamble =
+            Http11NotFound +
+            HeaderServer + Crlf +
+            HeaderContentTypeText + Crlf +
+            HeaderContentLengthRes + Http11NotFoundBody.Length.ToString() + Eoh;
+        
+        private async Task Default(AsciiString requestPath, PipeWriter pipeWriter, Socket socket)
+        {
+            _logger.Information("Unhandled request path: {RequestPath}", requestPath);
+            await pipeWriter.WriteAsync(_defaultPreamble.Data.AsMemory());
+            await pipeWriter.WriteAsync(Http11NotFoundBody.Data.AsMemory());
+        }
+        
+        
+        private static readonly AsciiString _pingPreamble =
+            Http11Ok +
+            HeaderServer + Crlf +
+            HeaderContentTypeText + Crlf +
+            HeaderContentLengthRes + PingBody.Length.ToString() + Eoh;
+
+        private static async Task Ping(PipeWriter pipeWriter, Socket socket)
+        {
+            await pipeWriter.WriteAsync(_pingPreamble.Data.AsMemory());
+            await pipeWriter.WriteAsync(PingBody.Data.AsMemory());
+        }
+        
+        private static readonly AsciiString _healthLivePreamble =
+            Http11Ok +
+            HeaderServer + Crlf +
+            HeaderContentTypeText + Crlf +
+            HeaderContentLengthRes + HealthLiveBody.Length.ToString() + Eoh;
+
+        private static async Task HealthLive(PipeWriter pipeWriter, Socket socket)
+        {
+            await pipeWriter.WriteAsync(_healthLivePreamble.Data.AsMemory());
+            await pipeWriter.WriteAsync(HealthLiveBody.Data.AsMemory());
+        }
+        
+        private static readonly AsciiString _healthReadyPreamble =
+            Http11Ok +
+            HeaderServer + Crlf +
+            HeaderContentTypeText + Crlf +
+            HeaderContentLengthRes + HealthLiveBody.Length.ToString() + Eoh;
+
+        private static async Task HealthReady(PipeWriter pipeWriter, Socket socket)
+        {
+            await pipeWriter.WriteAsync(_healthReadyPreamble.Data.AsMemory());
+            await pipeWriter.WriteAsync(HealthLiveBody.Data.AsMemory());
+        }
+
+
+        private async Task WriteJsonError(PipeWriter pipeWriter, int statusCode, AsciiString message)
+        {
+            AsciiString statusLine = statusCode switch
+            {
+                400 => Http11BadRequest,
+                _ => Http11InternalServerError
+            };
+
+            AsciiString preamble = statusLine + HeaderServer + Crlf + HeaderContentTypeJson + Crlf + HeaderContentLengthRes + message.Length.ToString() + Eoh;
+            await pipeWriter.WriteAsync(preamble.Data.AsMemory());
+            await pipeWriter.WriteAsync(message.Data.AsMemory());
+        }
+        
+        private static readonly AsciiString _ddcGetPreamble =
+            Http11Ok +
+            HeaderServer + Crlf +
+            HeaderContentTypeOctetStream + Crlf +
+            HeaderContentLengthRes;
+        
+        private async Task DdcGet(AsciiString requestPath, PipeWriter pipeWriter, Socket socket)
+        {
+            string[] parts = requestPath.ToString().Split("/");
+            if (parts.Length != 8)
+            {
+                await WriteJsonError(pipeWriter, 400, "Bad request path: " + requestPath);
+                return;
+            }
+            string ns = parts[5];
+            string bucket = parts[6];
+            string key = parts[7];
+            bool isRaw = key.EndsWith(".raw");
+            if (!isRaw)
+            {
+                await WriteJsonError(pipeWriter, 400, "Can only handle requests with .raw suffix");
+                return;
+            }
+
+            key = key[..^4]; // Strip .raw suffix
+
+            try
+            {
+                (RefResponse refRes, BlobContents? tempBlob) = await _ddcRefService.Get(new NamespaceId(ns), new BucketId(bucket), new KeyId(key), new[] {"blob"});
+                BlobContents blob = tempBlob!;
+
+                AsciiString contentLength = Convert.ToString(blob.Length) + Crlf + HeaderIoHashRes + refRes.ContentHash.ToString() + Eoh;
+                await pipeWriter.WriteAsync(_ddcGetPreamble.Data.AsMemory());
+                await pipeWriter.WriteAsync(contentLength.Data.AsMemory());
+                await blob.Stream.CopyToAsync(pipeWriter, CancellationToken.None);
+            }
+            catch (NamespaceNotFoundException e)
+            {
+                await WriteJsonError(pipeWriter, 400, "{\"title\": \"Namespace " + e.Namespace + " did not exist\", \"status\": 400}");
+            }
+            catch (RefRecordNotFoundException e)
+            {
+                await WriteJsonError(pipeWriter, 400, "{\"title\": \"Object " + e.Bucket + " " + e.Key + " did not exist\", \"status\": 400}");
+            }
+            catch (BlobNotFoundException e)
+            {
+                await WriteJsonError(pipeWriter, 400, "{\"title\": \"Object " + e.Blob + " in " + e.Ns + " not found\", \"status\": 400}");
+            }
+            catch (Exception e)
+            {
+                await WriteJsonError(pipeWriter, 500, "{\"title\": \"Unhandled internal server error\", \"status\": 500}");
+                _logger.Error(e, "Unhandled exception in DDC GET call");
+            }
+        }
+
+        private static readonly AsciiString _ddcHeadPreamble =
+            Http11NoContent +
+            HeaderServer + Crlf +
+            HeaderIoHashRes;
+        
+        private async Task DdcHead(AsciiString requestPath, PipeWriter pipeWriter, Socket socket)
+        {
+            string[] parts = requestPath.ToString().Split("/");
+            if (parts.Length != 8)
+            {
+                await WriteJsonError(pipeWriter, 400, "Bad request path: " + requestPath);
+                return;
+            }
+            string ns = parts[5];
+            string bucket = parts[6];
+            string key = parts[7];
+
+            try
+            {
+                RefRecord refRecord = await _ddcRefService.Exists(new NamespaceId(ns), new BucketId(bucket), new KeyId(key));
+
+                AsciiString contentHash = refRecord.ContentHash.ToString() + Eoh;
+                await pipeWriter.WriteAsync(_ddcHeadPreamble.Data.AsMemory());
+                await pipeWriter.WriteAsync(contentHash.Data.AsMemory());
+            }
+            catch (NamespaceNotFoundException e)
+            {
+                await WriteJsonError(pipeWriter, 400, "{\"title\": \"Namespace " + e.Namespace + " did not exist\", \"status\": 400}");
+            }
+            catch (RefRecordNotFoundException e)
+            {
+                await WriteJsonError(pipeWriter, 400, "{\"title\": \"Object " + e.Bucket + " " + e.Key + " did not exist\", \"status\": 400}");
+            }
+            catch (MissingBlobsException e)
+            {
+                await WriteJsonError(pipeWriter, 400, "{\"title\": \"Blobs " + e.Blobs + " from object " + e.Bucket + " " + e.Key + " in namespace " + e.Namespace + " did not exist\", \"status\": 400}");
+            }
+            catch (Exception e)
+            {
+                await WriteJsonError(pipeWriter, 500, "{\"title\": \"Unhandled internal server error\", \"status\": 500}");
+                _logger.Error(e, "Unhandled exception in DDC HEAD call");
+                throw;
+            }
+        }
+        
+        private static readonly AsciiString _ddcPutContent = "{\"TransactionId\": 1}";
+
+        private static readonly AsciiString _ddcPutNoOpPreamble =
+            Http11Ok +
+            HeaderServer + Crlf +
+            HeaderContentTypeJson + Crlf +
+            HeaderContentLengthRes + _ddcPutContent.Length.ToString() + Eoh;
+        
+        private async Task DdcPut(ReadOnlySequence<byte> buffer, AsciiString requestPath, PipeWriter pipeWriter, Socket socket)
+        {
+            await pipeWriter.WriteAsync(_ddcPutNoOpPreamble.Data.AsMemory());
+            await pipeWriter.WriteAsync(_ddcPutContent.Data.AsMemory());
+        }
+    }
+}
diff --git a/Engine/Source/Programs/Horde/HordeStorage/Horde.Storage/Implementation/LastAccess/LastAccessService.cs b/Engine/Source/Programs/Horde/HordeStorage/Horde.Storage/Implementation/LastAccess/LastAccessService.cs
new file mode 100644
--- /dev/null
+++ b/Engine/Source/Programs/Horde/HordeStorage/Horde.Storage/Implementation/LastAccess/LastAccessService.cs
@@ -0,0 +1,83 @@
+﻿// Copyright Epic Games, Inc. All Rights Reserved.
+
+using System;
+using System.Collections.Generic;
+using System.Threading;
+using System.Threading.Tasks;
+using Datadog.Trace;
+using Microsoft.Extensions.Hosting;
+using Microsoft.Extensions.Options;
+using Serilog;
+
+namespace Horde.Storage.Implementation
+{
+    // ReSharper disable once ClassNeverInstantiated.Global
+    public class LastAccessService : IHostedService, IDisposable
+    {
+        private readonly ILastAccessCache<RefRecord> _lastAccessCacheRecord;
+        private readonly IRefsStore _refsStore;
+        private readonly ILogger _logger = Log.ForContext<LastAccessService>();
+        private Timer? _timer;
+        private readonly HordeStorageSettings _settings;
+        
+        public bool Running { get; private set; }
+
+        public LastAccessService(IOptionsMonitor<HordeStorageSettings> settings, ILastAccessCache<RefRecord> lastAccessCache, IRefsStore refsStore)
+        {
+            _lastAccessCacheRecord = lastAccessCache;
+            _refsStore = refsStore;
+            _settings = settings.CurrentValue;
+        }
+
+        public Task StartAsync(CancellationToken cancellationToken)
+        {
+            _logger.Information("Last Access Aggregation service starting.");
+
+            _timer = new Timer(OnUpdate, null, TimeSpan.Zero,
+                period: TimeSpan.FromSeconds(_settings.LastAccessRollupFrequencySeconds));
+            Running = true;
+
+            return Task.CompletedTask;
+        }
+
+        public async Task StopAsync(CancellationToken cancellationToken)
+        {
+            _logger.Information("Last Access Aggregation service stopping.");
+
+            _timer?.Change(Timeout.Infinite, 0);
+            Running = false;
+
+            // process the last records we have built up
+            await ProcessLastAccessRecords();
+        }
+
+        private void OnUpdate(object? state)
+        {
+            // call results to make sure we join the task
+            Task.WaitAll(
+                ProcessLastAccessRecords()
+            );
+        }
+
+        internal async Task<List<(RefRecord, DateTime)>> ProcessLastAccessRecords()
+        {
+            using Scope _ = Tracer.Instance.StartActive("lastAccess.update");
+            _logger.Information("Running Last Access Aggregation");
+            List<(RefRecord, DateTime)> records = await _lastAccessCacheRecord.GetLastAccessedRecords();
+            foreach ((RefRecord record, DateTime lastAccessTime) in records)
+            {
+                using Scope scope = Tracer.Instance.StartActive("lastAccess.update.record");
+                scope.Span.ResourceName = $"{record.Namespace}:{record.Bucket}.{record.RefName}";
+                _logger.Debug("Updating last access time to {LastAccessTime} for {Record}", lastAccessTime, record);
+                await _refsStore.UpdateLastAccessTime(record, lastAccessTime);
+            }
+
+            return records;
+        }
+
+        public void Dispose()
+        {
+            _timer?.Dispose();
+        }
+    }
+}
diff --git a/Engine/Source/Programs/Horde/HordeStorage/Horde.Storage/Implementation/LeaderElection/ILeaderElection.cs b/Engine/Source/Programs/Horde/HordeStorage/Horde.Storage/Implementation/LeaderElection/ILeaderElection.cs
new file mode 100644
--- /dev/null
+++ b/Engine/Source/Programs/Horde/HordeStorage/Horde.Storage/Implementation/LeaderElection/ILeaderElection.cs
@@ -0,0 +1,47 @@
+﻿// Copyright Epic Games, Inc. All Rights Reserved.
+
+using System;
+using System.Collections.Generic;
+using System.Linq;
+using System.Threading.Tasks;
+
+namespace Horde.Storage.Implementation
+{
+    public interface ILeaderElection
+    {
+        public class OnLeaderChangedEventArgs
+        {
+            public bool IsLeader { get; }
+            public string LeaderName { get; }
+
+            public OnLeaderChangedEventArgs(bool isLeader, string leaderName)
+            {
+                IsLeader = isLeader;
+                LeaderName = leaderName;
+            }
+        }
+
+        bool IsThisInstanceLeader();
+
+        event EventHandler<OnLeaderChangedEventArgs>? OnLeaderChanged;
+    }
+    
+    public class LeaderElectionStub : ILeaderElection
+    {
+        private readonly bool _isLeader;
+
+        public LeaderElectionStub(bool isLeader)
+        {
+            _isLeader = isLeader;
+
+            OnLeaderChanged?.Invoke(this, new ILeaderElection.OnLeaderChangedEventArgs(isLeader, ""));
+        }
+
+        public bool IsThisInstanceLeader()
+        {
+            return _isLeader;
+        }
+
+        public event EventHandler<ILeaderElection.OnLeaderChangedEventArgs>? OnLeaderChanged;
+    }
+}
diff --git a/Engine/Source/Programs/Horde/HordeStorage/Horde.Storage/Implementation/LeaderElection/KubernetesLeaderElection.cs b/Engine/Source/Programs/Horde/HordeStorage/Horde.Storage/Implementation/LeaderElection/KubernetesLeaderElection.cs
new file mode 100644
--- /dev/null
+++ b/Engine/Source/Programs/Horde/HordeStorage/Horde.Storage/Implementation/LeaderElection/KubernetesLeaderElection.cs
@@ -0,0 +1,115 @@
+// Copyright Epic Games, Inc. All Rights Reserved.
+
+using System;
+using System.ComponentModel.DataAnnotations;
+using System.Net;
+using System.Threading;
+using System.Threading.Tasks;
+using Jupiter;
+using k8s;
+using k8s.LeaderElection;
+using k8s.LeaderElection.ResourceLock;
+using Microsoft.Extensions.Options;
+using Microsoft.Rest;
+using Serilog;
+
+namespace Horde.Storage.Implementation.LeaderElection
+{
+    public class KubernetesLeaderElectionSettings
+    {
+        // these defaults are based on recommendations from the kubernetes client - [REDACTED_PROFILE_URL]/client-go/blob/master/tools/leaderelection/leaderelection.go
+        public TimeSpan LeaseDuration { get; set; } = TimeSpan.FromSeconds(15);
+        public TimeSpan RetryPeriod { get; set; } = TimeSpan.FromMilliseconds(2000);
+        public TimeSpan RenewDeadline { get; set; } = TimeSpan.FromSeconds(10);
+
+        public string ConfigMapName { get; set; } = "jupiter-leader-configmap";
+
+        [Required]
+        public string Namespace { get; set; } = "";
+    }
+
+    public class KubernetesLeaderElectionState
+    {
+
+    }
+
+    public class KubernetesLeaderElection : PollingService<KubernetesLeaderElectionState>, ILeaderElection, IDisposable
+    {
+        private readonly ILogger _logger = Log.ForContext<KubernetesLeaderElection>();
+
+        private readonly Kubernetes _client;
+        private readonly LeaderElector _leaderElector;
+        private ConfigMapLock _configMapLock;
+        private string _identity;
+
+        public KubernetesLeaderElection(IOptionsMonitor<KubernetesLeaderElectionSettings> leaderSettings) : base("Kubernetes Leader Election", TimeSpan.FromSeconds(1), new KubernetesLeaderElectionState())
+        {
+            KubernetesLeaderElectionSettings settings = leaderSettings.CurrentValue;
+            // As we are determining if we are the leader we just assume we are running in a kubernetes cluster
+            KubernetesClientConfiguration config = KubernetesClientConfiguration.InClusterConfig();
+            _client = new Kubernetes(config);
+            
+            _identity = System.Net.Dns.GetHostName();
+            _logger.Information("Participating in kubernetes leadership election as {Identity} using {Resource} under {Namespace}", _identity, settings.ConfigMapName, settings.Namespace);
+            _configMapLock = new ConfigMapLock(_client, settings.Namespace, settings.ConfigMapName, _identity);
+
+            LeaderElectionConfig leaderElectionConfig = new LeaderElectionConfig(_configMapLock)
+            {
+                LeaseDuration = settings.LeaseDuration,
+                RetryPeriod = settings.RetryPeriod,
+                RenewDeadline = settings.RenewDeadline
+            };
+
+            _leaderElector = new LeaderElector(leaderElectionConfig);
+            _leaderElector.OnStartedLeading += () => _logger.Warning("Acquired leadership from kubernetes");
+            _leaderElector.OnStoppedLeading += () => _logger.Warning("Lost leadership in kubernetes");
+            _leaderElector.OnNewLeader += OnLeaderElectorNewLeader;
+        }
+
+        private void OnLeaderElectorNewLeader(string leaderName)
+        {
+            _logger.Warning("{Instance} is the new leader", leaderName);
+
+            bool isLeader = string.Equals(leaderName, _identity, StringComparison.InvariantCultureIgnoreCase);
+            OnLeaderChanged?.Invoke(this, new ILeaderElection.OnLeaderChangedEventArgs(isLeader, leaderName));
+        }
+
+        public bool IsThisInstanceLeader()
+        {
+            return _leaderElector.IsLeader();
+        }
+
+        public event EventHandler<ILeaderElection.OnLeaderChangedEventArgs>? OnLeaderChanged;
+
+        public void Dispose()
+        {
+            _leaderElector?.Dispose();
+            _client?.Dispose();
+        }
+
+        public override async Task<bool> OnPoll(KubernetesLeaderElectionState state, CancellationToken cancellationToken)
+        {
+            _logger.Information("Polling kubernetes to determine leadership status");
+
+            try
+            {
+                await _configMapLock.GetAsync(cancellationToken);
+            }
+            catch (HttpOperationException e)
+            {
+                if (e.Response.StatusCode != HttpStatusCode.NotFound)
+                {
+                    _logger.Error(e, "Failed to fetch config map lock");
+                }
+            }
+            catch (Exception e)
+            {
+                _logger.Error(e, "Failed to fetch config map lock");
+            }
+
+            await _leaderElector.RunAsync(cancellationToken);
+
+            return true;
+        }
+    }
+}
diff --git a/Engine/Source/Programs/Horde/HordeStorage/Horde.Storage/Implementation/ObjectService.cs b/Engine/Source/Programs/Horde/HordeStorage/Horde.Storage/Implementation/ObjectService.cs
new file mode 100644
--- /dev/null
+++ b/Engine/Source/Programs/Horde/HordeStorage/Horde.Storage/Implementation/ObjectService.cs
@@ -0,0 +1,157 @@
+// Copyright Epic Games, Inc. All Rights Reserved.
+
+using System;
+using System.Collections.Generic;
+using System.Linq;
+using System.Threading.Tasks;
+using async_enumerable_dotnet;
+using Datadog.Trace;
+using Jupiter.Implementation;
+using Jupiter.Utils;
+
+namespace Horde.Storage.Implementation
+{
+    public interface IObjectService
+    {
+        Task<(ObjectRecord, BlobContents)> Get(NamespaceId ns, BucketId bucket, KeyId key, string[] fields);
+        Task<PutObjectResult> Put(NamespaceId ns, BucketId bucket, KeyId key, BlobIdentifier blobHash, CompactBinaryObject payload);
+        Task<BlobIdentifier[]> Finalize(NamespaceId ns, BucketId bucket, KeyId key, BlobIdentifier blobHash);
+
+        IAsyncEnumerator<NamespaceId> GetNamespaces();
+
+        Task<long> Delete(NamespaceId ns, BucketId bucket, KeyId key);
+        Task<long> DropNamespace(NamespaceId ns);
+        Task<long> DeleteBucket(NamespaceId ns, BucketId bucket);
+    }
+
+    public class ObjectService : IObjectService
+    {
+        private readonly IReferencesStore _referencesStore;
+        private readonly IBlobStore _blobStore;
+        private readonly IReferenceResolver _referenceResolver;
+        private readonly IReplicationLog _replicationLog;
+
+        public ObjectService(IReferencesStore referencesStore, IBlobStore blobStore, IReferenceResolver referenceResolver, IReplicationLog replicationLog)
+        {
+            _referencesStore = referencesStore;
+            _blobStore = blobStore;
+            _referenceResolver = referenceResolver;
+            _replicationLog = replicationLog;
+        }
+
+        public async Task<(ObjectRecord, BlobContents)> Get(NamespaceId ns, BucketId bucket, KeyId key, string[]? fields = null)
+        {
+            ObjectRecord o = await _referencesStore.Get(ns, bucket, key);
+
+            BlobContents blobContents;
+            if (o.InlinePayload != null && o.InlinePayload.Length != 0)
+            {
+                blobContents = new BlobContents(o.InlinePayload);
+            }
+            else
+            {
+                blobContents = await _blobStore.GetObject(ns, o.BlobIdentifier);
+            }
+            return (o, blobContents);
+        }
+
+        public async Task<PutObjectResult> Put(NamespaceId ns, BucketId bucket, KeyId key, BlobIdentifier blobHash, CompactBinaryObject payload)
+        {
+            bool hasReferences = payload.GetAllFields().Any(field => field.IsAttachment());
+
+            // if we have no references we are always finalized, e.g. there are no referenced blobs to upload
+            bool isFinalized = !hasReferences;
+
+            Task objectStorePut = _referencesStore.Put(ns, bucket, key, blobHash, payload.Data, isFinalized);
+
+            Task<BlobIdentifier> blobStorePut = _blobStore.PutObject(ns, payload.Data, blobHash);
+
+            BlobIdentifier[] missingReferences = Array.Empty<BlobIdentifier>();
+            if (hasReferences)
+            {
+                using Scope _ = Tracer.Instance.StartActive("ObjectService.ResolveReferences");
+                try
+                {
+                    IAsyncEnumerable<BlobIdentifier> references = _referenceResolver.ResolveReferences(ns, payload);
+                    missingReferences = await _blobStore.FilterOutKnownBlobs(ns, references);
+                }
+                catch (PartialReferenceResolveException e)
+                {
+                    missingReferences = e.UnresolvedReferences.ToArray();
+                }
+            }
+
+            await Task.WhenAll(objectStorePut, blobStorePut);
+
+            if (missingReferences.Length == 0)
+            {
+                await _referencesStore.Finalize(ns, bucket, key, blobHash);
+                await _replicationLog.InsertAddEvent(ns, bucket, key, blobHash);
+            }
+
+            return new PutObjectResult(missingReferences);
+        }
+
+        public async Task<BlobIdentifier[]> Finalize(NamespaceId ns, BucketId bucket, KeyId key, BlobIdentifier blobHash)
+        {
+            (ObjectRecord o, BlobContents blob) = await Get(ns, bucket, key);
+            byte[] blobContents = await blob.Stream.ToByteArray();
+            CompactBinaryObject payload = CompactBinaryObject.Load(blobContents);
+
+            bool hasReferences = payload.GetAllFields().Any(field => field.IsAttachment());
+
+            BlobIdentifier[] missingReferences = Array.Empty<BlobIdentifier>();
+            if (hasReferences)
+            {
+                using Scope _ = Tracer.Instance.StartActive("ObjectService.ResolveReferences");
+                try
+                {
+                    IAsyncEnumerable<BlobIdentifier> references = _referenceResolver.ResolveReferences(ns, payload);
+                    missingReferences = await _blobStore.FilterOutKnownBlobs(ns, references);
+                }
+                catch (PartialReferenceResolveException e)
+                {
+                    missingReferences = e.UnresolvedReferences.ToArray();
+                }
+            }
+
+            if (missingReferences.Length == 0)
+            {
+                await _referencesStore.Finalize(ns, bucket, key, blobHash);
+                await _replicationLog.InsertAddEvent(ns, bucket, key, blobHash);
+            }
+
+            return missingReferences;
+        }
+
+        public IAsyncEnumerator<NamespaceId> GetNamespaces()
+        {
+            return _referencesStore.GetNamespaces();
+        }
+
+        public Task<long> Delete(NamespaceId ns, BucketId bucket, KeyId key)
+        {
+            return _referencesStore.Delete(ns, bucket, key);
+        }
+
+        public Task<long> DropNamespace(NamespaceId ns)
+        {
+            return _referencesStore.DropNamespace(ns);
+        }
+
+        public Task<long> DeleteBucket(NamespaceId ns, BucketId bucket)
+        {
+            return _referencesStore.DeleteBucket(ns, bucket);
+        }
+    }
+
+    public class PutObjectResult
+    {
+        public PutObjectResult(BlobIdentifier[] missingReferences)
+        {
+            MissingReferences = missingReferences;
+        }
+
+        public BlobIdentifier[] MissingReferences { get; set; }
+    }
+}
diff --git a/Engine/Source/Programs/Horde/HordeStorage/Horde.Storage/Implementation/Objects/IReferencesStore.cs b/Engine/Source/Programs/Horde/HordeStorage/Horde.Storage/Implementation/Objects/IReferencesStore.cs
new file mode 100644
--- /dev/null
+++ b/Engine/Source/Programs/Horde/HordeStorage/Horde.Storage/Implementation/Objects/IReferencesStore.cs
@@ -0,0 +1,61 @@
+// Copyright Epic Games, Inc. All Rights Reserved.
+
+using System;
+using System.Collections.Generic;
+using System.Threading.Tasks;
+using Jupiter.Implementation;
+
+namespace Horde.Storage.Implementation
+{
+    public interface IReferencesStore
+    {
+        Task<ObjectRecord> Get(NamespaceId ns, BucketId bucket, KeyId key);
+        Task Put(NamespaceId ns, BucketId bucket, KeyId key, BlobIdentifier blobHash, byte[] blob, bool isFinalized);
+        Task Finalize(NamespaceId ns, BucketId bucket, KeyId key, BlobIdentifier blobHash);
+
+
+        Task UpdateLastAccessTime(NamespaceId ns, BucketId bucket, KeyId key, DateTime newLastAccessTime);
+        IAsyncEnumerator<ObjectRecord> GetOldestRecords(NamespaceId ns);
+
+        IAsyncEnumerator<NamespaceId> GetNamespaces();
+        Task<long> Delete(NamespaceId ns, BucketId bucket, KeyId key);
+        Task<long> DropNamespace(NamespaceId ns);
+        Task<long> DeleteBucket(NamespaceId ns, BucketId bucket);
+    }
+
+    public class ObjectRecord
+    {
+        public ObjectRecord(NamespaceId ns, BucketId bucket, KeyId name, DateTime lastAccess, byte[]? inlinePayload, BlobIdentifier blobIdentifier, bool isFinalized)
+        {
+            Namespace = ns;
+            Bucket = bucket;
+            Name = name;
+            LastAccess = lastAccess;
+            InlinePayload = inlinePayload;
+            BlobIdentifier = blobIdentifier;
+            IsFinalized = isFinalized;
+        }
+
+        public NamespaceId Namespace { get; }
+        public BucketId Bucket { get; }
+        public KeyId Name { get; }
+        public DateTime LastAccess { get; }
+        public byte[]? InlinePayload { get; set; }
+        public BlobIdentifier BlobIdentifier { get; set; }
+        public bool IsFinalized {get;}
+    }
+
+    public class ObjectNotFoundException : Exception
+    {
+        public ObjectNotFoundException(NamespaceId ns, BucketId bucket, KeyId key) : base($"Object not found {key} in bucket {bucket} namespace {ns}")
+        {
+            Namespace = ns;
+            Bucket = bucket;
+            Key = key;
+        }
+
+        public NamespaceId Namespace { get; }
+        public BucketId Bucket { get; }
+        public KeyId Key { get; }
+    }
+}
diff --git a/Engine/Source/Programs/Horde/HordeStorage/Horde.Storage/Implementation/Objects/MemoryReferencesStore.cs b/Engine/Source/Programs/Horde/HordeStorage/Horde.Storage/Implementation/Objects/MemoryReferencesStore.cs
new file mode 100644
--- /dev/null
+++ b/Engine/Source/Programs/Horde/HordeStorage/Horde.Storage/Implementation/Objects/MemoryReferencesStore.cs
@@ -0,0 +1,189 @@
+// Copyright Epic Games, Inc. All Rights Reserved.
+
+using System;
+using System.Collections.Concurrent;
+using System.Collections.Generic;
+using System.Linq;
+using System.Threading.Tasks;
+using Dasync.Collections;
+using Jupiter.Implementation;
+
+namespace Horde.Storage.Implementation
+{
+    public class MemoryReferencesStore : IReferencesStore
+    {
+        private readonly ConcurrentDictionary<string, MemoryStoreObject> _objects = new ConcurrentDictionary<string, MemoryStoreObject>();
+        private readonly List<NamespaceId> _namespaces = new List<NamespaceId>();
+
+        public MemoryReferencesStore()
+        {
+
+        }
+
+        public Task<ObjectRecord> Get(NamespaceId ns, BucketId bucket, KeyId key)
+        {
+            if (_objects.TryGetValue(BuildKey(ns, bucket, key), out MemoryStoreObject? o))
+            {
+                return Task.FromResult(o.ToObjectRecord());
+            }
+
+            throw new ObjectNotFoundException(ns, bucket, key);
+        }
+
+        public Task Put(NamespaceId ns, BucketId bucket, KeyId key, BlobIdentifier blobHash, byte[] blob, bool isFinalized)
+        {
+            lock (_namespaces)
+            {
+                _namespaces.Add(ns);
+            }
+
+            MemoryStoreObject o = _objects.AddOrUpdate(BuildKey(ns, bucket, key),
+                s => new MemoryStoreObject(ns, bucket, key, blobHash, blob, isFinalized),
+                (s, o) => new MemoryStoreObject(ns, bucket, key, blobHash, blob, isFinalized));
+
+            return Task.FromResult(o);
+        }
+
+        public Task Finalize(NamespaceId ns, BucketId bucket, KeyId key, BlobIdentifier blobHash)
+        {
+            if (!_objects.TryGetValue(BuildKey(ns, bucket, key), out MemoryStoreObject? o))
+            {
+                throw new ObjectNotFoundException(ns, bucket, key);
+            }
+
+            o.FinalizeObject();
+            return Task.CompletedTask;
+        }
+
+        public Task UpdateLastAccessTime(NamespaceId ns, BucketId bucket, KeyId key, DateTime lastAccessTime)
+        {
+            if (!_objects.TryGetValue(BuildKey(ns, bucket, key), out MemoryStoreObject? o))
+            {
+                throw new ObjectNotFoundException(ns, bucket, key);
+            }
+
+            o.SetLastAccessTime(lastAccessTime);
+            return Task.CompletedTask;
+        }
+
+        public async IAsyncEnumerator<ObjectRecord> GetOldestRecords(NamespaceId ns)
+        {
+            foreach (MemoryStoreObject o in _objects.Values.Where(o => o.Namespace == ns).OrderBy(o => o.LastAccessTime))
+            {
+                await Task.CompletedTask;
+                yield return o.ToObjectRecord();
+            }
+        }
+
+        public IAsyncEnumerator<NamespaceId> GetNamespaces()
+        {
+            return _namespaces.GetAsyncEnumerator();
+        }
+
+        public Task<long> Delete(NamespaceId ns, BucketId bucket, KeyId key)
+        {
+            if (!_objects.TryRemove(BuildKey(ns, bucket, key), out MemoryStoreObject? o))
+            {
+                throw new ObjectNotFoundException(ns, bucket, key);
+            }
+
+            return Task.FromResult(1L);
+        }
+
+        public Task<long> DropNamespace(NamespaceId ns)
+        {
+            lock (_namespaces)
+            {
+                _namespaces.Remove(ns);
+            }
+
+            List<string> objectToRemove = new List<string>();
+
+            foreach (MemoryStoreObject o in _objects.Values)
+            {
+                if (o.Namespace == ns)
+                {
+                    objectToRemove.Add(BuildKey(o.Namespace, o.Bucket, o.Name));
+                }
+            }
+
+            long removedCount = 0L;
+            foreach (string key in objectToRemove)
+            {
+                if (_objects.TryRemove(key, out _))
+                {
+                    removedCount++;
+                }
+            }
+
+            return Task.FromResult(removedCount);
+        }
+
+        public Task<long> DeleteBucket(NamespaceId ns, BucketId bucket)
+        {
+            List<string> objectToRemove = new List<string>();
+
+            foreach (MemoryStoreObject o in _objects.Values)
+            {
+                if (o.Namespace == ns && o.Bucket == bucket)
+                {
+                    objectToRemove.Add(BuildKey(o.Namespace, o.Bucket, o.Name));
+                }
+            }
+
+            long removedCount = 0L;
+            foreach (string key in objectToRemove)
+            {
+                if (_objects.TryRemove(key, out _))
+                {
+                    removedCount++;
+                }
+            }
+
+            return Task.FromResult(removedCount);
+        }
+
+        private static string BuildKey(NamespaceId ns, BucketId bucket, KeyId name)
+        {
+            return $"{ns}.{bucket}.{name}";
+        }
+    }
+
+    public class MemoryStoreObject
+    {
+        public MemoryStoreObject(NamespaceId ns, BucketId bucket, KeyId key, BlobIdentifier blobHash, byte[] blob, bool isFinalized)
+        {
+            Namespace = ns;
+            Bucket = bucket;
+            Name = key;
+            BlobHash = blobHash;
+            Blob = blob;
+            IsFinalized = isFinalized;
+            LastAccessTime = DateTime.Now;
+        }
+
+        public NamespaceId Namespace { get; }
+        public BucketId Bucket { get; }
+        public KeyId Name { get; }
+        public byte[] Blob { get; }
+        public BlobIdentifier BlobHash { get; }
+
+        public DateTime LastAccessTime { get; private set; }
+        public bool IsFinalized { get; private set;}
+
+        public void FinalizeObject()
+        {
+            IsFinalized = true;
+        }
+        public void SetLastAccessTime(DateTime lastAccessTime)
+        {
+            LastAccessTime = lastAccessTime;
+        }
+
+        public ObjectRecord ToObjectRecord()
+        {
+            return new ObjectRecord(Namespace, Bucket, Name, LastAccessTime, Blob, BlobHash, IsFinalized);
+        }
+
+    }
+}
