Moved Horde.Storage from git to p4

diff --git a/Engine/Source/Programs/Horde/HordeStorage/Horde.Storage/Implementation/Objects/ScyllaReferencesStore.cs b/Engine/Source/Programs/Horde/HordeStorage/Horde.Storage/Implementation/Objects/ScyllaReferencesStore.cs
new file mode 100644
--- /dev/null
+++ b/Engine/Source/Programs/Horde/HordeStorage/Horde.Storage/Implementation/Objects/ScyllaReferencesStore.cs
@@ -0,0 +1,436 @@
+// Copyright Epic Games, Inc. All Rights Reserved.
+
+using System;
+using System.Collections.Generic;
+using System.ComponentModel.DataAnnotations;
+using System.Linq;
+using System.Threading.Tasks;
+using Cassandra;
+using Cassandra.Mapping;
+using Datadog.Trace;
+using Jupiter.Implementation;
+using Microsoft.Extensions.Options;
+
+namespace Horde.Storage.Implementation
+{
+    public class ScyllaReferencesStore : IReferencesStore
+    {
+        private readonly ISession _session;
+        private readonly IMapper _mapper;
+        private readonly IOptionsMonitor<ScyllaSettings> _settings;
+
+        public ScyllaReferencesStore(IScyllaSessionManager scyllaSessionManager, IOptionsMonitor<ScyllaSettings> settings)
+        {
+            _session = scyllaSessionManager.GetSessionForReplicatedKeyspace();
+            _settings = settings;
+
+            _mapper = new Mapper(_session);
+
+            _session.Execute(new SimpleStatement(@"CREATE TABLE IF NOT EXISTS objects (
+                namespace text, 
+                bucket text, 
+                name text, 
+                payload_hash blob_identifier, 
+                inline_payload blob, 
+                is_finalized boolean,
+                last_access_time timestamp,
+                PRIMARY KEY ((namespace, bucket, name))
+            );"
+            ));
+            
+            _session.Execute(new SimpleStatement(@"CREATE TABLE IF NOT EXISTS buckets (
+                namespace text, 
+                bucket set<text>, 
+                PRIMARY KEY (namespace)
+            );"
+            ));
+
+            _session.Execute(new SimpleStatement(@"CREATE TABLE IF NOT EXISTS object_last_access (
+                namespace text,
+                partition_index tinyint,
+                accessed_at date,
+                objects set<text>,
+                PRIMARY KEY ((namespace, partition_index, accessed_at))
+            );"
+            ));
+
+
+        }
+
+        public async Task<ObjectRecord> Get(NamespaceId ns, BucketId bucket, KeyId name)
+        {
+            using Scope _ = Tracer.Instance.StartActive("scylla.get");
+
+            ScyllaObject? o = await _mapper.SingleOrDefaultAsync<ScyllaObject>("WHERE namespace = ? AND bucket = ? AND name = ?", ns.ToString(), bucket.ToString(), name.ToString());
+
+            if (o == null)
+                throw new ObjectNotFoundException(ns, bucket, name);
+
+            // TODO: Check returned values for null
+            return new ObjectRecord(new NamespaceId(o.Namespace!), new BucketId(o.Bucket!), new KeyId(o.Name!), o.LastAccessTime, o.InlinePayload, o.PayloadHash!.AsBlobIdentifier(), o.IsFinalized!.Value);
+        }
+
+        public async Task Put(NamespaceId ns, BucketId bucket, KeyId name, BlobIdentifier blobHash, byte[] blob, bool isFinalized)
+        {
+            using Scope _ = Tracer.Instance.StartActive("scylla.put");
+            if (blob.LongLength > _settings.CurrentValue.InlineBlobMaxSize)
+            {
+                // do not inline large blobs
+                blob = Array.Empty<byte>();
+            }
+
+            // add the bucket in parallel with inserting the actual object
+            var addBucketTask = MaybeAddBucket(ns, bucket);
+            // add the ttl records in parallel with the object
+            var createTTLRecords = CreateTTLRecord(ns, bucket, name, (sbyte)blobHash.HashData[0]);
+
+            await _mapper.InsertAsync<ScyllaObject>(new ScyllaObject(ns, bucket, name, blob, blobHash, isFinalized));
+            await Task.WhenAll(addBucketTask, createTTLRecords);
+        }
+
+
+        public async Task Finalize(NamespaceId ns, BucketId bucket, KeyId name, BlobIdentifier blobHash)
+        {
+            using Scope _ = Tracer.Instance.StartActive("scylla.finalize");
+            ObjectRecord o = await Get(ns, bucket, name);
+            if (!o.BlobIdentifier.Equals(blobHash))
+                throw new ObjectHashMismatchException(ns, bucket, name, blobHash);
+
+            AppliedInfo<ScyllaObject> info = await _mapper.UpdateIfAsync<ScyllaObject>("SET is_finalized=true WHERE namespace=? AND bucket=? AND name=?", ns.ToString(), bucket.ToString(), name.ToString());
+            if (!info.Applied)
+                throw new InvalidOperationException("Failed to finalize object even though it existed and hashes matched");
+        }
+
+        private async Task CreateTTLRecord(NamespaceId ns, BucketId bucket, KeyId name, sbyte partitionIndex)
+        {
+            using Scope _ = Tracer.Instance.StartActive("scylla.create_ttl_record");
+            DateTime now = DateTime.Now;
+            LocalDate lastAccessAt = new LocalDate(now.Year, now.Month, now.Day);
+            await _mapper.UpdateAsync<ScyllaObjectLastAccessAt>("SET objects = objects + ? WHERE accessed_at = ? AND namespace = ? AND partition_index = ? ", new string[] {$"{bucket}#{name}"},  lastAccessAt, ns.ToString(), partitionIndex);
+        }
+
+        private Task RemoveTTLRecord(ObjectRecord record)
+        {
+            return RemoveTTLRecord(record.Namespace, record.Bucket, record.Name, record.LastAccess, (sbyte)record.BlobIdentifier.HashData[0]);
+        }
+
+        private async Task RemoveTTLRecord(NamespaceId ns, BucketId bucket, KeyId name, DateTime lastAccess, sbyte partitionIndex)
+        {
+            using Scope _ = Tracer.Instance.StartActive("scylla.remove_ttl_record");
+            LocalDate lastAccessAt = new LocalDate(lastAccess.Year, lastAccess.Month, lastAccess.Day);
+            await _mapper.UpdateAsync<ScyllaObjectLastAccessAt>("SET objects = objects - ? WHERE accessed_at = ? AND namespace = ? AND partition_index = ?", new string[] {$"{bucket}#{name}"},  lastAccessAt, ns.ToString(), partitionIndex);
+        }
+
+        public async Task UpdateLastAccessTime(NamespaceId ns, BucketId bucket, KeyId name, DateTime lastAccessTime)
+        {
+            using Scope _ = Tracer.Instance.StartActive("scylla.update_last_access_time");
+            // fetch the old record
+            ObjectRecord objectRecord = await Get(ns, bucket, name);
+
+            // update the object tracking
+            Task updateObjectTracking = _mapper.UpdateAsync<ScyllaObject>("SET last_access_time = ? WHERE namespace=? AND bucket = ? AND name = ?", lastAccessTime, ns.ToString(), bucket.ToString(), name.ToString());
+
+            // the object was accessed during the same day so no need to move it
+            if (objectRecord.LastAccess.Date != lastAccessTime.Date)
+            {
+                sbyte partitionIndex = (sbyte)objectRecord.BlobIdentifier.HashData[0];
+                // move which last access partition is used for the object
+                LocalDate lastAccessAt = new LocalDate(objectRecord.LastAccess.Year, objectRecord.LastAccess.Month, objectRecord.LastAccess.Day);
+                await _mapper.UpdateAsync<ScyllaObjectLastAccessAt>("SET objects = objects - ? WHERE accessed_at=? AND namespace = ? AND partition_index = ?", new string[] {$"{bucket}#{name}"}, lastAccessAt, ns.ToString(), partitionIndex);
+                LocalDate newLastAccessAt = new LocalDate(lastAccessTime.Year, lastAccessTime.Month, lastAccessTime.Day);
+                await _mapper.UpdateAsync<ScyllaObjectLastAccessAt>("SET objects = objects + ? WHERE accessed_at=? AND namespace = ? AND partition_index = ?", new string[] {$"{bucket}#{name}"}, newLastAccessAt, ns.ToString(), partitionIndex);
+            }
+
+            await updateObjectTracking;
+        }
+
+        public async IAsyncEnumerator<ObjectRecord> GetOldestRecords(NamespaceId ns)
+        {
+            using Scope _ = Tracer.Instance.StartActive("scylla.get_records");
+            for (sbyte partitionIndex = sbyte.MinValue; partitionIndex < sbyte.MaxValue; partitionIndex++)
+            {
+                RowSet rowSetDateBuckets = await _session.ExecuteAsync(new SimpleStatement("SELECT accessed_at FROM object_last_access WHERE namespace = ? AND partition_index = ? ALLOW FILTERING", ns.ToString(), partitionIndex));
+                SortedSet<DateTime> dateBuckets = new SortedSet<DateTime>();
+                foreach (Row row in rowSetDateBuckets)
+                {
+                    if (rowSetDateBuckets.GetAvailableWithoutFetching() == 0)
+                        await rowSetDateBuckets.FetchMoreResultsAsync();
+
+                    LocalDate date = row.GetValue<LocalDate>("accessed_at");
+
+                    dateBuckets.Add(new DateTime(date.Year, date.Month, date.Day));
+                }
+
+                foreach (DateTime dateBucket in dateBuckets)
+                {
+                    LocalDate date = new LocalDate(dateBucket.Year, dateBucket.Month, dateBucket.Day);
+                    foreach (ScyllaObjectLastAccessAt lastAccessRecord in await _mapper.FetchAsync<ScyllaObjectLastAccessAt>("WHERE namespace = ? AND accessed_at = ? AND partition_index = ?", ns.ToString(), date, partitionIndex))
+                    {
+                        foreach (string o in lastAccessRecord.Objects)
+                        {
+                            int bucketSeparator = o.IndexOf("#", StringComparison.InvariantCultureIgnoreCase);
+                            BucketId bucket = new BucketId(o.Substring(0, bucketSeparator));
+                            KeyId name = new KeyId(o.Substring(bucketSeparator + 1));
+
+                            ObjectRecord record;
+                            try
+                            {
+                                 record = await Get(ns, bucket, name);
+                            }
+                            catch (ObjectNotFoundException)
+                            {
+                                await RemoveTTLRecord(ns, bucket, name, dateBucket, partitionIndex);
+                                continue;
+                            }
+                            yield return record;
+                        }
+                    }
+                }
+
+            }
+        }
+
+        public async IAsyncEnumerator<NamespaceId> GetNamespaces()
+        {
+            using Scope _ = Tracer.Instance.StartActive("scylla.get_namespaces");
+            RowSet rowSet = await _session.ExecuteAsync(new SimpleStatement("SELECT DISTINCT namespace FROM buckets"));
+
+            foreach (Row row in rowSet)
+            {
+                if (rowSet.GetAvailableWithoutFetching() == 0)
+                    await rowSet.FetchMoreResultsAsync();
+
+                yield return new NamespaceId(row.GetValue<string>(0));
+            }
+        }
+
+        public async Task<long> Delete(NamespaceId ns, BucketId bucket, KeyId key)
+        {
+            using Scope _ = Tracer.Instance.StartActive("scylla.delete_record");
+            ObjectRecord record;
+            try
+            {
+                record = await Get(ns, bucket, key);
+            }
+            catch (ObjectNotFoundException)
+            {
+                // if the record does not exist we do not need to do anything
+                return 0L;
+            }
+
+            Task removeTTL = RemoveTTLRecord(record);
+            AppliedInfo<ScyllaObject> info = await _mapper.DeleteIfAsync<ScyllaObject>("WHERE namespace=? AND bucket=? AND name=?", ns.ToString(), bucket.ToString(), key.ToString());
+
+            await removeTTL;
+            if (info.Applied)
+                return 1L;
+
+            return 0L;
+        }
+
+        public async Task<long> DropNamespace(NamespaceId ns)
+        {
+            using Scope _ = Tracer.Instance.StartActive("scylla.delete_namespace");
+            RowSet rowSet = await _session.ExecuteAsync(new SimpleStatement("SELECT bucket, name FROM objects WHERE namespace = ? ALLOW FILTERING;", ns.ToString()));
+            long deletedCount = 0;
+            foreach (Row row in rowSet)
+            {
+                string bucket = row.GetValue<string>("bucket");
+                string name = row.GetValue<string>("name");
+
+                await Delete(ns, new BucketId(bucket), new KeyId(name));
+
+                deletedCount++;
+            }
+
+            // remove the tracking in the buckets table as well
+            await _session.ExecuteAsync(new SimpleStatement("DELETE FROM buckets WHERE namespace = ?", ns.ToString()));
+
+            return deletedCount;
+        }
+
+        public async Task<long> DeleteBucket(NamespaceId ns, BucketId bucket)
+        {
+            using Scope _ = Tracer.Instance.StartActive("scylla.delete_bucket");
+            RowSet rowSet = await _session.ExecuteAsync(new SimpleStatement("SELECT name FROM objects WHERE namespace = ? AND bucket = ? ALLOW FILTERING;", ns.ToString(), bucket.ToString()));
+            long deletedCount = 0;
+            foreach (Row row in rowSet)
+            {
+                string name = row.GetValue<string>("name");
+
+                await Delete(ns, bucket, new KeyId(name));
+                deletedCount++;
+            }
+
+            // remove the tracking in the buckets table as well
+            await _mapper.UpdateAsync<ScyllaBucket>("SET bucket = bucket - ? WHERE namespace = ?", new string[] {bucket.ToString()}, ns.ToString());
+
+            return deletedCount;
+        }
+
+        private async Task MaybeAddBucket(NamespaceId ns, BucketId bucket)
+        {
+            using Scope _ = Tracer.Instance.StartActive("scylla.add_bucket");
+            await _mapper.UpdateAsync<ScyllaBucket>("SET bucket = bucket + ? WHERE namespace = ?", new string[] {bucket.ToString()}, ns.ToString());
+        }
+    }
+
+    public class ObjectHashMismatchException : Exception
+    {
+        public ObjectHashMismatchException(NamespaceId ns, BucketId bucket, KeyId name, BlobIdentifier blobHash) : base($"Object {name} in bucket {bucket} and namespace {ns} did not reference hash {blobHash}")
+        {
+        }
+    }
+
+    public class ScyllaBlobIdentifier
+    {
+        public ScyllaBlobIdentifier()
+        {
+            Hash = null;
+        }
+
+        public ScyllaBlobIdentifier(ContentHash hash)
+        {
+            Hash = hash.HashData;
+        }
+
+        public byte[]? Hash { get;set; }
+
+        public BlobIdentifier AsBlobIdentifier()
+        {
+            return new BlobIdentifier(Hash!); 
+        }
+    }
+
+    public class ScyllaSettings
+    {
+        public long InlineBlobMaxSize { get; set; } = 32 * 1024; // default to 32 kb blobs max
+        public string[] ContactPoints { get; set; } = new string[] {"localhost", "scylla"};
+        public int MaxSnapshotsPerNamespace { get; set; } = 10;
+
+        /// <summary>
+        /// Set to override the replication strategy used to create keyspace
+        /// Note that this only applies when creating the keyspace
+        /// To modify a existing keyspace see https://docs.datastax.com/en/dse/6.7/dse-admin/datastax_enterprise/operations/opsChangeKSStrategy.html
+        /// </summary>
+        public Dictionary<string, string>? KeyspaceReplicationStrategy { get; set; }
+
+        /// <summary>
+        /// Set to override the replication strategy used to create the local keyspace
+        /// Note that this only applies when creating the keyspace
+        /// To modify a existing keyspace see https://docs.datastax.com/en/dse/6.7/dse-admin/datastax_enterprise/operations/opsChangeKSStrategy.html
+        /// </summary>
+        public Dictionary<string, string>? LocalKeyspaceReplicationStrategy { get; set; }
+
+        /// <summary>
+        /// Used to configure the load balancing policy to stick to this specified datacenter
+        /// </summary>
+        [Required]
+        public string LocalDatacenterName { get; set; } = null!;
+        
+        [Required] 
+        public string LocalKeyspaceSuffix { get; set; } = null!;
+
+        /// <summary>
+        /// Max number of connections for each scylla host before switching to another host
+        /// See https://docs.datastax.com/en/developer/nodejs-driver/4.6/features/connection-pooling/
+        /// </summary>
+        public int MaxConnectionForLocalHost { get; set; } = 8192;
+
+        /// <summary>
+        /// The time for a replication log event to live in the incremental state before being deleted, assumption is that the snapshot will have processed the event within this time
+        /// </summary>
+        public TimeSpan ReplicationLogTimeToLive { get; set; } = TimeSpan.FromDays(7);
+    }
+
+    [Cassandra.Mapping.Attributes.Table("objects")]
+    public class ScyllaObject
+    {
+        public ScyllaObject()
+        {
+
+        }
+
+        public ScyllaObject(NamespaceId ns, BucketId bucket, KeyId name, byte[] payload, BlobIdentifier payloadHash, bool isFinalized)
+        {
+            Namespace = ns.ToString();
+            Bucket = bucket.ToString();
+            Name = name.ToString();
+            InlinePayload = payload;
+            PayloadHash = new ScyllaBlobIdentifier(payloadHash);
+
+            IsFinalized = isFinalized;
+
+            LastAccessTime = DateTime.Now;
+        }
+
+        public static string BuildKey(string ns, string bucket, string name)
+        {
+            return $"{ns}.{bucket}.{name}";
+        }
+
+        [Cassandra.Mapping.Attributes.PartitionKey(0)]
+        public string? Namespace { get; set; }
+
+        [Cassandra.Mapping.Attributes.PartitionKey(1)]
+        public string? Bucket { get; set; }
+
+        [Cassandra.Mapping.Attributes.PartitionKey(2)]
+        public string? Name { get; set; }
+
+        [Cassandra.Mapping.Attributes.Column("payload_hash")]
+        public ScyllaBlobIdentifier? PayloadHash { get; set; }
+
+        [Cassandra.Mapping.Attributes.Column("inline_payload")]
+        public byte[]? InlinePayload {get; set; }
+
+        [Cassandra.Mapping.Attributes.Column("is_finalized")]
+        public bool? IsFinalized { get;set; }
+        [Cassandra.Mapping.Attributes.Column("last_access_time")]
+        public DateTime LastAccessTime { get; set; }
+    }
+
+    [Cassandra.Mapping.Attributes.Table("buckets")]
+    public class ScyllaBucket
+    {
+        public ScyllaBucket()
+        {
+
+        }
+
+        public ScyllaBucket(NamespaceId ns, BucketId[] buckets)
+        {
+            Namespace = ns.ToString();
+            Buckets = buckets.Select(b => b.ToString()).ToList();
+        }
+
+        [Cassandra.Mapping.Attributes.PartitionKey]
+        public string? Namespace { get; set; }
+
+        public List<string> Buckets { get; set; } = new List<string>();
+    }
+
+    [Cassandra.Mapping.Attributes.Table("object_last_access")]
+    public class ScyllaObjectLastAccessAt
+    {
+        public ScyllaObjectLastAccessAt()
+        {
+
+        }
+
+        [Cassandra.Mapping.Attributes.PartitionKey(0)]
+        public string? Namespace { get; set; }
+
+        [Cassandra.Mapping.Attributes.PartitionKey(1)]
+        [Cassandra.Mapping.Attributes.Column("partition_index")]
+        public sbyte? PartitionIndex { get;set; }
+
+        [Cassandra.Mapping.Attributes.PartitionKey(2)]
+        [Cassandra.Mapping.Attributes.Column("accessed_at")]
+        public LocalDate? AccessedAt { get; set; }
+
+
+        public List<string> Objects { get; set; } = new List<string>();
+
+    }
+}
diff --git a/Engine/Source/Programs/Horde/HordeStorage/Horde.Storage/Implementation/Refs/CachedRefStore.cs b/Engine/Source/Programs/Horde/HordeStorage/Horde.Storage/Implementation/Refs/CachedRefStore.cs
new file mode 100644
--- /dev/null
+++ b/Engine/Source/Programs/Horde/HordeStorage/Horde.Storage/Implementation/Refs/CachedRefStore.cs
@@ -0,0 +1,176 @@
+﻿using System;
+using System.Collections.Concurrent;
+using System.Collections.Generic;
+using System.Threading;
+using System.Threading.Tasks;
+using Jupiter.Implementation;
+using Microsoft.Extensions.Caching.Memory;
+using Microsoft.Extensions.Options;
+using Microsoft.Extensions.Primitives;
+using Serilog;
+using Serilog.Context;
+
+namespace Horde.Storage.Implementation
+{
+    public class CachedRefStore : IRefsStore
+    {
+        private readonly IRefsStore _backingStore;
+        private readonly IOptionsMonitor<MemoryCacheRefSettings> _memoryOptions;
+        private readonly ConcurrentDictionary<NamespaceId, BoolChangeToken> _namespaceChangeTokens;
+        private readonly ConcurrentDictionary<string, BoolChangeToken> _bucketChangeTokens;
+        private readonly MemoryCache _memoryCache;
+
+        private readonly ILogger _logger = Log.ForContext<CachedRefStore>();
+
+        // ReSharper disable once NotAccessedField.Local
+        private Timer _cacheInfoTimer;
+
+        public CachedRefStore(IRefsStore backingStore, IOptionsMonitor<MemoryCacheRefSettings> memoryOptions)
+        {
+            _backingStore = backingStore;
+            _memoryOptions = memoryOptions;
+            _memoryCache = new MemoryCache(_memoryOptions.CurrentValue);
+            _namespaceChangeTokens = new ConcurrentDictionary<NamespaceId, BoolChangeToken>();
+            _bucketChangeTokens = new ConcurrentDictionary<string, BoolChangeToken>();
+
+            _cacheInfoTimer = new Timer(LogCacheInfo, null, TimeSpan.Zero, TimeSpan.FromSeconds(60));
+        }
+
+        public IRefsStore BackingStore
+        {
+            get { return _backingStore; }
+        }
+
+        private void LogCacheInfo(object? state)
+        {
+            LogContext.Reset();
+            _logger.Information("Memory db-cache {Count}", _memoryCache.Count);
+        }
+
+        private static string BuildCacheKey(NamespaceId ns, BucketId bucket, KeyId key)
+        {
+            return $"{ns}.{bucket}.{key}";
+        }
+
+        public async Task<RefRecord?> Get(NamespaceId ns, BucketId bucket, KeyId key, IRefsStore.ExtraFieldsFlag fields)
+        {
+            // we do not cache last access as it is frequently updated, if this fields is requested this cache can not be used
+            if (fields.HasFlag(IRefsStore.ExtraFieldsFlag.LastAccess))
+            {
+                return await _backingStore.Get(ns, bucket, key, fields);
+            }
+
+            string cacheKey = BuildCacheKey(ns, bucket, key);
+            if (_memoryCache.TryGetValue(cacheKey, out RefRecord? cachedKey))
+            {
+                return cachedKey;
+            }
+
+            // override the fields used so that we get a full object for use by the cache
+            IRefsStore.ExtraFieldsFlag localFields = IRefsStore.ExtraFieldsFlag.Metadata;
+            RefRecord? backendRecord = await _backingStore.Get(ns, bucket, key, localFields);
+            
+            // do not refs a miss as it could get added at any time
+            if (backendRecord == null)
+                return null;
+
+            CreateMemoryCacheEntry(cacheKey, backendRecord, ns, bucket);
+
+            return backendRecord;
+        }
+
+        private void CreateMemoryCacheEntry(string cacheKey, object value, NamespaceId ns, BucketId bucket)
+        {
+            using ICacheEntry cacheEntry = _memoryCache.CreateEntry(cacheKey);
+
+            if (_memoryOptions.CurrentValue.EnableSlidingExpiry)
+            {
+                cacheEntry.SlidingExpiration = TimeSpan.FromMinutes(_memoryOptions.CurrentValue.SlidingExpirationMinutes);
+            }
+
+            cacheEntry.RegisterPostEvictionCallback((key, value, reason, state) =>
+            {
+                _logger.Information("{Key} evicted from db cache because {Reason}", key, reason);
+            });
+            cacheEntry.Value = value;
+
+            BoolChangeToken namespaceChangeToken = _namespaceChangeTokens.GetOrAdd(ns, new BoolChangeToken());
+            cacheEntry.ExpirationTokens.Add(namespaceChangeToken);
+
+            BoolChangeToken bucketChangeToken = _bucketChangeTokens.GetOrAdd($"{ns}.{bucket}", new BoolChangeToken());
+            cacheEntry.ExpirationTokens.Add(bucketChangeToken);
+        }
+
+        public Task Add(RefRecord record)
+        {
+            string cacheKey = BuildCacheKey(record.Namespace, record.Bucket, record.RefName);
+            CreateMemoryCacheEntry(cacheKey, record, record.Namespace, record.Bucket);
+
+            return _backingStore.Add(record);
+        }
+
+        public Task<long> DeleteBucket(NamespaceId ns, BucketId bucket)
+        {
+            if (_bucketChangeTokens.TryGetValue($"{ns}.{bucket}", out BoolChangeToken? bucketToken))
+            {
+                // sets the objects as invalid and will be cleaned up when the next compact runs
+                bucketToken.HasChanged = true;
+            }
+
+            return _backingStore.DeleteBucket(ns, bucket);
+        }
+
+        public Task<long> Delete(NamespaceId ns, BucketId bucket, KeyId key)
+        {
+            string cacheKey = BuildCacheKey(ns, bucket, key);
+            _memoryCache.Remove(cacheKey);
+
+            return _backingStore.Delete(ns, bucket, key);
+        }
+
+        public Task UpdateLastAccessTime(RefRecord record, DateTime lastAccessTime)
+        {
+            // we do not cache last access time so we just calling the backing store
+            return _backingStore.UpdateLastAccessTime(record, lastAccessTime);
+        }
+
+        public IAsyncEnumerable<NamespaceId> GetNamespaces()
+        {
+            // we do not cache namespaces so we just call the backing store
+            return _backingStore.GetNamespaces();
+        }
+
+        public IAsyncEnumerable<OldRecord> GetOldRecords(NamespaceId ns, TimeSpan oldRecordCutoff)
+        {
+            // we do not cache the last access time so we can not filter out old records, thus we call the backing store
+            return _backingStore.GetOldRecords(ns, oldRecordCutoff);
+        }
+
+        public Task DropNamespace(NamespaceId ns)
+        {
+            if (_namespaceChangeTokens.TryGetValue(ns, out BoolChangeToken? namespaceToken))
+            {
+                // sets the objects as invalid and will be cleaned up when the next compact runs
+                namespaceToken.HasChanged = true;
+            }
+
+            return _backingStore.DropNamespace(ns);
+        }
+    }
+
+    class BoolChangeToken : IChangeToken
+    {
+        public IDisposable RegisterChangeCallback(Action<object> callback, object state)
+        {
+            throw new NotImplementedException();
+        }
+
+        public bool HasChanged { get; set; }
+
+        public bool ActiveChangeCallbacks
+        {
+            get { return false; }
+        }
+    }
+
+}
diff --git a/Engine/Source/Programs/Horde/HordeStorage/Horde.Storage/Implementation/Refs/CosmosRefsStore.cs b/Engine/Source/Programs/Horde/HordeStorage/Horde.Storage/Implementation/Refs/CosmosRefsStore.cs
new file mode 100644
--- /dev/null
+++ b/Engine/Source/Programs/Horde/HordeStorage/Horde.Storage/Implementation/Refs/CosmosRefsStore.cs
@@ -0,0 +1,90 @@
+﻿// Copyright Epic Games, Inc. All Rights Reserved.
+
+using System;
+using System.Threading.Tasks;
+using Microsoft.Extensions.Options;
+using MongoDB.Bson;
+using MongoDB.Driver;
+using Serilog;
+
+namespace Horde.Storage.Implementation
+{
+    /// <summary>
+    /// A specialized form of the Mongo refs store that adds some Cosmos specific management
+    /// </summary>
+    // ReSharper disable once ClassNeverInstantiated.Global
+    public class CosmosRefsStore : MongoRefsStore
+    {
+        private readonly ILogger _logger = Log.ForContext<CosmosRefsStore>();
+        private readonly IOptionsMonitor<CosmosSettings> _cosmosSettings;
+
+        public CosmosRefsStore(IOptionsMonitor<MongoSettings> mongoSettings, IOptionsMonitor<CosmosSettings> cosmosSettings) : base(mongoSettings)
+        {
+            _cosmosSettings = cosmosSettings;
+        }
+
+
+        public override async Task Add(RefRecord record)
+        {
+            string dbName = GetDatabaseName(record.Namespace);
+
+            bool collectionExists = await CollectionExistsAsync(dbName, ModelName);
+            if (!collectionExists)
+            {
+                var createDbCommand = new BsonDocument
+                {
+                    {
+                        "customAction","CreateCollection"
+                    },
+                    {
+                        "collection", ModelName
+                    },
+                    {
+                        "offerThroughput",_cosmosSettings.CurrentValue.DefaultRU
+                    }
+                };
+                var db = _client.GetDatabase(dbName);
+                BsonDocument? result;
+                try
+                {
+                    result = await db.RunCommandAsync<BsonDocument>(createDbCommand);
+                }
+                catch (MongoCommandException e)
+                {
+                    if (e.Code == 117) // http 409 - mongo 117
+                    {
+                        // the database got created by someone else, that is okay lets just use the db and ignore this error
+                        result = null;
+                    }
+                    else
+                    {
+                        _logger.Warning("Mongo exception: {Code} {Exception}", e.Code, e.ToString());
+                        // rethrow if we do not recognize the exception
+                        throw;
+                    }
+                }
+
+                if (result != null)
+                {
+                    //https://docs.microsoft.com/en-us/azure/cosmos-db/mongodb-custom-commands#default-output
+                    int ok = result["ok"].AsInt32;
+                    if (ok == 0) // failure
+                        throw new Exception($"Failed to create MongoDB {dbName} in Cosmos. Response: {result}");
+                }
+            }
+
+            await base.Add(record);
+        }
+
+        private async Task<bool> CollectionExistsAsync(string dbName, string collectionName)
+        {
+            var db = _client.GetDatabase(dbName);
+
+            var filter = new BsonDocument("name", collectionName);
+            //filter by collection name
+            var collections = await db.ListCollectionsAsync(new ListCollectionsOptions { Filter = filter });
+            //check for existence
+            return await collections.AnyAsync();
+        }
+    }
+}
diff --git a/Engine/Source/Programs/Horde/HordeStorage/Horde.Storage/Implementation/Refs/DynamoDbRefsStore.cs b/Engine/Source/Programs/Horde/HordeStorage/Horde.Storage/Implementation/Refs/DynamoDbRefsStore.cs
new file mode 100644
--- /dev/null
+++ b/Engine/Source/Programs/Horde/HordeStorage/Horde.Storage/Implementation/Refs/DynamoDbRefsStore.cs
@@ -0,0 +1,468 @@
+// Copyright Epic Games, Inc. All Rights Reserved.
+
+using System;
+using System.Collections.Generic;
+using System.Linq;
+using System.Threading.Tasks;
+using Amazon.DynamoDBv2;
+using Amazon.DynamoDBv2.DataModel;
+using Amazon.DynamoDBv2.DocumentModel;
+using Amazon.DynamoDBv2.Model;
+using Datadog.Trace;
+using Jupiter.Implementation;
+using Microsoft.Extensions.Options;
+using Serilog;
+
+namespace Horde.Storage.Implementation
+{
+    internal class DynamoDbRefsStore : DynamoStore, IRefsStore
+    {
+        private readonly ILogger _logger = Log.ForContext<DynamoDbRefsStore>();
+        private readonly IOptionsMonitor<DynamoDbSettings> _settings;
+        private readonly Random _lastAccessIndexGenerator = new Random();
+
+        private const int LastAccessIndexMax = 100;
+
+        public const string MainTableName = "Europa_Cache";
+
+        private readonly DynamoNamespaceStore _namespaceStore;
+
+        public DynamoDbRefsStore(IOptionsMonitor<DynamoDbSettings> settings, IAmazonDynamoDB dynamoDb, DynamoNamespaceStore namespaceStore) : base(settings, dynamoDb)
+        {
+            _settings = settings;
+            _namespaceStore = namespaceStore;
+        }
+
+        protected override async Task CreateTables()
+        {
+            // TODO: Add a Bucket Table to be able query which buckets exist within a ns
+
+            Task createMainTable = CreateTable(MainTableName, new[]
+            {
+                new KeySchemaElement("Id", KeyType.HASH),
+            }, new[]
+            {
+                new AttributeDefinition("Id", ScalarAttributeType.S),
+                // these two are used by the GSI for LastAccess
+                new AttributeDefinition("LastAccessTime", ScalarAttributeType.S),
+                new AttributeDefinition("LastAccessIndex", ScalarAttributeType.N)
+            }, new ProvisionedThroughput
+            {
+                ReadCapacityUnits = _settings.CurrentValue.ReadCapacityUnits,
+                WriteCapacityUnits = _settings.CurrentValue.WriteCapacityUnits
+            }, 
+            // LastAccessIndex is sorted on last access time to allow us to quickly find out of date records
+            new [] {
+                new GlobalSecondaryIndex
+                {
+                    IndexName = "LastAccessIndex",
+                    KeySchema = new[]
+                    {
+                        new KeySchemaElement("LastAccessIndex", KeyType.HASH),
+                        new KeySchemaElement("LastAccessTime", KeyType.RANGE)
+                    }.ToList(),
+                    Projection = new Projection
+                    {
+                        ProjectionType = ProjectionType.KEYS_ONLY,
+                    }
+                }}
+            );
+
+            await Task.WhenAll(_namespaceStore.Initialize(), createMainTable);
+        }
+
+        public async Task<RefRecord?> Get(NamespaceId ns, BucketId bucket, KeyId key, IRefsStore.ExtraFieldsFlag fields)
+        {
+            await Initialize();
+
+            using Scope scope = Tracer.Instance.StartActive("dynamo.get");
+            Span span = scope.Span;
+            span.ResourceName = $"{ns}.{bucket}.{key}";
+
+            Task<DynamoBaseRefRecord> baseRefRecordTask = Context.LoadAsync<DynamoBaseRefRecord>(DynamoBaseRefRecord.BuildCacheKey(ns, bucket, key));
+            Task<DynamoMetadataRefRecord>? metadataRefRecordTask = fields.HasFlag(IRefsStore.ExtraFieldsFlag.Metadata) ? Context.LoadAsync<DynamoMetadataRefRecord>(DynamoMetadataRefRecord.BuildCacheKey(ns, bucket, key)) : null;
+            Task<DynamoLastAccessRefRecord>? lastAccessRefRecordTask = fields.HasFlag(IRefsStore.ExtraFieldsFlag.LastAccess) ? Context.LoadAsync<DynamoLastAccessRefRecord>(DynamoLastAccessRefRecord.BuildCacheKey(ns, bucket, key)) : null;
+
+            DynamoBaseRefRecord? baseRefRecord = await baseRefRecordTask;
+            DynamoLastAccessRefRecord? lastAccessRefRecord = lastAccessRefRecordTask != null ? await lastAccessRefRecordTask : null;
+            DynamoMetadataRefRecord? metadataRefRecord = metadataRefRecordTask != null ? await metadataRefRecordTask : null;
+
+            // LoadAsync can return null records even if its not annotated for it
+            // ReSharper disable once ConditionIsAlwaysTrueOrFalse
+            if (baseRefRecord == null)
+                return null; // this key is not valid
+
+            if (!baseRefRecord.IsValid())
+                throw new InvalidOperationException("Refs record was not fully populated, version mismatch?");
+
+            if (!lastAccessRefRecord?.IsValid() ?? false)
+                throw new InvalidOperationException("LastAccess Refs record was not fully populated, version mismatch?");
+
+            if (!metadataRefRecord?.IsValid() ?? false)
+                throw new InvalidOperationException("Metadata refs record was not fully populated, version mismatch?");
+
+            return new DynamoRefRecord(new NamespaceId(baseRefRecord.Namespace!), new BucketId(baseRefRecord.Bucket!), new KeyId(baseRefRecord.Name!), 
+                baseRefRecord.Blobs!.Select(s => new BlobIdentifier(s)).ToArray(), 
+                lastAccessRefRecord?.LastAccessTime ?? null, lastAccessRefRecord?.LastAccessIndex ?? null, 
+                new ContentHash(baseRefRecord.ContentHash!), metadataRefRecord?.Metadata ?? null);
+        }
+
+        public async Task Add(RefRecord record)
+        {
+            await Initialize();
+
+            using Scope scope = Tracer.Instance.StartActive("dynamo.add");
+            Span span = scope.Span;
+            span.ResourceName = $"{record.Namespace}.{record.Bucket}.{record.RefName}";
+
+            Task namespaceTableTask = _namespaceStore.AddToNamespaceTable(record.Namespace, DynamoNamespaceStore.NamespaceUsage.Cache);
+            // TODO: write to a bucket table
+
+            DynamoBaseRefRecord dynamoRefRecord = new DynamoBaseRefRecord(record);
+            Task baseInsertTask = Context.SaveAsync(dynamoRefRecord);
+
+            int lastAccessIndex = _lastAccessIndexGenerator.Next(LastAccessIndexMax);
+            DynamoLastAccessRefRecord dynamoLastAccessRefRecord = new DynamoLastAccessRefRecord(record, lastAccessIndex);
+            Task lastAccessInsertTask = Context.SaveAsync(dynamoLastAccessRefRecord);
+
+            DynamoMetadataRefRecord dynamoMetadataRefRecord = new DynamoMetadataRefRecord(record);
+            Task metadataInsertTask = Context.SaveAsync(dynamoMetadataRefRecord);
+
+            await Task.WhenAll(
+                baseInsertTask, 
+                lastAccessInsertTask, 
+                metadataInsertTask,
+                namespaceTableTask
+            );
+        }
+
+        public async Task<long> DeleteBucket(NamespaceId ns, BucketId bucket)
+        {
+            await Initialize();
+
+            AsyncSearch<DynamoBaseRefRecord> search = Context.ScanAsync<DynamoBaseRefRecord>(new[]
+            {
+                new ScanCondition("Namespace", ScanOperator.Equal, ns.ToString()),
+                new ScanCondition("Bucket", ScanOperator.Equal, bucket.ToString())
+            });
+
+            int countOfDeletedItems = 0;
+            do
+            {
+                List<DynamoBaseRefRecord> newSet = await search.GetNextSetAsync();
+                foreach (DynamoBaseRefRecord document in newSet)
+                {
+                    if (!document.IsValid())
+                    {
+                        _logger.Warning("Invalid refs record with id {Id}, missing one or more member, unable to delete", document.Id);
+                        continue;
+                    }
+
+                    Task deleteBase = Context.DeleteAsync(document);
+
+                    NamespaceId documentNs = new NamespaceId(document.Namespace!);
+                    BucketId documentBucket = new BucketId(document.Bucket!);
+                    KeyId documentKey = new KeyId(document.Name!);
+
+                    Task deleteLastAccess = Context.DeleteAsync<DynamoLastAccessRefRecord>(DynamoLastAccessRefRecord.BuildCacheKey(documentNs, documentBucket, documentKey));
+                    Task deleteMetadata = Context.DeleteAsync<DynamoMetadataRefRecord>(DynamoMetadataRefRecord.BuildCacheKey(documentNs, documentBucket, documentKey));
+                    countOfDeletedItems++;
+
+                    await Task.WhenAll(deleteBase, deleteLastAccess, deleteMetadata);
+                }
+            } while (!search.IsDone);
+
+            return countOfDeletedItems;
+        }
+
+
+        public async Task<long> Delete(NamespaceId ns, BucketId bucket, KeyId key)
+        {
+            await Initialize();
+
+            var deleteBaseTask = new DeleteItemRequest(MainTableName, new Dictionary<string, AttributeValue> {{"Id", new AttributeValue(DynamoBaseRefRecord.BuildCacheKey(ns, bucket, key)) }})
+            {
+                ConditionExpression = "attribute_exists(Id)"
+            };
+            var deleteLastAccessTask = new DeleteItemRequest(MainTableName, new Dictionary<string, AttributeValue> {{"Id", new AttributeValue(DynamoLastAccessRefRecord.BuildCacheKey(ns, bucket, key)) }});
+            var deleteMetadataTask = new DeleteItemRequest(MainTableName, new Dictionary<string, AttributeValue> {{"Id", new AttributeValue(DynamoMetadataRefRecord.BuildCacheKey(ns, bucket, key)) }});
+
+            try
+            {
+                await Task.WhenAll(Client.DeleteItemAsync(deleteBaseTask), Client.DeleteItemAsync(deleteLastAccessTask), Client.DeleteItemAsync(deleteMetadataTask));
+            }
+            catch (ConditionalCheckFailedException)
+            {
+                // our condition is that it exists, so if it fails it means we had nothing to delete
+                return 0;
+            }
+
+            return 1;
+        }
+
+
+        public async Task UpdateLastAccessTime(RefRecord record, DateTime lastAccessTime)
+        {
+            await Initialize();
+
+            using Scope scope = Tracer.Instance.StartActive("dynamo.update_last_access");
+            Span span = scope.Span;
+            span.ResourceName = $"{record.Namespace}.{record.Bucket}.{record.RefName}";
+
+            Table table = Table.LoadTable(Client, MainTableName);
+            // we use the document access here so that we can update just the last access time field
+            Document? item = await table.GetItemAsync(DynamoLastAccessRefRecord.BuildCacheKey(record.Namespace, record.Bucket, record.RefName));
+            // if the item has been deleted after it was scheduled to be updated it will be null
+            if (item == null)
+                return;
+            item["LastAccessTime"] = lastAccessTime;
+            await table.UpdateItemAsync(item);
+        }
+
+        public async IAsyncEnumerable<NamespaceId> GetNamespaces()
+        {
+            await Initialize();
+
+            await foreach (var ns in _namespaceStore.GetNamespaces(DynamoNamespaceStore.NamespaceUsage.Cache))
+            {
+                yield return ns;
+            }
+        }
+
+        public async IAsyncEnumerable<OldRecord> GetOldRecords(NamespaceId ns, TimeSpan oldRecordCutoff)
+        {
+            await Initialize();
+
+            DateTime cutOffDate = DateTime.Now - oldRecordCutoff;
+            using var context = new DynamoDBContext(Client, new DynamoDBContextConfig { Conversion = DynamoDBEntryConversion.V2 });
+            AsyncSearch<DynamoLastAccessRefRecord> search = context.ScanAsync<DynamoLastAccessRefRecord>(new[]
+            {
+                new ScanCondition("LastAccessTime", ScanOperator.Between, DateTime.MinValue, cutOffDate)
+            });
+
+            do
+            {
+                List<DynamoLastAccessRefRecord> newSet = await search.GetNextSetAsync();
+                foreach (DynamoLastAccessRefRecord document in newSet)
+                {
+                    if (document == null)
+                        continue;
+                    string parentKey = DynamoLastAccessRefRecord.ToBaseRefRecordsKey(document.Id!);
+
+                    DynamoBaseRefRecord? baseRecord = await context.LoadAsync<DynamoBaseRefRecord>(parentKey);
+                    if (baseRecord == null || baseRecord.Namespace == null || baseRecord.Bucket == null || baseRecord.Name == null)
+                        continue;
+                    yield return new OldRecord(new NamespaceId(baseRecord.Namespace), new BucketId(baseRecord.Bucket), new KeyId(baseRecord.Name));
+                }
+            } while (!search.IsDone);
+        }
+
+        public async Task DropNamespace(NamespaceId ns)
+        {
+            using var context = new DynamoDBContext(Client, new DynamoDBContextConfig { Conversion = DynamoDBEntryConversion.V2 });
+            try
+            {
+                AsyncSearch<DynamoBaseRefRecord> search = context.ScanAsync<DynamoBaseRefRecord>(new[]
+                {
+                    new ScanCondition("Namespace", ScanOperator.Equal, ns.ToString())
+                });
+
+                int countOfDeletedItems = 0;
+                do
+                {
+                    List<DynamoBaseRefRecord> newSet = await search.GetNextSetAsync();
+                    foreach (DynamoBaseRefRecord document in newSet)
+                    {
+                        if (!document.IsValid())
+                        {
+                            _logger.Warning("Invalid refs record with id {Id}, missing one or more member, unable to delete", document.Id);
+                            continue;
+                        }
+
+                        Task deleteBase = Context.DeleteAsync(document);
+
+                        NamespaceId documentNs = new NamespaceId(document.Namespace!);
+                        BucketId documentBucket = new BucketId(document.Bucket!);
+                        KeyId documentKey = new KeyId(document.Name!);
+
+                        Task deleteLastAccess = Context.DeleteAsync<DynamoLastAccessRefRecord>(DynamoLastAccessRefRecord.BuildCacheKey(documentNs, documentBucket, documentKey));
+                        Task deleteMetadata = Context.DeleteAsync<DynamoMetadataRefRecord>(DynamoMetadataRefRecord.BuildCacheKey(documentNs, documentBucket, documentKey));
+                        countOfDeletedItems++;
+
+                        await Task.WhenAll(deleteBase, deleteLastAccess, deleteMetadata);
+                    }
+                } while (!search.IsDone);
+            }
+            catch (ResourceNotFoundException)
+            {
+                // if the db doesn't exist we consider the drop a success
+                return;
+            }
+
+            await _namespaceStore.RemoveNamespace(ns);
+
+        }
+    }
+
+    [DynamoDBTable(DynamoDbRefsStore.MainTableName)] 
+    public class DynamoBaseRefRecord
+    {
+        // Required by DynamoDb Serialization
+        // ReSharper disable once UnusedMember.Global
+        public DynamoBaseRefRecord()
+        {
+
+        }
+
+        public DynamoBaseRefRecord(RefRecord record)
+        {
+            Namespace = record.Namespace.ToString();
+            Bucket = record.Bucket.ToString();
+            Name = record.RefName.ToString();
+            Blobs = record.Blobs.Select(identifier => identifier.ToString()).ToList();
+            ContentHash = record.ContentHash.ToString();
+
+            Id = BuildCacheKey(record.Namespace, record.Bucket, record.RefName);
+        }
+
+
+        [DynamoDBHashKey] 
+        public string? Id { get; set; }
+
+        public string? Namespace { get; set; }
+
+        public string? Bucket { get; set; }
+
+        public string? Name { get; set; }
+
+        [DynamoDBProperty]
+        public List<string>? Blobs { get; set; }
+
+        public string? ContentHash { get; set; }
+
+        public static string BuildCacheKey(NamespaceId ns, BucketId bucket, KeyId key)
+        {
+            return $"{ns}.{bucket}.{key}";
+        }
+
+        public bool IsValid()
+        {
+            if (Id == null)
+                return false;
+            if (Namespace == null)
+                return false;
+            if (Bucket == null)
+                return false;
+            if (Name == null)
+                return false;
+            if (Blobs == null)
+                return false;
+            if (ContentHash == null)
+                return false;
+
+            return true;
+        }
+    }
+
+    [DynamoDBTable(DynamoDbRefsStore.MainTableName)]
+    public class DynamoLastAccessRefRecord
+    {
+        // Required by DynamoDb Serialization
+        // ReSharper disable once UnusedMember.Global
+        public DynamoLastAccessRefRecord()
+        {
+        }
+
+        public DynamoLastAccessRefRecord(RefRecord record, int lastAccessIndex)
+        {
+            LastAccessTime = record.LastAccessTime;
+            LastAccessIndex = lastAccessIndex;
+
+            Id = BuildCacheKey(record.Namespace, record.Bucket, record.RefName);
+        }
+
+        public DynamoLastAccessRefRecord(DynamoRefRecord record)
+        {
+            LastAccessTime = record.LastAccessTime;
+            LastAccessIndex = record.LastAccessIndex;
+
+            Id = BuildCacheKey(record.Namespace, record.Bucket, record.RefName);
+        }
+
+        [DynamoDBHashKey]
+        public string? Id { get; set; }
+
+        public DateTime? LastAccessTime { get; set; }
+
+        // a random index used to distribute requests over a bucket of records
+        public int? LastAccessIndex { get; set; }
+
+        public static string BuildCacheKey(NamespaceId ns, BucketId bucket, KeyId key)
+        {
+            return $"{DynamoBaseRefRecord.BuildCacheKey(ns, bucket, key)}#last-access";
+        }
+
+        public static string ToBaseRefRecordsKey(string cacheKey)
+        {
+            return cacheKey.Replace("#last-access", "");
+        }
+
+
+        public bool IsValid()
+        {
+            if (Id == null)
+                return false;
+            if (LastAccessTime == null)
+                return false;
+            if (LastAccessIndex == null)
+                return false;
+
+            return true;
+        }
+    }
+
+    [DynamoDBTable(DynamoDbRefsStore.MainTableName)]
+    public class DynamoMetadataRefRecord
+    {
+        // Required by DynamoDb Serialization
+        // ReSharper disable once UnusedMember.Global
+        public DynamoMetadataRefRecord()
+        {
+        }
+
+        public DynamoMetadataRefRecord(RefRecord record)
+        {
+            Id = BuildCacheKey(record.Namespace, record.Bucket, record.RefName);
+            Metadata = record.Metadata;
+        }
+
+        [DynamoDBHashKey]
+        public string? Id { get; set; }
+
+        [DynamoDBProperty(typeof(DynamoJsonConverter))]
+        public Dictionary<string, object>? Metadata { get; set; }
+
+        public static string BuildCacheKey(NamespaceId ns, BucketId bucket, KeyId key)
+        {
+            return $"{DynamoBaseRefRecord.BuildCacheKey(ns, bucket, key)}#metadata";
+        }
+
+        public bool IsValid()
+        {
+            if (Id == null)
+                return false;
+
+            return true;
+        }
+    }
+
+    public class DynamoRefRecord : RefRecord
+    {
+        public DynamoRefRecord(NamespaceId ns, BucketId bucket, KeyId refName, BlobIdentifier[] blobs, DateTime? lastAccessTime, int? lastAccessIndex, ContentHash contentHash, Dictionary<string, object>? metadata) : base(ns, bucket, refName, blobs, lastAccessTime, contentHash, metadata)
+        {
+            LastAccessIndex = lastAccessIndex;
+        }
+
+        public int? LastAccessIndex { get; set; }
+    }
+}
diff --git a/Engine/Source/Programs/Horde/HordeStorage/Horde.Storage/Implementation/Refs/DynamoNamespaceStore.cs b/Engine/Source/Programs/Horde/HordeStorage/Horde.Storage/Implementation/Refs/DynamoNamespaceStore.cs
new file mode 100644
--- /dev/null
+++ b/Engine/Source/Programs/Horde/HordeStorage/Horde.Storage/Implementation/Refs/DynamoNamespaceStore.cs
@@ -0,0 +1,114 @@
+// Copyright Epic Games, Inc. All Rights Reserved.
+
+using System;
+using System.Collections.Generic;
+using System.Threading.Tasks;
+using Amazon.DynamoDBv2;
+using Amazon.DynamoDBv2.DataModel;
+using Amazon.DynamoDBv2.DocumentModel;
+using Amazon.DynamoDBv2.Model;
+using Jupiter.Implementation;
+using Microsoft.Extensions.Options;
+
+namespace Horde.Storage.Implementation
+{
+    public class DynamoNamespaceStore : DynamoStore
+    {
+        private const string NamespaceTableName = "Europa_Namespace";
+
+        public DynamoNamespaceStore(IOptionsMonitor<DynamoDbSettings> settings, IAmazonDynamoDB client) : base(settings, client)
+        {
+        }
+
+        public async Task AddToNamespaceTable(NamespaceId ns, NamespaceUsage usage)
+        {
+            // insert the namespace if it does not exist
+            PutItemRequest putItem = new PutItemRequest(NamespaceTableName,
+                new Dictionary<string, AttributeValue>
+                {
+                    { "Namespace", new AttributeValue(ns.ToString()) },
+                    { "Usage", new AttributeValue {N = ((int)usage).ToString()}},
+
+                })
+            {
+                ConditionExpression = "attribute_not_exists(Namespace)"
+            };
+
+            try
+            {
+                await Client.PutItemAsync(putItem);
+            }
+            catch (ConditionalCheckFailedException)
+            {
+                // if it already exists we do not need to do anything
+            }
+        }
+
+
+        public async IAsyncEnumerable<NamespaceId> GetNamespaces(NamespaceUsage usage)
+        {
+            await Initialize();
+
+            AsyncSearch<NamespaceMappingRecord> search = Context.ScanAsync<NamespaceMappingRecord>(new ScanCondition[]
+            {
+                new ScanCondition("Usage", ScanOperator.Equal, usage),
+            });
+
+            do
+            {
+                List<NamespaceMappingRecord> newSet = await search.GetNextSetAsync();
+                foreach (NamespaceMappingRecord document in newSet)
+                {
+                    if (document.Namespace == null)
+                        continue;
+                    yield return new NamespaceId(document.Namespace);
+                }
+            } while (!search.IsDone);
+        }
+
+        protected override async Task CreateTables()
+        {
+            await CreateTable(NamespaceTableName, new[]
+            {
+                new KeySchemaElement("Namespace", KeyType.HASH),
+            }, new[]
+            {
+                new AttributeDefinition("Namespace", ScalarAttributeType.S)
+            }, new ProvisionedThroughput
+            {
+                // hard coding capacity as it should be really low for this table
+                ReadCapacityUnits = 20,
+                WriteCapacityUnits = 10,
+            });
+
+        }
+
+        public async Task RemoveNamespace(NamespaceId ns)
+        {
+            try
+            {
+                await Context.DeleteAsync(new Primitive(ns.ToString(), true));
+            }
+            catch (ResourceNotFoundException)
+            {
+                // if the table does note exist it is considered a success
+            }
+        }
+
+        public enum NamespaceUsage
+        {
+            // TODO: This should be called Ref but we keep it for legacy reasons
+            Cache,
+            Tree
+        }
+
+
+        [DynamoDBTable(NamespaceTableName)]
+        public class NamespaceMappingRecord
+        {
+            [DynamoDBHashKey] public string Namespace { get; set; } = null!;
+
+            public NamespaceUsage Usage { get; set; }
+        }
+    }
+}
diff --git a/Engine/Source/Programs/Horde/HordeStorage/Horde.Storage/Implementation/Refs/DynamoStore.cs b/Engine/Source/Programs/Horde/HordeStorage/Horde.Storage/Implementation/Refs/DynamoStore.cs
new file mode 100644
--- /dev/null
+++ b/Engine/Source/Programs/Horde/HordeStorage/Horde.Storage/Implementation/Refs/DynamoStore.cs
@@ -0,0 +1,124 @@
+// Copyright Epic Games, Inc. All Rights Reserved.
+
+using System;
+using System.Collections.Generic;
+using System.Linq;
+using System.Threading;
+using System.Threading.Tasks;
+using Amazon.DynamoDBv2;
+using Amazon.DynamoDBv2.DataModel;
+using Amazon.DynamoDBv2.DocumentModel;
+using Amazon.DynamoDBv2.Model;
+using Datadog.Trace;
+using Microsoft.Extensions.Options;
+using Newtonsoft.Json;
+
+namespace Horde.Storage.Implementation
+{
+    public abstract class DynamoStore
+    {
+        private readonly IOptionsMonitor<DynamoDbSettings> _settings;
+        private readonly SemaphoreSlim _initializeSemaphore = new SemaphoreSlim(1, 1);
+        private bool _initialized = false;
+
+        protected DynamoStore(IOptionsMonitor<DynamoDbSettings> settings, IAmazonDynamoDB client)
+        {
+            _settings = settings;
+            Client = client;
+            Context = CreateDynamoContext(client);
+        }
+
+        protected IDynamoDBContext Context { get; }
+
+        protected IAmazonDynamoDB Client { get; }
+
+        protected abstract Task CreateTables();
+
+        public async Task Initialize()
+        {
+            if (_initialized)
+                return;
+
+            await _initializeSemaphore.WaitAsync();
+
+            if (_initialized)
+                return;
+            try
+            {
+                await CreateTables();
+                _initialized = true;
+            }
+            finally
+            {
+                _initializeSemaphore.Release();
+            }
+        }
+
+        private static IDynamoDBContext CreateDynamoContext(IAmazonDynamoDB client)
+        {
+            return new DynamoDBContext(client, new DynamoDBContextConfig { Conversion = DynamoDBEntryConversion.V2 });
+        }
+
+        protected async Task CreateTable(string tableName, KeySchemaElement[] keySchemaElements, AttributeDefinition[] attributeDefinitions, ProvisionedThroughput provisionedThroughput, IEnumerable<GlobalSecondaryIndex>? secondaryIndices = null)
+        {
+            var tables = await Client.ListTablesAsync();
+            if (tables.TableNames.Contains(tableName))
+                return;
+
+            bool shouldCreateTable = _settings.CurrentValue.CreateTablesOnDemand;
+
+            if (shouldCreateTable)
+            {
+                using Scope scope = Tracer.Instance.StartActive("dynamo.create_tables");
+                Span span = scope.Span;
+                span.ResourceName = $"CREATE {tableName}";
+
+                CreateTableRequest createTableRequest = new CreateTableRequest(tableName, keySchemaElements.ToList(), attributeDefinitions.ToList(), _settings.CurrentValue.UseOndemandCapacityProvisioning ? null : provisionedThroughput)
+                {
+                    BillingMode = _settings.CurrentValue.UseOndemandCapacityProvisioning ? BillingMode.PAY_PER_REQUEST : BillingMode.PROVISIONED,
+                };
+
+                if (secondaryIndices != null)
+                {
+                    // recreate the GlobalSecondaryIndex instance so we can control the provisioned thruput
+                    createTableRequest.GlobalSecondaryIndexes.AddRange(from index in secondaryIndices
+                        select new GlobalSecondaryIndex
+                        {
+                            IndexName = index.IndexName,
+                            KeySchema = index.KeySchema,
+                            Projection = index.Projection,
+                            ProvisionedThroughput = _settings.CurrentValue.UseOndemandCapacityProvisioning ? null : provisionedThroughput
+                        });
+                }
+
+                await Client.CreateTableAsync(createTableRequest);
+
+                // wait for the table to have transitioned into active, e.g. that its finished being created
+                while ((await Client.DescribeTableAsync(tableName)).Table.TableStatus != "ACTIVE")
+                {
+                    await Task.Delay(50);
+                }
+            }
+        }
+
+    }
+
+    public class DynamoJsonConverter : IPropertyConverter
+    {
+        public object? FromEntry(DynamoDBEntry entry)
+        {
+            Primitive? primitive = entry as Primitive;
+            if (primitive == null || !(primitive.Value is String) || string.IsNullOrEmpty((string)primitive.Value))
+                throw new ArgumentOutOfRangeException();
+            Dictionary<string, object>? ret = JsonConvert.DeserializeObject<Dictionary<string, object>>((primitive.Value as string)!);
+            return ret;
+        }
+
+        public DynamoDBEntry ToEntry(object value)
+        {
+            var jsonString = JsonConvert.SerializeObject(value);
+            DynamoDBEntry ret = new Primitive(jsonString);
+            return ret;
+        }
+    }
+}
diff --git a/Engine/Source/Programs/Horde/HordeStorage/Horde.Storage/Implementation/Refs/IRefsStore.cs b/Engine/Source/Programs/Horde/HordeStorage/Horde.Storage/Implementation/Refs/IRefsStore.cs
new file mode 100644
--- /dev/null
+++ b/Engine/Source/Programs/Horde/HordeStorage/Horde.Storage/Implementation/Refs/IRefsStore.cs
@@ -0,0 +1,79 @@
+﻿// Copyright Epic Games, Inc. All Rights Reserved.
+
+using System;
+using System.Collections.Generic;
+using System.Threading.Tasks;
+using Jupiter.Implementation;
+
+namespace Horde.Storage.Implementation
+{
+    public class RefRecord
+    {
+        public RefRecord(NamespaceId ns, BucketId bucket, KeyId refName, BlobIdentifier[] blobs, DateTime? lastAccessTime,
+            ContentHash contentHash, Dictionary<string, object>? metadata)
+        {
+            Namespace = ns;
+            Bucket = bucket;
+            RefName = refName;
+            Blobs = blobs;
+            Metadata = metadata;
+            LastAccessTime = lastAccessTime;
+            ContentHash = contentHash;
+        }
+
+
+        public NamespaceId Namespace { get; }
+
+        public BucketId Bucket { get; }
+        public KeyId RefName { get; }
+
+        public BlobIdentifier[] Blobs { get; }
+        public Dictionary<string, object>? Metadata { get; }
+        public DateTime? LastAccessTime { get; }
+        public ContentHash ContentHash { get; }
+
+        public TransactionEvent ToAddTransactionEvent()
+        {
+            return new AddTransactionEvent(RefName.ToString(), Bucket.ToString(), Blobs, Metadata);
+        }
+    }
+
+    public class OldRecord
+    {
+        public OldRecord(NamespaceId ns, BucketId bucket, KeyId refName)
+        {
+            Namespace = ns;
+            Bucket = bucket;
+            RefName = refName;
+        }
+
+
+        public NamespaceId Namespace { get; }
+
+        public BucketId Bucket { get; }
+        public KeyId RefName { get; }
+
+    }
+
+
+    public interface IRefsStore
+    {
+        [Flags]
+        enum ExtraFieldsFlag
+        {
+            None = 0,
+            LastAccess = 1,
+            Metadata = 2
+        }
+
+        Task<RefRecord?> Get(NamespaceId ns, BucketId bucket, KeyId key, ExtraFieldsFlag fields);
+        Task Add(RefRecord record);
+        Task<long> DeleteBucket(NamespaceId ns, BucketId bucket);
+        Task<long> Delete(NamespaceId ns, BucketId bucket, KeyId key);
+        Task UpdateLastAccessTime(RefRecord record, DateTime lastAccessTime);
+
+        IAsyncEnumerable<NamespaceId> GetNamespaces();
+        IAsyncEnumerable<OldRecord> GetOldRecords(NamespaceId ns, TimeSpan oldRecordCutoff);
+        Task DropNamespace(NamespaceId ns);
+    }
+}
diff --git a/Engine/Source/Programs/Horde/HordeStorage/Horde.Storage/Implementation/Refs/MongoRefsStore.cs b/Engine/Source/Programs/Horde/HordeStorage/Horde.Storage/Implementation/Refs/MongoRefsStore.cs
new file mode 100644
--- /dev/null
+++ b/Engine/Source/Programs/Horde/HordeStorage/Horde.Storage/Implementation/Refs/MongoRefsStore.cs
@@ -0,0 +1,334 @@
+﻿// Copyright Epic Games, Inc. All Rights Reserved.
+
+using System;
+using System.Collections.Generic;
+using System.Linq;
+using System.Security.Authentication;
+using System.Threading.Tasks;
+using Jupiter.Implementation;
+using Microsoft.Extensions.Options;
+using MongoDB.Bson;
+using MongoDB.Bson.Serialization;
+using MongoDB.Bson.Serialization.Attributes;
+using MongoDB.Driver;
+using Serilog;
+
+namespace Horde.Storage.Implementation
+{
+    [BsonIgnoreExtraElements]
+    public abstract class MongoRefModelBase
+    {
+        protected MongoRefModelBase(string ns, string bucket, string name)
+        {
+            Ns = ns;
+            Bucket = bucket;
+            Name = name;
+        }
+
+        [BsonElement("Namespace")] public string Ns { get; set; }
+        public string Bucket { get; set; }
+        public string Name { get; set; }
+
+        public abstract RefRecord ToRefRecord();
+    }
+
+    // we do versioning by writing a discriminator into object
+    [BsonDiscriminator("cache.v0")]
+    [BsonIgnoreExtraElements]
+    public class MongoRefModelV0 : MongoRefModelBase
+    {
+        [BsonConstructor]
+        public MongoRefModelV0(string ns, string bucket, string name, string[] blobIdentifiers,
+            string contentHash, DateTime lastAccessTime) : base(ns, bucket, name)
+        {
+            BlobIdentifiers = blobIdentifiers;
+            LastAccessTime = lastAccessTime;
+            ContentHash = contentHash;
+        }
+
+        [BsonConstructor]
+        public MongoRefModelV0(string ns, string bucket, string name, string[] blobIdentifiers,
+            DateTime lastAccessTime, string contentHash, BsonDocument metadata) : base(ns, bucket, name)
+        {
+            BlobIdentifiers = blobIdentifiers;
+            Metadata = metadata;
+            LastAccessTime = lastAccessTime;
+            ContentHash = contentHash;
+        }
+
+        public BsonDocument? Metadata { get; set; }
+        public string[] BlobIdentifiers { get; set; }
+        public DateTime LastAccessTime { get; set; }
+        public string ContentHash { get; set; }
+
+        public override RefRecord ToRefRecord()
+        {
+            return new RefRecord(new NamespaceId(Ns), new BucketId(Bucket), new KeyId(Name), BlobIdentifiers.Select(s => new BlobIdentifier(s)).ToArray(),
+                LastAccessTime, new ContentHash(ContentHash), metadata: Metadata?.ToDictionary());
+        }
+    }
+
+    internal static class RefRecordExtensions
+    {
+        public static MongoRefModelBase ToLatestRecord(this RefRecord record)
+        {
+            MongoRefModelV0 newModel = new MongoRefModelV0(
+                record.Namespace.ToString(),
+                record.Bucket.ToString(),
+                record.RefName.ToString(),
+                record.Blobs.Select(identifier => identifier.ToString()).ToArray(),
+                record.LastAccessTime.GetValueOrDefault(DateTime.Now),
+                contentHash: record.ContentHash.ToString(), 
+                metadata: record.Metadata.ToBsonDocument()
+            );
+
+            return newModel;
+        }
+    }
+
+    public class MongoRefsStore : IRefsStore
+    {
+        protected readonly MongoClient _client;
+        protected const string ModelName = "Cache";
+        private const string MetadataModel = "Metadata";
+        private readonly ILogger _logger = Log.ForContext<MongoRefsStore>();
+
+        public MongoRefsStore(IOptionsMonitor<MongoSettings> mongoSettings)
+        {
+            MongoSettings settings = mongoSettings.CurrentValue;
+            MongoClientSettings mongoClientSettings = MongoClientSettings.FromUrl(
+                new MongoUrl(settings.ConnectionString)
+            );
+            if (settings.RequireTls12)
+            {
+                mongoClientSettings.SslSettings = new SslSettings {EnabledSslProtocols = SslProtocols.Tls12};
+            }
+
+            _client = new MongoClient(mongoClientSettings);
+
+            SetupClassMaps();
+        }
+
+        private static bool s_classMapsCreated = false;
+
+        private static void SetupClassMaps()
+        {
+            if (s_classMapsCreated)
+            {
+                return;
+            }
+
+            BsonClassMap.RegisterClassMap<MongoRefModelV0>();
+            s_classMapsCreated = true;
+        }
+
+        protected IMongoCollection<MongoRefModelBase> GetCollection(NamespaceId ns)
+        {
+            return GetCollection<MongoRefModelBase>(ns, ModelName);
+        }
+
+        protected IMongoCollection<T> GetCollection<T>(NamespaceId ns, string collectionName)
+        {
+            string dbName = GetDatabaseName(ns);
+            return _client.GetDatabase(dbName).GetCollection<T>(collectionName);
+        }
+
+        protected string GetDatabaseName(NamespaceId ns)
+        {
+            string cleanedNs = ns.ToString().Replace(".", "_");
+            string dbName = $"Europa_{cleanedNs}";
+            return dbName;
+        }
+
+        public async Task<RefRecord?> Get(NamespaceId ns, BucketId bucket, KeyId key, IRefsStore.ExtraFieldsFlag fields)
+        {
+            MongoRefModelBase model = await ExecuteWithRetries(() =>
+            {
+                return GetCollection(ns).Find(m => m.Name == key.ToString() && m.Bucket == bucket.ToString()).FirstOrDefaultAsync();
+            });
+
+            return model?.ToRefRecord();
+        }
+
+        public virtual async Task Add(RefRecord record)
+        {
+            Task addMetadataRecord = AddMetadataRecord(record.Namespace);
+            MongoRefModelBase model = record.ToLatestRecord();
+
+            ReplaceOneResult result = await ExecuteWithRetries(() =>
+            {
+                return GetCollection(record.Namespace)
+                    .ReplaceOneAsync(m => m.Name == model.Name && m.Bucket == model.Bucket, model,
+                        options: new ReplaceOptions {IsUpsert = true});
+            });
+
+            if (!result.IsAcknowledged)
+            {
+                _logger.Error("Failed to create {Record} using mongo model {Model}. Got Mongo response {Mongo}", record,
+                    model, result);
+                throw new Exception(
+                    $"Failed to create refs record {record.Bucket} {record.RefName} in namespace {record.Namespace}");
+            }
+
+            await addMetadataRecord;
+        }
+
+        private async Task AddMetadataRecord(NamespaceId ns)
+        {
+            // makes sure we have a metadata collection with information about this database
+            // since this metadata record always have the same id we can just always insert as if it exists this will not do anything
+            try
+            {
+                await GetCollection<MongoMetadata>(ns, MetadataModel).InsertOneAsync(new MongoMetadata(ns.ToString()));
+            }
+            catch (MongoWriteException e)
+            {
+                if (e.WriteError.Code == 11000)
+                {
+                    // the object already existed so lets ignore this
+                }
+                else
+                {
+                    throw;
+                }
+            }
+        }
+
+        public async Task<long> DeleteBucket(NamespaceId ns, BucketId bucket)
+        {
+            DeleteResult actionResult = await GetCollection(ns).DeleteManyAsync(b => b.Bucket == bucket.ToString());
+
+            return actionResult.DeletedCount;
+        }
+
+        public async Task<long> Delete(NamespaceId ns, BucketId bucket, KeyId key)
+        {
+            DeleteResult actionResult =
+                await GetCollection(ns).DeleteOneAsync(b => b.Bucket == bucket.ToString() && b.Name == key.ToString());
+            return actionResult.DeletedCount;
+        }
+
+        public async Task UpdateLastAccessTime(RefRecord record, DateTime lastAccessTime)
+        {
+            NamespaceId ns = record.Namespace;
+
+            UpdateDefinition<MongoRefModelV0> update =
+                Builders<MongoRefModelV0>.Update.Set(b => b.LastAccessTime, lastAccessTime);
+            UpdateResult result = await GetCollection<MongoRefModelV0>(ns, ModelName)
+                .UpdateOneAsync(b => b.Bucket == record.Bucket.ToString() && b.Name == record.RefName.ToString(), update);
+
+            if (!result.IsAcknowledged || result.MatchedCount != 1 || result.ModifiedCount != 1)
+            {
+                _logger.Error("Failed to update {Record} to {Time}. Got Mongo response {Mongo}", record, lastAccessTime,
+                    result);
+                throw new Exception(
+                    $"Failed to update refs record {record.Bucket} {record.RefName} in namespace {ns}");
+            }
+        }
+
+        public async IAsyncEnumerable<NamespaceId> GetNamespaces()
+        {
+            using IAsyncCursor<BsonDocument> cursor = await _client.ListDatabasesAsync();
+            while (await cursor.MoveNextAsync())
+            {
+                foreach (BsonDocument doc in cursor.Current)
+                {
+                    string dbName = doc["name"].AsString;
+                    MongoMetadata? result = await _client.GetDatabase(dbName).GetCollection<MongoMetadata>("Metadata")
+                        .Find(FilterDefinition<MongoMetadata>.Empty).FirstOrDefaultAsync();
+                    if (result == null)
+                        continue;
+
+                    NamespaceId ns;
+                    try
+                    {
+                        ns = new NamespaceId(result.Namespace);
+                    }
+                    catch (Exception)
+                    {
+                        continue;
+                    }
+                    yield return ns;
+                }
+            }
+        }
+
+        public async IAsyncEnumerable<OldRecord> GetOldRecords(NamespaceId ns, TimeSpan oldRecordCutoff)
+        {
+            IAsyncCursor<MongoRefModelV0> cursor = await GetCollection<MongoRefModelV0>(ns, ModelName).FindAsync(b => b.LastAccessTime < DateTime.Now - oldRecordCutoff);
+            while (await cursor.MoveNextAsync())
+            {
+                foreach (MongoRefModelV0 doc in cursor.Current)
+                {
+                    yield return new OldRecord(new NamespaceId(doc.Ns), new BucketId(doc.Bucket), new KeyId(doc.Name));
+                }
+            }
+        }
+
+        public async Task DropNamespace(NamespaceId ns)
+        {
+            string dbName = GetDatabaseName(ns);
+            await _client.GetDatabase(dbName).DropCollectionAsync(ModelName);
+        }
+
+        private async Task<V> ExecuteWithRetries<V>(Func<Task<V>> function)
+        {
+            while (true)
+            {
+                TimeSpan HandleMongoException(MongoCommandException de)
+                {
+                    // to many requests (http 429 is 16500 from mongo)
+                    if (de.Code != 16500)
+                    {
+                        throw de;
+                    }
+
+                    _logger.Warning(de, "Mongo retry exception encountered, sleeping and retrying. ");
+                    //sleepTime = de.RetryAfter;
+                    //sleepTime = (TimeSpan?)de.Data["RetryAfter"] ?? TimeSpan.Zero;
+                    // the mongo api in cosmos does not return back when we should retry
+                    // https://feedback.azure.com/forums/263030-azure-cosmos-db/suggestions/32956231-better-handling-of-request-limit-too-large
+                    // so lets just guess that we should retry after a while
+                    return TimeSpan.FromMilliseconds(250);
+
+                };
+
+                TimeSpan sleepTime;
+                try
+                {
+                    return await function();
+                }
+                catch (MongoCommandException de)
+                {
+                    sleepTime = HandleMongoException(de);
+                }
+                catch (AggregateException ae)
+                {
+                    if (!(ae.InnerException is MongoCommandException))
+                    {
+                        throw;
+                    }
+
+                    MongoCommandException de = (MongoCommandException)ae.InnerException;
+                    sleepTime = HandleMongoException(de);
+                }
+
+                await Task.Delay(sleepTime);
+            }
+        }
+    }
+
+    public class MongoMetadata
+    {
+        [BsonConstructor]
+        public MongoMetadata(string @namespace)
+        {
+            Namespace = @namespace;
+        }
+
+        // we define the object id as a hardcoded string to make sure there is only ever one record of this in each database
+        [BsonId] public string ObjectId { get; } = "CollectionMetadata";
+
+        public string Namespace { get; set; }
+    }
+}
